[mutable] Optimize mesh clip versus switching of clipping mesh.
Cleanup usage of most of STL from mutable runtime. Exception is string for now.
Cleanup a lot of unused code like reference implementations, node cloning boilerplate, etc.

diff --git a/Engine/Plugins/Experimental/Mutable/Source/CustomizableObject/Private/MuCO/CustomizableObjectSystem.cpp b/Engine/Plugins/Experimental/Mutable/Source/CustomizableObject/Private/MuCO/CustomizableObjectSystem.cpp
--- a/Engine/Plugins/Experimental/Mutable/Source/CustomizableObject/Private/MuCO/CustomizableObjectSystem.cpp
+++ b/Engine/Plugins/Experimental/Mutable/Source/CustomizableObject/Private/MuCO/CustomizableObjectSystem.cpp
@@ -1455,7 +1455,7 @@ namespace impl
 							{
 								for (int32 i = 0; i < boneIndexComponents; ++i)
 								{
-									size_t SectionBoneIndex = ((uint16_t*)VertexBoneIndexPtr)[i];
+									size_t SectionBoneIndex = ((uint16*)VertexBoneIndexPtr)[i];
 									UsedBones[SectionBoneIndex] = true;
 								}
 								VertexBoneIndexPtr += elemSize;
diff --git a/Engine/Plugins/Experimental/Mutable/Source/CustomizableObjectEditor/Private/MuCOE/CustomizableObjectBakeHelpers.cpp b/Engine/Plugins/Experimental/Mutable/Source/CustomizableObjectEditor/Private/MuCOE/CustomizableObjectBakeHelpers.cpp
--- a/Engine/Plugins/Experimental/Mutable/Source/CustomizableObjectEditor/Private/MuCOE/CustomizableObjectBakeHelpers.cpp
+++ b/Engine/Plugins/Experimental/Mutable/Source/CustomizableObjectEditor/Private/MuCOE/CustomizableObjectBakeHelpers.cpp
@@ -99,13 +99,13 @@ namespace
 	{
 		uint32_t temp, code;
 
-		uint16_t color0, color1;
+		uint16 color0, color1;
 		uint8_t r0, g0, b0, r1, g1, b1;
 
 		int i, j;
 
-		color0 = *(const uint16_t*)(block);
-		color1 = *(const uint16_t*)(block + 2);
+		color0 = *(const uint16*)(block);
+		color1 = *(const uint16*)(block + 2);
 
 		temp = (color0 >> 11) * 255 + 16;
 		r0 = (uint8_t)((temp / 32 + temp) / 32);
@@ -204,7 +204,7 @@ namespace
 		uint8_t alpha0, alpha1;
 		uint8_t alphaIndices[16];
 
-		uint16_t color0, color1;
+		uint16 color0, color1;
 		uint8_t r0, g0, b0, r1, g1, b1;
 
 		int i, j;
@@ -216,8 +216,8 @@ namespace
 
 		Decompress16x3bitIndices(blockStorage + 2, alphaIndices);
 
-		color0 = *(const uint16_t*)(blockStorage + 8);
-		color1 = *(const uint16_t*)(blockStorage + 10);
+		color0 = *(const uint16*)(blockStorage + 8);
+		color1 = *(const uint16*)(blockStorage + 10);
 
 		temp = (color0 >> 11) * 255 + 16;
 		r0 = (uint8_t)((temp / 32 + temp) / 32);
@@ -299,7 +299,7 @@ namespace
 		uint8_t alphaValues[16] = { 0 };
 
 		for (i = 0; i < 4; ++i) {
-			const uint16_t* alphaData = (const uint16_t*)(blockStorage);
+			const uint16* alphaData = (const uint16*)(blockStorage);
 
 			alphaValues[i * 4 + 0] = (((*alphaData) >> 0) & 0xF) * 17;
 			alphaValues[i * 4 + 1] = (((*alphaData) >> 4) & 0xF) * 17;
diff --git a/Engine/Plugins/Experimental/Mutable/Source/CustomizableObjectEditor/Private/MuCOE/MutableMeshPreviewUtils.cpp b/Engine/Plugins/Experimental/Mutable/Source/CustomizableObjectEditor/Private/MuCOE/MutableMeshPreviewUtils.cpp
--- a/Engine/Plugins/Experimental/Mutable/Source/CustomizableObjectEditor/Private/MuCOE/MutableMeshPreviewUtils.cpp
+++ b/Engine/Plugins/Experimental/Mutable/Source/CustomizableObjectEditor/Private/MuCOE/MutableMeshPreviewUtils.cpp
@@ -207,7 +207,6 @@ namespace MutableMeshPreviewUtils
 				true,
 				true,
 				false,
-				false,
 				false );
 			
 			return ResultMutableMesh;
@@ -412,7 +411,7 @@ namespace MutableMeshPreviewUtils
 							{
 								for (int32 i = 0; i < BoneIndexComponents; ++i)
 								{
-									const size_t SectionBoneIndex = ((uint16_t*)VertexBoneIndexPtr)[i];
+									const size_t SectionBoneIndex = ((uint16*)VertexBoneIndexPtr)[i];
 									UsedBones[SectionBoneIndex] = true;
 								}
 								VertexBoneIndexPtr += ElementSize;
diff --git a/Engine/Plugins/Experimental/Mutable/Source/CustomizableObjectEditor/Private/MuCOE/Nodes/CustomizableObjectNodeStaticMesh.cpp b/Engine/Plugins/Experimental/Mutable/Source/CustomizableObjectEditor/Private/MuCOE/Nodes/CustomizableObjectNodeStaticMesh.cpp
--- a/Engine/Plugins/Experimental/Mutable/Source/CustomizableObjectEditor/Private/MuCOE/Nodes/CustomizableObjectNodeStaticMesh.cpp
+++ b/Engine/Plugins/Experimental/Mutable/Source/CustomizableObjectEditor/Private/MuCOE/Nodes/CustomizableObjectNodeStaticMesh.cpp
@@ -47,8 +47,6 @@
 #include "Widgets/SBoxPanel.h"
 #include "Widgets/SOverlay.h"
 
-#include <functional>
-
 class UCustomizableObjectNodeRemapPinsByName;
 class UObject;
 struct FSlateBrush;
diff --git a/Engine/Plugins/Experimental/Mutable/Source/CustomizableObjectEditor/Private/MuCOE/SMutableCodeViewer.cpp b/Engine/Plugins/Experimental/Mutable/Source/CustomizableObjectEditor/Private/MuCOE/SMutableCodeViewer.cpp
--- a/Engine/Plugins/Experimental/Mutable/Source/CustomizableObjectEditor/Private/MuCOE/SMutableCodeViewer.cpp
+++ b/Engine/Plugins/Experimental/Mutable/Source/CustomizableObjectEditor/Private/MuCOE/SMutableCodeViewer.cpp
@@ -82,8 +82,6 @@
 #include "Widgets/Views/STableRow.h"
 #include "Widgets/Views/STreeView.h"
 
-#include <functional>
-
 class FExtender;
 class FReferenceCollector;
 class FUICommandList;
@@ -351,7 +349,7 @@ void SMutableCodeViewer::Construct(const FArguments& InArgs, const mu::ModelPtr&
 	//TabManager = FGlobalTabmanager::Get()->NewTabManager(ConstructUnderMajorTab.ToSharedRef());
 
 	const mu::Model::Private* ModelPrivate = MutableModel->GetPrivate();
-	const int32 StateCount = int32( ModelPrivate->m_program.m_states.size() );
+	const int32 StateCount = ModelPrivate->m_program.m_states.Num();
 	for ( int32 StateIndex=0; StateIndex<StateCount; ++StateIndex )
 	{
 		const mu::PROGRAM::STATE& State = ModelPrivate->m_program.m_states[StateIndex];
@@ -1811,7 +1809,7 @@ void SMutableCodeViewer::CacheRootNodeAddresses()
 	TArray<mu::OP::ADDRESS> FoundRootNodeAddresses;
 	
 	const mu::Model::Private* ModelPrivate = MutableModel->GetPrivate();
-	const int32 StateCount = int32( ModelPrivate->m_program.m_states.size() );
+	const int32 StateCount = ModelPrivate->m_program.m_states.Num();
 	for ( int32 StateIndex=0; StateIndex < StateCount; ++StateIndex )
 	{
 		const mu::PROGRAM::STATE& State = ModelPrivate->m_program.m_states[StateIndex];
diff --git a/Engine/Plugins/Experimental/Mutable/Source/CustomizableObjectEditor/Private/MuCOE/SMutableCodeViewer.h b/Engine/Plugins/Experimental/Mutable/Source/CustomizableObjectEditor/Private/MuCOE/SMutableCodeViewer.h
--- a/Engine/Plugins/Experimental/Mutable/Source/CustomizableObjectEditor/Private/MuCOE/SMutableCodeViewer.h
+++ b/Engine/Plugins/Experimental/Mutable/Source/CustomizableObjectEditor/Private/MuCOE/SMutableCodeViewer.h
@@ -32,8 +32,6 @@
 #include "Widgets/Views/STableViewBase.h"
 #include "Widgets/Views/STreeView.h"
 
-#include <utility>
-
 class FMutableCodeTreeElement;
 class FReferenceCollector;
 class ITableRow;
@@ -614,7 +612,7 @@ public:
 				for (auto& DynamicResource : CurrentState.m_dynamicResources)
 				{
 					// If the operation gets located then mark it as dynamic resource
-					if (DynamicResource.first == MutableOperation)
+					if (DynamicResource.Key == MutableOperation)
 					{
 						bIsDynamicResource = true;
 						break;
@@ -628,15 +626,7 @@ public:
 				}
 				
 				// Check if it is a state constant
-				for (auto& CachedOperation : CurrentState.m_updateCache)
-				{
-					// If the operation gets located then mark it as dynamic resource
-					if (CachedOperation == MutableOperation)
-					{
-						bIsStateConstant = true;
-						break;
-					}
-				}
+				bIsStateConstant = CurrentState.m_updateCache.Contains(MutableOperation);
 			}
 		}
 		
diff --git a/Engine/Plugins/Experimental/Mutable/Source/CustomizableObjectEditor/Private/MuCOE/SMutableConstantsWidget.cpp b/Engine/Plugins/Experimental/Mutable/Source/CustomizableObjectEditor/Private/MuCOE/SMutableConstantsWidget.cpp
--- a/Engine/Plugins/Experimental/Mutable/Source/CustomizableObjectEditor/Private/MuCOE/SMutableConstantsWidget.cpp
+++ b/Engine/Plugins/Experimental/Mutable/Source/CustomizableObjectEditor/Private/MuCOE/SMutableConstantsWidget.cpp
@@ -1151,7 +1151,7 @@ void SMutableConstantsWidget::LoadConstantStrings()
 	ConstantStringElements.Empty();
 	uint32 ConstantStringsAccumulatedSize = 0;
 	
-	for (uint32 StringAddressIndex = 0; StringAddressIndex < MutableProgramPtr->m_constantStrings.size(); StringAddressIndex++)
+	for (int32 StringAddressIndex = 0; StringAddressIndex < MutableProgramPtr->m_constantStrings.Num(); StringAddressIndex++)
 	{
 		TSharedPtr<FMutableConstantStringElement> ConstantStringElement = MakeShared<FMutableConstantStringElement>();
 		ConstantStringElement->MutableString = &(MutableProgramPtr->m_constantStrings[StringAddressIndex]);
@@ -1226,7 +1226,7 @@ void SMutableConstantsWidget::LoadConstantLayouts()
 	ConstantLayoutElements.Empty();
 	uint32 ConstantLayoutsAccumulatedSize = 0;
 	
-	for (uint32 LayoutIndex = 0; LayoutIndex < MutableProgramPtr->m_constantLayouts.size(); LayoutIndex++)
+	for (int32 LayoutIndex = 0; LayoutIndex < MutableProgramPtr->m_constantLayouts.Num(); LayoutIndex++)
 	{		
 		TSharedPtr<FMutableConstantLayoutElement> ConstantLayoutElement = MakeShared<FMutableConstantLayoutElement>();
 		ConstantLayoutElement->Layout = MutableProgramPtr->m_constantLayouts[LayoutIndex];
@@ -1247,7 +1247,7 @@ void SMutableConstantsWidget::LoadConstantSkeletons()
 	ConstantSkeletonElements.Empty();
 	uint32 ConstantSkeletonsAccumulatedSize = 0;
 	
-	for (uint32 SkeletonIndex = 0; SkeletonIndex < MutableProgramPtr->m_constantSkeletons.size(); SkeletonIndex++)
+	for (int32 SkeletonIndex = 0; SkeletonIndex < MutableProgramPtr->m_constantSkeletons.Num(); SkeletonIndex++)
 	{
 		TSharedPtr<FMutableConstantSkeletonElement> ConstantSkeletonElement = MakeShared<FMutableConstantSkeletonElement>();
 		ConstantSkeletonElement->Skeleton = MutableProgramPtr->m_constantSkeletons[SkeletonIndex];
@@ -1267,7 +1267,7 @@ void SMutableConstantsWidget::LoadConstantProjectors()
 	ConstantProjectorElements.Empty();
 	uint32 ConstantProjectorsAccumulatedSize = 0;
 	
-	for (uint32 ProjectorIndex = 0; ProjectorIndex < MutableProgramPtr->m_constantProjectors.size(); ProjectorIndex++)
+	for (int32 ProjectorIndex = 0; ProjectorIndex < MutableProgramPtr->m_constantProjectors.Num(); ProjectorIndex++)
 	{
 		TSharedPtr<FMutableConstantProjectorElement> ConstantProjectorElement = MakeShared<FMutableConstantProjectorElement>();
 		ConstantProjectorElement->Projector = &(MutableProgramPtr->m_constantProjectors[ProjectorIndex]);
@@ -1288,7 +1288,7 @@ void SMutableConstantsWidget::LoadConstantMatrices()
 	ConstantMatrixElements.Empty();
 	uint32 ConstantMatricesAccumulatedSize = 0;
 	
-	for (uint32 MatrixIndex = 0; MatrixIndex < MutableProgramPtr->m_constantMatrices.size(); MatrixIndex++)
+	for (int32 MatrixIndex = 0; MatrixIndex < MutableProgramPtr->m_constantMatrices.Num(); MatrixIndex++)
 	{
 		TSharedPtr<FMutableConstantMatrixElement> ConstantMatrixElement = MakeShared<FMutableConstantMatrixElement>();
 		ConstantMatrixElement->Matrix = &(MutableProgramPtr->m_constantMatrices[MatrixIndex]);
@@ -1309,7 +1309,7 @@ void SMutableConstantsWidget::LoadConstantShapes()
 	ConstantShapeElements.Empty();
 	uint32 ConstantShapesAccumulatedSize = 0;
 	
-	for (uint32 ShapeIndex = 0; ShapeIndex < MutableProgramPtr->m_constantShapes.size(); ShapeIndex++)
+	for (int32 ShapeIndex = 0; ShapeIndex < MutableProgramPtr->m_constantShapes.Num(); ShapeIndex++)
 	{
 		TSharedPtr<FMutableConstantShapeElement> ConstantShapeElement = MakeShared<FMutableConstantShapeElement>();
 		ConstantShapeElement->Shape = &(MutableProgramPtr->m_constantShapes[ShapeIndex]);
@@ -1330,7 +1330,7 @@ void SMutableConstantsWidget::LoadConstantCurves()
 	ConstantCurveElements.Empty();
 	uint32 ConstantCurvesAccumulatedSize = 0;
 	
-	for (uint32 CurveIndex = 0; CurveIndex < MutableProgramPtr->m_constantCurves.size(); CurveIndex++)
+	for (int32 CurveIndex = 0; CurveIndex < MutableProgramPtr->m_constantCurves.Num(); CurveIndex++)
 	{
 		TSharedPtr<FMutableConstantCurveElement> ConstantCurveElement = MakeShared<FMutableConstantCurveElement>();
 		ConstantCurveElement->Curve = &(MutableProgramPtr->m_constantCurves[CurveIndex]);
diff --git a/Engine/Plugins/Experimental/Mutable/Source/CustomizableObjectEditor/Private/MuCOE/SMutableCurveViewer.cpp b/Engine/Plugins/Experimental/Mutable/Source/CustomizableObjectEditor/Private/MuCOE/SMutableCurveViewer.cpp
--- a/Engine/Plugins/Experimental/Mutable/Source/CustomizableObjectEditor/Private/MuCOE/SMutableCurveViewer.cpp
+++ b/Engine/Plugins/Experimental/Mutable/Source/CustomizableObjectEditor/Private/MuCOE/SMutableCurveViewer.cpp
@@ -484,7 +484,7 @@ void SMutableCurveViewer::SetCurve(const mu::Curve& InMutableCurve)
 
 void SMutableCurveViewer::SetupMutableCurveListView()
 {
-	const int32 KeyFrameCount = MutableCurve.keyFrames.size() ;
+	const int32 KeyFrameCount = MutableCurve.keyFrames.Num() ;
 	CurveElements.SetNum(KeyFrameCount);
 	for	(int32 KeyFrameIndex = 0; KeyFrameIndex < KeyFrameCount; KeyFrameIndex++)
 	{
@@ -524,8 +524,8 @@ void SMutableCurveViewer::SetupMutableCurveGraph() const
 	float PreviousFrameTime =  TNumericLimits<float>::Lowest();
 	
 	// Fill the curve with data
-	const uint32 KeyFrameCount = MutableCurve.keyFrames.size() ;
-	for (uint32 KeyframeIndex = 0; KeyframeIndex < KeyFrameCount; KeyframeIndex++)
+	const int32 KeyFrameCount = MutableCurve.keyFrames.Num() ;
+	for (int32 KeyframeIndex = 0; KeyframeIndex < KeyFrameCount; KeyframeIndex++)
 	{
 		// Load the mutable data
 		const mu::CurveKeyFrame& CurrentKeyframe = MutableCurve.keyFrames[KeyframeIndex];
diff --git a/Engine/Plugins/Experimental/Mutable/Source/MutableRuntime/Private/MuR/CodeRunner.cpp b/Engine/Plugins/Experimental/Mutable/Source/MutableRuntime/Private/MuR/CodeRunner.cpp
--- a/Engine/Plugins/Experimental/Mutable/Source/MutableRuntime/Private/MuR/CodeRunner.cpp
+++ b/Engine/Plugins/Experimental/Mutable/Source/MutableRuntime/Private/MuR/CodeRunner.cpp
@@ -94,13 +94,13 @@ namespace mu
 		uint32 lodMask, uint8 executionOptions, SCHEDULED_OP::EType Type )
 		: m_pSettings(pSettings), m_pSystem(s), m_pModel(pModel), m_pParams(pParams), m_lodMask(lodMask)
 	{
-		ScheduledStagePerOp.resize(m_pModel->GetPrivate()->m_program.m_opAddress.size());
+		ScheduledStagePerOp.resize(m_pModel->GetPrivate()->m_program.m_opAddress.Num());
 
 		// We will read this in the end, so make sure we keep it.
 		m_pSystem->m_memory->IncreaseHitCount(CACHE_ADDRESS(at, 0, executionOptions));
 
 		PROGRAM& program = m_pModel->GetPrivate()->m_program;
-		m_romPendingOps.resize(program.m_roms.Num(), 0);
+		m_romPendingOps.SetNum(program.m_roms.Num());
 
 		// Push the root operation
 		SCHEDULED_OP rootOp;
@@ -372,7 +372,7 @@ namespace mu
                     vec4f value = GetMemory().GetColour( CACHE_ADDRESS(args.value,item) );
 
                     OP::ADDRESS nameAd = args.name;
-                    check(  nameAd < pModel->GetPrivate()->m_program.m_constantStrings.size() );
+                    check(  nameAd < (uint32)pModel->GetPrivate()->m_program.m_constantStrings.Num() );
                     const char* strName = pModel->GetPrivate()->m_program.m_constantStrings[ nameAd ].c_str();
 
                     pResult->GetPrivate()->AddVector( 0, 0, 0, value, strName );
@@ -417,7 +417,7 @@ namespace mu
                     float value = GetMemory().GetScalar( CACHE_ADDRESS(args.value,item) );
 
                     OP::ADDRESS nameAd = args.name;
-                    check(  nameAd < pModel->GetPrivate()->m_program.m_constantStrings.size() );
+                    check(  nameAd < (uint32)pModel->GetPrivate()->m_program.m_constantStrings.Num() );
                     const char* strName = pModel->GetPrivate()->m_program.m_constantStrings[ nameAd ].c_str();
 
                     pResult->GetPrivate()->AddScalar( 0, 0, 0, value, strName );
@@ -463,8 +463,7 @@ namespace mu
                         GetMemory().GetString( CACHE_ADDRESS( args.value, item ) );
 
                     OP::ADDRESS nameAd = args.name;
-                    check( nameAd <
-                                    pModel->GetPrivate()->m_program.m_constantStrings.size() );
+                    check( nameAd < (uint32)pModel->GetPrivate()->m_program.m_constantStrings.Num() );
                     const char* strName =
                         pModel->GetPrivate()->m_program.m_constantStrings[nameAd].c_str();
 
@@ -511,11 +510,11 @@ namespace mu
 
                     int cindex = pResult->GetPrivate()->AddComponent( 0 );
 
-                    if ( !pComp->GetPrivate()->m_lods.empty()
+                    if ( !pComp->GetPrivate()->m_lods.IsEmpty()
                          &&
-                         !pResult->GetPrivate()->m_lods.empty()
+                         !pResult->GetPrivate()->m_lods.IsEmpty()
                          &&
-                         !pComp->GetPrivate()->m_lods[0].m_components.empty() )
+                         !pComp->GetPrivate()->m_lods[0].m_components.IsEmpty() )
                     {
                         pResult->GetPrivate()->m_lods[0].m_components[cindex] =
                                 pComp->GetPrivate()->m_lods[0].m_components[0];
@@ -524,7 +523,7 @@ namespace mu
                     	
                         // Name
                         OP::ADDRESS nameAd = args.name;
-                        check( nameAd < pModel->GetPrivate()->m_program.m_constantStrings.size() );
+                        check( nameAd < (uint32)pModel->GetPrivate()->m_program.m_constantStrings.Num() );
                         const char* strName = pModel->GetPrivate()->m_program.m_constantStrings[ nameAd ].c_str();
                         pResult->GetPrivate()->SetComponentName( 0, cindex, strName );
                     }
@@ -570,22 +569,22 @@ namespace mu
                     // Surface data
                     if (pSurf
                             &&
-                            pSurf->GetPrivate()->m_lods.size()
+                            pSurf->GetPrivate()->m_lods.Num()
                             &&
-                            pSurf->GetPrivate()->m_lods[0].m_components.size()
+                            pSurf->GetPrivate()->m_lods[0].m_components.Num()
                             &&
-                            pSurf->GetPrivate()->m_lods[0].m_components[0].m_surfaces.size())
+                            pSurf->GetPrivate()->m_lods[0].m_components[0].m_surfaces.Num())
                     {
                         pResult->GetPrivate()->m_lods[0].m_components[0].m_surfaces[sindex] =
                             pSurf->GetPrivate()->m_lods[0].m_components[0].m_surfaces[0];
 
                         // Meshes must be added later.
-                        check(!pSurf->GetPrivate()->m_lods[0].m_components[0].m_meshes.size());
+                        check(!pSurf->GetPrivate()->m_lods[0].m_components[0].m_meshes.Num());
                     }
 
                     // Name
                     OP::ADDRESS nameAd = args.name;
-                    check( nameAd < pModel->GetPrivate()->m_program.m_constantStrings.size() );
+                    check( nameAd < (uint32)pModel->GetPrivate()->m_program.m_constantStrings.Num() );
                     const char* strName = pModel->GetPrivate()->m_program.m_constantStrings[ nameAd ].c_str();
                     pResult->GetPrivate()->SetSurfaceName( 0, 0, sindex, strName );
 
@@ -611,7 +610,7 @@ namespace mu
             {
             case 0:
             {                
-                    vector<SCHEDULED_OP> deps;
+                    TArray<SCHEDULED_OP> deps;
                     for ( int i=0; i<MUTABLE_OP_MAX_ADD_COUNT; ++i )
                     {
                         if ( args.lod[i] )
@@ -620,7 +619,7 @@ namespace mu
 
                             if ( selectedLod )
                             {
-                                deps.emplace_back(args.lod[i], item);
+                                deps.Emplace(args.lod[i], item);
                             }
                         }
                     }
@@ -648,7 +647,7 @@ namespace mu
                             int lindex = pResult->GetPrivate()->AddLOD();
 
                             // In a degenerated case, the returned pLOD may not have an LOD inside
-                            if ( pLOD && !pLOD->GetPrivate()->m_lods.empty() )
+                            if ( pLOD && !pLOD->GetPrivate()->m_lods.IsEmpty() )
                             {
                                 pResult->GetPrivate()->m_lods[lindex] = pLOD->GetPrivate()->m_lods[0];
                             }
@@ -723,7 +722,7 @@ namespace mu
                                 args.value,
                                 pParams);
                     OP::ADDRESS nameAd = args.name;
-                    check(  nameAd < pModel->GetPrivate()->m_program.m_constantStrings.size() );
+                    check(  nameAd < (uint32)pModel->GetPrivate()->m_program.m_constantStrings.Num() );
                     const char* strName = pModel->GetPrivate()->m_program.m_constantStrings[ nameAd ].c_str();
                     pResult->GetPrivate()->AddMesh( 0, 0, meshId, strName );
                 }
@@ -767,7 +766,7 @@ namespace mu
                                 args.value,
                                 pParams);
                     OP::ADDRESS nameAd = args.name;
-                    check(  nameAd < pModel->GetPrivate()->m_program.m_constantStrings.size() );
+                    check(  nameAd < (uint32)pModel->GetPrivate()->m_program.m_constantStrings.Num() );
                     const char* strName = pModel->GetPrivate()->m_program.m_constantStrings[ nameAd ].c_str();
                     pResult->GetPrivate()->AddImage( 0, 0, 0, imageId, strName );
                 }
@@ -815,14 +814,14 @@ namespace mu
             // Set the separate skeleton if necessary
             if (args.skeleton>=0)
             {
-                check( program.m_constantSkeletons.size()>size_t(args.skeleton)  );
+                check( program.m_constantSkeletons.Num()>size_t(args.skeleton)  );
                 Ptr<const Skeleton> pSkeleton = program.m_constantSkeletons[args.skeleton];
                 pSource->SetSkeleton(pSkeleton);
             }
 
 			if (args.physicsBody >= 0)
 			{
-                check( program.m_constantPhysicsBodies.size()>size_t(args.physicsBody)  );
+                check( program.m_constantPhysicsBodies.Num()>size_t(args.physicsBody)  );
                 Ptr<const PhysicsBody> pPhysicsBody = program.m_constantPhysicsBodies[args.physicsBody];
                 pSource->SetPhysicsBody(pPhysicsBody);
 			}
@@ -1075,7 +1074,6 @@ namespace mu
                         Ptr<const Mesh> pMorph = GetMemory().GetMesh( CACHE_ADDRESS(args.targets[min],item) );
                         if (pMorph)
                         {
-                            //pResult = MeshMorphReference( pBase.get(), pMorph.get() );
 							pResult = MeshMorph(pBase.get(), pMorph.get());
                         }
                     }
@@ -1086,7 +1084,6 @@ namespace mu
                         Ptr<const Mesh> pMorph = GetMemory().GetMesh( CACHE_ADDRESS(args.targets[max],item) );
                         if (pMorph)
                         {
-                            //pResult = MeshMorphReference( pBase.get(), pMorph.get() );
 							pResult = MeshMorph(pBase.get(), pMorph.get());
                         }
                     }
@@ -1097,7 +1094,6 @@ namespace mu
 						Ptr<const Mesh> pMorph = GetMemory().GetMesh(CACHE_ADDRESS(args.targets[max], item));
 						if (pMorph)
 						{
-							//pResult = MeshMorphReference(pBase.get(), pMorph.get(), bifactor);
 							pResult = MeshMorph(pBase.get(), pMorph.get(), bifactor);
 						}
 					}
@@ -1108,7 +1104,6 @@ namespace mu
                         Ptr<const Mesh> pMax = GetMemory().GetMesh( CACHE_ADDRESS(args.targets[max],item) );
                         if (pMin && pMax)
                         {
-						    //pResult = MeshMorph2Reference( pBase.get(), pMin.get(), pMax.get(), bifactor );
                             pResult = MeshMorph2( pBase.get(), pMin.get(), pMax.get(), bifactor );
                         }
                     }
@@ -1513,7 +1508,6 @@ namespace mu
                                       (flags & OP::MeshFormatArgs::BT_VERTEX) != 0,
                                       (flags & OP::MeshFormatArgs::BT_INDEX) != 0,
                                       (flags & OP::MeshFormatArgs::BT_FACE) != 0,
-                                      (flags & OP::MeshFormatArgs::BT_REBUILD_TANGENTS) != 0,
                                       (flags & OP::MeshFormatArgs::BT_IGNORE_MISSING) != 0
                                       );
 
@@ -1569,7 +1563,7 @@ namespace mu
 
                 // Access with memcpy necessary for unaligned arm issues.
                 uint32 blocks[1024];
-                memcpy( blocks, data, sizeof(uint32)*std::min(1024,int(blockCount)) );
+                memcpy( blocks, data, sizeof(uint32)*FMath::Min(1024,int(blockCount)) );
 
                 MeshPtr pResult;
                 pResult = MeshExtractLayoutBlock( pSource.get(),
@@ -1684,9 +1678,7 @@ namespace mu
 
                 MeshPtr pResult;
 
-                check( args.morphShape
-                                <
-                                pModel->GetPrivate()->m_program.m_constantShapes.size() );
+                check( args.morphShape < (uint32)pModel->GetPrivate()->m_program.m_constantShapes.Num() );
 
                 // Should be an ellipse
                 const SHAPE& morphShape = pModel->GetPrivate()->m_program.
@@ -1697,9 +1689,7 @@ namespace mu
 
                 if (args.vertexSelectionType == OP::MeshClipMorphPlaneArgs::VS_SHAPE)
                 {
-                    check( args.vertexSelectionShapeOrBone
-                                    <
-                                    pModel->GetPrivate()->m_program.m_constantShapes.size() );
+                    check( args.vertexSelectionShapeOrBone < (uint32)pModel->GetPrivate()->m_program.m_constantShapes.Num() );
 
                     // Should be None or an axis aligned box
                     const SHAPE& selectionShape = pModel->GetPrivate()->m_program.m_constantShapes[args.vertexSelectionShapeOrBone];
@@ -1708,9 +1698,7 @@ namespace mu
 
                 else if (args.vertexSelectionType == OP::MeshClipMorphPlaneArgs::VS_BONE_HIERARCHY)
                 {
-                    check( args.vertexSelectionShapeOrBone
-                                    <
-                                    pModel->GetPrivate()->m_program.m_constantStrings.size() );
+                    check( args.vertexSelectionShapeOrBone < (uint32)pModel->GetPrivate()->m_program.m_constantStrings.Num() );
 
                     SHAPE selectionShape;
                     selectionShape.type = (uint8)SHAPE::Type::None;
@@ -2212,8 +2200,8 @@ namespace mu
             OP::ADDRESS source;
             FMemory::Memcpy(&source,data,sizeof(OP::ADDRESS)); data += sizeof(OP::ADDRESS);
 
-            vector<SCHEDULED_OP> conditions;
-            vector<OP::ADDRESS> masks;
+            TArray<SCHEDULED_OP> conditions;
+			TArray<OP::ADDRESS> masks;
 
             uint16 removes;
 			FMemory::Memcpy(&removes,data,sizeof(uint16)); data += sizeof(uint16);
@@ -2222,11 +2210,11 @@ namespace mu
             {
                 OP::ADDRESS condition;
 				FMemory::Memcpy(&condition,data,sizeof(OP::ADDRESS)); data += sizeof(OP::ADDRESS);
-                conditions.emplace_back( condition, item );
+                conditions.Emplace( condition, item );
 
                 OP::ADDRESS mask;
 				FMemory::Memcpy(&mask,data,sizeof(OP::ADDRESS)); data += sizeof(OP::ADDRESS);
-                masks.push_back( mask );
+                masks.Add( mask );
             }
 
 
@@ -2249,9 +2237,9 @@ namespace mu
             {
                 // Request the source and the necessary masks
                 // \todo: store condition values in heap?
-                vector<SCHEDULED_OP> deps;
-                deps.push_back( SCHEDULED_OP( source, item ) );
-                for( size_t r=0; source && r<conditions.size(); ++r )
+                TArray<SCHEDULED_OP> deps;
+                deps.Emplace( source, item );
+                for( size_t r=0; source && r<conditions.Num(); ++r )
                 {
                     // If there is no expression, we'll assume true.
                     bool value = true;
@@ -2262,7 +2250,7 @@ namespace mu
 
                     if (value)
                     {
-                        deps.emplace_back( masks[r], item );
+                        deps.Emplace( masks[r], item );
                     }
                 }
 
@@ -2280,7 +2268,7 @@ namespace mu
 
                 MeshPtrConst pResult = pSource;
 
-                for( size_t r=0; pResult && r<conditions.size(); ++r )
+                for( size_t r=0; pResult && r<conditions.Num(); ++r )
                 {
                     // If there is no expression, we'll assume true.
                     bool value = true;
@@ -2478,19 +2466,20 @@ namespace mu
                 SCHEDULED_OP itemCopy = item;
                 ExecutionIndex index = GetMemory().GetRageIndex( item.executionIndex );
 
-                    vector<SCHEDULED_OP> deps(iterations*2+1);
-                    int d=0;
-                    deps[d++] = SCHEDULED_OP( args.base, item );
-                    for (int i=0; i<iterations; ++i)
-                    {
-                        index.SetFromModelRangeIndex( args.rangeId, i );
+				TArray<SCHEDULED_OP> deps;
+				deps.SetNumUninitialized(iterations * 2 + 1);
+                int d=0;
+                deps[d++] = SCHEDULED_OP( args.base, item );
+                for (int i=0; i<iterations; ++i)
+                {
+                    index.SetFromModelRangeIndex( args.rangeId, i );
 
-                        itemCopy.executionIndex = GetMemory().GetRageIndexIndex(index);
+                    itemCopy.executionIndex = GetMemory().GetRageIndexIndex(index);
 
-                        deps[d++] = SCHEDULED_OP( args.blended, itemCopy );
-                        deps[d++] = SCHEDULED_OP( args.mask, itemCopy );
-                    }
-                    AddOp( SCHEDULED_OP( item.at, item, 2, iterations), deps );
+                    deps[d++] = SCHEDULED_OP( args.blended, itemCopy );
+                    deps[d++] = SCHEDULED_OP( args.mask, itemCopy );
+                }
+                AddOp( SCHEDULED_OP( item.at, item, 2, iterations), deps );
 
                 break;
             }
@@ -3807,8 +3796,8 @@ namespace mu
 
                     // Allocate memory for the temporary buffers
                     SCRATCH_IMAGE_PROJECT scratch;
-                    scratch.vertices.resize( pMesh->GetVertexCount() );
-                    scratch.culledVertex.resize( pMesh->GetVertexCount() );
+                    scratch.vertices.SetNum( pMesh->GetVertexCount() );
+                    scratch.culledVertex.SetNum( pMesh->GetVertexCount() );
 
                     // Layout is always 0 because the previous mesh project operations take care of
                     // moving the right layout channel to the 0.
@@ -4047,7 +4036,7 @@ namespace mu
         const auto& program = pModel->GetPrivate()->m_program;
         const PARAMETER_DESC& paramDesc = program.m_parameters[ parameterIndex ];
         for( size_t rangeIndexInParam=0;
-             rangeIndexInParam<paramDesc.m_ranges.size();
+             rangeIndexInParam<paramDesc.m_ranges.Num();
              ++rangeIndexInParam )
         {
             auto rangeIndexInModel = paramDesc.m_ranges[rangeIndexInParam];
@@ -4491,9 +4480,8 @@ namespace mu
 
         case OP_TYPE::ST_CONSTANT:
         {
-            auto args =
-                pModel->GetPrivate()->m_program.GetOpArgs<OP::ResourceConstantArgs>( item.at );
-            check( args.value < pModel->GetPrivate()->m_program.m_constantStrings.size() );
+			OP::ResourceConstantArgs args = pModel->GetPrivate()->m_program.GetOpArgs<OP::ResourceConstantArgs>( item.at );
+            check( args.value < (uint32)pModel->GetPrivate()->m_program.m_constantStrings.Num() );
 
             auto result = pModel->GetPrivate()->m_program.m_constantStrings[args.value];
             GetMemory().SetString( item, new String(result.c_str()) );
@@ -4503,7 +4491,7 @@ namespace mu
 
         case OP_TYPE::ST_PARAMETER:
         {
-            auto args = pModel->GetPrivate()->m_program.GetOpArgs<OP::ParameterArgs>( item.at );
+			OP::ParameterArgs args = pModel->GetPrivate()->m_program.GetOpArgs<OP::ParameterArgs>( item.at );
 			Ptr<RangeIndex> index = BuildCurrentOpRangeIndex( item, pParams, pModel, args.variable );
             string result = pParams->GetStringValue( args.variable, index );
             GetMemory().SetString( item, new String( result.c_str() ) );
@@ -4891,9 +4879,7 @@ namespace mu
         case OP_TYPE::LA_CONSTANT:
         {
             auto args = pModel->GetPrivate()->m_program.GetOpArgs<OP::ResourceConstantArgs>(item.at);
-            check( args.value
-                            <
-                            pModel->GetPrivate()->m_program.m_constantLayouts.size() );
+            check( args.value < (uint32)pModel->GetPrivate()->m_program.m_constantLayouts.Num() );
 
             LayoutPtrConst pResult = pModel->GetPrivate()->m_program.m_constantLayouts
                     [ args.value ];
@@ -5396,7 +5382,7 @@ namespace mu
 
 				// At least keep the levels we already have.
 				int startLevel = BaseDesc.m_lods;
-				levelCount = std::max(startLevel, levelCount);
+				levelCount = FMath::Max(startLevel, levelCount);
 
 				// Update result.
 				m_heapImageDesc[item.customState].m_lods = levelCount;
diff --git a/Engine/Plugins/Experimental/Mutable/Source/MutableRuntime/Private/MuR/CodeRunner.h b/Engine/Plugins/Experimental/Mutable/Source/MutableRuntime/Private/MuR/CodeRunner.h
--- a/Engine/Plugins/Experimental/Mutable/Source/MutableRuntime/Private/MuR/CodeRunner.h
+++ b/Engine/Plugins/Experimental/Mutable/Source/MutableRuntime/Private/MuR/CodeRunner.h
@@ -250,7 +250,7 @@ class RangeIndex;
 			AddChildren(dep4);
 		}
 
-		void AddOp(const SCHEDULED_OP& op, const vector<SCHEDULED_OP>& deps)
+		void AddOp(const SCHEDULED_OP& op, const TArray<SCHEDULED_OP>& deps)
 		{
 			FTask Task(op);
 			for (const auto& d : deps)
@@ -312,12 +312,12 @@ class RangeIndex;
 			int32 m_romIndex = 0;
 			DATATYPE ConstantType = DT_NONE;
 			ModelStreamer::OPERATION_ID m_streamID;
-			vector<uint8> m_streamBuffer;
+			TArray<uint8> m_streamBuffer;
 		};
-		vector<ROM_LOAD_OP> m_romLoadOps;
+		TArray<ROM_LOAD_OP> m_romLoadOps;
 
 		//! Count of pending operations for every rom index
-		vector<uint16> m_romPendingOps;
+		TArray<uint16> m_romPendingOps;
 
 	private:
 
diff --git a/Engine/Plugins/Experimental/Mutable/Source/MutableRuntime/Private/MuR/CodeRunnerTasked.cpp b/Engine/Plugins/Experimental/Mutable/Source/MutableRuntime/Private/MuR/CodeRunnerTasked.cpp
--- a/Engine/Plugins/Experimental/Mutable/Source/MutableRuntime/Private/MuR/CodeRunnerTasked.cpp
+++ b/Engine/Plugins/Experimental/Mutable/Source/MutableRuntime/Private/MuR/CodeRunnerTasked.cpp
@@ -67,7 +67,7 @@ namespace mu
 		{
 			MUTABLE_CPUPROFILER_SCOPE(Unserialise);
 
-			InputMemoryStream stream(o.m_streamBuffer.data(), o.m_streamBuffer.size());
+			InputMemoryStream stream(o.m_streamBuffer.GetData(), o.m_streamBuffer.Num());
 			InputArchive arch(&stream);
 
 			int32 ResIndex = program.m_roms[o.m_romIndex].ResourceIndex;
@@ -93,7 +93,7 @@ namespace mu
 			}
 
 			o.m_romIndex = -1;
-			o.m_streamBuffer = vector<uint8_t>();
+			o.m_streamBuffer.Empty();
 		}
 	}
 
@@ -1014,7 +1014,7 @@ namespace mu
 						// This now happens more often since the generation of specific mips on request. For this reason
 						// this warning is usually acceptable.
 						//UE_LOG(LogMutableCore, Log, TEXT("Required image resize for ImageCompose: performance warning."));
-						FImageSize blockSize((uint16_t)rect.size[0], (uint16_t)rect.size[1]);
+						FImageSize blockSize((uint16)rect.size[0], (uint16)rect.size[1]);
 						m_block = ImageResizeLinear(m_imageCompressionQuality, m_block.get(), blockSize);
 					}
 
@@ -1111,22 +1111,22 @@ namespace mu
 
 			if (!op)
 			{
-				runner->m_romLoadOps.push_back(CodeRunner::ROM_LOAD_OP());
-				op = &runner->m_romLoadOps.back();
+				runner->m_romLoadOps.Add(CodeRunner::ROM_LOAD_OP());
+				op = &runner->m_romLoadOps.Last();
 			}
 
 			op->m_romIndex = RomIndex;
 			op->ConstantType = DT_MESH;
 
-			if (op->m_streamBuffer.size() < RomSize)
+			if (uint32(op->m_streamBuffer.Num()) < RomSize)
 			{
 				// This could happen in 32-bit platforms
 				check(RomSize < std::numeric_limits<size_t>::max());
-				op->m_streamBuffer.resize((size_t)RomSize);
+				op->m_streamBuffer.SetNum((size_t)RomSize);
 			}
 
 			uint32 RomId = program.m_roms[RomIndex].Id;
-			op->m_streamID = runner->m_pSystem->m_pStreamInterface->BeginReadBlock(InModel, RomId, op->m_streamBuffer.data(), RomSize);
+			op->m_streamID = runner->m_pSystem->m_pStreamInterface->BeginReadBlock(InModel, RomId, op->m_streamBuffer.GetData(), RomSize);
 			if (op->m_streamID < 0)
 			{
 				bOutFailed = true;
@@ -1231,22 +1231,22 @@ namespace mu
 
 			if (!op)
 			{
-				runner->m_romLoadOps.push_back(CodeRunner::ROM_LOAD_OP());
-				op = &runner->m_romLoadOps.back();
+				runner->m_romLoadOps.Add(CodeRunner::ROM_LOAD_OP());
+				op = &runner->m_romLoadOps.Last();
 			}
 
 			op->m_romIndex = RomIndex;
 			op->ConstantType = DT_IMAGE;
 
-			if (op->m_streamBuffer.size() < RomSize)
+			if (uint32(op->m_streamBuffer.Num()) < RomSize)
 			{
 				// This could happen in 32-bit platforms
 				check(RomSize < std::numeric_limits<size_t>::max());
-				op->m_streamBuffer.resize((size_t)RomSize);
+				op->m_streamBuffer.SetNum(RomSize);
 			}
 
 			uint32 RomId = program.m_roms[RomIndex].Id;
-			op->m_streamID = runner->m_pSystem->m_pStreamInterface->BeginReadBlock(InModel, RomId, op->m_streamBuffer.data(), RomSize);
+			op->m_streamID = runner->m_pSystem->m_pStreamInterface->BeginReadBlock(InModel, RomId, op->m_streamBuffer.GetData(), RomSize);
 			if (op->m_streamID < 0)
 			{
 				bOutFailed = true;
diff --git a/Engine/Plugins/Experimental/Mutable/Source/MutableRuntime/Private/MuR/CodeVisitor.cpp b/Engine/Plugins/Experimental/Mutable/Source/MutableRuntime/Private/MuR/CodeVisitor.cpp
--- a/Engine/Plugins/Experimental/Mutable/Source/MutableRuntime/Private/MuR/CodeVisitor.cpp
+++ b/Engine/Plugins/Experimental/Mutable/Source/MutableRuntime/Private/MuR/CodeVisitor.cpp
@@ -16,10 +16,10 @@ namespace mu
 	void SubtreeParametersVisitor::Run(OP::ADDRESS root, PROGRAM& program)
 	{
 		// Cached?
-		auto it = m_resultCache.find(root);
-		if (it != m_resultCache.end())
+		auto it = m_resultCache.Find(root);
+		if (it)
 		{
-			m_params = it->second;
+			m_params = *it;
 			return;
 		}
 
@@ -27,25 +27,24 @@ namespace mu
 		{
 			MUTABLE_CPUPROFILER_SCOPE(SubtreeParametersVisitor);
 
-			m_visited.resize(program.m_opAddress.size());
-			if (program.m_opAddress.size())
+			m_visited.SetNum(program.m_opAddress.Num());
+			if (program.m_opAddress.Num())
 			{
-				FMemory::Memzero(&m_visited[0], m_visited.size());
+				FMemory::Memzero(&m_visited[0], m_visited.Num());
 			}
 
-			m_currentParams.resize(program.m_parameters.size());
-			if (m_currentParams.size())
+			m_currentParams.SetNum(program.m_parameters.Num());
+			if (m_currentParams.Num())
 			{
-				FMemory::Memzero(&m_currentParams[0], sizeof(int) * m_currentParams.size());
+				FMemory::Memzero(&m_currentParams[0], sizeof(int) * m_currentParams.Num());
 			}
 
-			m_pending.reserve(program.m_opAddress.size() / 4);
-			m_pending.push_back(root);
+			m_pending.Reserve(program.m_opAddress.Num() / 4);
+			m_pending.Add(root);
 
-			while (m_pending.size())
+			while (m_pending.Num())
 			{
-				OP::ADDRESS at = m_pending.back();
-				m_pending.pop_back();
+				OP::ADDRESS at = m_pending.Pop();
 
 				if (!m_visited[at])
 				{
@@ -70,7 +69,7 @@ namespace mu
 						{
 							if (ref)
 							{
-								m_pending.push_back(ref);
+								m_pending.Add(ref);
 							}
 						});
 				}
@@ -78,16 +77,16 @@ namespace mu
 
 
 			// Build result
-			m_params.clear();
-			for (size_t i = 0; i < m_currentParams.size(); ++i)
+			m_params.Empty();
+			for (size_t i = 0; i < m_currentParams.Num(); ++i)
 			{
 				if (m_currentParams[i])
 				{
-					m_params.push_back(int(i));
+					m_params.Add(int(i));
 				}
 			}
 
-			m_resultCache[root] = m_params;
+			m_resultCache.Add(root, m_params);
 		}
 	}
 
diff --git a/Engine/Plugins/Experimental/Mutable/Source/MutableRuntime/Private/MuR/CodeVisitor.h b/Engine/Plugins/Experimental/Mutable/Source/MutableRuntime/Private/MuR/CodeVisitor.h
--- a/Engine/Plugins/Experimental/Mutable/Source/MutableRuntime/Private/MuR/CodeVisitor.h
+++ b/Engine/Plugins/Experimental/Mutable/Source/MutableRuntime/Private/MuR/CodeVisitor.h
@@ -45,7 +45,7 @@ class Parameters;
         UniqueConstCodeVisitorIterative( bool skipResources=false )
         {
             // Default state
-            m_states.push_back(STATE());
+            m_states.Add(STATE());
             m_currentState = 0;
             m_skipResources = skipResources;
         }
@@ -77,52 +77,50 @@ class Parameters;
         //! For manual recursion that changes the state for a specific path.
         void RecurseWithState(OP::ADDRESS at, const STATE& newState)
         {
-            auto it = std::find(m_states.begin(),m_states.end(),newState);
-            if (it==m_states.end())
+            auto it = m_states.Find(newState);
+            if (it==INDEX_NONE)
             {
-                m_states.push_back(newState);
-                it = m_states.end() - 1;
+                m_states.Add(newState);
             }
-            int stateIndex = (int)(it - m_states.begin());
+            int stateIndex = m_states.IndexOfByKey(newState);
 
-            m_pending.push_back( PENDING(at,stateIndex) );
+            m_pending.Add( PENDING(at,stateIndex) );
         }
 
         //! For manual recursion that doesn't change the state for a specific path.
         void RecurseWithCurrentState(OP::ADDRESS at)
         {
-            m_pending.push_back( PENDING(at,m_currentState) );
+            m_pending.Add( PENDING(at,m_currentState) );
         }
 
         //! Can be called from visit to set the state to visit all children ops
         void SetCurrentState(const STATE& newState)
         {
-            auto it = std::find(m_states.begin(),m_states.end(),newState);
-            if (it==m_states.end())
+            auto it = m_states.Find(newState);
+            if (it==INDEX_NONE)
             {
-                m_states.push_back(newState);
-                it = m_states.end() - 1;
+                m_states.Add(newState);
             }
-            m_currentState = (int)(it - m_states.begin());
+            m_currentState = m_states.Find(newState);
         }
 
 
         void Traverse( OP::ADDRESS root, PROGRAM& program, bool visitDecorators = true )
         {
-            m_pending.reserve( program.m_opAddress.size() );
+            m_pending.Reserve( program.m_opAddress.Num() );
 
             // Visit the given root
-            m_pending.push_back( PENDING(root,0) );
+            m_pending.Add( PENDING(root,0) );
             Recurse( program );
 
             if (visitDecorators)
             {
                 // Fix the code used in the parameter descriptions
-                for ( std::size_t p=0; p<program.m_parameters.size(); ++p )
+                for ( std::size_t p=0; p<program.m_parameters.Num(); ++p )
                 {
-                    for ( std::size_t d=0; d<program.m_parameters[p].m_descImages.size(); ++d )
+                    for ( std::size_t d=0; d<program.m_parameters[p].m_descImages.Num(); ++d )
                     {
-                        m_pending.push_back(PENDING(program.m_parameters[p].m_descImages[d],0) );
+                        m_pending.Add(PENDING(program.m_parameters[p].m_descImages[d],0) );
                         Recurse( program );
                     }
                 }
@@ -132,20 +130,20 @@ class Parameters;
         void FullTraverse( PROGRAM& program, bool visitDecorators = true )
         {
             // Visit all the state roots
-            for ( std::size_t p=0; p<program.m_states.size(); ++p )
+            for ( std::size_t p=0; p<program.m_states.Num(); ++p )
             {
-                m_pending.push_back(PENDING(program.m_states[p].m_root,0) );
+                m_pending.Add(PENDING(program.m_states[p].m_root,0) );
                 Recurse( program );
             }
 
             if (visitDecorators)
             {
                 // Fix the code used in the parameter descriptions
-                for ( std::size_t p=0; p<program.m_parameters.size(); ++p )
+                for ( std::size_t p=0; p<program.m_parameters.Num(); ++p )
                 {
-                    for ( std::size_t d=0; d<program.m_parameters[p].m_descImages.size(); ++d )
+                    for ( std::size_t d=0; d<program.m_parameters[p].m_descImages.Num(); ++d )
                     {
-                        m_pending.push_back(PENDING(program.m_parameters[p].m_descImages[d],0) );
+                        m_pending.Add(PENDING(program.m_parameters[p].m_descImages[d],0) );
                         Recurse( program );
                     }
                 }
@@ -176,10 +174,10 @@ class Parameters;
             OP::ADDRESS at;
             int stateIndex;
         };
-        vector<PENDING> m_pending;
+		TArray<PENDING> m_pending;
 
         //! States found so far
-        vector<STATE> m_states;
+		TArray<STATE> m_states;
 
         //! Index of the current state, from the m_states vector.
         int m_currentState;
@@ -190,27 +188,26 @@ class Parameters;
 
         //! Array of states visited for each operation.
         //! Empty vector means operation not visited at all.
-        vector<vector<int>> m_visited;
+		TArray<TArray<int>> m_visited;
 
         //! Process all the pending operations and visit all children if necessary
         void Recurse( PROGRAM& program )
         {
-            m_visited = vector<vector<int>>(program.m_opAddress.size());
+			m_visited.Empty();
+			m_visited.SetNum(program.m_opAddress.Num());
 
-            while ( m_pending.size() )
+            while ( m_pending.Num() )
             {
-                OP::ADDRESS at = m_pending.back().at;
-                m_currentState = m_pending.back().stateIndex;
-                m_pending.pop_back();
+                OP::ADDRESS at = m_pending.Last().at;
+                m_currentState = m_pending.Last().stateIndex;
+                m_pending.Pop();
 
                 bool recurse = false;
 
-                bool visitedInThisState = std::find(m_visited[at].begin(), m_visited[at].end(), m_currentState)
-                    !=
-                    m_visited[at].end();
+                bool visitedInThisState = m_visited[at].Contains(m_currentState);
                 if (!visitedInThisState)
                 {
-                    m_visited[at].push_back(m_currentState);
+                    m_visited[at].Add(m_currentState);
 
                     // Visit may change current state
                     recurse = Visit( at, program );
@@ -222,7 +219,7 @@ class Parameters;
                     {
                         if (ref)
                         {
-                            m_pending.push_back( PENDING(ref,m_currentState) );
+                            m_pending.Add( PENDING(ref,m_currentState) );
                         }
                     });
                 }
@@ -247,7 +244,7 @@ class Parameters;
         RepeatConstCodeVisitorIterative( bool skipResources=false )
         {
             // Default state
-            m_states.push_back(STATE());
+            m_states.Add(STATE());
             m_currentState = 0;
             m_skipResources = skipResources;
         }
@@ -282,20 +279,20 @@ class Parameters;
             auto it = std::find(m_states.begin(),m_states.end(),newState);
             if (it==m_states.end())
             {
-                m_states.push_back(newState);
+                m_states.Add(newState);
                 it = m_states.end() - 1;
             }
             int stateIndex = (int)(it - m_states.begin());
 
             //check(at<1000000);
-            m_pending.push_back( PENDING(at,stateIndex) );
+            m_pending.Add( PENDING(at,stateIndex) );
         }
 
         //! For manual recursion that doesn't change the state for a specific path.
         void RecurseWithCurrentState(OP::ADDRESS at)
         {
             //check(at<1000000);
-            m_pending.push_back( PENDING(at,m_currentState) );
+            m_pending.Add( PENDING(at,m_currentState) );
         }
 
         //! Can be called from visit to set the state to visit all children ops
@@ -304,7 +301,7 @@ class Parameters;
             auto it = std::find(m_states.begin(),m_states.end(),newState);
             if (it==m_states.end())
             {
-                m_states.push_back(newState);
+                m_states.Add(newState);
                 it = m_states.end() - 1;
             }
             m_currentState = (int)(it - m_states.begin());
@@ -313,20 +310,20 @@ class Parameters;
 
         void Traverse( OP::ADDRESS root, PROGRAM& program, bool visitDecorators = true )
         {
-            m_pending.reserve( program.m_opAddress.size() );
+            m_pending.reserve( program.m_opAddress.Num() );
 
             // Visit the given root
-            m_pending.push_back( PENDING(root,0) );
+            m_pending.Add( PENDING(root,0) );
             Recurse( program );
 
             if (visitDecorators)
             {
                 // Fix the code used in the parameter descriptions
-                for ( std::size_t p=0; p<program.m_parameters.size(); ++p )
+                for ( std::size_t p=0; p<program.m_parameters.Num(); ++p )
                 {
-                    for ( std::size_t d=0; d<program.m_parameters[p].m_descImages.size(); ++d )
+                    for ( std::size_t d=0; d<program.m_parameters[p].m_descImages.Num(); ++d )
                     {
-                        m_pending.push_back(PENDING(program.m_parameters[p].m_descImages[d],0) );
+                        m_pending.Add(PENDING(program.m_parameters[p].m_descImages[d],0) );
                         Recurse( program );
                     }
                 }
@@ -336,20 +333,20 @@ class Parameters;
         void FullTraverse( PROGRAM& program, bool visitDecorators = true )
         {
             // Visit all the state roots
-            for ( std::size_t p=0; p<program.m_states.size(); ++p )
+            for ( std::size_t p=0; p<program.m_states.Num(); ++p )
             {
-                m_pending.push_back(PENDING(program.m_states[p].m_root,0) );
+                m_pending.Add(PENDING(program.m_states[p].m_root,0) );
                 Recurse( program );
             }
 
             if (visitDecorators)
             {
                 // Fix the code used in the parameter descriptions
-                for ( std::size_t p=0; p<program.m_parameters.size(); ++p )
+                for ( std::size_t p=0; p<program.m_parameters.Num(); ++p )
                 {
-                    for ( std::size_t d=0; d<program.m_parameters[p].m_descImages.size(); ++d )
+                    for ( std::size_t d=0; d<program.m_parameters[p].m_descImages.Num(); ++d )
                     {
-                        m_pending.push_back(PENDING(program.m_parameters[p].m_descImages[d],0) );
+                        m_pending.Add(PENDING(program.m_parameters[p].m_descImages[d],0) );
                         Recurse( program );
                     }
                 }
@@ -380,10 +377,10 @@ class Parameters;
             OP::ADDRESS at;
             int stateIndex;
         };
-        vector<PENDING> m_pending;
+        TArray<PENDING> m_pending;
 
         //! States found so far
-        vector<STATE> m_states;
+		TArray<STATE> m_states;
 
         //! Index of the current state, from the m_states vector.
         int m_currentState;
@@ -396,7 +393,7 @@ class Parameters;
         //! Process all the pending operations and visit all children if necessary
         void Recurse( PROGRAM& program )
         {
-            while ( m_pending.size() )
+            while ( m_pending.Num() )
             {
                 OP::ADDRESS at = m_pending.back().at;
                 m_currentState = m_pending.back().stateIndex;
@@ -417,8 +414,8 @@ class Parameters;
                         OP::ADDRESS base = args.instance;
                         if (base)
                         {
-							check(base<program.m_opAddress.size());
-                            m_pending.push_back( PENDING(base,m_currentState) );
+							check(base<program.m_opAddress.Num());
+                            m_pending.Add( PENDING(base,m_currentState) );
                         }
                     }
                     else
@@ -427,8 +424,8 @@ class Parameters;
                         {
                             if (ref)
                             {
-								check(ref<program.m_opAddress.size());
-                                m_pending.push_back( PENDING(ref,m_currentState) );
+								check(ref<program.m_opAddress.Num());
+                                m_pending.Add( PENDING(ref,m_currentState) );
                             }
                         });
                     }
@@ -449,7 +446,7 @@ class Parameters;
     //---------------------------------------------------------------------------------------------
     struct COVERED_CODE_VISITOR_STATE
     {
-        uint16_t m_underResourceCount = 0;
+        uint16 m_underResourceCount = 0;
 
         bool operator==(const COVERED_CODE_VISITOR_STATE& o) const
         {
@@ -714,16 +711,16 @@ class Parameters;
         void Run( OP::ADDRESS root, PROGRAM& program );
 
         //! After Run, list of relevant parameters.
-        vector<int> m_params;
+		TArray<int> m_params;
 
     private:
 
-        vector<int> m_currentParams;
-        vector<uint8_t> m_visited;
-        vector<OP::ADDRESS> m_pending;
+		TArray<int> m_currentParams;
+		TArray<uint8_t> m_visited;
+		TArray<OP::ADDRESS> m_pending;
 
         // Result cache
-        map< OP::ADDRESS, vector<int> > m_resultCache;
+        TMap< OP::ADDRESS, TArray<int> > m_resultCache;
     };
 
 
diff --git a/Engine/Plugins/Experimental/Mutable/Source/MutableRuntime/Private/MuR/ConvertData.h b/Engine/Plugins/Experimental/Mutable/Source/MutableRuntime/Private/MuR/ConvertData.h
--- a/Engine/Plugins/Experimental/Mutable/Source/MutableRuntime/Private/MuR/ConvertData.h
+++ b/Engine/Plugins/Experimental/Mutable/Source/MutableRuntime/Private/MuR/ConvertData.h
@@ -395,9 +395,9 @@ namespace mu
 			{
 				const float* pTypedSource = reinterpret_cast<const float*>( pSource );
                 pTypedResult[channel] = (uint8)
-                        std::min<uint32>(
+                        FMath::Min<uint32>(
 							0xFF,
-                            std::max<uint32>(
+                            FMath::Max<uint32>(
 								0,
                                 (uint32)(pTypedSource[channel])
 								)
@@ -409,9 +409,9 @@ namespace mu
 			{
 				const float16* pTypedSource = reinterpret_cast<const float16*>( pSource );
                 pTypedResult[channel] = (uint8)
-                        std::min<uint32>(
+                        FMath::Min<uint32>(
 							0xFF,
-                            std::max<uint32>(
+                            FMath::Max<uint32>(
 								0,
                                 (uint32)halfToFloat(pTypedSource[channel])
 								)
@@ -422,7 +422,7 @@ namespace mu
 			case MBF_INT8:
 			{
                 const int8* pTypedSource = reinterpret_cast<const int8*>( pSource );
-                pTypedResult[channel] =  (uint8)std::max<int8>( 0, pTypedSource[channel] );
+                pTypedResult[channel] =  (uint8)FMath::Max<int8>( 0, pTypedSource[channel] );
 				break;
 			}
 
@@ -437,9 +437,9 @@ namespace mu
 			{
                 const int16* pTypedSource = reinterpret_cast<const int16*>( pSource );
                 pTypedResult[channel] =  (uint8)
-                        std::min<int16>(
+                        FMath::Min<int16>(
 							0xFF,
-                            std::max<int16>(
+                            FMath::Max<int16>(
 								0,
 								pTypedSource[channel]
 								)
@@ -452,9 +452,9 @@ namespace mu
 				// Clamp
                 const uint16* pTypedSource = reinterpret_cast<const uint16*>( pSource );
                 pTypedResult[channel] = (uint8)
-                        std::min<uint16>(
+                        FMath::Min<uint16>(
 							0xFF,
-                            std::max<uint16>(
+                            FMath::Max<uint16>(
 								0,
 								pTypedSource[channel]
 								)
@@ -466,9 +466,9 @@ namespace mu
 			{
                 const int32* pTypedSource = reinterpret_cast<const int32*>( pSource );
                 pTypedResult[channel] =  (uint8)
-                        std::min<int32>(
+                        FMath::Min<int32>(
 							0xFF,
-                            std::max<int32>(
+                            FMath::Max<int32>(
 								0,
 								pTypedSource[channel]
 								)
@@ -481,9 +481,9 @@ namespace mu
 				// Clamp
                 const uint32* pTypedSource = reinterpret_cast<const uint32*>( pSource );
                 pTypedResult[channel] = (uint8)
-                        std::min<uint32>(
+                        FMath::Min<uint32>(
 							0xFF,
-                            std::max<uint32>(
+                            FMath::Max<uint32>(
 								0,
 								pTypedSource[channel]
 								)
@@ -517,9 +517,9 @@ namespace mu
 			{
 				const float* pTypedSource = reinterpret_cast<const float*>( pSource );
                 pTypedResult[channel] = (uint16)
-                        std::min<uint32>(
+                        FMath::Min<uint32>(
 							0xFFFF,
-                            std::max<uint32>(
+                            FMath::Max<uint32>(
 								0,
                                 (uint32)(pTypedSource[channel])
 								)
@@ -531,9 +531,9 @@ namespace mu
 			{
 				const float16* pTypedSource = reinterpret_cast<const float16*>( pSource );
                 pTypedResult[channel] = (uint16)
-                        std::min<uint32>(
+                        FMath::Min<uint32>(
 							0xFF,
-                            std::max<uint32>(
+                            FMath::Max<uint32>(
 								0,
                                 (uint32)halfToFloat(pTypedSource[channel])
 								)
@@ -551,7 +551,7 @@ namespace mu
 			case MBF_INT8:
 			{
                 const int8* pTypedSource = reinterpret_cast<const int8*>( pSource );
-                pTypedResult[channel] =  (uint16)std::max<int8>( 0, pTypedSource[channel] );
+                pTypedResult[channel] =  (uint16)FMath::Max<int8>( 0, pTypedSource[channel] );
 				break;
 			}
 
@@ -565,7 +565,7 @@ namespace mu
 			case MBF_INT16:
 			{
                 const int16* pTypedSource = reinterpret_cast<const int16*>( pSource );
-                pTypedResult[channel] =  (uint16)std::max<int16>( 0, pTypedSource[channel] );
+                pTypedResult[channel] =  (uint16)FMath::Max<int16>( 0, pTypedSource[channel] );
 				break;
 			}
 
@@ -574,9 +574,9 @@ namespace mu
 				// Clamp
                 const uint32* pTypedSource = reinterpret_cast<const uint32*>( pSource );
                 pTypedResult[channel] = (uint16)
-                        std::min<uint32>(
+                        FMath::Min<uint32>(
 							0xFFFF,
-                            std::max<uint32>(
+                            FMath::Max<uint32>(
 								0,
 								pTypedSource[channel]
 								)
@@ -589,9 +589,9 @@ namespace mu
 				// Clamp
                 const int32* pTypedSource = reinterpret_cast<const int32*>( pSource );
                 pTypedResult[channel] = (uint16)
-                        std::min<int32>(
+                        FMath::Min<int32>(
 							0xFFFF,
-                            std::max<int32>(
+                            FMath::Max<int32>(
 								0,
 								pTypedSource[channel]
 								)
@@ -625,9 +625,9 @@ namespace mu
 			{
 				const float* pTypedSource = reinterpret_cast<const float*>( pSource );
 				pTypedResult[channel] =
-                        std::min<uint32>(
+                        FMath::Min<uint32>(
 							0xFFFF,
-                            std::max<uint32>(
+                            FMath::Max<uint32>(
 								0,
                                 (uint32)(pTypedSource[channel])
 								)
@@ -639,9 +639,9 @@ namespace mu
 			{
 				const float16* pTypedSource = reinterpret_cast<const float16*>( pSource );
 				pTypedResult[channel] =
-                        std::min<uint32>(
+                        FMath::Min<uint32>(
 							0xFF,
-                            std::max<uint32>(
+                            FMath::Max<uint32>(
 								0,
                                 (uint32)halfToFloat(pTypedSource[channel])
 								)
@@ -659,7 +659,7 @@ namespace mu
 			case MBF_INT8:
 			{
                 const int8* pTypedSource = reinterpret_cast<const int8*>( pSource );
-                pTypedResult[channel] =  (uint16)std::max<int8>( 0, pTypedSource[channel] );
+                pTypedResult[channel] =  (uint16)FMath::Max<int8>( 0, pTypedSource[channel] );
 				break;
 			}
 
@@ -673,7 +673,7 @@ namespace mu
 			case MBF_INT16:
 			{
                 const int16* pTypedSource = reinterpret_cast<const int16*>( pSource );
-                pTypedResult[channel] =  (uint16)std::max<int16>( 0, pTypedSource[channel] );
+                pTypedResult[channel] =  (uint16)FMath::Max<int16>( 0, pTypedSource[channel] );
 				break;
 			}
 
@@ -688,7 +688,7 @@ namespace mu
 			{
 				// Clamp
                 const int32* pTypedSource = reinterpret_cast<const int32*>( pSource );
-                pTypedResult[channel] = (uint32)std::max<int32>( 0, pTypedSource[channel] );
+                pTypedResult[channel] = (uint32)FMath::Max<int32>( 0, pTypedSource[channel] );
 				break;
 			}
 
@@ -718,9 +718,9 @@ namespace mu
 			{
 				const float* pTypedSource = reinterpret_cast<const float*>( pSource );
                 pTypedResult[channel] = (int8)
-                        std::min<int32>(
+                        FMath::Min<int32>(
 							127,
-                            std::max<int32>(
+                            FMath::Max<int32>(
 								-128,
                                 (int32)(pTypedSource[channel])
 								)
@@ -732,9 +732,9 @@ namespace mu
 			{
 				const float16* pTypedSource = reinterpret_cast<const float16*>( pSource );
                 pTypedResult[channel] = (int8)
-                        std::min<int32>(
+                        FMath::Min<int32>(
 							127,
-                            std::max<int32>(
+                            FMath::Max<int32>(
 								-128,
                                 (int32)halfToFloat(pTypedSource[channel])
 								)
@@ -775,9 +775,9 @@ namespace mu
 			{
 				const float* pTypedSource = reinterpret_cast<const float*>( pSource );
                 pTypedResult[channel] = (int16)
-                        std::min<int32>(
+                        FMath::Min<int32>(
 							32767,
-                            std::max<int32>(
+                            FMath::Max<int32>(
 								-32768,
                                 (int32)(pTypedSource[channel])
 								)
@@ -789,9 +789,9 @@ namespace mu
 			{
 				const float16* pTypedSource = reinterpret_cast<const float16*>( pSource );
                 pTypedResult[channel] = (int16)
-                        std::min<int32>(
+                        FMath::Min<int32>(
 							32767,
-                            std::max<int32>(
+                            FMath::Max<int32>(
 								-32768,
                                 (int32)halfToFloat(pTypedSource[channel])
 								)
@@ -817,7 +817,7 @@ namespace mu
 			{
                 const uint16* pTypedSource = reinterpret_cast<const uint16*>( pSource );
                 pTypedResult[channel] = (int16)
-                        std::min<int32>(
+                        FMath::Min<int32>(
                             32767, (int32)pTypedSource[channel]
 							);
 				break;
@@ -827,9 +827,9 @@ namespace mu
 			{
                 const int32* pTypedSource = reinterpret_cast<const int32*>( pSource );
                 pTypedResult[channel] = (int16)
-                        std::min<int32>(
+                        FMath::Min<int32>(
 							32767,
-                            std::max<int32>(
+                            FMath::Max<int32>(
 								-32768,
 								pTypedSource[channel]
 								)
@@ -841,7 +841,7 @@ namespace mu
 			{
                 const uint32* pTypedSource = reinterpret_cast<const uint32*>( pSource );
                 pTypedResult[channel] = (int16)
-                        std::min<int32>(
+                        FMath::Min<int32>(
 							32767, pTypedSource[channel]
 							);
 				break;
@@ -966,9 +966,9 @@ namespace mu
 			{
 				const float* pTypedSource = reinterpret_cast<const float*>(pSource);
 				pTypedResult[channel] = (uint8)
-					std::min<uint32>(
+					FMath::Min<uint32>(
 						0xFF,
-						std::max<uint32>(
+						FMath::Max<uint32>(
 							0,
 							(uint32)(((float)0xFF)*pTypedSource[channel] + 0.5f)
 							)
@@ -980,9 +980,9 @@ namespace mu
 			{
 				const float16* pTypedSource = reinterpret_cast<const float16*>( pSource );
                 pTypedResult[channel] = (uint8)
-                        std::min<uint32>(
+                        FMath::Min<uint32>(
 							0xFF,
-                            std::max<uint32>(
+                            FMath::Max<uint32>(
 								0,
                                 (uint32)(((float)0xFF)*halfToFloat(pTypedSource[channel])+0.5f)
 								)
@@ -1031,9 +1031,9 @@ namespace mu
 			{
 				const float* pTypedSource = reinterpret_cast<const float*>( pSource );
                 pTypedResult[channel] = (uint16)
-                        std::min<uint32>(
+                        FMath::Min<uint32>(
 							0xFFFF,
-                            std::max<uint32>(
+                            FMath::Max<uint32>(
 								0,
                                 (uint32)(((float)0xFFFF)*pTypedSource[channel]+0.5f)
 								)
@@ -1045,9 +1045,9 @@ namespace mu
 			{
 				const float16* pTypedSource = reinterpret_cast<const float16*>( pSource );
                 pTypedResult[channel] = (uint16)
-                        std::min<uint32>(
+                        FMath::Min<uint32>(
 							0xFFFF,
-                            std::max<uint32>(
+                            FMath::Max<uint32>(
 								0,
                                 (uint32)(((float)0xFFFF)*halfToFloat(pTypedSource[channel])+0.5f)
 								)
@@ -1082,9 +1082,9 @@ namespace mu
 			{
 				const float* pTypedSource = reinterpret_cast<const float*>( pSource );
                 pTypedResult[channel] = (uint32)
-                        std::min<uint32>(
+                        FMath::Min<uint32>(
 							0xFFFFFFFF,
-                            std::max<uint32>(
+                            FMath::Max<uint32>(
 								0,
                                 (uint32)(((float)0xFFFFFFFF)*pTypedSource[channel])
 								)
@@ -1096,9 +1096,9 @@ namespace mu
 			{
 				const float16* pTypedSource = reinterpret_cast<const float16*>( pSource );
                 pTypedResult[channel] = (uint32)
-                        std::min<uint32>(
+                        FMath::Min<uint32>(
 							0xFFFFFFFF,
-                            std::max<uint32>(
+                            FMath::Max<uint32>(
 								0,
                                 (uint32)(((float)0xFFFFFFFF)*halfToFloat(pTypedSource[channel])+0.5f)
 								)
@@ -1133,9 +1133,9 @@ namespace mu
 			{
 				const float* pTypedSource = reinterpret_cast<const float*>( pSource );
                 pTypedResult[channel] = (int8)
-                        std::min<int32>(
+                        FMath::Min<int32>(
 							127,
-                            std::max<int32>(
+                            FMath::Max<int32>(
 								-128,
                                 (int32)(128.0f*pTypedSource[channel]+0.5f)
 								)
@@ -1147,9 +1147,9 @@ namespace mu
 			{
 				const float16* pTypedSource = reinterpret_cast<const float16*>( pSource );
                 pTypedResult[channel] = (int8)
-                        std::min<int32>(
+                        FMath::Min<int32>(
 							127,
-                            std::max<int32>(
+                            FMath::Max<int32>(
 								-128,
                                 (int32)(128.0f*halfToFloat(pTypedSource[channel])+0.5f)
 								)
@@ -1184,9 +1184,9 @@ namespace mu
 			{
 				const float* pTypedSource = reinterpret_cast<const float*>( pSource );
                 pTypedResult[channel] = (int16)
-                        std::min<int32>(
+                        FMath::Min<int32>(
 							32767,
-                            std::max<int32>(
+                            FMath::Max<int32>(
 								-32768,
                                 (int32)(32768.0f*pTypedSource[channel]+0.5f)
 								)
@@ -1198,9 +1198,9 @@ namespace mu
 			{
 				const float16* pTypedSource = reinterpret_cast<const float16*>( pSource );
                 pTypedResult[channel] = (int16)
-                        std::min<int32>(
+                        FMath::Min<int32>(
 							32767,
-                            std::max<int32>(
+                            FMath::Max<int32>(
 								-32768,
                                 (int32)(32768.0f*halfToFloat(pTypedSource[channel])+0.5f)
 								)
@@ -1285,7 +1285,7 @@ namespace mu
                 float source = pTypedSource[channel];
                 source = (source*0.5f+0.5f)*255.0f;
                 pTypedResult[channel] =
-                        (uint8)std::min<float>( 255.0f, std::max<float>( 0.0f, source ) );
+                        (uint8)FMath::Min<float>( 255.0f, FMath::Max<float>( 0.0f, source ) );
                 break;
             }
 
@@ -1318,7 +1318,7 @@ namespace mu
                 float source = pTypedSource[channel];
                 source = source*0.5f*255.0f;
                 pTypedResult[channel] =
-                        (int8)std::min<float>( 127.0f, std::max<float>( -128.0f, source ) );
+                        (int8)FMath::Min<float>( 127.0f, FMath::Max<float>( -128.0f, source ) );
                 break;
             }
 
diff --git a/Engine/Plugins/Experimental/Mutable/Source/MutableRuntime/Private/MuR/Image.cpp b/Engine/Plugins/Experimental/Mutable/Source/MutableRuntime/Private/MuR/Image.cpp
--- a/Engine/Plugins/Experimental/Mutable/Source/MutableRuntime/Private/MuR/Image.cpp
+++ b/Engine/Plugins/Experimental/Mutable/Source/MutableRuntime/Private/MuR/Image.cpp
@@ -70,7 +70,7 @@ namespace mu
 		LLM_SCOPE_BYNAME(TEXT("MutableRuntime"));
 
         m_format = format;
-        m_size = FImageSize( (uint16_t)sizeX, (uint16_t)sizeY );
+        m_size = FImageSize( (uint16)sizeX, (uint16)sizeY );
         m_lods = (uint8_t)lods;
         m_internalId = 0;
 
@@ -142,14 +142,14 @@ namespace mu
 
 
     //---------------------------------------------------------------------------------------------
-    uint16_t Image::GetSizeX() const
+    uint16 Image::GetSizeX() const
     {
         return m_size[0];
     }
 
 
     //---------------------------------------------------------------------------------------------
-    uint16_t Image::GetSizeY() const
+    uint16 Image::GetSizeY() const
     {
         return m_size[1];
     }
@@ -230,13 +230,13 @@ namespace mu
 
             for ( int l=0; l<FMath::Max(1,(int)m_lods); ++l )
             {
-                int blocksX = FMath::DivideAndRoundUp( s[0], (uint16_t)fdata.m_pixelsPerBlockX );
-                int blocksY = FMath::DivideAndRoundUp( s[1], (uint16_t)fdata.m_pixelsPerBlockY );
+                int blocksX = FMath::DivideAndRoundUp( s[0], (uint16)fdata.m_pixelsPerBlockX );
+                int blocksY = FMath::DivideAndRoundUp( s[1], (uint16)fdata.m_pixelsPerBlockY );
 
                 res += (blocksX*blocksY) * fdata.m_bytesPerBlock;
 
-                s[0] = FMath::DivideAndRoundUp( s[0], (uint16_t)2 );
-                s[1] = FMath::DivideAndRoundUp( s[1], (uint16_t)2 );
+                s[0] = FMath::DivideAndRoundUp( s[0], (uint16)2 );
+                s[1] = FMath::DivideAndRoundUp( s[1], (uint16)2 );
             }
         }
 
@@ -256,8 +256,8 @@ namespace mu
 
             for ( int l=0; l< FMath::Max(1,(int)m_lods); ++l )
             {
-                int blocksX = FMath::DivideAndRoundUp( s[0], (uint16_t)fdata.m_pixelsPerBlockX );
-                int blocksY = FMath::DivideAndRoundUp( s[1], (uint16_t)fdata.m_pixelsPerBlockY );
+                int blocksX = FMath::DivideAndRoundUp( s[0], (uint16)fdata.m_pixelsPerBlockX );
+                int blocksY = FMath::DivideAndRoundUp( s[1], (uint16)fdata.m_pixelsPerBlockY );
 
                 if ( mip==l )
                 {
@@ -265,8 +265,8 @@ namespace mu
                     break;
                 }
 
-                s[0] = FMath::DivideAndRoundUp( s[0], (uint16_t)2 );
-                s[1] = FMath::DivideAndRoundUp( s[1], (uint16_t)2 );
+                s[0] = FMath::DivideAndRoundUp( s[0], (uint16)2 );
+                s[1] = FMath::DivideAndRoundUp( s[1], (uint16)2 );
             }
         }
 
@@ -286,13 +286,13 @@ namespace mu
 
             for ( int l=0; l<FMath::Max(1,(int)m_lods); ++l )
             {
-                int blocksX = FMath::DivideAndRoundUp( s[0], (uint16_t)fdata.m_pixelsPerBlockX );
-                int blocksY = FMath::DivideAndRoundUp( s[1], (uint16_t)fdata.m_pixelsPerBlockY );
+                int blocksX = FMath::DivideAndRoundUp( s[0], (uint16)fdata.m_pixelsPerBlockX );
+                int blocksY = FMath::DivideAndRoundUp( s[1], (uint16)fdata.m_pixelsPerBlockY );
 
                 res += (blocksX*blocksY) * fdata.m_pixelsPerBlockX * fdata.m_pixelsPerBlockY;
 
-                s[0] = FMath::DivideAndRoundUp( s[0], (uint16_t)2 );
-                s[1] = FMath::DivideAndRoundUp( s[1], (uint16_t)2 );
+                s[0] = FMath::DivideAndRoundUp( s[0], (uint16)2 );
+                s[1] = FMath::DivideAndRoundUp( s[1], (uint16)2 );
             }
         }
         else
@@ -321,16 +321,16 @@ namespace mu
 
             for ( int l=0; l<mip+1; ++l )
             {
-                int blocksX = FMath::DivideAndRoundUp( s[0], (uint16_t)fdata.m_pixelsPerBlockX );
-                int blocksY = FMath::DivideAndRoundUp( s[1], (uint16_t)fdata.m_pixelsPerBlockY );
+                int blocksX = FMath::DivideAndRoundUp( s[0], (uint16)fdata.m_pixelsPerBlockX );
+                int blocksY = FMath::DivideAndRoundUp( s[1], (uint16)fdata.m_pixelsPerBlockY );
 
                 if ( l==mip )
                 {
                     res = (blocksX*blocksY) * fdata.m_pixelsPerBlockX * fdata.m_pixelsPerBlockY;
                 }
 
-                s[0] = FMath::DivideAndRoundUp( s[0], (uint16_t)2 );
-                s[1] = FMath::DivideAndRoundUp( s[1], (uint16_t)2 );
+                s[0] = FMath::DivideAndRoundUp( s[0], (uint16)2 );
+                s[1] = FMath::DivideAndRoundUp( s[1], (uint16)2 );
             }
         }
         else
@@ -360,8 +360,8 @@ namespace mu
                 return res;
             }
 
-            s[0] = FMath::DivideAndRoundUp( s[0], (uint16_t)2 );
-            s[1] = FMath::DivideAndRoundUp( s[1], (uint16_t)2 );
+            s[0] = FMath::DivideAndRoundUp( s[0], (uint16)2 );
+            s[1] = FMath::DivideAndRoundUp( s[1], (uint16)2 );
         }
 
         return res;
@@ -397,13 +397,13 @@ namespace mu
 
 			for (int l = 0; l < mip; ++l)
 			{
-				int blocksX = FMath::DivideAndRoundUp(s[0], (uint16_t)fdata.m_pixelsPerBlockX);
-				int blocksY = FMath::DivideAndRoundUp(s[1], (uint16_t)fdata.m_pixelsPerBlockY);
+				int blocksX = FMath::DivideAndRoundUp(s[0], (uint16)fdata.m_pixelsPerBlockX);
+				int blocksY = FMath::DivideAndRoundUp(s[1], (uint16)fdata.m_pixelsPerBlockY);
 
 				pResult += (blocksX * blocksY) * fdata.m_bytesPerBlock;
 
-				s[0] = FMath::DivideAndRoundUp(s[0], (uint16_t)2);
-				s[1] = FMath::DivideAndRoundUp(s[1], (uint16_t)2);
+				s[0] = FMath::DivideAndRoundUp(s[0], (uint16)2);
+				s[1] = FMath::DivideAndRoundUp(s[1], (uint16)2);
 			}
 		}
 
@@ -442,13 +442,13 @@ namespace mu
 
 			for (int l = 0; l < mip; ++l)
 			{
-				int blocksX = FMath::DivideAndRoundUp(s[0], (uint16_t)fdata.m_pixelsPerBlockX);
-				int blocksY = FMath::DivideAndRoundUp(s[1], (uint16_t)fdata.m_pixelsPerBlockY);
+				int blocksX = FMath::DivideAndRoundUp(s[0], (uint16)fdata.m_pixelsPerBlockX);
+				int blocksY = FMath::DivideAndRoundUp(s[1], (uint16)fdata.m_pixelsPerBlockY);
 
 				pResult += (blocksX * blocksY) * fdata.m_bytesPerBlock;
 
-				s[0] = FMath::DivideAndRoundUp(s[0], (uint16_t)2);
-				s[1] = FMath::DivideAndRoundUp(s[1], (uint16_t)2);
+				s[0] = FMath::DivideAndRoundUp(s[0], (uint16)2);
+				s[1] = FMath::DivideAndRoundUp(s[1], (uint16)2);
 			}
 		}
 
@@ -639,18 +639,18 @@ namespace mu
              return;
 
          bool first = true;
-         uint16_t sx = m_size[0];
-         uint16_t sy = m_size[1];
+         uint16 sx = m_size[0];
+         uint16 sy = m_size[1];
 
          // Slow reference implementation
-         uint16_t top = 0;
-         uint16_t left = 0;
-         uint16_t right = sx - 1;
-         uint16_t bottom = sy - 1;
+         uint16 top = 0;
+         uint16 left = 0;
+         uint16 right = sx - 1;
+         uint16 bottom = sy - 1;
 
-         for ( uint16_t y=0; y<sy; ++y )
+         for ( uint16 y=0; y<sy; ++y )
          {
-             for ( uint16_t x=0; x<sx; ++x )
+             for ( uint16 x=0; x<sx; ++x )
              {
                  bool black = false;
                  switch( m_format )
@@ -734,14 +734,14 @@ namespace mu
   //      if ( !rect.size[0] || !rect.size[1] )
   //          return;
 
-  //      uint16_t sx = m_size[0];
-  //      uint16_t sy = m_size[1];
+  //      uint16 sx = m_size[0];
+  //      uint16 sy = m_size[1];
 
   //      // Somewhat faster implementation
-  //      uint16_t top = 0;
-  //      uint16_t left = 0;
-  //      uint16_t right = sx - 1;
-  //      uint16_t bottom = sy - 1;
+  //      uint16 top = 0;
+  //      uint16 left = 0;
+  //      uint16 right = sx - 1;
+  //      uint16 bottom = sy - 1;
 
   //      switch ( m_format )
   //      {
@@ -776,7 +776,7 @@ namespace mu
   //          pRow = m_data.GetData() + rowStride * currentRow;
   //          while ( currentRow <= bottom && ( left > 0 || right < ( sx - 1 ) ) )
   //          {
-  //              for ( uint16_t x = 0; x < left; ++x )
+  //              for ( uint16 x = 0; x < left; ++x )
   //              {
   //                  if ( pRow[x] )
   //                  {
@@ -784,7 +784,7 @@ namespace mu
   //                      break;
   //                  }
   //              }
-  //              for ( uint16_t x = sx - 1; x > right; --x )
+  //              for ( uint16 x = sx - 1; x > right; --x )
   //              {
   //                  if ( pRow[x] )
   //                  {
@@ -835,7 +835,7 @@ namespace mu
   //          check(bytesPerPixel<16);
   //          while ( currentRow <= bottom && (left > 0 || right<(sx-1)) )
   //          {
-  //              for ( uint16_t x = 0; x < left; ++x )
+  //              for ( uint16 x = 0; x < left; ++x )
   //              {
   //                  if ( memcmp( pRow + x * bytesPerPixel, zeroPixel, bytesPerPixel ) )
   //                  {
@@ -843,7 +843,7 @@ namespace mu
   //                      break;
   //                  }
   //              }
-  //              for ( uint16_t x = sx - 1; x > right; --x )
+  //              for ( uint16 x = sx - 1; x > right; --x )
   //              {
   //                  if ( memcmp( pRow + x * bytesPerPixel, zeroPixel, bytesPerPixel ) )
   //                  {
@@ -897,8 +897,8 @@ namespace mu
 			}
 
 			vec2<int> FinalSize = CalculateMipSize(LODSToSkip);
-			m_size[0] = uint16_t(FinalSize[0]);
-			m_size[1] = uint16_t(FinalSize[1]);
+			m_size[0] = uint16(FinalSize[0]);
+			m_size[1] = uint16(FinalSize[1]);
 			m_lods = NewLODCount;
 
 			int32 FinalDataSize = m_data.Num() - int32(DataToSkip);
@@ -924,8 +924,8 @@ namespace mu
 			}
 
 			vec2<int> FinalSize = CalculateMipSize(LODSToSkip);
-			m_size[0] = uint16_t(FinalSize[0]);
-			m_size[1] = uint16_t(FinalSize[1]);
+			m_size[0] = uint16(FinalSize[0]);
+			m_size[1] = uint16(FinalSize[1]);
 			m_lods -= LODSToSkip;
 
 			int32 FinalDataSize = m_data.Num() - int32(DataToSkip);
diff --git a/Engine/Plugins/Experimental/Mutable/Source/MutableRuntime/Private/MuR/ImagePrivate.h b/Engine/Plugins/Experimental/Mutable/Source/MutableRuntime/Private/MuR/ImagePrivate.h
--- a/Engine/Plugins/Experimental/Mutable/Source/MutableRuntime/Private/MuR/ImagePrivate.h
+++ b/Engine/Plugins/Experimental/Mutable/Source/MutableRuntime/Private/MuR/ImagePrivate.h
@@ -195,7 +195,7 @@ namespace mu
         case EImageFormat::IF_L_UBYTE:
         {
             uint8* pData = pImage->GetData();
-            uint8 v = std::min<uint8>( 255, std::max<uint8>( 0, uint8( 255.0f * c[0] ) ) );
+            uint8 v = FMath::Min<uint8>( 255, FMath::Max<uint8>( 0, uint8( 255.0f * c[0] ) ) );
             FMemory::Memset( pData, v, pixelCount );
             break;
         }
diff --git a/Engine/Plugins/Experimental/Mutable/Source/MutableRuntime/Private/MuR/ImageRLE.cpp b/Engine/Plugins/Experimental/Mutable/Source/MutableRuntime/Private/MuR/ImageRLE.cpp
--- a/Engine/Plugins/Experimental/Mutable/Source/MutableRuntime/Private/MuR/ImageRLE.cpp
+++ b/Engine/Plugins/Experimental/Mutable/Source/MutableRuntime/Private/MuR/ImageRLE.cpp
@@ -170,11 +170,11 @@ namespace mu
                              uint8* destData,
                              uint32 destDataSize )
     {
-        vector<int8_t> rle;
-        rle.reserve(  (width * rows) / 2 );
+        TArray<int8_t> rle;
+        rle.Reserve(  (width * rows) / 2 );
 
         uint32 offset = sizeof(uint32)*(rows+1);
-        rle.resize( offset );
+        rle.SetNum( offset, false );
 
         for ( int r=0; r<rows; ++r )
         {
@@ -203,7 +203,7 @@ namespace mu
                 }
 
                 // Copy block
-                rle.resize( rle.size()+4 );
+                rle.SetNum( rle.Num()+4, false);
                 FMemory::Memmove(&rle[ offset ], &zeroPixels, sizeof(uint16));
                 offset += 2;
 
@@ -212,7 +212,7 @@ namespace mu
             }
         }
 
-        if (destDataSize<rle.size())
+        if (destDataSize<(uint32)rle.Num())
         {
             // Failed
             return 0;
@@ -276,11 +276,11 @@ namespace mu
                            TArray<uint8>& destData )
     {
         // TODO: Support for compression from compressed data size, like L_RLE formats.
-        vector<int8_t> rle;
-        rle.reserve(  (width*rows) );
+        TArray<int8_t> rle;
+        rle.Reserve(  (width*rows) );
 
         const uint32* pBaseData = (const uint32*)pBaseDataByte;
-        rle.resize( rows*4 );
+        rle.SetNum( rows*4, false);
         uint32 offset = sizeof(uint32)*rows;
         for ( int r=0; r<rows; ++r )
         {
@@ -322,7 +322,7 @@ namespace mu
                 }
 
                 // Copy header
-                rle.resize( rle.size()+8 );
+                rle.SetNum( rle.Num()+8, false);
                 FMemory::Memmove(&rle[ offset ], &equal, sizeof(uint16));
                 offset += 2;
                 FMemory::Memmove(&rle[ offset ], &different, sizeof(uint16));
@@ -336,7 +336,7 @@ namespace mu
 					// If we are at the end of a row, maybe there isn't a block of 4 pixels
 					uint16 BytesToCopy = FMath::Min(different * 4 * 4, uint16(pBaseRowEnd - pDifferentPixels) * 4);
 
-					rle.resize( rle.size()+ BytesToCopy);
+					rle.SetNum( rle.Num()+ BytesToCopy, false);
                     FMemory::Memmove( &rle[offset], pDifferentPixels, BytesToCopy);
 					offset += BytesToCopy;
 				}
@@ -427,11 +427,11 @@ namespace mu
                           const uint8* pBaseDataByte,
                           TArray<uint8>& destData )
     {
-        vector<int8_t> rle;
-        rle.reserve(  (width*rows) );
+        TArray<int8_t> rle;
+        rle.Reserve(  (width*rows) );
 
         const UINT24* pBaseData = (const UINT24*)pBaseDataByte;
-        rle.resize( rows*4 );
+        rle.SetNum( rows*4, false );
         uint32 offset = sizeof(uint32)*rows;
         for ( int r=0; r<rows; ++r )
         {
@@ -473,7 +473,7 @@ namespace mu
                 }
 
                 // Copy header
-                rle.resize( rle.size()+8 );
+                rle.SetNum( rle.Num()+8, false);
                 FMemory::Memmove( &rle[offset], &equal, sizeof(uint16) );
                 offset += 2;
                 FMemory::Memmove( &rle[offset], &different, sizeof(uint16) );
@@ -487,7 +487,7 @@ namespace mu
 					// If we are at the end of a row, maybe there isn't a block of 4 pixels
 					uint16 BytesToCopy = FMath::Min(different * 4 * 3, uint16(pBaseRowEnd- pDifferentPixels)*3 );
 
-					rle.resize( rle.size()+BytesToCopy );
+					rle.SetNum( rle.Num()+BytesToCopy, false );
                     FMemory::Memmove( &rle[offset], pDifferentPixels, BytesToCopy );
 					offset += BytesToCopy;
 				}
