[mutable] Optimize mesh clip versus switching of clipping mesh.

diff --git a/Engine/Plugins/Experimental/Mutable/Source/MutableRuntime/Private/MuR/Instance.cpp b/Engine/Plugins/Experimental/Mutable/Source/MutableRuntime/Private/MuR/Instance.cpp
--- a/Engine/Plugins/Experimental/Mutable/Source/MutableRuntime/Private/MuR/Instance.cpp
+++ b/Engine/Plugins/Experimental/Mutable/Source/MutableRuntime/Private/MuR/Instance.cpp
@@ -68,7 +68,7 @@ namespace mu
 	//---------------------------------------------------------------------------------------------
 	int Instance::Private::GetLODCount() const
 	{
-		return (int)m_lods.size();
+		return m_lods.Num();
 	}
 
 
@@ -82,10 +82,10 @@ namespace mu
 	//---------------------------------------------------------------------------------------------
 	int Instance::Private::GetComponentCount( int lod ) const
 	{
-		check( lod>=0 && lod<(int)m_lods.size() );
-        if ( lod>=0 && lod<(int)m_lods.size() )
+		check( lod>=0 && lod<m_lods.Num() );
+        if ( lod>=0 && lod<m_lods.Num() )
         {
-            return (int)m_lods[lod].m_components.size();
+            return m_lods[lod].m_components.Num();
         }
 
         return 0;
@@ -95,8 +95,8 @@ namespace mu
     //---------------------------------------------------------------------------------------------
     const char* Instance::GetComponentName( int lod, int comp ) const
     {
-        if ( lod>=0 && lod<(int)m_pD->m_lods.size() &&
-             comp>=0 && comp<(int)m_pD->m_lods[lod].m_components.size() )
+        if ( lod>=0 && lod<m_pD->m_lods.Num() &&
+             comp>=0 && comp<m_pD->m_lods[lod].m_components.Num() )
         {
             return m_pD->m_lods[lod].m_components[comp].m_name.c_str();
         }
@@ -112,8 +112,8 @@ namespace mu
 	//---------------------------------------------------------------------------------------------
 	uint16 Instance::GetComponentId( int lod, int comp ) const
 	{
-		if ( lod>=0 && lod<(int)m_pD->m_lods.size() &&
-			 comp>=0 && comp<(int)m_pD->m_lods[lod].m_components.size() )
+		if ( lod>=0 && lod<m_pD->m_lods.Num() &&
+			 comp>=0 && comp<m_pD->m_lods[lod].m_components.Num() )
 		{
 			return m_pD->m_lods[lod].m_components[comp].m_id;
 		}
@@ -136,10 +136,10 @@ namespace mu
     //---------------------------------------------------------------------------------------------
     int Instance::Private::GetSurfaceCount( int lod, int comp ) const
     {
-        if ( lod>=0 && lod<(int)m_lods.size() &&
-             comp>=0 && comp<(int)m_lods[lod].m_components.size() )
+        if ( lod>=0 && lod<m_lods.Num() &&
+             comp>=0 && comp<m_lods[lod].m_components.Num() )
         {
-            return (int)m_lods[lod].m_components[comp].m_surfaces.size();
+            return m_lods[lod].m_components[comp].m_surfaces.Num();
         }
 		else
 		{
@@ -153,9 +153,9 @@ namespace mu
 	//---------------------------------------------------------------------------------------------
     const char* Instance::GetSurfaceName( int lod, int comp, int surf ) const
     {
-        if ( lod>=0 && lod<(int)m_pD->m_lods.size() &&
-             comp>=0 && comp<(int)m_pD->m_lods[lod].m_components.size() &&
-             surf>=0 && surf<(int)m_pD->m_lods[lod].m_components[comp].m_surfaces.size() )
+        if ( lod>=0 && lod<m_pD->m_lods.Num() &&
+             comp>=0 && comp<m_pD->m_lods[lod].m_components.Num() &&
+             surf>=0 && surf<m_pD->m_lods[lod].m_components[comp].m_surfaces.Num() )
         {
             return m_pD->m_lods[lod].m_components[comp].m_surfaces[surf].m_name.c_str();
         }
@@ -171,9 +171,9 @@ namespace mu
     //---------------------------------------------------------------------------------------------
     uint32_t Instance::GetSurfaceId( int lod, int comp, int surf ) const
     {
-        if ( lod>=0 && lod<(int)m_pD->m_lods.size() &&
-             comp>=0 && comp<(int)m_pD->m_lods[lod].m_components.size() &&
-             surf>=0 && surf<(int)m_pD->m_lods[lod].m_components[comp].m_surfaces.size() )
+        if ( lod>=0 && lod<m_pD->m_lods.Num() &&
+             comp>=0 && comp<m_pD->m_lods[lod].m_components.Num() &&
+             surf>=0 && surf<m_pD->m_lods[lod].m_components[comp].m_surfaces.Num() )
         {
             return m_pD->m_lods[lod].m_components[comp].m_surfaces[surf].m_internalID;
         }
@@ -189,10 +189,10 @@ namespace mu
     //---------------------------------------------------------------------------------------------
     int Instance::FindSurfaceById( int lod, int comp, uint32_t id ) const
     {
-		if (lod >= 0 && lod < (int)m_pD->m_lods.size() &&
-			comp >= 0 && comp < (int)m_pD->m_lods[lod].m_components.size())
+		if (lod >= 0 && lod < m_pD->m_lods.Num() &&
+			comp >= 0 && comp < m_pD->m_lods[lod].m_components.Num())
 		{
-			for (int i = 0; i < (int)m_pD->m_lods[lod].m_components[comp].m_surfaces.size(); ++i)
+			for (int i = 0; i < m_pD->m_lods[lod].m_components[comp].m_surfaces.Num(); ++i)
 			{
 				if (m_pD->m_lods[lod].m_components[comp].m_surfaces[i].m_internalID == id)
 				{
@@ -212,9 +212,9 @@ namespace mu
     //---------------------------------------------------------------------------------------------
     uint32_t Instance::GetSurfaceCustomId( int lod, int comp, int surf ) const
     {
-        if ( lod>=0 && lod<(int)m_pD->m_lods.size() &&
-             comp>=0 && comp<(int)m_pD->m_lods[lod].m_components.size() &&
-             surf>=0 && surf<(int)m_pD->m_lods[lod].m_components[comp].m_surfaces.size() )
+        if ( lod>=0 && lod<m_pD->m_lods.Num() &&
+             comp>=0 && comp<m_pD->m_lods[lod].m_components.Num() &&
+             surf>=0 && surf<m_pD->m_lods[lod].m_components[comp].m_surfaces.Num() )
         {
             return m_pD->m_lods[lod].m_components[comp].m_surfaces[surf].m_customID;
         }
@@ -237,10 +237,10 @@ namespace mu
 	//---------------------------------------------------------------------------------------------
     int Instance::Private::GetMeshCount( int lod, int comp ) const
 	{
-		check( lod>=0 && lod<(int)m_lods.size() );
-		check( comp>=0 && comp<(int)m_lods[lod].m_components.size() );
+		check( lod>=0 && lod<m_lods.Num() );
+		check( comp>=0 && comp<m_lods[lod].m_components.Num() );
 
-        return (int)m_lods[lod].m_components[comp].m_meshes.size();
+        return m_lods[lod].m_components[comp].m_meshes.Num();
 	}
 
 
@@ -254,11 +254,11 @@ namespace mu
 	//---------------------------------------------------------------------------------------------
     int Instance::Private::GetImageCount( int lod, int comp, int surf ) const
 	{
-		check( lod>=0 && lod<(int)m_lods.size() );
-		check( comp>=0 && comp<(int)m_lods[lod].m_components.size() );
-        check( surf>=0 && surf<(int)m_lods[lod].m_components[comp].m_surfaces.size() );
+		check( lod>=0 && lod<m_lods.Num() );
+		check( comp>=0 && comp<m_lods[lod].m_components.Num() );
+        check( surf>=0 && surf<m_lods[lod].m_components[comp].m_surfaces.Num() );
 
-        return (int)m_lods[lod].m_components[comp].m_surfaces[surf].m_images.size();
+        return m_lods[lod].m_components[comp].m_surfaces[surf].m_images.Num();
 	}
 
 
@@ -272,11 +272,11 @@ namespace mu
 	//---------------------------------------------------------------------------------------------
     int Instance::Private::GetVectorCount( int lod, int comp, int surf ) const
 	{
-		check( lod>=0 && lod<(int)m_lods.size() );
-		check( comp>=0 && comp<(int)m_lods[lod].m_components.size() );
-        check( surf>=0 && surf<(int)m_lods[lod].m_components[comp].m_surfaces.size() );
+		check( lod>=0 && lod<m_lods.Num() );
+		check( comp>=0 && comp<m_lods[lod].m_components.Num() );
+        check( surf>=0 && surf<m_lods[lod].m_components[comp].m_surfaces.Num() );
 
-        return (int)m_lods[lod].m_components[comp].m_surfaces[surf].m_vectors.size();
+        return m_lods[lod].m_components[comp].m_surfaces[surf].m_vectors.Num();
 	}
 
 
@@ -290,11 +290,11 @@ namespace mu
 	//---------------------------------------------------------------------------------------------
     int Instance::Private::GetScalarCount( int lod, int comp, int surf ) const
 	{
-		check( lod>=0 && lod<(int)m_lods.size() );
-		check( comp>=0 && comp<(int)m_lods[lod].m_components.size() );
-        check( surf>=0 && surf<(int)m_lods[lod].m_components[comp].m_surfaces.size() );
+		check( lod>=0 && lod<m_lods.Num() );
+		check( comp>=0 && comp<m_lods[lod].m_components.Num() );
+        check( surf>=0 && surf<m_lods[lod].m_components[comp].m_surfaces.Num() );
 
-        return (int)m_lods[lod].m_components[comp].m_surfaces[surf].m_scalars.size();
+        return m_lods[lod].m_components[comp].m_surfaces[surf].m_scalars.Num();
 	}
 
 
@@ -308,20 +308,20 @@ namespace mu
     //---------------------------------------------------------------------------------------------
     int Instance::Private::GetStringCount( int lod, int comp, int surf ) const
     {
-        check( lod >= 0 && lod < (int)m_lods.size() );
-        check( comp >= 0 && comp < (int)m_lods[lod].m_components.size() );
-        check( surf >= 0 && surf < (int)m_lods[lod].m_components[comp].m_surfaces.size() );
+        check( lod >= 0 && lod < m_lods.Num() );
+        check( comp >= 0 && comp < m_lods[lod].m_components.Num() );
+        check( surf >= 0 && surf < m_lods[lod].m_components[comp].m_surfaces.Num() );
 
-        return (int)m_lods[lod].m_components[comp].m_surfaces[surf].m_strings.size();
+        return m_lods[lod].m_components[comp].m_surfaces[surf].m_strings.Num();
     }
 
 
     //---------------------------------------------------------------------------------------------
     RESOURCE_ID Instance::GetMeshId( int lod, int comp, int mesh ) const
     {
-        check( lod>=0 && lod<(int)m_pD->m_lods.size() );
-        check( comp>=0 && comp<(int)m_pD->m_lods[lod].m_components.size() );
-        check( mesh>=0 && mesh<(int)m_pD->m_lods[lod].m_components[comp].m_meshes.size() );
+        check( lod>=0 && lod<m_pD->m_lods.Num() );
+        check( comp>=0 && comp<m_pD->m_lods[lod].m_components.Num() );
+        check( mesh>=0 && mesh<m_pD->m_lods[lod].m_components[comp].m_meshes.Num() );
 
         RESOURCE_ID result = m_pD->m_lods[lod].m_components[comp].m_meshes[mesh].m_meshId;
         return result;
@@ -331,9 +331,9 @@ namespace mu
 	//---------------------------------------------------------------------------------------------
     const char* Instance::GetMeshName( int lod, int comp, int mesh ) const
 	{
-		check( lod>=0 && lod<(int)m_pD->m_lods.size() );
-		check( comp>=0 && comp<(int)m_pD->m_lods[lod].m_components.size() );
-        check( mesh>=0 && mesh<(int)m_pD->m_lods[lod].m_components[comp].m_meshes.size() );
+		check( lod>=0 && lod<m_pD->m_lods.Num() );
+		check( comp>=0 && comp<m_pD->m_lods[lod].m_components.Num() );
+        check( mesh>=0 && mesh<m_pD->m_lods[lod].m_components[comp].m_meshes.Num() );
 
         const char* strResult = m_pD->m_lods[lod].m_components[comp].m_meshes[mesh].m_name.c_str();
 		return strResult;
@@ -343,10 +343,10 @@ namespace mu
 	//---------------------------------------------------------------------------------------------
     RESOURCE_ID Instance::GetImageId( int lod, int comp, int surf, int img ) const
 	{
-		check( lod>=0 && lod<(int)m_pD->m_lods.size() );
-		check( comp>=0 && comp<(int)m_pD->m_lods[lod].m_components.size() );
-        check( surf>=0 && surf<(int)m_pD->m_lods[lod].m_components[comp].m_surfaces.size() );
-        check( img>=0 && img<(int)m_pD->m_lods[lod].m_components[comp].m_surfaces[surf].m_images.size() );
+		check( lod>=0 && lod<m_pD->m_lods.Num() );
+		check( comp>=0 && comp<m_pD->m_lods[lod].m_components.Num() );
+        check( surf>=0 && surf<m_pD->m_lods[lod].m_components[comp].m_surfaces.Num() );
+        check( img>=0 && img<m_pD->m_lods[lod].m_components[comp].m_surfaces[surf].m_images.Num() );
 
         RESOURCE_ID result = m_pD->m_lods[lod].m_components[comp].m_surfaces[surf].m_images[img].m_imageId;
         return result;
@@ -356,10 +356,10 @@ namespace mu
 	//---------------------------------------------------------------------------------------------
     const char* Instance::GetImageName( int lod, int comp, int surf, int img ) const
 	{
-		check( lod>=0 && lod<(int)m_pD->m_lods.size() );
-		check( comp>=0 && comp<(int)m_pD->m_lods[lod].m_components.size() );
-        check( surf>=0 && surf<(int)m_pD->m_lods[lod].m_components[comp].m_surfaces.size() );
-        check( img>=0 && img<(int)m_pD->m_lods[lod].m_components[comp].m_surfaces[surf].m_images.size() );
+		check( lod>=0 && lod<m_pD->m_lods.Num() );
+		check( comp>=0 && comp<m_pD->m_lods[lod].m_components.Num() );
+        check( surf>=0 && surf<m_pD->m_lods[lod].m_components[comp].m_surfaces.Num() );
+        check( img>=0 && img<m_pD->m_lods[lod].m_components[comp].m_surfaces[surf].m_images.Num() );
 
         const char* strResult = m_pD->m_lods[lod].m_components[comp].m_surfaces[surf].m_images[img].m_name.c_str();
 		return strResult;
@@ -370,10 +370,10 @@ namespace mu
     void Instance::GetVector( int lod, int comp, int surf, int vec,
 							  float* pX, float* pY, float* pZ, float* pW ) const
 	{
-		check( lod>=0 && lod<(int)m_pD->m_lods.size() );
-		check( comp>=0 && comp<(int)m_pD->m_lods[lod].m_components.size() );
-        check( surf>=0 && surf<(int)m_pD->m_lods[lod].m_components[comp].m_surfaces.size() );
-        check( vec>=0 && vec<(int)m_pD->m_lods[lod].m_components[comp].m_surfaces[surf].m_vectors.size() );
+		check( lod>=0 && lod<m_pD->m_lods.Num() );
+		check( comp>=0 && comp<m_pD->m_lods[lod].m_components.Num() );
+        check( surf>=0 && surf<m_pD->m_lods[lod].m_components[comp].m_surfaces.Num() );
+        check( vec>=0 && vec<m_pD->m_lods[lod].m_components[comp].m_surfaces[surf].m_vectors.Num() );
 
         vec4<float> r = m_pD->m_lods[lod].m_components[comp].m_surfaces[surf].m_vectors[vec].m_vec;
 
@@ -387,10 +387,10 @@ namespace mu
 	//---------------------------------------------------------------------------------------------
     const char* Instance::GetVectorName( int lod, int comp, int surf, int vec ) const
 	{
-		check( lod>=0 && lod<(int)m_pD->m_lods.size() );
-		check( comp>=0 && comp<(int)m_pD->m_lods[lod].m_components.size() );
-        check( surf>=0 && surf<(int)m_pD->m_lods[lod].m_components[comp].m_surfaces.size() );
-        check( vec>=0 && vec<(int)m_pD->m_lods[lod].m_components[comp].m_surfaces[surf].m_vectors.size() );
+		check( lod>=0 && lod<m_pD->m_lods.Num() );
+		check( comp>=0 && comp<m_pD->m_lods[lod].m_components.Num() );
+        check( surf>=0 && surf<m_pD->m_lods[lod].m_components[comp].m_surfaces.Num() );
+        check( vec>=0 && vec<m_pD->m_lods[lod].m_components[comp].m_surfaces[surf].m_vectors.Num() );
 
         const char* strResult = m_pD->m_lods[lod].m_components[comp].m_surfaces[surf].m_vectors[vec].m_name.c_str();
 		return strResult;
@@ -400,10 +400,10 @@ namespace mu
 	//---------------------------------------------------------------------------------------------
     float Instance::GetScalar( int lod, int comp, int surf, int sca ) const
 	{
-		check( lod>=0 && lod<(int)m_pD->m_lods.size() );
-		check( comp>=0 && comp<(int)m_pD->m_lods[lod].m_components.size() );
-        check( surf>=0 && surf<(int)m_pD->m_lods[lod].m_components[comp].m_surfaces.size() );
-        check( sca>=0 && sca<(int)m_pD->m_lods[lod].m_components[comp].m_surfaces[surf].m_scalars.size() );
+		check( lod>=0 && lod<m_pD->m_lods.Num() );
+		check( comp>=0 && comp<m_pD->m_lods[lod].m_components.Num() );
+        check( surf>=0 && surf<m_pD->m_lods[lod].m_components[comp].m_surfaces.Num() );
+        check( sca>=0 && sca<m_pD->m_lods[lod].m_components[comp].m_surfaces[surf].m_scalars.Num() );
 
         float result = m_pD->m_lods[lod].m_components[comp].m_surfaces[surf].m_scalars[sca].m_scalar;
 		return result;
@@ -413,10 +413,10 @@ namespace mu
 	//---------------------------------------------------------------------------------------------
     const char* Instance::GetScalarName( int lod, int comp, int surf, int sca ) const
 	{
-		check( lod>=0 && lod<(int)m_pD->m_lods.size() );
-		check( comp>=0 && comp<(int)m_pD->m_lods[lod].m_components.size() );
-        check( surf>=0 && surf<(int)m_pD->m_lods[lod].m_components[comp].m_surfaces.size() );
-        check( sca>=0 && sca<(int)m_pD->m_lods[lod].m_components[comp].m_surfaces[surf].m_scalars.size() );
+		check( lod>=0 && lod<m_pD->m_lods.Num() );
+		check( comp>=0 && comp<m_pD->m_lods[lod].m_components.Num() );
+        check( surf>=0 && surf<m_pD->m_lods[lod].m_components[comp].m_surfaces.Num() );
+        check( sca>=0 && sca<m_pD->m_lods[lod].m_components[comp].m_surfaces[surf].m_scalars.Num() );
 
         const char* strResult = m_pD->m_lods[lod].m_components[comp].m_surfaces[surf].m_scalars[sca].m_name.c_str();
 		return strResult;
@@ -426,14 +426,14 @@ namespace mu
     //---------------------------------------------------------------------------------------------
     const char* Instance::GetString( int lod, int comp, int surf, int str ) const
     {
-        check( lod >= 0 && lod < (int)m_pD->m_lods.size() );
-        check( comp >= 0 && comp < (int)m_pD->m_lods[lod].m_components.size() );
+        check( lod >= 0 && lod < m_pD->m_lods.Num() );
+        check( comp >= 0 && comp < m_pD->m_lods[lod].m_components.Num() );
         check( surf >= 0 &&
-                        surf < (int)m_pD->m_lods[lod].m_components[comp].m_surfaces.size() );
+                        surf < m_pD->m_lods[lod].m_components[comp].m_surfaces.Num() );
 
         bool valid =
             str >= 0 &&
-            str < (int)m_pD->m_lods[lod].m_components[comp].m_surfaces[surf].m_strings.size();
+            str < m_pD->m_lods[lod].m_components[comp].m_surfaces[surf].m_strings.Num();
         check(valid);
 
         if (valid)
@@ -453,13 +453,13 @@ namespace mu
     //---------------------------------------------------------------------------------------------
     const char* Instance::GetStringName( int lod, int comp, int surf, int str ) const
     {
-        check( lod >= 0 && lod < (int)m_pD->m_lods.size() );
-        check( comp >= 0 && comp < (int)m_pD->m_lods[lod].m_components.size() );
+        check( lod >= 0 && lod < m_pD->m_lods.Num() );
+        check( comp >= 0 && comp < m_pD->m_lods[lod].m_components.Num() );
         check( surf >= 0 &&
-                        surf < (int)m_pD->m_lods[lod].m_components[comp].m_surfaces.size() );
+                        surf < m_pD->m_lods[lod].m_components[comp].m_surfaces.Num() );
         bool valid =
             str >= 0 &&
-            str < (int)m_pD->m_lods[lod].m_components[comp].m_surfaces[surf].m_strings.size();
+            str < m_pD->m_lods[lod].m_components[comp].m_surfaces[surf].m_strings.Num();
         check( valid );
 
         if (valid)
@@ -476,8 +476,8 @@ namespace mu
     //---------------------------------------------------------------------------------------------
 	int Instance::Private::AddLOD()
 	{
-		int result = (int)m_lods.size();
-		m_lods.push_back( INSTANCE_LOD() );
+		int result = m_lods.Num();
+		m_lods.Add( INSTANCE_LOD() );
 		return result;
 	}
 
@@ -491,8 +491,8 @@ namespace mu
             AddLOD();
         }
 
-        int result = (int)m_lods[lod].m_components.size();
-        m_lods[lod].m_components.push_back( INSTANCE_COMPONENT() );
+        int result = m_lods[lod].m_components.Num();
+        m_lods[lod].m_components.Add( INSTANCE_COMPONENT() );
         return result;
     }
 
@@ -510,8 +510,8 @@ namespace mu
             AddComponent(lod);
         }
 
-        int result = (int)m_lods[lod].m_components[comp].m_surfaces.size();
-        m_lods[lod].m_components[comp].m_surfaces.push_back( INSTANCE_SURFACE() );
+        int result = m_lods[lod].m_components[comp].m_surfaces.Num();
+        m_lods[lod].m_components[comp].m_surfaces.Add( INSTANCE_SURFACE() );
         return result;
     }
 
@@ -586,8 +586,8 @@ namespace mu
 		}
 
 		INSTANCE_COMPONENT& component = m_lods[lod].m_components[comp];
-        int result = (int)component.m_meshes.size();
-        component.m_meshes.emplace_back( meshId, strName );
+        int result = component.m_meshes.Num();
+        component.m_meshes.Emplace( meshId, strName );
 
 		return result;
 	}
@@ -613,8 +613,8 @@ namespace mu
         }
 
         INSTANCE_SURFACE& surface = m_lods[lod].m_components[comp].m_surfaces[surf];
-        int result = (int)surface.m_images.size();
-        surface.m_images.emplace_back( imageId, strName );
+        int result = surface.m_images.Num();
+        surface.m_images.Emplace( imageId, strName );
 
         return result;
 	}
@@ -638,8 +638,8 @@ namespace mu
         }
 
         INSTANCE_SURFACE& surface = m_lods[lod].m_components[comp].m_surfaces[surf];
-        int result = (int)surface.m_vectors.size();
-        surface.m_vectors.emplace_back( vec, strName );
+        int result = surface.m_vectors.Num();
+        surface.m_vectors.Emplace( vec, strName );
 
         return result;
 	}
@@ -663,8 +663,8 @@ namespace mu
         }
 
         INSTANCE_SURFACE& surface = m_lods[lod].m_components[comp].m_surfaces[surf];
-        int result = (int)surface.m_scalars.size();
-        surface.m_scalars.emplace_back( sca, strName );
+        int result = surface.m_scalars.Num();
+        surface.m_scalars.Emplace( sca, strName );
 
         return result;
     }
@@ -689,8 +689,8 @@ namespace mu
         }
 
         INSTANCE_SURFACE& surface = m_lods[lod].m_components[comp].m_surfaces[surf];
-        int result = (int)surface.m_strings.size();
-        surface.m_strings.emplace_back( str, strName );
+        int result = surface.m_strings.Num();
+        surface.m_strings.Emplace( str, strName );
 
         return result;
     }
diff --git a/Engine/Plugins/Experimental/Mutable/Source/MutableRuntime/Private/MuR/InstancePrivate.h b/Engine/Plugins/Experimental/Mutable/Source/MutableRuntime/Private/MuR/InstancePrivate.h
--- a/Engine/Plugins/Experimental/Mutable/Source/MutableRuntime/Private/MuR/InstancePrivate.h
+++ b/Engine/Plugins/Experimental/Mutable/Source/MutableRuntime/Private/MuR/InstancePrivate.h
@@ -33,7 +33,7 @@ namespace mu
 			string m_name;
 		};
 
-		vector<IMAGE> m_images;
+		TArray<IMAGE> m_images;
 
 		struct VECTOR
 		{
@@ -50,7 +50,7 @@ namespace mu
 			string m_name;
 		};
 
-		vector<VECTOR> m_vectors;
+		TArray<VECTOR> m_vectors;
 
         struct SCALAR
         {
@@ -67,7 +67,7 @@ namespace mu
             string m_name;
         };
 
-        vector<SCALAR> m_scalars;
+		TArray<SCALAR> m_scalars;
 
         struct STRING
         {
@@ -84,7 +84,7 @@ namespace mu
             string m_name;
         };
 
-        vector<STRING> m_strings;
+		TArray<STRING> m_strings;
     };
 
 
@@ -108,16 +108,16 @@ namespace mu
             RESOURCE_ID m_meshId;
 			string m_name;
 		};
-		vector<MESH> m_meshes;
+		TArray<MESH> m_meshes;
 
 		// The order must match the meshes surfaces
-		vector<INSTANCE_SURFACE> m_surfaces;
+		TArray<INSTANCE_SURFACE> m_surfaces;
 	};
 
 
     struct INSTANCE_LOD
     {
-        vector<INSTANCE_COMPONENT> m_components;
+		TArray<INSTANCE_COMPONENT> m_components;
     };
 
 
@@ -129,7 +129,7 @@ namespace mu
         Instance::ID m_id;
 
 		//!
-		vector<INSTANCE_LOD> m_lods;
+		TArray<INSTANCE_LOD> m_lods;
 
         Private()
         {
diff --git a/Engine/Plugins/Experimental/Mutable/Source/MutableRuntime/Private/MuR/Layout.cpp b/Engine/Plugins/Experimental/Mutable/Source/MutableRuntime/Private/MuR/Layout.cpp
--- a/Engine/Plugins/Experimental/Mutable/Source/MutableRuntime/Private/MuR/Layout.cpp
+++ b/Engine/Plugins/Experimental/Mutable/Source/MutableRuntime/Private/MuR/Layout.cpp
@@ -67,8 +67,8 @@ namespace mu {
 	void Layout::SetGridSize( int sizeX, int sizeY )
 	{
 		check( sizeX>=0 && sizeY>=0 );
-		m_size[0] = (uint16_t)sizeX;
-		m_size[1] = (uint16_t)sizeY;
+		m_size[0] = (uint16)sizeX;
+		m_size[1] = (uint16)sizeY;
 	}
 
 
@@ -89,8 +89,8 @@ namespace mu {
 	void Layout::SetMaxGridSize(int sizeX, int sizeY)
 	{
 		check(sizeX >= 0 && sizeY >= 0);
-		m_maxsize[0] = (uint16_t)sizeX;
-		m_maxsize[1] = (uint16_t)sizeY;
+		m_maxsize[0] = (uint16)sizeX;
+		m_maxsize[1] = (uint16)sizeY;
 	}
 
 
@@ -145,8 +145,8 @@ namespace mu {
 		check( index >=0 && index < m_blocks.Num() );
 
 		// Keeps the id
-		m_blocks[index].m_min = vec2<uint16_t>((uint16_t)minx, (uint16_t)miny);
-		m_blocks[index].m_size = vec2<uint16_t>((uint16_t)sizex, (uint16_t)sizey);
+		m_blocks[index].m_min = vec2<uint16>((uint16)minx, (uint16)miny);
+		m_blocks[index].m_size = vec2<uint16>((uint16)sizex, (uint16)sizey);
 	}
 
 
@@ -211,7 +211,7 @@ namespace mu {
 	bool Layout::IsSingleBlockAndFull() const
 	{
 		if (m_blocks.Num() == 1
-			&& m_blocks[0].m_min == vec2<uint16_t>(0, 0)
+			&& m_blocks[0].m_min == vec2<uint16>(0, 0)
 			&& m_blocks[0].m_size == m_size)
 		{
 			return true;
diff --git a/Engine/Plugins/Experimental/Mutable/Source/MutableRuntime/Private/MuR/MemoryPrivate.h b/Engine/Plugins/Experimental/Mutable/Source/MutableRuntime/Private/MuR/MemoryPrivate.h
--- a/Engine/Plugins/Experimental/Mutable/Source/MutableRuntime/Private/MuR/MemoryPrivate.h
+++ b/Engine/Plugins/Experimental/Mutable/Source/MutableRuntime/Private/MuR/MemoryPrivate.h
@@ -4,112 +4,3 @@
 
 #include "MuR/Platform.h"
 #include "MuR/MutableMemory.h"
-
-#include <unordered_map>
-#include <map>
-#include <set>
-
-
-namespace mu
-{
-
-    //!
-	MUTABLERUNTIME_API void mutable_memlog();
-
-
-    //! Custom memory allocator to be used by STL continers inside the library.
-    //! It uses the functions from the memory hooks
-    template <class T>
-    class MemAlloc
-    {
-    public:
-
-        // type definitions
-        using value_type = T;
-        using pointer = T*;
-        using const_pointer = const T*;
-        using reference = T&;
-        using const_reference = const T&;
-        using size_type = std::size_t;
-        using difference_type = std::ptrdiff_t;
-
-        // constructors and destructor
-        // - nothing to do because the allocator has no state
-        MemAlloc() noexcept = default;
-
-        template <class U>
-        MemAlloc( const MemAlloc<U>& ) noexcept
-        {
-        }
-
-        // allocate but don't initialize num elements of type T
-        pointer allocate (size_type num, const void* = nullptr)
-        {
-            auto ret = (pointer)(mutable_malloc(num*sizeof(T)));
-            return ret;
-        }
-
-        // deallocate storage p of deleted elements
-        void deallocate (pointer p, size_type)
-        {
-			mutable_free((void*)p);
-        }
-    };
-
-    //! Return that all specializations of this allocator are interchangeable
-    template<class T1, class T2>
-    bool operator==( const MemAlloc<T1>&, const MemAlloc<T2>& ) noexcept
-    {
-        return true;
-    }
-
-    template<class T1, class T2>
-    bool operator!=( const MemAlloc<T1>& a, const MemAlloc<T2>& b ) noexcept
-    {
-        return !(a==b);
-    }
-
-
-    //! STL-like containers using this allocator
-    template< typename T >
-    using vector = std::vector<T, MemAlloc<T> >;
-
-    template< typename T >
-    using basic_string = std::basic_string<T, std::char_traits<T>, MemAlloc<T> >;
-
-    using string = std::basic_string<char, std::char_traits<char>, MemAlloc<char> >;
-
-    template< typename K, typename T >
-    using map = std::map< K, T, std::less<K>, MemAlloc< std::pair<const K,T> > >;
-
-    template< typename T >
-    using set = std::set< T, std::less<T>, MemAlloc<T> >;
-
-    template< typename T >
-    using multiset = std::multiset< T, std::less<T>, MemAlloc<T> >;
-
-    template< typename K, typename T >
-    using unordered_map = std::unordered_map< K, T, std::hash<K>, std::equal_to<K>, MemAlloc< std::pair<const K,T> > >;
-
-    template< typename K, typename T >
-    using pair = std::pair<K, T>;
-
-}
-
-
-#ifdef __GNUC__
-#include <string>
-namespace std
-{
-    // This is necessary for GCC, but it may be inefficient, causing additional string copies.
-    // In theory it can be avoided in C++0x17 with the use of string_view
-    template<>
-    struct hash<mu::string>
-    {
-        std::size_t operator()(mu::string const& s) const noexcept
-        {
-            return std::hash<std::basic_string<char, std::char_traits<char>>>{}(std::basic_string<char, std::char_traits<char>>(s.c_str()));
-        }
-    };
-}
-#endif //#ifdef __GNUC__
diff --git a/Engine/Plugins/Experimental/Mutable/Source/MutableRuntime/Private/MuR/Mesh.cpp b/Engine/Plugins/Experimental/Mutable/Source/MutableRuntime/Private/MuR/Mesh.cpp
--- a/Engine/Plugins/Experimental/Mutable/Source/MutableRuntime/Private/MuR/Mesh.cpp
+++ b/Engine/Plugins/Experimental/Mutable/Source/MutableRuntime/Private/MuR/Mesh.cpp
@@ -743,10 +743,10 @@ bool Mesh::IsSameVertex
 //---------------------------------------------------------------------------------------------
 bool Mesh::VERTEX_MATCH_MAP::Matches(int v, int ov) const
 {
-	if (v >= 0 && v < (int)m_firstMatch.size())
+	if (v >= 0 && v < (int)m_firstMatch.Num())
 	{
 		int start = m_firstMatch[v];
-		int end = v + 1 < (int)m_firstMatch.size() ? m_firstMatch[v + 1] : (int)m_matches.size();
+		int end = v + 1 < m_firstMatch.Num() ? m_firstMatch[v + 1] : m_matches.Num();
 		bool res = false;
 
 		while (!res && start < end)
@@ -774,8 +774,8 @@ void Mesh::GetVertexMap
     ) const
 {
     int vertexCount = m_VertexBuffers.GetElementCount();
-    vertexMap.m_firstMatch.resize( vertexCount );
-    vertexMap.m_matches.reserve( vertexCount+(vertexCount>>2) );
+    vertexMap.m_firstMatch.SetNum( vertexCount );
+    vertexMap.m_matches.SetNum( vertexCount+(vertexCount>>2) );
 
     int otherVertexCount = other.m_VertexBuffers.GetElementCount();
 
@@ -806,10 +806,10 @@ void Mesh::GetVertexMap
     rangeMin -= tolerance;
     rangeMax += tolerance;
 
-    vector<int> buckets[MUTABLE_NUM_BUCKETS];
+    TArray<int> buckets[MUTABLE_NUM_BUCKETS];
     for ( int b=0; b<MUTABLE_NUM_BUCKETS; ++b )
     {
-        buckets[b].reserve( otherVertexCount/MUTABLE_NUM_BUCKETS*2 );
+        buckets[b].Reserve( otherVertexCount/MUTABLE_NUM_BUCKETS*2 );
     }
 
     float bucketSize = (rangeMax-rangeMin)/float(MUTABLE_NUM_BUCKETS);
@@ -819,15 +819,15 @@ void Mesh::GetVertexMap
         float v = (*itop)[MUTABLE_BUCKET_CHANNEL];
 
         int bucket0 = int( floor( (v-tolerance-rangeMin)/bucketSize ) );
-        bucket0 = std::min( MUTABLE_NUM_BUCKETS-1, std::max( 0, bucket0 ) );
-        buckets[bucket0].push_back(ov);
+        bucket0 = FMath::Min( MUTABLE_NUM_BUCKETS-1, FMath::Max( 0, bucket0 ) );
+        buckets[bucket0].Add(ov);
 
         int bucket1 = int( floor( (v+tolerance-rangeMin)/bucketSize ) );
-        bucket1 = std::min( MUTABLE_NUM_BUCKETS-1, std::max( 0, bucket1 ) );
+        bucket1 = FMath::Min( MUTABLE_NUM_BUCKETS-1, FMath::Max( 0, bucket1 ) );
 
         if (bucket1!=bucket0)
         {
-            buckets[bucket1].push_back(ov);
+            buckets[bucket1].Add(ov);
         }
 
         ++itop;
@@ -868,14 +868,14 @@ void Mesh::GetVertexMap
         // Use buckets
         for ( int v=0; v<vertexCount; ++v )
         {
-            vertexMap.m_firstMatch[v] = (int)vertexMap.m_matches.size();
+            vertexMap.m_firstMatch[v] = (int)vertexMap.m_matches.Num();
 
             float vbucket = (*itp)[MUTABLE_BUCKET_CHANNEL];
             int bucket = int( floor( (vbucket-rangeMin)/bucketSize ) );
 
             if (bucket>=0 && bucket<MUTABLE_NUM_BUCKETS)
             {
-                int bucketVertexCount = (int)buckets[bucket].size();
+                int bucketVertexCount = (int)buckets[bucket].Num();
                 for ( int ov=0; ov<bucketVertexCount; ++ov )
                 {
                     int otherVertexIndex = buckets[bucket][ov];
@@ -890,7 +890,7 @@ void Mesh::GetVertexMap
 
                     if ( same )
                     {
-                        vertexMap.m_matches.push_back( otherVertexIndex );
+                        vertexMap.m_matches.Add( otherVertexIndex );
                     }
                 }
             }
@@ -904,13 +904,13 @@ void Mesh::GetVertexMap
         // Slow generic way
         for ( int v=0; v<vertexCount; ++v )
         {
-            vertexMap.m_firstMatch[v] = (int)vertexMap.m_matches.size();
+            vertexMap.m_firstMatch[v] = (int)vertexMap.m_matches.Num();
             for ( int ov=0; ov<otherVertexCount; ++ov )
             {
                 // TODO: Optimize this by not using IsSameVertex
                 if ( IsSameVertex( v, other, ov, tolerance ) )
                 {
-                    vertexMap.m_matches.push_back(ov);
+                    vertexMap.m_matches.Add(ov);
                 }
             }
         }
@@ -1056,8 +1056,8 @@ void Mesh::CheckIntegrity() const
                         }
                         case MBF_UINT16:
                         {
-                            uint16_t index = *(const uint16_t*)pData;
-                            check( index < uint16_t( elemCount ) );
+                            uint16 index = *(const uint16*)pData;
+                            check( index < uint16( elemCount ) );
                             break;
                         }
                         case MBF_UINT8:
@@ -1135,7 +1135,7 @@ void Mesh::CheckIntegrity() const
                     {
                         for ( int d = 0; d < components; ++d )
                         {
-                            uint16_t index = *(uint16_t*)pData;
+                            uint16 index = *(uint16*)pData;
                             check( index < uint64_t( boneCount ) );
                             pData+=2;
                         }
@@ -1378,11 +1378,11 @@ namespace
                         case MBF_UINT32:
                         case MBF_NUINT32: out += FString::Printf(TEXT("%d"), *(const uint32_t*)pChanData); pChanData +=4; break;
                         case MBF_UINT16:
-                        case MBF_NUINT16: out += FString::Printf(TEXT("%d"), *(const uint16_t*)pChanData); pChanData +=2; break;
+                        case MBF_NUINT16: out += FString::Printf(TEXT("%d"), *(const uint16*)pChanData); pChanData +=2; break;
                         case MBF_UINT8:
                         case MBF_NUINT8: out += FString::Printf(TEXT("%d"), *(const uint8_t*)pChanData); pChanData +=1; break;
 						case MBF_FLOAT32: out += FString::Printf(TEXT("%.3f"), *(const float*)pChanData); pChanData += 4; break;
-                        case MBF_FLOAT16: out += FString::Printf(TEXT("%d"), *(const uint16_t*)pChanData); pChanData +=2; break;
+                        case MBF_FLOAT16: out += FString::Printf(TEXT("%d"), *(const uint16*)pChanData); pChanData +=2; break;
                         default: break;
                         }
                         out += ",";
diff --git a/Engine/Plugins/Experimental/Mutable/Source/MutableRuntime/Private/MuR/MeshBufferSet.cpp b/Engine/Plugins/Experimental/Mutable/Source/MutableRuntime/Private/MuR/MeshBufferSet.cpp
--- a/Engine/Plugins/Experimental/Mutable/Source/MutableRuntime/Private/MuR/MeshBufferSet.cpp
+++ b/Engine/Plugins/Experimental/Mutable/Source/MutableRuntime/Private/MuR/MeshBufferSet.cpp
@@ -173,7 +173,7 @@ namespace mu
 			chan.m_semantic = pSemantics ? pSemantics[c] : MBS_NONE;
 			chan.m_semanticIndex = pSemanticIndices ? ((uint8_t)pSemanticIndices[c]) : 0;
 			chan.m_format = pFormats ? pFormats[c] : MBF_NONE;
-			chan.m_componentCount = pComponentCount ? ((uint16_t)pComponentCount[c]) : 0;
+			chan.m_componentCount = pComponentCount ? ((uint16)pComponentCount[c]) : 0;
 			chan.m_offset = pOffsets ? ((uint8_t)pOffsets[c]) : 0;
 
 			minElemSize = FMath::Max
@@ -222,7 +222,7 @@ namespace mu
 		chan.m_semantic = semantic;
 		chan.m_semanticIndex = uint8_t(semanticIndex);
 		chan.m_format = format;
-		chan.m_componentCount = uint16_t(componentCount);
+		chan.m_componentCount = uint16(componentCount);
 		chan.m_offset = uint8(offset);
 	}
 
diff --git a/Engine/Plugins/Experimental/Mutable/Source/MutableRuntime/Private/MuR/MeshPrivate.h b/Engine/Plugins/Experimental/Mutable/Source/MutableRuntime/Private/MuR/MeshPrivate.h
--- a/Engine/Plugins/Experimental/Mutable/Source/MutableRuntime/Private/MuR/MeshPrivate.h
+++ b/Engine/Plugins/Experimental/Mutable/Source/MutableRuntime/Private/MuR/MeshPrivate.h
@@ -92,7 +92,7 @@ namespace mu
 	//---------------------------------------------------------------------------------------------
 	class UntypedMeshBufferIteratorConst;
 
-	class MUTABLERUNTIME_API UntypedMeshBufferIterator : public Base
+	class MUTABLERUNTIME_API UntypedMeshBufferIterator 
 	{
 	public:
 		inline UntypedMeshBufferIterator()
@@ -364,7 +364,7 @@ namespace mu
 	//---------------------------------------------------------------------------------------------
 	//! Class to iterate a specific buffer channel of unknown type
 	//---------------------------------------------------------------------------------------------
-	class MUTABLERUNTIME_API UntypedMeshBufferIteratorConst : public Base
+	class MUTABLERUNTIME_API UntypedMeshBufferIteratorConst
 	{
 	public:
 
diff --git a/Engine/Plugins/Experimental/Mutable/Source/MutableRuntime/Private/MuR/Model.cpp b/Engine/Plugins/Experimental/Mutable/Source/MutableRuntime/Private/MuR/Model.cpp
--- a/Engine/Plugins/Experimental/Mutable/Source/MutableRuntime/Private/MuR/Model.cpp
+++ b/Engine/Plugins/Experimental/Mutable/Source/MutableRuntime/Private/MuR/Model.cpp
@@ -95,7 +95,7 @@ namespace mu
             countPerType[(int)type]++;
         }
 
-        vector< pair<uint64,OP_TYPE> > sorted((int)OP_TYPE::COUNT);
+		TArray< TPair<uint64,OP_TYPE> > sorted((int)OP_TYPE::COUNT);
         for (int i=0; i<(int)OP_TYPE::COUNT; ++i)
         {
             sorted[i].second = (OP_TYPE)i;
@@ -107,10 +107,10 @@ namespace mu
             return a.first>b.first;
         });
 
-        UE_LOG(LogMutableCore,Log, TEXT("Op histogram (%llu ops):"), m_opAddress.size());
+        UE_LOG(LogMutableCore,Log, TEXT("Op histogram (%llu ops):"), m_opAddress.Num());
         for(int i=0; i<8; ++i)
         {
-            float p = sorted[i].first/float(m_opAddress.size())*100.0f;
+            float p = sorted[i].first/float(m_opAddress.Num())*100.0f;
             UE_LOG(LogMutableCore,Log, TEXT("  %3.2f%% : %d"), p, (int)sorted[i].second );
         }
 #endif
@@ -332,7 +332,7 @@ namespace mu
 
         if (m_pD)
         {
-            m_pD->m_generatedResources.clear();
+            m_pD->m_generatedResources.Empty();
         }
     }
 
@@ -346,8 +346,8 @@ namespace mu
 
         pRes->GetPrivate()->m_pModel = this;
 
-        pRes->GetPrivate()->m_values.resize( m_pD->m_program.m_parameters.size() );
-        for ( std::size_t p=0; p<m_pD->m_program.m_parameters.size(); ++p )
+        pRes->GetPrivate()->m_values.SetNum( m_pD->m_program.m_parameters.Num() );
+        for ( int32 p=0; p<m_pD->m_program.m_parameters.Num(); ++p )
         {
             pRes->GetPrivate()->m_values[p] = m_pD->m_program.m_parameters[p].m_defaultValue;
         }
@@ -414,7 +414,7 @@ namespace mu
     //---------------------------------------------------------------------------------------------
     int Model::GetStateCount() const
     {
-        return (int)m_pD->m_program.m_states.size();
+        return (int)m_pD->m_program.m_states.Num();
     }
 
 
@@ -423,7 +423,7 @@ namespace mu
     {
         const char* strRes = 0;
 
-        if ( index>=0 && index<(int)m_pD->m_program.m_states.size() )
+        if ( index>=0 && index<(int)m_pD->m_program.m_states.Num() )
         {
             strRes = m_pD->m_program.m_states[index].m_name.c_str();
         }
@@ -437,7 +437,7 @@ namespace mu
     {
         int res = -1;
 
-        for ( int i=0; res<0 && i<(int)m_pD->m_program.m_states.size(); ++i )
+        for ( int i=0; res<0 && i<(int)m_pD->m_program.m_states.Num(); ++i )
         {
             if ( m_pD->m_program.m_states[i].m_name == strName )
             {
@@ -454,9 +454,9 @@ namespace mu
     {
         int res = -1;
 
-        if ( stateIndex>=0 && stateIndex<(int)m_pD->m_program.m_states.size() )
+        if ( stateIndex>=0 && stateIndex<(int)m_pD->m_program.m_states.Num() )
         {
-            res = (int)m_pD->m_program.m_states[stateIndex].m_runtimeParameters.size();
+            res = (int)m_pD->m_program.m_states[stateIndex].m_runtimeParameters.Num();
         }
 
         return res;
@@ -468,10 +468,10 @@ namespace mu
     {
         int res = -1;
 
-        if ( stateIndex>=0 && stateIndex<(int)m_pD->m_program.m_states.size() )
+        if ( stateIndex>=0 && stateIndex<(int)m_pD->m_program.m_states.Num() )
         {
             const PROGRAM::STATE& state = m_pD->m_program.m_states[stateIndex];
-            if ( paramIndex>=0 && paramIndex<(int)state.m_runtimeParameters.size() )
+            if ( paramIndex>=0 && paramIndex<(int)state.m_runtimeParameters.Num() )
             {
                 res = (int)state.m_runtimeParameters[paramIndex];
             }
@@ -482,24 +482,24 @@ namespace mu
 
 
     //---------------------------------------------------------------------------------------------
-    static size_t AddMultiValueKeys( vector<uint8_t>& parameterValuesBlob, size_t pos,
-                                  const map< vector<int>, PARAMETER_VALUE >& multi )
+    static size_t AddMultiValueKeys(TArray<uint8_t>& parameterValuesBlob, size_t pos,
+                                  const TMap< TArray<int>, PARAMETER_VALUE >& multi )
     {
-        parameterValuesBlob.resize( pos+4 );
-        uint32 s = uint32( multi.size() );
-        memcpy( &parameterValuesBlob[pos], &s, sizeof(uint32) );
+        parameterValuesBlob.SetNum( pos+4 );
+        uint32 s = uint32( multi.Num() );
+        FMemory::Memcpy( &parameterValuesBlob[pos], &s, sizeof(uint32) );
         pos+=4;
 
         for(const auto& v: multi)
         {
-            uint32 ds = uint32( v.first.size() );
+            uint32 ds = uint32( v.Key.Num() );
 
-            parameterValuesBlob.resize( pos + 4 + ds*4 );
+            parameterValuesBlob.SetNum( pos + 4 + ds*4 );
 
-            memcpy( &parameterValuesBlob[pos], &s, sizeof(uint32) );
+			FMemory::Memcpy( &parameterValuesBlob[pos], &s, sizeof(uint32) );
             pos+=4;
 
-            memcpy( &parameterValuesBlob[pos],v.first.data(), ds*sizeof(int32) );
+			FMemory::Memcpy( &parameterValuesBlob[pos],v.Key.GetData(), ds*sizeof(int32) );
             pos += ds*sizeof(int32);
         }
 
@@ -513,8 +513,8 @@ namespace mu
                                              const Parameters* pParams )
     {
         // Find the list of relevant parameters
-        const vector<uint16_t>* params = nullptr;
-        if (paramListIndex<m_program.m_parameterLists.size())
+        const TArray<uint16>* params = nullptr;
+        if (paramListIndex<(uint32)m_program.m_parameterLists.Num())
         {
             params = &m_program.m_parameterLists[paramListIndex];
         }
@@ -525,28 +525,28 @@ namespace mu
         }
 
         // Generate the relevant parameters blob
-        vector<uint8_t> parameterValuesBlob;
-		parameterValuesBlob.reserve(1024);
+		TArray<uint8_t> parameterValuesBlob;
+		parameterValuesBlob.Reserve(1024);
         for (int param: *params)
         {
-            size_t pos = parameterValuesBlob.size();
+            size_t pos = parameterValuesBlob.Num();
             size_t dataSize = 0;
 
             switch(m_program.m_parameters[param].m_type)
             {
             case PARAMETER_TYPE::T_BOOL:
                 dataSize = 1;
-                parameterValuesBlob.push_back( pParams->GetPrivate()->m_values[param].m_bool ? 1 : 0 );
+                parameterValuesBlob.Add( pParams->GetPrivate()->m_values[param].m_bool ? 1 : 0 );
                 pos += dataSize;
 
                 // Multi-values
-                if ( param < int(pParams->GetPrivate()->m_multiValues.size()) )
+                if ( param < int(pParams->GetPrivate()->m_multiValues.Num()) )
                 {
                     const auto& multi = pParams->GetPrivate()->m_multiValues[param];
                     pos = AddMultiValueKeys( parameterValuesBlob, pos, multi );
                     for(const auto& v: multi)
                     {
-                        parameterValuesBlob.push_back( v.second.m_bool ? 1 : 0 );
+                        parameterValuesBlob.Add( v.Value.m_bool ? 1 : 0 );
                         pos += dataSize;
                     }
                 }
@@ -554,21 +554,21 @@ namespace mu
 
             case PARAMETER_TYPE::T_INT:
                 dataSize = sizeof(int32);
-                parameterValuesBlob.resize( pos+dataSize );
-                memcpy( &parameterValuesBlob[pos],
+                parameterValuesBlob.SetNum( pos+dataSize );
+                FMemory::Memcpy( &parameterValuesBlob[pos],
                         &pParams->GetPrivate()->m_values[param].m_int,
                         dataSize );
                 pos += dataSize;
 
                 // Multi-values
-                if ( param < int(pParams->GetPrivate()->m_multiValues.size()) )
+                if ( param < int(pParams->GetPrivate()->m_multiValues.Num()) )
                 {
                     const auto& multi = pParams->GetPrivate()->m_multiValues[param];
                     pos = AddMultiValueKeys( parameterValuesBlob, pos, multi );
-                    parameterValuesBlob.resize( pos+multi.size()*dataSize );
+                    parameterValuesBlob.SetNum( pos+multi.Num()*dataSize );
                     for(const auto& v: multi)
                     {
-                        memcpy( &parameterValuesBlob[pos], &v.second.m_int, dataSize );
+						FMemory::Memcpy( &parameterValuesBlob[pos], &v.Value.m_int, dataSize );
                         pos += dataSize;
                     }
                 }
@@ -577,21 +577,21 @@ namespace mu
             //! Floating point value in the range of 0.0 to 1.0
             case PARAMETER_TYPE::T_FLOAT:
                 dataSize = sizeof(float);
-                parameterValuesBlob.resize( pos+dataSize );
-                memcpy( &parameterValuesBlob[pos],
+                parameterValuesBlob.SetNum( pos+dataSize );
+				FMemory::Memcpy( &parameterValuesBlob[pos],
                         &pParams->GetPrivate()->m_values[param].m_float,
                         dataSize );
                 pos += dataSize;
 
                 // Multi-values
-                if ( param < int(pParams->GetPrivate()->m_multiValues.size()) )
+                if ( param < pParams->GetPrivate()->m_multiValues.Num() )
                 {
                     const auto& multi = pParams->GetPrivate()->m_multiValues[param];
                     pos = AddMultiValueKeys( parameterValuesBlob, pos, multi );
-                    parameterValuesBlob.resize( pos+multi.size()*dataSize );
+                    parameterValuesBlob.SetNum( pos+multi.Num()*dataSize );
                     for(const auto& v: multi)
                     {
-                        memcpy( &parameterValuesBlob[pos], &v.second.m_float,dataSize );
+						FMemory::Memcpy( &parameterValuesBlob[pos], &v.Value.m_float,dataSize );
                         pos += dataSize;
                     }
                 }
@@ -600,21 +600,21 @@ namespace mu
             //! Floating point RGBA colour, with each channel ranging from 0.0 to 1.0
             case PARAMETER_TYPE::T_COLOUR:
                 dataSize = 4*sizeof(float);
-                parameterValuesBlob.resize( pos+dataSize );
-                memcpy( &parameterValuesBlob[pos],
+                parameterValuesBlob.SetNum( pos+dataSize );
+				FMemory::Memcpy( &parameterValuesBlob[pos],
                         &pParams->GetPrivate()->m_values[param].m_colour,
                         dataSize );
                 pos += dataSize;
 
                 // Multi-values
-                if ( param < int(pParams->GetPrivate()->m_multiValues.size()) )
+                if ( param < int(pParams->GetPrivate()->m_multiValues.Num()) )
                 {
                     const auto& multi = pParams->GetPrivate()->m_multiValues[param];
                     pos = AddMultiValueKeys( parameterValuesBlob, pos, multi );
-                    parameterValuesBlob.resize( pos+multi.size()*dataSize );
+                    parameterValuesBlob.SetNum( pos+multi.Num()*dataSize );
                     for(const auto& v: multi)
                     {
-                        memcpy( &parameterValuesBlob[pos], &v.second.m_colour,dataSize );
+						FMemory::Memcpy( &parameterValuesBlob[pos], &v.Value.m_colour,dataSize );
                         pos += dataSize;
                     }
                 }
@@ -626,21 +626,21 @@ namespace mu
                 dataSize = sizeof(PROJECTOR);
 
                 // \todo: padding will be random?
-                parameterValuesBlob.resize( pos+dataSize );
-                memcpy( &parameterValuesBlob[pos],
+                parameterValuesBlob.SetNum( pos+dataSize );
+				FMemory::Memcpy( &parameterValuesBlob[pos],
                         &pParams->GetPrivate()->m_values[param].m_projector,
                         dataSize );
 				pos += dataSize;
 
                 // Multi-values
-                if ( param < int(pParams->GetPrivate()->m_multiValues.size()) )
+                if ( param < int(pParams->GetPrivate()->m_multiValues.Num()) )
                 {
                     const auto& multi = pParams->GetPrivate()->m_multiValues[param];
                     pos = AddMultiValueKeys( parameterValuesBlob, pos, multi );
-                    parameterValuesBlob.resize( pos+multi.size()*dataSize );
+                    parameterValuesBlob.SetNum( pos+multi.Num()*dataSize );
                     for(const auto& v: multi)
                     {
-                        memcpy( &parameterValuesBlob[pos], &v.second.m_projector,dataSize );
+						FMemory::Memcpy( &parameterValuesBlob[pos], &v.Value.m_projector,dataSize );
                         pos += dataSize;
                     }
                 }
@@ -648,21 +648,21 @@ namespace mu
 
             case PARAMETER_TYPE::T_IMAGE:
                 dataSize = sizeof(EXTERNAL_IMAGE_ID);
-                parameterValuesBlob.resize( pos+dataSize );
-                memcpy( &parameterValuesBlob[pos],
+                parameterValuesBlob.SetNum( pos+dataSize );
+				FMemory::Memcpy( &parameterValuesBlob[pos],
                         &pParams->GetPrivate()->m_values[param].m_image,
                         dataSize );
 				pos += dataSize;
 
                 // Multi-values
-                if ( param < int(pParams->GetPrivate()->m_multiValues.size()) )
+                if ( param < int(pParams->GetPrivate()->m_multiValues.Num()) )
                 {
                     const auto& multi = pParams->GetPrivate()->m_multiValues[param];
                     pos = AddMultiValueKeys( parameterValuesBlob, pos, multi );
-                    parameterValuesBlob.resize( pos+multi.size()*dataSize );
+                    parameterValuesBlob.SetNum( pos+multi.Num()*dataSize );
                     for(const auto& v: multi)
                     {
-                        memcpy( &parameterValuesBlob[pos], &v.second.m_image,dataSize );
+						FMemory::Memcpy( &parameterValuesBlob[pos], &v.Value.m_image,dataSize );
                         pos += dataSize;
                     }
                 }
@@ -679,7 +679,7 @@ namespace mu
 
         // See if we already have this id
         size_t oldestCachePosition = 0;
-        for (size_t i=0; i<m_generatedResources.size(); ++i)
+        for (size_t i=0; i<m_generatedResources.Num(); ++i)
         {
             auto& key = m_generatedResources[i];
             if (key.m_rootAddress==rootAt
@@ -710,13 +710,13 @@ namespace mu
 
         // TODO: Move the constant to settings?
         const size_t maxGeneratedResourcesIDCacheSize = 1024;
-        if (m_generatedResources.size()>=maxGeneratedResourcesIDCacheSize)
+        if (m_generatedResources.Num()>=maxGeneratedResourcesIDCacheSize)
         {
             m_generatedResources[oldestCachePosition] = std::move(newKey);
         }
         else
         {
-            m_generatedResources.push_back(std::move(newKey));
+            m_generatedResources.Add(MoveTemp(newKey));
         }
 
         return newId;
@@ -739,12 +739,13 @@ namespace mu
         m_pD->m_pModel = pModel;
         m_pD->m_pSystem = pSystem;
 
-        size_t paramCount = pModel->GetPrivate()->m_program.m_parameters.size();
+        size_t paramCount = pModel->GetPrivate()->m_program.m_parameters.Num();
         m_pD->m_considerRelevancy = considerRelevancy;
         if (m_pD->m_considerRelevancy)
         {
-            vector<int> currentValues( paramCount, 0 );
-            m_pD->m_intervals.resize( paramCount );
+			TArray<int> currentValues;
+			currentValues.SetNumZeroed(paramCount, 0);
+            m_pD->m_intervals.SetNum( paramCount );
             m_pD->m_instanceCount = (int)m_pD->BuildIntervals( 0, 0, currentValues );
         }
         else
@@ -759,7 +760,7 @@ namespace mu
 
                 case PARAMETER_TYPE::T_INT:
                 {
-                    m_pD->m_instanceCount *= param.m_possibleValues.size();
+                    m_pD->m_instanceCount *= param.m_possibleValues.Num();
                     break;
                 }
 
@@ -813,30 +814,30 @@ namespace mu
 
         if (m_pD->m_considerRelevancy)
         {
-            vector<int> values = m_pD->GetParameters( (int)index );
-            for ( size_t p=0; p<values.size(); ++p )
+			TArray<int> values = m_pD->GetParameters( (int)index );
+            for ( int32 p=0; p<values.Num(); ++p )
             {
-                switch ( res->GetType( (int)p ) )
+                switch ( res->GetType( p ) )
                 {
                 case PARAMETER_TYPE::T_BOOL:
-                    res->SetBoolValue( (int)p, values[p]!=0 );
+                    res->SetBoolValue( p, values[p]!=0 );
                     break;
 
                 case PARAMETER_TYPE::T_INT:
-                    res->SetIntValue( (int)p, res->GetIntPossibleValue( (int)p, values[p] ) );
+                    res->SetIntValue( p, res->GetIntPossibleValue( p, values[p] ) );
                     break;
 
                 case PARAMETER_TYPE::T_FLOAT:
                     if ( randomGenerator )
                     {
-                        res->SetFloatValue( (int)p, randomGenerator() );
+                        res->SetFloatValue( p, randomGenerator() );
                     }
                     break;
 
                 case PARAMETER_TYPE::T_COLOUR:
                     if ( randomGenerator )
                     {
-                        res->SetColourValue( (int)p,
+                        res->SetColourValue( p,
                                              randomGenerator(), randomGenerator(), randomGenerator() );
                     }
                     break;
@@ -848,7 +849,7 @@ namespace mu
         }
         else
         {
-            int paramCount = (int)m_pD->m_pModel->GetPrivate()->m_program.m_parameters.size();
+            int paramCount = (int)m_pD->m_pModel->GetPrivate()->m_program.m_parameters.Num();
             int64 currentInstance = index;
             for (int i=0;i<paramCount;++i)
             {
@@ -897,13 +898,13 @@ namespace mu
 
 
     //---------------------------------------------------------------------------------------------
-    ParametersPtr ModelParametersGenerator::GetRandomInstance( std::function<float()> randomGenerator)
+    ParametersPtr ModelParametersGenerator::GetRandomInstance( TFunctionRef<float()> randomGenerator)
     {
 		LLM_SCOPE_BYNAME(TEXT("MutableRuntime"));
 
         ParametersPtr res = m_pD->m_pModel->NewParameters();
 
-        int paramCount = (int)m_pD->m_pModel->GetPrivate()->m_program.m_parameters.size();
+        int paramCount = (int)m_pD->m_pModel->GetPrivate()->m_program.m_parameters.Num();
 
         for (int i=0;i<paramCount;++i)
         {
@@ -914,7 +915,7 @@ namespace mu
             case PARAMETER_TYPE::T_INT:
             {
                 int numOptions = res->GetIntPossibleValueCount( i );
-                int valueIndex = (int)(std::min(numOptions-1,int(randomGenerator()*numOptions)));
+                int valueIndex = (int)(FMath::Min(numOptions-1,int(randomGenerator()*numOptions)));
                 int value = res->GetIntPossibleValue( i, valueIndex );
                 res->SetIntValue( i, value );
                 break;
@@ -927,17 +928,11 @@ namespace mu
             }
 
             case PARAMETER_TYPE::T_FLOAT:
-                if ( randomGenerator )
-                {
-                    res->SetFloatValue( i, randomGenerator() );
-                }
+                res->SetFloatValue( i, randomGenerator() );
                 break;
 
             case PARAMETER_TYPE::T_COLOUR:
-                if ( randomGenerator )
-                {
-                    res->SetColourValue( i, randomGenerator(), randomGenerator(), randomGenerator() );
-                }
+                res->SetColourValue( i, randomGenerator(), randomGenerator(), randomGenerator() );
                 break;
 
             default:
@@ -967,7 +962,7 @@ namespace mu
 //            )
 //            : PartialDiscreteCoveredCodeVisitor( pSystem, pModel, pParams, relevantParameters )
 //        {
-//            m_visited.resize( pModel->GetPrivate()->m_program.m_opAddress.size(), false );
+//            m_visited.resize( pModel->GetPrivate()->m_program.m_opAddress.Num(), false );
 //            m_thisParameter = thisParameter;
 //            m_relevant = false;
 
@@ -1014,16 +1009,17 @@ namespace mu
 //        int m_thisParameter;
 
 //        //! Flags for visited instructions.
-//        vector<bool> m_visited;
+//        TArray<bool> m_visited;
 //    };
 
 
     //---------------------------------------------------------------------------------------------
-    size_t ModelParametersGenerator::Private::BuildIntervals( size_t currentInstanceIndex,
-                                                              size_t currentParameter,
-                                                              vector<int>& currentValues )
+	uint32 ModelParametersGenerator::Private::BuildIntervals(
+		uint32 currentInstanceIndex,
+		uint32 currentParameter,
+		TArray<int>& currentValues )
     {
-        size_t paramCount = m_pModel->GetPrivate()->m_program.m_parameters.size();
+        size_t paramCount = m_pModel->GetPrivate()->m_program.m_parameters.Num();
 
         if ( currentParameter>=paramCount )
         {
@@ -1078,7 +1074,7 @@ namespace mu
                     PARAMETER_INTERVAL_VALUE falseValue;
                     falseValue.m_minIndex = (int)currentInstanceIndex;
                     falseValue.m_value = 0;
-                    m_intervals[ currentParameter ].m_intervalValue.push_back( falseValue );
+                    m_intervals[ currentParameter ].m_intervalValue.Add( falseValue );
                     currentValues[ currentParameter ] = 0;
                     currentInstanceIndex = BuildIntervals( currentInstanceIndex,
                                                            currentParameter+1,
@@ -1088,7 +1084,7 @@ namespace mu
                     PARAMETER_INTERVAL_VALUE trueValue;
                     trueValue.m_minIndex = (int)currentInstanceIndex;
                     trueValue.m_value = 1;
-                    m_intervals[ currentParameter ].m_intervalValue.push_back( trueValue );
+                    m_intervals[ currentParameter ].m_intervalValue.Add( trueValue );
                     currentValues[ currentParameter ] = 1;
                     currentInstanceIndex = BuildIntervals( currentInstanceIndex,
                                                            currentParameter+1,
@@ -1098,13 +1094,13 @@ namespace mu
 
                 case PARAMETER_TYPE::T_INT:
                 {
-                    for ( size_t v=0; v<param.m_possibleValues.size(); ++v )
+                    for ( int32 v=0; v<param.m_possibleValues.Num(); ++v )
                     {
                         PARAMETER_INTERVAL_VALUE value;
                         value.m_minIndex = (int)currentInstanceIndex;
-                        value.m_value = (int)v;
-                        m_intervals[ currentParameter ].m_intervalValue.push_back( value );
-                        currentValues[ currentParameter ] = (int)v;
+                        value.m_value = v;
+                        m_intervals[ currentParameter ].m_intervalValue.Add( value );
+                        currentValues[ currentParameter ] = v;
                         currentInstanceIndex = BuildIntervals( currentInstanceIndex,
                                                                currentParameter+1,
                                                                currentValues);
@@ -1134,14 +1130,15 @@ namespace mu
 
 
     //---------------------------------------------------------------------------------------------
-    vector<int> ModelParametersGenerator::Private::GetParameters( int instanceIndex )
+	TArray<int> ModelParametersGenerator::Private::GetParameters( int instanceIndex )
     {
-        vector<int> res( m_pModel->GetPrivate()->m_program.m_parameters.size(), 0 );
+		TArray<int> res;
+		res.SetNumZeroed(m_pModel->GetPrivate()->m_program.m_parameters.Num());
 
-        for ( size_t p=0; p<res.size(); ++p )
+        for ( size_t p=0; p<res.Num(); ++p )
         {
             size_t e = 0;
-            while ( m_intervals[p].m_intervalValue.size()>e+1
+            while ( m_intervals[p].m_intervalValue.Num()>e+1
                     &&
                     m_intervals[p].m_intervalValue[e+1].m_minIndex<=instanceIndex )
             {
@@ -1149,7 +1146,7 @@ namespace mu
             }
 
             // Avoid degenerated case (parameter always irrelevant?)
-            if ( m_intervals[p].m_intervalValue.size()>e )
+            if ( m_intervals[p].m_intervalValue.Num()>e )
             {
                 res[p] = m_intervals[p].m_intervalValue[e].m_value;
             }
diff --git a/Engine/Plugins/Experimental/Mutable/Source/MutableRuntime/Private/MuR/ModelPrivate.h b/Engine/Plugins/Experimental/Mutable/Source/MutableRuntime/Private/MuR/ModelPrivate.h
--- a/Engine/Plugins/Experimental/Mutable/Source/MutableRuntime/Private/MuR/ModelPrivate.h
+++ b/Engine/Plugins/Experimental/Mutable/Source/MutableRuntime/Private/MuR/ModelPrivate.h
@@ -45,10 +45,10 @@ namespace mu
 
     //!
     template<typename DATA>
-    inline void AppendCode( vector<uint8>& code, const DATA& data )
+    inline void AppendCode(TArray<uint8>& code, const DATA& data )
     {
-        size_t pos = code.size();
-        code.resize( code.size()+sizeof(DATA) );
+        int32 pos = code.Num();
+        code.SetNum( pos+sizeof(DATA) );
 		FMemory::Memcpy (&code[pos], &data, sizeof(DATA));
     }
 
@@ -71,7 +71,7 @@ namespace mu
             // Add the null instruction at address 0.
             // TODO: Will do it in the linker
             AppendCode( m_byteCode, OP_TYPE::NONE );
-            m_opAddress.push_back(0);
+            m_opAddress.Add(0);
         }
 
         struct STATE
@@ -84,16 +84,16 @@ namespace mu
 
             //! List of parameters index (to PROGRAM::m_parameters) of the runtime parameters of
             //! this state.
-            vector<int> m_runtimeParameters;
+			TArray<int> m_runtimeParameters;
 
             //! List of instructions that need to be cached to efficiently update this state
-            vector<OP::ADDRESS> m_updateCache;
+			TArray<OP::ADDRESS> m_updateCache;
 
             //! List of root instructions for the dynamic resources that depend on the runtime
             //! parameters of this state, with a mask of relevant runtime parameters.
             //! The mask has a bit on for every runtime parameter in the m_runtimeParameters
             //! vector.
-            vector< std::pair<OP::ADDRESS,uint64> > m_dynamicResources;
+			TArray< TPair<OP::ADDRESS,uint64> > m_dynamicResources;
 
             //!
             inline void Serialise( OutputArchive& arch ) const
@@ -120,11 +120,11 @@ namespace mu
             {
                 uint64 res = 0;
 
-                for ( size_t i=0; !res && i<m_dynamicResources.size(); ++i )
+                for ( int32 i=0; !res && i<m_dynamicResources.Num(); ++i )
                 {
-                    if ( m_dynamicResources[i].first==at )
+                    if ( m_dynamicResources[i].Key==at )
                     {
-                        res = m_dynamicResources[i].second;
+                        res = m_dynamicResources[i].Value;
                     }
                 }
 
@@ -136,7 +136,7 @@ namespace mu
             {
                 bool res = false;
 
-                for ( size_t i=0; !res && i<m_updateCache.size(); ++i )
+                for ( size_t i=0; !res && i<m_updateCache.Num(); ++i )
                 {
                     if ( m_updateCache[i]==at )
                     {
@@ -152,19 +152,19 @@ namespace mu
             {
                 if ( !IsUpdateCache(at) )
                 {
-                    m_updateCache.push_back( at );
+                    m_updateCache.Add( at );
                 }
             }
         };
 
         //! Location in the m_byteCode of the beginning of each operation
-        vector<uint32> m_opAddress;
+		TArray<uint32> m_opAddress;
 
         //! Byte-coded representation of the program, using flexible-sized ops.
-        vector<uint8> m_byteCode;
+		TArray<uint8> m_byteCode;
 
         //!
-        vector<STATE> m_states;
+		TArray<STATE> m_states;
 
 		//! Data for every rom.
 		TArray<FRomData> m_roms;
@@ -182,40 +182,40 @@ namespace mu
 		TArray<TPair<int32, Ptr<const Mesh>>> m_constantMeshes;
 
         //! Constant string data
-        vector<string> m_constantStrings;
+		TArray<string> m_constantStrings;
 
         //! Constant layout data
-        vector<Ptr<const Layout>> m_constantLayouts;
+		TArray<Ptr<const Layout>> m_constantLayouts;
 
         //! Constant projectors
-        vector<PROJECTOR> m_constantProjectors;
+		TArray<PROJECTOR> m_constantProjectors;
 
         //! Constant matrices, usually used for transforms
-        vector<mat4f> m_constantMatrices;
+		TArray<mat4f> m_constantMatrices;
 
 		//! Constant projectors
-		vector<SHAPE> m_constantShapes;
+		TArray<SHAPE> m_constantShapes;
 
         //! Constant curves
-        vector<Curve> m_constantCurves;
+		TArray<Curve> m_constantCurves;
 
         //! Constant curves
-        vector<Ptr<const Skeleton>> m_constantSkeletons;
+		TArray<Ptr<const Skeleton>> m_constantSkeletons;
 
 		//! Constant Physics Bodies
-		vector<Ptr<const PhysicsBody>> m_constantPhysicsBodies;
+		TArray<Ptr<const PhysicsBody>> m_constantPhysicsBodies;
 
         //! Parameters of the model.
         //! The value stored here is the default value.
-        vector<PARAMETER_DESC> m_parameters;
+		TArray<PARAMETER_DESC> m_parameters;
 
         //! Ranges for interation of the model operations.
-        vector<RANGE_DESC> m_ranges;
+		TArray<RANGE_DESC> m_ranges;
 
         //! List of parameter lists. These are used in several places, like storing the
         //! pregenerated list of parameters influencing a resource.
         //! The parameter lists are sorted.
-        vector<vector<uint16>> m_parameterLists;
+		TArray<TArray<uint16>> m_parameterLists;
 
 
         //!
@@ -403,7 +403,7 @@ namespace mu
 		OP::ADDRESS AddConstant( Ptr<const Layout> pLayout )
         {
             // Ensure unique
-            for ( size_t i=0; i<m_constantLayouts.size(); ++i)
+            for ( size_t i=0; i<m_constantLayouts.Num(); ++i)
             {
                 if (m_constantLayouts[i]==pLayout)
                 {
@@ -411,15 +411,15 @@ namespace mu
                 }
             }
 
-            OP::ADDRESS index = OP::ADDRESS( m_constantLayouts.size() );
-            m_constantLayouts.push_back( pLayout );
+            OP::ADDRESS index = OP::ADDRESS( m_constantLayouts.Num() );
+            m_constantLayouts.Add( pLayout );
             return index;
         }
 
         OP::ADDRESS AddConstant( Ptr<const Skeleton> pSkeleton )
         {
             // Ensure unique
-            for ( size_t i=0; i<m_constantSkeletons.size(); ++i)
+            for ( size_t i=0; i<m_constantSkeletons.Num(); ++i)
             {
                 if ( m_constantSkeletons[i]==pSkeleton
                      ||
@@ -430,15 +430,15 @@ namespace mu
                 }
             }
 
-            OP::ADDRESS index = OP::ADDRESS( m_constantSkeletons.size() );
-            m_constantSkeletons.push_back( pSkeleton );
+            OP::ADDRESS index = OP::ADDRESS( m_constantSkeletons.Num() );
+            m_constantSkeletons.Add( pSkeleton );
             return index;
         }
 
         OP::ADDRESS AddConstant( Ptr<const PhysicsBody> pPhysicsBody )
         {
             // Ensure unique
-            for ( size_t i=0; i<m_constantPhysicsBodies.size(); ++i)
+            for ( size_t i=0; i<m_constantPhysicsBodies.Num(); ++i)
             {
                 if ( m_constantPhysicsBodies[i]==pPhysicsBody
                      ||
@@ -449,15 +449,15 @@ namespace mu
                 }
             }
 
-            OP::ADDRESS index = OP::ADDRESS( m_constantPhysicsBodies.size() );
-            m_constantPhysicsBodies.push_back( pPhysicsBody );
+            OP::ADDRESS index = OP::ADDRESS( m_constantPhysicsBodies.Num() );
+            m_constantPhysicsBodies.Add( pPhysicsBody );
             return index;
         }
 
         OP::ADDRESS AddConstant( const string& str )
         {            
             // Ensure unique
-            for ( size_t i=0; i<m_constantStrings.size(); ++i)
+            for ( size_t i=0; i<m_constantStrings.Num(); ++i)
             {
                 if (m_constantStrings[i]==str)
                 {
@@ -465,15 +465,15 @@ namespace mu
                 }
             }
 
-            OP::ADDRESS index = OP::ADDRESS( m_constantStrings.size() );
-            m_constantStrings.push_back( str );
+            OP::ADDRESS index = OP::ADDRESS( m_constantStrings.Num() );
+            m_constantStrings.Add( str );
             return index;
         }
 
         OP::ADDRESS AddConstant( const mat4f& m )
         {
             // Ensure unique
-            for ( size_t i=0; i<m_constantMatrices.size(); ++i)
+            for ( size_t i=0; i<m_constantMatrices.Num(); ++i)
             {
                 if (m_constantMatrices[i]==m)
                 {
@@ -481,15 +481,15 @@ namespace mu
                 }
             }
 
-            OP::ADDRESS index = OP::ADDRESS( m_constantMatrices.size() );
-            m_constantMatrices.push_back( m );
+            OP::ADDRESS index = OP::ADDRESS( m_constantMatrices.Num() );
+            m_constantMatrices.Add( m );
             return index;
         }
 
         OP::ADDRESS AddConstant( const SHAPE& m )
         {
             // Ensure unique
-            for ( size_t i=0; i<m_constantShapes.size(); ++i)
+            for ( size_t i=0; i<m_constantShapes.Num(); ++i)
             {
                 if (m_constantShapes[i]==m)
                 {
@@ -497,15 +497,15 @@ namespace mu
                 }
             }
 
-            OP::ADDRESS index = OP::ADDRESS( m_constantShapes.size() );
-            m_constantShapes.push_back( m );
+            OP::ADDRESS index = OP::ADDRESS( m_constantShapes.Num() );
+            m_constantShapes.Add( m );
             return index;
         }
 
         OP::ADDRESS AddConstant( const PROJECTOR& m )
         {
             // Ensure unique
-            for ( size_t i=0; i<m_constantProjectors.size(); ++i)
+            for ( size_t i=0; i<m_constantProjectors.Num(); ++i)
             {
                 if (m_constantProjectors[i]==m)
                 {
@@ -513,15 +513,15 @@ namespace mu
                 }
             }
 
-            OP::ADDRESS index = OP::ADDRESS( m_constantProjectors.size() );
-            m_constantProjectors.push_back( m );
+            OP::ADDRESS index = OP::ADDRESS( m_constantProjectors.Num() );
+            m_constantProjectors.Add( m );
             return index;
         }
 
         OP::ADDRESS AddConstant( const Curve& m )
         {
             // Ensure unique
-//            for ( size_t i=0; i<m_constantCurves.size(); ++i)
+//            for ( size_t i=0; i<m_constantCurves.Num(); ++i)
 //            {
 //                if (m_constantCurves[i]==m)
 //                {
@@ -529,8 +529,8 @@ namespace mu
 //                }
 //            }
 
-            OP::ADDRESS index = OP::ADDRESS( m_constantCurves.size() );
-            m_constantCurves.push_back( m );
+            OP::ADDRESS index = OP::ADDRESS( m_constantCurves.Num() );
+            m_constantCurves.Add( m );
             return index;
         }
 
@@ -610,7 +610,7 @@ namespace mu
 
         inline OP_TYPE GetOpType( OP::ADDRESS at ) const
         {
-            if (at>=OP::ADDRESS(m_opAddress.size())) return OP_TYPE::NONE;
+            if (at>=OP::ADDRESS(m_opAddress.Num())) return OP_TYPE::NONE;
 
             OP_TYPE result;
             uint64 byteCodeAddress = m_opAddress[at];
@@ -693,7 +693,7 @@ namespace mu
             OP::ADDRESS m_rootAddress;
 
             //! An opaque blob with the values of the relevant parameters
-            vector<uint8> m_parameterValuesBlob;
+			TArray<uint8> m_parameterValuesBlob;
         };
 
         //! The last id generated for a resource
@@ -705,7 +705,7 @@ namespace mu
 
         //! Cached ids for returned assets
         //! This is non-persistent runtime data
-        vector<RESOURCE_KEY> m_generatedResources;
+		TArray<RESOURCE_KEY> m_generatedResources;
 
         //! Get a resource key for a given resource with given parameter values.
         uint32 GetResourceKey( uint32 paramListIndex, OP::ADDRESS rootAt, const Parameters* pParams );
@@ -742,22 +742,20 @@ namespace mu
         {
             //! As many entries as necessary
             //! Sorted by PARAMETER_INCREMENTS_PER_VALUE::m_minIndex
-            vector<PARAMETER_INTERVAL_VALUE> m_intervalValue;
+			TArray<PARAMETER_INTERVAL_VALUE> m_intervalValue;
         };
 
         //! An entry for every parameter in the model
-        vector< PARAMETER_INTERVALS > m_intervals;
+		TArray< PARAMETER_INTERVALS > m_intervals;
 
         //! Whether to use brute force, or consider the parameter relevancy
         bool m_considerRelevancy;
 
         //! Build the interval information
-        size_t BuildIntervals( size_t currentInstanceIndex,
-                               size_t currentParameter,
-                               vector<int>& currentValues );
+		uint32 BuildIntervals( uint32 currentInstanceIndex, uint32 currentParameter, TArray<int>& currentValues );
 
         //! Get the parameter values for a particular instance index
-        vector<int> GetParameters( int instanceIndex );
+		TArray<int> GetParameters( int instanceIndex );
     };
 
 }
diff --git a/Engine/Plugins/Experimental/Mutable/Source/MutableRuntime/Private/MuR/OpEvalCurve.h b/Engine/Plugins/Experimental/Mutable/Source/MutableRuntime/Private/MuR/OpEvalCurve.h
--- a/Engine/Plugins/Experimental/Mutable/Source/MutableRuntime/Private/MuR/OpEvalCurve.h
+++ b/Engine/Plugins/Experimental/Mutable/Source/MutableRuntime/Private/MuR/OpEvalCurve.h
@@ -31,7 +31,7 @@ namespace mu
 	inline float EvalCurve(const Curve& curve, float time)
 	{
 		float interpolated_value = curve.defaultValue;
-        const int32_t numKeyFrames = (int32_t)curve.keyFrames.size();
+        const int32_t numKeyFrames = curve.keyFrames.Num();
 
 		if (numKeyFrames < 1)
 		{
diff --git a/Engine/Plugins/Experimental/Mutable/Source/MutableRuntime/Private/MuR/OpImageBurn.h b/Engine/Plugins/Experimental/Mutable/Source/MutableRuntime/Private/MuR/OpImageBurn.h
--- a/Engine/Plugins/Experimental/Mutable/Source/MutableRuntime/Private/MuR/OpImageBurn.h
+++ b/Engine/Plugins/Experimental/Mutable/Source/MutableRuntime/Private/MuR/OpImageBurn.h
@@ -12,8 +12,8 @@ namespace mu
 	{
 		// R = 1 - (1-Base) / Blend
 		unsigned burn =
-				std::min( 255,
-						  std::max( 0,
+				FMath::Min( 255,
+						  FMath::Max( 0,
 									255 - ( ( ( 255 - (int)base  ) << 8 ) / ((int)blended+1) )
 									)
 						  );
diff --git a/Engine/Plugins/Experimental/Mutable/Source/MutableRuntime/Private/MuR/OpImageColourMap.h b/Engine/Plugins/Experimental/Mutable/Source/MutableRuntime/Private/MuR/OpImageColourMap.h
--- a/Engine/Plugins/Experimental/Mutable/Source/MutableRuntime/Private/MuR/OpImageColourMap.h
+++ b/Engine/Plugins/Experimental/Mutable/Source/MutableRuntime/Private/MuR/OpImageColourMap.h
@@ -31,10 +31,10 @@ namespace mu
 		for ( int i=0; i<256; ++i )
 		{
 			vec4<float> c = pMap->Sample( vec2<float>( float(i)/256.0f, 0.0f ) );
-            palette[i][0] = (uint8_t)std::max( 0, std::min( 255, int(c[0]*255.0f) ) );
-            palette[i][1] = (uint8_t)std::max( 0, std::min( 255, int(c[1]*255.0f) ) );
-            palette[i][2] = (uint8_t)std::max( 0, std::min( 255, int(c[2]*255.0f) ) );
-            palette[i][3] = (uint8_t)std::max( 0, std::min( 255, int(c[3]*255.0f) ) );
+            palette[i][0] = (uint8_t)FMath::Max( 0, FMath::Min( 255, int(c[0]*255.0f) ) );
+            palette[i][1] = (uint8_t)FMath::Max( 0, FMath::Min( 255, int(c[1]*255.0f) ) );
+            palette[i][2] = (uint8_t)FMath::Max( 0, FMath::Min( 255, int(c[2]*255.0f) ) );
+            palette[i][3] = (uint8_t)FMath::Max( 0, FMath::Min( 255, int(c[3]*255.0f) ) );
 		}
 
 		switch ( pSource->GetFormat() )
diff --git a/Engine/Plugins/Experimental/Mutable/Source/MutableRuntime/Private/MuR/OpImageDifference.h b/Engine/Plugins/Experimental/Mutable/Source/MutableRuntime/Private/MuR/OpImageDifference.h
--- a/Engine/Plugins/Experimental/Mutable/Source/MutableRuntime/Private/MuR/OpImageDifference.h
+++ b/Engine/Plugins/Experimental/Mutable/Source/MutableRuntime/Private/MuR/OpImageDifference.h
@@ -36,7 +36,7 @@ namespace mu
                     uint8_t b = pBBuf[3*i+c];
 					if ( abs(a-b)>3 )
 					{
-                        diff = std::max( diff, (uint8_t)255 );
+                        diff = FMath::Max( diff, (uint8_t)255 );
 					}
 				}
 				pDestBuf[i] = diff;
@@ -54,7 +54,7 @@ namespace mu
                     uint8_t b = pBBuf[4*i+c];
 					if ( abs(a-b)>3 )
 					{
-                        diff = std::max( diff, (uint8_t)255 );
+                        diff = FMath::Max( diff, (uint8_t)255 );
 					}
 				}
 				pDestBuf[i] = diff;
diff --git a/Engine/Plugins/Experimental/Mutable/Source/MutableRuntime/Private/MuR/OpImageDisplace.h b/Engine/Plugins/Experimental/Mutable/Source/MutableRuntime/Private/MuR/OpImageDisplace.h
--- a/Engine/Plugins/Experimental/Mutable/Source/MutableRuntime/Private/MuR/OpImageDisplace.h
+++ b/Engine/Plugins/Experimental/Mutable/Source/MutableRuntime/Private/MuR/OpImageDisplace.h
@@ -264,7 +264,7 @@ namespace mu
                     while ( x<sizeX )
                     {
                         // Decode header
-                        uint16_t equal = *(const uint16_t*)pMapData;
+                        uint16 equal = *(const uint16*)pMapData;
                         pMapData += 2;
 
                         uint8_t different = *pMapData;
diff --git a/Engine/Plugins/Experimental/Mutable/Source/MutableRuntime/Private/MuR/OpImageGradient.h b/Engine/Plugins/Experimental/Mutable/Source/MutableRuntime/Private/MuR/OpImageGradient.h
--- a/Engine/Plugins/Experimental/Mutable/Source/MutableRuntime/Private/MuR/OpImageGradient.h
+++ b/Engine/Plugins/Experimental/Mutable/Source/MutableRuntime/Private/MuR/OpImageGradient.h
@@ -14,7 +14,7 @@ namespace mu
 								   const vec3<float>& c1,
 								   int sizeX, int sizeY )
 	{
-		ImagePtr pDest = new Image((uint16_t)sizeX, (uint16_t)sizeY, 1, EImageFormat::IF_RGB_UBYTE );
+		ImagePtr pDest = new Image((uint16)sizeX, (uint16)sizeY, 1, EImageFormat::IF_RGB_UBYTE );
 
         uint8_t* pDestBuf = pDest->GetData();
 
@@ -24,9 +24,9 @@ namespace mu
 			vec3<float> c = c0 * (1.0f-delta) + c1 * delta;
 
             uint8_t colour[4];
-            colour[0] = (uint8_t)std::max( 0, std::min( 255, int(c[0]*255.0f) ) );
-            colour[1] = (uint8_t)std::max( 0, std::min( 255, int(c[1]*255.0f) ) );
-            colour[2] = (uint8_t)std::max( 0, std::min( 255, int(c[2]*255.0f) ) );
+            colour[0] = (uint8_t)FMath::Max( 0, FMath::Min( 255, int(c[0]*255.0f) ) );
+            colour[1] = (uint8_t)FMath::Max( 0, FMath::Min( 255, int(c[1]*255.0f) ) );
+            colour[2] = (uint8_t)FMath::Max( 0, FMath::Min( 255, int(c[2]*255.0f) ) );
 			colour[3] = 255;
 
 			for ( int j=0; j<sizeY; ++j )
diff --git a/Engine/Plugins/Experimental/Mutable/Source/MutableRuntime/Private/MuR/OpImageGrow.h b/Engine/Plugins/Experimental/Mutable/Source/MutableRuntime/Private/MuR/OpImageGrow.h
--- a/Engine/Plugins/Experimental/Mutable/Source/MutableRuntime/Private/MuR/OpImageGrow.h
+++ b/Engine/Plugins/Experimental/Mutable/Source/MutableRuntime/Private/MuR/OpImageGrow.h
@@ -29,9 +29,9 @@ namespace mu
 
 		for ( int y=0; y<sizeY; ++y )
 		{
-			int y0 = std::max( 0, std::min( sizeY-1, y-1 ) );
-			int y1 = std::max( 0, std::min( sizeY-1, y   ) );
-			int y2 = std::max( 0, std::min( sizeY-1, y+1 ) );
+			int y0 = FMath::Max( 0, FMath::Min( sizeY-1, y-1 ) );
+			int y1 = FMath::Max( 0, FMath::Min( sizeY-1, y   ) );
+			int y2 = FMath::Max( 0, FMath::Min( sizeY-1, y+1 ) );
 
             const uint8_t* pS0 = pSource->GetData() + y0 * rowSize;
             const uint8_t* pS1 = pSource->GetData() + y1 * rowSize;
diff --git a/Engine/Plugins/Experimental/Mutable/Source/MutableRuntime/Private/MuR/OpImageHardLight.h b/Engine/Plugins/Experimental/Mutable/Source/MutableRuntime/Private/MuR/OpImageHardLight.h
--- a/Engine/Plugins/Experimental/Mutable/Source/MutableRuntime/Private/MuR/OpImageHardLight.h
+++ b/Engine/Plugins/Experimental/Mutable/Source/MutableRuntime/Private/MuR/OpImageHardLight.h
@@ -20,7 +20,7 @@ namespace mu
 				? 255 - ( ( ( 255 - base ) * ( 255 - 2*(blended - 128 ) ) ) >> 8 )
 				: ( base * (2 * blended) ) >> 8;
 
-		hardlight = std::min( 255u, hardlight );
+		hardlight = FMath::Min( 255u, hardlight );
 
 		return hardlight;
 	}
diff --git a/Engine/Plugins/Experimental/Mutable/Source/MutableRuntime/Private/MuR/OpImageInterpolate.h b/Engine/Plugins/Experimental/Mutable/Source/MutableRuntime/Private/MuR/OpImageInterpolate.h
--- a/Engine/Plugins/Experimental/Mutable/Source/MutableRuntime/Private/MuR/OpImageInterpolate.h
+++ b/Engine/Plugins/Experimental/Mutable/Source/MutableRuntime/Private/MuR/OpImageInterpolate.h
@@ -26,7 +26,7 @@ namespace mu
         check( pA->GetFormat() == pB->GetFormat() );
 
 		// Clamp the factor
-		factor = std::max( 0.0f, std::min( 1.0f, factor ) );
+		factor = FMath::Max( 0.0f, FMath::Min( 1.0f, factor ) );
 
         uint8_t* pDestBuf = pDest->GetData();
         const uint8_t* pABuf = pA->GetData();
@@ -100,9 +100,9 @@ namespace mu
 		check( p0->GetFormat() == p2->GetFormat() );
 
 		// Clamp the factor
-		factor1 = std::max( 0.0f, std::min( 1.0f, factor1 ) );
-		factor2 = std::max( 0.0f, std::min( 1.0f-factor1, factor2 ) );
-		float factor0 = std::max( 0.0f, 1.0f-factor1-factor2 );
+		factor1 = FMath::Max( 0.0f, FMath::Min( 1.0f, factor1 ) );
+		factor2 = FMath::Max( 0.0f, FMath::Min( 1.0f-factor1, factor2 ) );
+		float factor0 = FMath::Max( 0.0f, 1.0f-factor1-factor2 );
 
         ImagePtr pDest = new Image( p0->GetSizeX(), p0->GetSizeY(),
                                     p0->GetLODCount(),
diff --git a/Engine/Plugins/Experimental/Mutable/Source/MutableRuntime/Private/MuR/OpImageLayer.h b/Engine/Plugins/Experimental/Mutable/Source/MutableRuntime/Private/MuR/OpImageLayer.h
--- a/Engine/Plugins/Experimental/Mutable/Source/MutableRuntime/Private/MuR/OpImageLayer.h
+++ b/Engine/Plugins/Experimental/Mutable/Source/MutableRuntime/Private/MuR/OpImageLayer.h
@@ -53,7 +53,7 @@ namespace mu
 					unsigned result = BLEND_FUNC(base, top[c]);
 					if (CLAMP)
 					{
-						pDestBuf[3 * i + c] = (uint8)std::min(255u, result);
+						pDestBuf[3 * i + c] = (uint8)FMath::Min(255u, result);
 					}
 					else
 					{
@@ -87,7 +87,7 @@ namespace mu
 					unsigned result = BLEND_FUNC(base, top[c]);
 					if (CLAMP)
 					{
-						pDestBuf[4 * i + c] = (uint8)std::min(255u, result);
+						pDestBuf[4 * i + c] = (uint8)FMath::Min(255u, result);
 					}
 					else
 					{
@@ -122,7 +122,7 @@ namespace mu
 						unsigned result = BLEND_FUNC(base, top[2 - c]);
 						if (CLAMP)
 						{
-							pDestBuf[4 * i + c] = (uint8)std::min(255u, result);
+							pDestBuf[4 * i + c] = (uint8)FMath::Min(255u, result);
 						}
 						else
 						{
@@ -155,7 +155,7 @@ namespace mu
 					unsigned result = BLEND_FUNC(base, top[0]);
 					if (CLAMP)
 					{
-						pDestBuf[i] = (uint8)std::min(255u, result);
+						pDestBuf[i] = (uint8)FMath::Min(255u, result);
 					}
 					else
 					{
@@ -831,7 +831,7 @@ namespace mu
 									unsigned result = BLEND_FUNC_MASKED( base, blended, mask );
 									if ( CLAMP )
 									{
-										pDestBuf[NC*i+c] = (uint8)std::min( 255u, result );
+										pDestBuf[NC*i+c] = (uint8)FMath::Min( 255u, result );
 									}
 									else
 									{
@@ -851,7 +851,7 @@ namespace mu
 									unsigned result = BLEND_FUNC_MASKED( base, blended, mask );
 									if ( CLAMP )
 									{
-										pDestBuf[NC*i+c] = (uint8)std::min( 255u, result );
+										pDestBuf[NC*i+c] = (uint8)FMath::Min( 255u, result );
 									}
 									else
 									{
diff --git a/Engine/Plugins/Experimental/Mutable/Source/MutableRuntime/Private/MuR/OpImageLuminance.h b/Engine/Plugins/Experimental/Mutable/Source/MutableRuntime/Private/MuR/OpImageLuminance.h
--- a/Engine/Plugins/Experimental/Mutable/Source/MutableRuntime/Private/MuR/OpImageLuminance.h
+++ b/Engine/Plugins/Experimental/Mutable/Source/MutableRuntime/Private/MuR/OpImageLuminance.h
@@ -25,7 +25,7 @@ namespace mu
 			for ( int i=0; i<pixelCount; ++i )
 			{
                 uint32_t l_16 = 76 * pABuf[3*i+0] + 150 * pABuf[3*i+1] + 29 * pABuf[3*i+2];
-                pDestBuf[i] = (uint8_t)std::min( l_16>>8, 255u );
+                pDestBuf[i] = (uint8_t)FMath::Min( l_16>>8, 255u );
 			}
 			break;
 		}
@@ -35,7 +35,7 @@ namespace mu
             for ( int i=0; i<pixelCount; ++i )
             {
                 uint32_t l_16 = 76 * pABuf[4*i+0] + 150 * pABuf[4*i+1] + 29 * pABuf[4*i+2];
-                pDestBuf[i] = (uint8_t)std::min( l_16>>8, 255u );
+                pDestBuf[i] = (uint8_t)FMath::Min( l_16>>8, 255u );
             }
             break;
         }
@@ -45,7 +45,7 @@ namespace mu
             for ( int i=0; i<pixelCount; ++i )
             {
                 uint32_t l_16 = 76 * pABuf[4*i+2] + 150 * pABuf[4*i+1] + 29 * pABuf[4*i+0];
-                pDestBuf[i] = (uint8_t)std::min( l_16>>8, 255u );
+                pDestBuf[i] = (uint8_t)FMath::Min( l_16>>8, 255u );
             }
             break;
         }
diff --git a/Engine/Plugins/Experimental/Mutable/Source/MutableRuntime/Private/MuR/OpImageMipmap.h b/Engine/Plugins/Experimental/Mutable/Source/MutableRuntime/Private/MuR/OpImageMipmap.h
--- a/Engine/Plugins/Experimental/Mutable/Source/MutableRuntime/Private/MuR/OpImageMipmap.h
+++ b/Engine/Plugins/Experimental/Mutable/Source/MutableRuntime/Private/MuR/OpImageMipmap.h
@@ -58,12 +58,12 @@ namespace mu
 		inline void BuildGaussianBlurKernel(float(&kernel)[N], float variance)
 		{
 			constexpr float sqrtTau = 2.5066282746310002f;
-			const auto normalDist = [sdSqrtTauInv = 1.0f / (std::sqrt(variance) * sqrtTau),
+			const auto normalDist = [sdSqrtTauInv = 1.0f / (FMath::Sqrt(variance) * sqrtTau),
 				twoVInv = 1.0f / (2.0f * variance)]
 				(float x)
 				-> float
 			{
-				return std::exp(-x * x * twoVInv) * sdSqrtTauInv;
+				return FMath::Exp(-x * x * twoVInv) * sdSqrtTauInv;
 			};
 
 			const float center = static_cast<float>(N) * 0.5f;
@@ -270,26 +270,26 @@ namespace mu
 					kernel,
 					sourceSampler,
 					pDest, destSize,
-					vec2<int32_t>(0, 0), vec2<int32_t>(destSize.x(), std::min(2, destSize.y())));
+					vec2<int32_t>(0, 0), vec2<int32_t>(destSize.x(), FMath::Min(2, destSize.y())));
 
 				GenerateMipSharpenedRegion(
 					kernel,
 					sourceSampler,
 					pDest, destSize,
-					vec2<int32_t>(0, std::max(destSize.y() - 2, 0)), vec2<int32_t>(destSize.x(), destSize.y()));
+					vec2<int32_t>(0, FMath::Max(destSize.y() - 2, 0)), vec2<int32_t>(destSize.x(), destSize.y()));
 
 				// Vertical borders
 				GenerateMipSharpenedRegion(
 					kernel,
 					sourceSampler,
 					pDest, destSize,
-					vec2<int32_t>(0, std::min(2, destSize.y())), vec2<int32_t>(std::min(2, destSize.x()), destSize.y() - 2));
+					vec2<int32_t>(0, FMath::Min(2, destSize.y())), vec2<int32_t>(FMath::Min(2, destSize.x()), destSize.y() - 2));
 
 				GenerateMipSharpenedRegion(
 					kernel,
 					sourceSampler,
 					pDest, destSize,
-					vec2<int32_t>(destSize.x() - 2, std::min(2, destSize.y())), vec2<int32_t>(destSize.x(), destSize.y() - 2));
+					vec2<int32_t>(destSize.x() - 2, FMath::Min(2, destSize.y())), vec2<int32_t>(destSize.x(), destSize.y() - 2));
 
 				sourceSize = destSize;
 
@@ -545,22 +545,22 @@ namespace mu
 			// Uncompress the last mip that we already have
 			vec2<int> uncompressedSize = pBase->CalculateMipSize(startLevel);
 			scratch->pUncompressed = new Image(
-				(uint16_t)uncompressedSize[0],
-				(uint16_t)uncompressedSize[1],
+				(uint16)uncompressedSize[0],
+				(uint16)uncompressedSize[1],
 				1,
 				EImageFormat::IF_RGBA_UBYTE);
 
 			// Generate the mipmaps from there on
 			scratch->pUncompressedMips = new Image(
-				(uint16_t)uncompressedSize[0],
-				(uint16_t)uncompressedSize[1],
+				(uint16)uncompressedSize[0],
+				(uint16)uncompressedSize[1],
 				levelCount - startLevel,
 				EImageFormat::IF_RGBA_UBYTE);
 
 			// Compress the mipmapped image
 			scratch->pCompressedMips = new Image(
-				(uint16_t)uncompressedSize[0],
-				(uint16_t)uncompressedSize[1],
+				(uint16)uncompressedSize[0],
+				(uint16)uncompressedSize[1],
 				scratch->pUncompressedMips->GetLODCount(),
 				pBase->GetFormat());
 
@@ -573,22 +573,22 @@ namespace mu
 			// Uncompress the last mip that we already have
 			vec2<int> uncompressedSize = pBase->CalculateMipSize(startLevel);
 			scratch->pUncompressed = new Image(
-				(uint16_t)uncompressedSize[0],
-				(uint16_t)uncompressedSize[1],
+				(uint16)uncompressedSize[0],
+				(uint16)uncompressedSize[1],
 				1,
 				EImageFormat::IF_L_UBYTE);
 
 			// Generate the mipmaps from there on
 			scratch->pUncompressedMips = new Image(
-				(uint16_t)uncompressedSize[0],
-				(uint16_t)uncompressedSize[1],
+				(uint16)uncompressedSize[0],
+				(uint16)uncompressedSize[1],
 				levelCount - startLevel,
 				EImageFormat::IF_L_UBYTE);
 
 			// Compress the mipmapped image
 			scratch->pCompressedMips = new Image(
-				(uint16_t)uncompressedSize[0],
-				(uint16_t)uncompressedSize[1],
+				(uint16)uncompressedSize[0],
+				(uint16)uncompressedSize[1],
 				scratch->pUncompressedMips->GetLODCount(),
 				pBase->GetFormat());
 
diff --git a/Engine/Plugins/Experimental/Mutable/Source/MutableRuntime/Private/MuR/OpImagePixelFormat.cpp b/Engine/Plugins/Experimental/Mutable/Source/MutableRuntime/Private/MuR/OpImagePixelFormat.cpp
--- a/Engine/Plugins/Experimental/Mutable/Source/MutableRuntime/Private/MuR/OpImagePixelFormat.cpp
+++ b/Engine/Plugins/Experimental/Mutable/Source/MutableRuntime/Private/MuR/OpImagePixelFormat.cpp
@@ -39,7 +39,7 @@ ImagePtr ImagePixelFormat( int imageCompressionQuality, const Image* pBase,
         resultLODCount = 1;
     }
 
-    ImagePtr result = new Image( (uint16_t)resultSize[0], (uint16_t)resultSize[1], resultLODCount, targetFormat );
+    ImagePtr result = new Image( (uint16)resultSize[0], (uint16)resultSize[1], resultLODCount, targetFormat );
 	result->m_flags = pBase->m_flags;
 	bool bSuccess = ImagePixelFormatInPlace( imageCompressionQuality, result.get(), pBase, onlyLOD );
 	int32 OriginalDataSize = FMath::Max(result->m_data.Num(), pBase->m_data.Num());
@@ -150,7 +150,7 @@ bool ImagePixelFormatInPlace( int imageCompressionQuality, Image* pResult, const
                     unsigned result = 76 * pBaseBuf[3*i+0]
                             + 150 * pBaseBuf[3*i+1]
                             + 29 * pBaseBuf[3*i+2];
-                    pDestBuf[i] = (uint8_t)std::min( 255u, result >> 8 );
+                    pDestBuf[i] = (uint8_t)FMath::Min( 255u, result >> 8 );
                 }
                 break;
             }
@@ -172,7 +172,7 @@ bool ImagePixelFormatInPlace( int imageCompressionQuality, Image* pResult, const
                     unsigned result = 76 * pBaseBuf[3*i+0]
                             + 150 * pBaseBuf[3*i+1]
                             + 29 * pBaseBuf[3*i+2];
-                    pDestBuf[i] = (uint8_t)std::min( 255u, result >> 8 );
+                    pDestBuf[i] = (uint8_t)FMath::Min( 255u, result >> 8 );
                 }
                 break;
             }
@@ -184,7 +184,7 @@ bool ImagePixelFormatInPlace( int imageCompressionQuality, Image* pResult, const
                     unsigned result = 76 * pBaseBuf[4*i+0]
                             + 150 * pBaseBuf[4*i+1]
                             + 29 * pBaseBuf[4*i+2];
-                    pDestBuf[i] = (uint8_t)std::min( 255u, result >> 8 );
+                    pDestBuf[i] = (uint8_t)FMath::Min( 255u, result >> 8 );
                 }
                 break;
             }
@@ -206,7 +206,7 @@ bool ImagePixelFormatInPlace( int imageCompressionQuality, Image* pResult, const
                     unsigned result = 76 * pBaseBuf[4*i+0]
                             + 150 * pBaseBuf[4*i+1]
                             + 29 * pBaseBuf[4*i+2];
-                    pDestBuf[i] = (uint8_t)std::min( 255u, result >> 8 );
+                    pDestBuf[i] = (uint8_t)FMath::Min( 255u, result >> 8 );
                 }
                 break;
             }
@@ -218,7 +218,7 @@ bool ImagePixelFormatInPlace( int imageCompressionQuality, Image* pResult, const
                     unsigned result = 76 * pBaseBuf[4*i+2]
                             + 150 * pBaseBuf[4*i+1]
                             + 29 * pBaseBuf[4*i+0];
-                    pDestBuf[i] = (uint8_t)std::min( 255u, result >> 8 );
+                    pDestBuf[i] = (uint8_t)FMath::Min( 255u, result >> 8 );
                 }
                 break;
             }
@@ -243,7 +243,7 @@ bool ImagePixelFormatInPlace( int imageCompressionQuality, Image* pResult, const
                     {
                         unsigned result = 76 * pBaseBuf[3 * i + 0] + 150 * pBaseBuf[3 * i + 1] +
                                           29 * pBaseBuf[3 * i + 2];
-                        pDestBuf[i] = (uint8_t)std::min(255u, result >> 8);
+                        pDestBuf[i] = (uint8_t)FMath::Min(255u, result >> 8);
                     }
                 }
                 break;
