[mutable] Optimize mesh clip versus switching of clipping mesh.

diff --git a/Engine/Plugins/Experimental/Mutable/Source/MutableRuntime/Private/MuR/OpImageProject.cpp b/Engine/Plugins/Experimental/Mutable/Source/MutableRuntime/Private/MuR/OpImageProject.cpp
--- a/Engine/Plugins/Experimental/Mutable/Source/MutableRuntime/Private/MuR/OpImageProject.cpp
+++ b/Engine/Plugins/Experimental/Mutable/Source/MutableRuntime/Private/MuR/OpImageProject.cpp
@@ -2,12 +2,6 @@
 
 
 #include "MuR/OpImageProject.h"
-
-#include "Async/ParallelFor.h"
-#include "Containers/Array.h"
-#include "Math/IntPoint.h"
-#include "Math/UnrealMathSSE.h"
-#include "Misc/AssertionMacros.h"
 #include "MuR/ConvertData.h"
 #include "MuR/Image.h"
 #include "MuR/ImagePrivate.h"
@@ -23,8 +17,11 @@
 #include "MuR/Ptr.h"
 #include "MuR/Raster.h"
 
-#include <queue>
-#include <unordered_set>
+#include "Async/ParallelFor.h"
+#include "Containers/Array.h"
+#include "Math/IntPoint.h"
+#include "Math/UnrealMathSSE.h"
+#include "Misc/AssertionMacros.h"
 
 
 namespace mu
@@ -680,7 +677,8 @@ void ImageRasterProjected_Generic( const Mesh* pMesh,
 
     // Get the vertices
     int vertexCount = pMesh->GetVertexCount();
-    vector< RasterVertex<4> > vertices( vertexCount );
+	TArray< RasterVertex<4> > vertices;
+	vertices.SetNum(vertexCount);
 
     UntypedMeshBufferIteratorConst texIt( pMesh->GetVertexBuffers(), MBS_TEXCOORDS, layout );
     UntypedMeshBufferIteratorConst posIt( pMesh->GetVertexBuffers(), MBS_POSITION, 0 );
@@ -713,7 +711,8 @@ void ImageRasterProjected_Generic( const Mesh* pMesh,
 
     // Get the indices
     int faceCount = pMesh->GetFaceCount();
-    vector<int> indices( faceCount*3 );
+	TArray<int> indices;
+	indices.SetNum(faceCount * 3);
 
     UntypedMeshBufferIteratorConst indIt( pMesh->GetIndexBuffers(), MBS_VERTEXINDEX, 0 );
     for ( int i=0; i<faceCount*3; ++i )
@@ -726,14 +725,15 @@ void ImageRasterProjected_Generic( const Mesh* pMesh,
     }
 
     // Get the block per face
-    vector<int> blocks( vertexCount );
+	TArray<int> blocks;
+	blocks.SetNum(vertexCount);
 
     UntypedMeshBufferIteratorConst bloIt( pMesh->GetVertexBuffers(), MBS_LAYOUTBLOCK, layout );
     if (bloIt.ptr())
     {
         for ( int i=0; i<vertexCount; ++i )
         {
-            uint16_t index=0;
+            uint16 index=0;
             ConvertData( 0, &index, MBF_UINT16, bloIt.ptr(), bloIt.GetFormat() );
 
             blocks[i] = index;
@@ -748,7 +748,8 @@ void ImageRasterProjected_Generic( const Mesh* pMesh,
         const Layout* pLayout = layout<pMesh->GetLayoutCount() ? pMesh->GetLayout( layout ) : nullptr;
         if (pLayout && bloIt.ptr())
         {
-            vector< box< vec2<float> > > transforms( pLayout->GetBlockCount() );
+			TArray< box< vec2<float> > > transforms;
+			transforms.SetNum(pLayout->GetBlockCount());
             for ( int b=0; b<pLayout->GetBlockCount(); ++b )
             {
                 FIntPoint grid = pLayout->GetGridSize();
@@ -860,8 +861,8 @@ void ImageRasterProjected_Optimised( const Mesh* pMesh,
     // Get the vertices
     int vertexCount = pMesh->GetVertexCount();
 
-    check( (int)scratch->vertices.size()==vertexCount );
-    check( (int)scratch->culledVertex.size()==vertexCount );
+    check( (int)scratch->vertices.Num()==vertexCount );
+    check( (int)scratch->culledVertex.Num()==vertexCount );
 
     check( pMesh->GetVertexBuffers().GetElementSize(0)==sizeof(OPTIMISED_VERTEX) );
     auto pVertices = reinterpret_cast<const OPTIMISED_VERTEX*>( pMesh->GetVertexBuffers().GetBufferData(0) );
@@ -939,8 +940,8 @@ void ImageRasterProjected_OptimisedWrapping( const Mesh* pMesh,
     // Get the vertices
     int vertexCount = pMesh->GetVertexCount();
 
-    check( (int)scratch->vertices.size()==vertexCount );
-    check( (int)scratch->culledVertex.size()==vertexCount );
+    check( (int)scratch->vertices.Num()==vertexCount );
+    check( (int)scratch->culledVertex.Num()==vertexCount );
 
     check( pMesh->GetVertexBuffers().GetElementSize(0)==sizeof(OPTIMISED_VERTEX_WRAPPING) );
     auto pVertices = reinterpret_cast<const OPTIMISED_VERTEX_WRAPPING*>( pMesh->GetVertexBuffers().GetBufferData(0) );
@@ -1275,20 +1276,23 @@ int* assert_aux = 0;
 #define assert(x) if((x) == 0) assert_aux[0] = 1;
 #endif
 
+constexpr float vert_collapse_eps = 0.0001f;
+
+
 //---------------------------------------------------------------------------------------------
 //! Create a map from vertices into vertices, collapsing vertices that have the same position
 //---------------------------------------------------------------------------------------------
 inline void MeshCreateCollapsedVertexMap( const Mesh* pMesh,
-                                    mu::vector<int>& collapsedVertexMap,
-                                    mu::vector<vec3f>& vertices,
-									std::multimap<int, int>& collapsedVertsMap
-                                    )
+	TArray<int>& collapsedVertexMap,
+	TArray<vec3f>& vertices,
+	TMultiMap<int, int>& collapsedVertsMap
+)
 {
 	MUTABLE_CPUPROFILER_SCOPE(CreateCollapseMap);
 
     int vcount = pMesh->GetVertexCount();
-    collapsedVertexMap.resize(vcount);
-    vertices.resize(vcount);
+    collapsedVertexMap.SetNum(vcount);
+    vertices.SetNum(vcount);
 
     const FMeshBufferSet& MBSPriv = pMesh->GetVertexBuffers();
     for (int32 b = 0; b < MBSPriv.m_buffers.Num(); ++b)
@@ -1324,27 +1328,27 @@ inline void MeshCreateCollapsedVertexMap( const Mesh* pMesh,
 				//Initializing Octree boundingbox
 				meshBoundingBox.min = vertices[0];
 				meshBoundingBox.size = vec3f(0.0f, 0.0f, 0.0f);
-				for (size_t i = 1; i < vertices.size(); ++i)
+				for (int32 i = 1; i < vertices.Num(); ++i)
 				{
 					meshBoundingBox.Bound(vertices[i]);
 				}
 
 				//Initializing and filling Octree
 				pOctree = new Octree(meshBoundingBox.min, meshBoundingBox.min + meshBoundingBox.size,0.001f);
-				for (size_t itr = 0; itr < vertices.size(); ++itr)
+				for (int32 itr = 0; itr < vertices.Num(); ++itr)
 				{
 					pOctree->InsertElement(vertices[itr], int(itr));
 				}
 
 				// Create map to store which vertices are the same (collapse nearby vertices)
-				for (size_t itr = 0; itr < vertices.size(); ++itr)
+				for (int32 itr = 0; itr < vertices.Num(); ++itr)
 				{
 					int collapsed_candidate_v_idx = pOctree->GetNearests(vertices[itr], vert_collapse_eps, int(itr));
 					collapsedVertexMap[itr] = collapsed_candidate_v_idx;
 
 					if (collapsed_candidate_v_idx != int(itr))
 					{
-						collapsedVertsMap.insert(std::make_pair<>(collapsed_candidate_v_idx, (int)itr));
+						collapsedVertsMap.Add(collapsed_candidate_v_idx, itr);
 					}
 				}
                 break;
@@ -1499,7 +1503,7 @@ struct NeighborFace
 	{
 		//if (a.numUVIslandChanges == b.numUVIslandChanges)
 		//{
-			return a.step > b.step;
+			return a.step < b.step;
 		//}
 
 		//return a.numUVIslandChanges > b.numUVIslandChanges;
@@ -1596,9 +1600,11 @@ void MeshProject_Optimised_Planar( const OPTIMISED_VERTEX* pVertices, int vertex
 {
 	MUTABLE_CPUPROFILER_SCOPE(MeshProject_Optimised_Planar)
 
-    vector<int32> oldToNewVertex(vertexCount,-1);
+	TArray<int32> oldToNewVertex;
+	oldToNewVertex.Init(-1,vertexCount);
 
-    vector<PROJECTED_VERTEX> projectedPositions(vertexCount);
+	TArray<PROJECTED_VERTEX> projectedPositions;
+	projectedPositions.SetNumZeroed(vertexCount);
 
     for ( int v=0; v<vertexCount; ++v )
     {
@@ -1678,9 +1684,11 @@ void MeshProject_Optimised_Cylindrical( const OPTIMISED_VERTEX* pVertices, int v
 {
 	MUTABLE_CPUPROFILER_SCOPE(MeshProject_Optimised_Cylindrical)
 
-	vector<int32> oldToNewVertex(vertexCount,-1);
+	TArray<int32> oldToNewVertex;
+	oldToNewVertex.Init( -1, vertexCount);
 
-    vector<PROJECTED_VERTEX> projectedPositions(vertexCount);
+	TArray<PROJECTED_VERTEX> projectedPositions;
+	projectedPositions.SetNumZeroed(vertexCount);
 
     // TODO: support for non uniform scale?
     float radius = projectorScale[1];
@@ -1780,30 +1788,33 @@ void MeshProject_Optimised_Wrapping( const Mesh* pMesh,
     auto pIndices = reinterpret_cast<const uint32*>( pMesh->GetIndexBuffers().GetBufferData(0) );
     int faceCount = pMesh->GetFaceCount();
 
-    vector<PROJECTED_VERTEX> projectedPositions(vertexCount);
+	TArray<PROJECTED_VERTEX> projectedPositions;
+	projectedPositions.SetNum(vertexCount);
 
     // Iterate the faces and trace a ray to find the origin face of the projection
     const float maxDist = 100000.f;
     float min_t = maxDist;
     float rayLength = maxDist;
     int intersectedFace = -1;
-    std::set<int> processedVertices;
+	TSet<int> processedVertices;
     vec3f projectionPlaneNormal;
     vec3f out_intersection;
 
-    vector<AdjacentFaces> faceConnectivity(faceCount);
-    std::set<int> processedFaces;
-    std::set<int> discardedWrapAroundFaces;
-    vector<int> faceStep(faceCount, 0);
+	TArray<AdjacentFaces> faceConnectivity;
+	faceConnectivity.SetNum(faceCount);
+    TSet<int> processedFaces;
+	TSet<int> discardedWrapAroundFaces;
+	TArray<int> faceStep;
+	faceStep.SetNum(faceCount);
 
     // Map vertices to the one they are collapsed to because they are very similar, if they aren't collapsed then they are mapped to themselves
-    vector<int> collapsedVertexMap;
-    vector<vec3f> vertices;
-    std::multimap<int, int> collapsedVertsMap; // Maps a collapsed vertex to all the vertices that collapse to it
+    TArray<int> collapsedVertexMap;
+	TArray<vec3f> vertices;
+    TMultiMap<int, int> collapsedVertsMap; // Maps a collapsed vertex to all the vertices that collapse to it
     MeshCreateCollapsedVertexMap(pMesh, collapsedVertexMap, vertices, collapsedVertsMap);
 
     // Used to speed up connectivity building
-    std::multimap<int, int> vertToFacesMap;
+	TMultiMap<int, int> vertToFacesMap;
 
     for (int f = 0; f < faceCount; ++f)
     {
@@ -1811,9 +1822,9 @@ void MeshProject_Optimised_Wrapping( const Mesh* pMesh,
         int i1 = pIndices[f * 3 + 1];
         int i2 = pIndices[f * 3 + 2];
 
-        vertToFacesMap.insert(std::pair<int, int>(collapsedVertexMap[i0], f));
-        vertToFacesMap.insert(std::pair<int, int>(collapsedVertexMap[i1], f));
-        vertToFacesMap.insert(std::pair<int, int>(collapsedVertexMap[i2], f));
+        vertToFacesMap.Add(collapsedVertexMap[i0], f);
+        vertToFacesMap.Add(collapsedVertexMap[i1], f);
+        vertToFacesMap.Add(collapsedVertexMap[i2], f);
         //vertToFacesMap.insert(std::pair<int, int>(i0, f));
         //vertToFacesMap.insert(std::pair<int, int>(i1, f));
         //vertToFacesMap.insert(std::pair<int, int>(i2, f));
@@ -1856,15 +1867,10 @@ void MeshProject_Optimised_Wrapping( const Mesh* pMesh,
             int v = collapsedVertexMap[pIndices[f * 3 + i]];
             //int v = pIndices[f * 3 + i];
 
-            for (auto it = vertToFacesMap.find(v); it != vertToFacesMap.end(); ++it)
+			TArray<int> FoundValues;
+			vertToFacesMap.MultiFind(v, FoundValues);
+            for (int f2 : FoundValues )
             {
-                if (it->first != v) // The return iterator points at the first instance of v, but will iterate to the end of the map
-                {
-                    break;
-                }
-
-                int f2 = it->second;
-
                 if (f != f2 && f2 >= 0)
                 {
                     int commonVertices = 0;
@@ -1960,8 +1966,11 @@ void MeshProject_Optimised_Wrapping( const Mesh* pMesh,
     // Do a BFS walk of the mesh, unfolding a face at each step
     if (intersectedFace >= 0)
     {
-        std::priority_queue<NeighborFace> pendingFaces; // Queue of pending face + new vertex
-        std::unordered_set<int> pendingFacesUnique; // Used to quickly check uniqueness in the pendingFaces queue
+        TArray<NeighborFace> pendingFaces; // Queue of pending face + new vertex
+        
+		// \TODO: Bool array to speed up?
+		TSet<int> pendingFacesUnique; // Used to quickly check uniqueness in the pendingFaces queue
+
         //vec3f hitFaceProjectedNormal;
         bool hitFaceHasPositiveArea = false;
 
@@ -1972,25 +1981,25 @@ void MeshProject_Optimised_Wrapping( const Mesh* pMesh,
         neighborFace.numUVIslandChanges = 0;
         neighborFace.step = 0;
         neighborFace.changesUVIsland = false;
-        pendingFaces.push(neighborFace);
-        pendingFacesUnique.insert(intersectedFace);
+        pendingFaces.HeapPush(neighborFace);
+        pendingFacesUnique.Add(intersectedFace);
 
         int step = 0;
         float totalUVAreaCovered = 0.f;
-        vector<int> oldVertices;
-        oldVertices.reserve(3);
+        TArray<int> oldVertices;
+        oldVertices.Reserve(3);
 
-        while (!pendingFaces.empty())
+        while (!pendingFaces.IsEmpty())
         {
-            NeighborFace currentFaceStruct = pendingFaces.top();
+			NeighborFace currentFaceStruct;
+			pendingFaces.HeapPop(currentFaceStruct);
             int currentFace = currentFaceStruct.neighborFace;
 #ifdef DEBUG_PROJECTION
             int newVertexFromQueue = currentFaceStruct.newVertex;
 #endif
-            pendingFaces.pop();
-            pendingFacesUnique.erase(currentFace);
+            pendingFacesUnique.Remove(currentFace);
 
-            processedFaces.insert(currentFace);
+            processedFaces.Add(currentFace);
             faceStep[currentFace] = step;
             float currentTriangleArea = -2.f;
 
@@ -2059,7 +2068,7 @@ void MeshProject_Optimised_Wrapping( const Mesh* pMesh,
 
                     int collapsedVert = collapsedVertexMap[v];
 
-                    processedVertices.insert(collapsedVert);
+                    processedVertices.Add(collapsedVert);
 
                     //PlanarlyProjectVertex(pVertices[collapsedVert].pos, projectedPositions[collapsedVert], projector2, projectorPosition2, projectorDirection2, s2, u2);
                     vec2f projectedVertex = ChangeBase2D(pVertices[v].uv, OrigUVs_Origin, OrigUVs_BaseU, OrigUVs_BaseV);
@@ -2096,15 +2105,11 @@ void MeshProject_Optimised_Wrapping( const Mesh* pMesh,
                                                            projectedPositions[collapsedVert].pos1 >= 0.0f && projectedPositions[collapsedVert].pos1 <= 1.0f;
 
                     // Copy the new info to all the vertices that collapse to the same vertex
-                    for (auto it = collapsedVertsMap.find(collapsedVert); it != collapsedVertsMap.end(); ++it)
+					TArray<int> FoundValues;
+					collapsedVertsMap.MultiFind(collapsedVert,FoundValues);
+                    for (int otherVert:FoundValues)
                     {
-                        if (it->first != collapsedVert)
-                        {
-                            break;
-                        }
-
-                        int otherVert = it->second;
-                        processedVertices.insert(otherVert);
+                        processedVertices.Add(otherVert);
 
                         projectedPositions[otherVert] = projectedPositions[collapsedVert];
                     }
@@ -2130,14 +2135,14 @@ void MeshProject_Optimised_Wrapping( const Mesh* pMesh,
 #endif
                 int newVertex = -1;
                 //int oldVertices[2] = { -1, -1 };
-                oldVertices.clear();
+                oldVertices.Empty();
                 bool reverseOrder = false;
 
                 for (int i = 0; i < 3; ++i)
                 {
                     int v = pIndices[currentFace * 3 + i];
 
-                    if (processedVertices.count(collapsedVertexMap[v]) == 0)
+                    if (!processedVertices.Contains(collapsedVertexMap[v]))
                     {
                         newVertex = v;
 
@@ -2151,11 +2156,11 @@ void MeshProject_Optimised_Wrapping( const Mesh* pMesh,
                         //int index = oldVertices[0] == -1 ? 0 : 1;
                         //oldVertices[index] = v;
                         //oldVertices.push_back(collapsedVertexMap[v]);
-                        oldVertices.push_back(v);
+                        oldVertices.Add(v);
                     }
                 }
 
-                if (reverseOrder && oldVertices.size() == 2)
+                if (reverseOrder && oldVertices.Num() == 2)
                 {
                     int aux = oldVertices[0];
                     oldVertices[0] = oldVertices[1];
@@ -2168,7 +2173,7 @@ void MeshProject_Optimised_Wrapping( const Mesh* pMesh,
                     assert(newVertex >= 0);
                     //assert(collapsedVertexMap[newVertex] == newVertexFromQueue);
                     assert(oldVertices[0] >= 0 && oldVertices[1] >= 0);
-                    assert(oldVertices.size() == 2);
+                    assert(oldVertices.Num() == 2);
                     assert(newVertex < vertexCount && oldVertices[0] < vertexCount && oldVertices[1] < vertexCount);
 #endif
 
@@ -2301,7 +2306,7 @@ void MeshProject_Optimised_Wrapping( const Mesh* pMesh,
 #endif
 
                     int collapsedNewVertex = collapsedVertexMap[newVertex];
-                    processedVertices.insert(collapsedNewVertex);
+                    processedVertices.Add(collapsedNewVertex);
 
                     projectedPositions[collapsedNewVertex].pos0 = newVertex_proj[0];
                     projectedPositions[collapsedNewVertex].pos1 = newVertex_proj[1];
@@ -2310,15 +2315,11 @@ void MeshProject_Optimised_Wrapping( const Mesh* pMesh,
                                                                 newVertex_proj[1] >= 0.0f && newVertex_proj[1] <= 1.0f;
 
                     // Copy the new info to all the vertices that collapse to the same vertex
-                    for (auto it = collapsedVertsMap.find(collapsedNewVertex); it != collapsedVertsMap.end(); ++it)
+					TArray<int> FoundValues;
+					collapsedVertsMap.MultiFind(collapsedNewVertex,FoundValues);
+                    for (int otherVert : FoundValues)
                     {
-                        if (it->first != collapsedNewVertex)
-                        {
-                            break;
-                        }
-
-                        int otherVert = it->second;
-                        processedVertices.insert(otherVert);
+                        processedVertices.Add(otherVert);
 
                         projectedPositions[otherVert] = projectedPositions[collapsedNewVertex];
                     }
@@ -2334,7 +2335,7 @@ void MeshProject_Optimised_Wrapping( const Mesh* pMesh,
                 else
                 {
 #ifdef DEBUG_PROJECTION
-                    assert(oldVertices.size() == 3);
+                    assert(oldVertices.Num() == 3);
 #endif
                 }
 
@@ -2352,7 +2353,7 @@ void MeshProject_Optimised_Wrapping( const Mesh* pMesh,
                 //if (hitFaceProjectedNormal.z() * currentFaceNormal.z() < 0)
                 if(hitFaceHasPositiveArea != currentTriangleHasPositiveArea) // Is the current face wound in the opposite direction?
                 {
-                    discardedWrapAroundFaces.insert(currentFace); // If so, discard it since it's probably a wrap-around face
+                    discardedWrapAroundFaces.Add(currentFace); // If so, discard it since it's probably a wrap-around face
                 }
             }
 
@@ -2384,7 +2385,7 @@ void MeshProject_Optimised_Wrapping( const Mesh* pMesh,
                 }
             }
 
-            if (anyVertexInUVSpace && anyVertexInObjSpaceRange && discardedWrapAroundFaces.count(currentFace) == 0)
+            if (anyVertexInUVSpace && anyVertexInObjSpaceRange && !discardedWrapAroundFaces.Contains(currentFace) )
             {
 #ifdef DEBUG_PROJECTION
                 assert(currentTriangleArea != -2.f);
@@ -2397,7 +2398,7 @@ void MeshProject_Optimised_Wrapping( const Mesh* pMesh,
                     int newVertex = faceConnectivity[currentFace].newVertices[i];
                     bool changesUVIsland = faceConnectivity[currentFace].changesUVIsland[i];
 
-                    if (neighborFace2 >= 0 && processedFaces.count(neighborFace2) == 0 && pendingFacesUnique.count(neighborFace2) == 0)
+                    if (neighborFace2 >= 0 && !processedFaces.Contains(neighborFace2) && pendingFacesUnique.Contains(neighborFace2) == 0)
                     {
                         NeighborFace neighborFaceStruct;
                         neighborFaceStruct.neighborFace = neighborFace2;
@@ -2407,14 +2408,14 @@ void MeshProject_Optimised_Wrapping( const Mesh* pMesh,
                                                                                                    : currentFaceStruct.numUVIslandChanges;
                         neighborFaceStruct.step = step;
                         neighborFaceStruct.changesUVIsland = changesUVIsland;
-                        pendingFaces.push(neighborFaceStruct);
-                        pendingFacesUnique.insert(neighborFace2);
+                        pendingFaces.Add(neighborFaceStruct);
+                        pendingFacesUnique.Add(neighborFace2);
                     }
                 }
             }
             else
             {
-                discardedWrapAroundFaces.insert(currentFace);
+                discardedWrapAroundFaces.Add(currentFace);
             }
 
             //if(step == 1000)
@@ -2431,12 +2432,13 @@ void MeshProject_Optimised_Wrapping( const Mesh* pMesh,
         }
     }
 
-    vector<int32> oldToNewVertex(vertexCount,-1);
+	TArray<int32> oldToNewVertex;
+	oldToNewVertex.Init(-1,vertexCount);
 
     // Add the projected face
     for(int f : processedFaces)
     {
-        if (discardedWrapAroundFaces.count(f) > 0)
+        if (discardedWrapAroundFaces.Contains(f))
         {
             continue;
         }
diff --git a/Engine/Plugins/Experimental/Mutable/Source/MutableRuntime/Private/MuR/OpImageProject.h b/Engine/Plugins/Experimental/Mutable/Source/MutableRuntime/Private/MuR/OpImageProject.h
--- a/Engine/Plugins/Experimental/Mutable/Source/MutableRuntime/Private/MuR/OpImageProject.h
+++ b/Engine/Plugins/Experimental/Mutable/Source/MutableRuntime/Private/MuR/OpImageProject.h
@@ -14,8 +14,8 @@ namespace mu
 
     struct SCRATCH_IMAGE_PROJECT
     {
-        vector< RasterVertex<4> > vertices;
-        vector<uint8> culledVertex;
+        TArray< RasterVertex<4> > vertices;
+		TArray<uint8> culledVertex;
     };
 
 
diff --git a/Engine/Plugins/Experimental/Mutable/Source/MutableRuntime/Private/MuR/OpImageRasterMesh.h b/Engine/Plugins/Experimental/Mutable/Source/MutableRuntime/Private/MuR/OpImageRasterMesh.h
--- a/Engine/Plugins/Experimental/Mutable/Source/MutableRuntime/Private/MuR/OpImageRasterMesh.h
+++ b/Engine/Plugins/Experimental/Mutable/Source/MutableRuntime/Private/MuR/OpImageRasterMesh.h
@@ -33,7 +33,8 @@ namespace mu
 
 		// Get the vertices
 		int vertexCount = pMesh->GetVertexCount();
-		vector< RasterVertex<1> > vertices( vertexCount );
+		TArray< RasterVertex<1> > vertices;
+		vertices.SetNumZeroed(vertexCount);
 
 		UntypedMeshBufferIteratorConst texIt( pMesh->GetVertexBuffers(), MBS_TEXCOORDS, 0 );
 		for ( int v=0; v<vertexCount; ++v )
@@ -49,7 +50,8 @@ namespace mu
 
 		// Get the indices
 		int faceCount = pMesh->GetFaceCount();
-		vector<int> indices( faceCount*3 );
+		TArray<int> indices;
+		indices.SetNumZeroed(faceCount * 3);
 
 		UntypedMeshBufferIteratorConst indIt( pMesh->GetIndexBuffers(), MBS_VERTEXINDEX, 0 );
 		for ( int i=0; i<faceCount*3; ++i )
@@ -89,11 +91,12 @@ namespace mu
 			// Raster only the faces in the selected block
 
 			// Get the block per face
-			vector<int> blocks( vertexCount );
+			TArray<int> blocks;
+			blocks.SetNumZeroed(vertexCount);
 
 			for ( int i=0; i<vertexCount; ++i )
 			{
-                uint16_t index=0;
+                uint16 index=0;
 				ConvertData( 0, &index, MBF_UINT16, bloIt.ptr(), bloIt.GetFormat() );
 
 				blocks[i] = index;
diff --git a/Engine/Plugins/Experimental/Mutable/Source/MutableRuntime/Private/MuR/OpImageResize.h b/Engine/Plugins/Experimental/Mutable/Source/MutableRuntime/Private/MuR/OpImageResize.h
--- a/Engine/Plugins/Experimental/Mutable/Source/MutableRuntime/Private/MuR/OpImageResize.h
+++ b/Engine/Plugins/Experimental/Mutable/Source/MutableRuntime/Private/MuR/OpImageResize.h
@@ -836,7 +836,7 @@ namespace mu
 					format);
 			}
 
-			FImageSize blockSize((uint16_t)size[0], (uint16_t)size[1]);
+			FImageSize blockSize((uint16)size[0], (uint16)size[1]);
 			Result = ImageResizeLinear(imageCompressionQuality, pImage, blockSize);
 		}
 
diff --git a/Engine/Plugins/Experimental/Mutable/Source/MutableRuntime/Private/MuR/OpImageSaturate.h b/Engine/Plugins/Experimental/Mutable/Source/MutableRuntime/Private/MuR/OpImageSaturate.h
--- a/Engine/Plugins/Experimental/Mutable/Source/MutableRuntime/Private/MuR/OpImageSaturate.h
+++ b/Engine/Plugins/Experimental/Mutable/Source/MutableRuntime/Private/MuR/OpImageSaturate.h
@@ -13,7 +13,7 @@ namespace mu
 	{
 		// Clamp the factor
 		// TODO: See what happens if we don't
-		factor = std::max( 0.0f, factor );
+		factor = FMath::Max( 0.0f, factor );
         int32_t f_8 = (int32_t)(factor*255);
 
         ImagePtr pDest = new Image( pA->GetSizeX(), pA->GetSizeY(),
@@ -38,7 +38,7 @@ namespace mu
 				{
                     int32_t d_16 = (pABuf[3*i+c]<<8) - l_16;
                     int32_t r_16 = l_16 + ( (d_16 * f_8) >> 8 );
-                    pDestBuf[3*i+c] = (uint8_t)std::min( (uint32_t)r_16>>8, 255u );
+                    pDestBuf[3*i+c] = (uint8_t)FMath::Min( (uint32_t)r_16>>8, 255u );
 				}
 			}
 			break;
@@ -54,7 +54,7 @@ namespace mu
                 {
                     int32_t d_16 = (pABuf[4*i+c]<<8) - l_16;
                     int32_t r_16 = l_16 + ( (d_16 * f_8) >> 8 );
-                    pDestBuf[4*i+c] = (uint8_t)std::min( (uint32_t)r_16>>8, 255u );
+                    pDestBuf[4*i+c] = (uint8_t)FMath::Min( (uint32_t)r_16>>8, 255u );
                 }
 
                 pDestBuf[4*i+3] = pABuf[4*i+3];
@@ -72,7 +72,7 @@ namespace mu
                 {
                     int32_t d_16 = (pABuf[4*i+c]<<8) - l_16;
                     int32_t r_16 = l_16 + ( (d_16 * f_8) >> 8 );
-                    pDestBuf[4*i+c] = (uint8_t)std::min( (uint32_t)r_16>>8, 255u );
+                    pDestBuf[4*i+c] = (uint8_t)FMath::Min( (uint32_t)r_16>>8, 255u );
                 }
 
                 pDestBuf[4*i+3] = pABuf[4*i+3];
diff --git a/Engine/Plugins/Experimental/Mutable/Source/MutableRuntime/Private/MuR/OpImageSelectColour.h b/Engine/Plugins/Experimental/Mutable/Source/MutableRuntime/Private/MuR/OpImageSelectColour.h
--- a/Engine/Plugins/Experimental/Mutable/Source/MutableRuntime/Private/MuR/OpImageSelectColour.h
+++ b/Engine/Plugins/Experimental/Mutable/Source/MutableRuntime/Private/MuR/OpImageSelectColour.h
@@ -24,7 +24,7 @@ namespace mu
 		case EImageFormat::IF_L_UBYTE:
 		{
             uint8_t c;
-            c = (uint8_t)(std::max( 0.0f, std::min( 255.0f, colour[0] * 255.0f ) )+0.5f);
+            c = (uint8_t)(FMath::Max( 0.0f, FMath::Min( 255.0f, colour[0] * 255.0f ) )+0.5f);
 
 			for ( int i=0; i<pixelCount; ++i )
 			{
@@ -43,9 +43,9 @@ namespace mu
 		case EImageFormat::IF_RGB_UBYTE:
 		{
             uint8_t c[3];
-            c[0] = (uint8_t)(std::max( 0.0f, std::min( 255.0f, colour[0] * 255.0f ) )+0.5f);
-            c[1] = (uint8_t)(std::max( 0.0f, std::min( 255.0f, colour[1] * 255.0f ) )+0.5f);
-            c[2] = (uint8_t)(std::max( 0.0f, std::min( 255.0f, colour[2] * 255.0f ) )+0.5f);
+            c[0] = (uint8_t)(FMath::Max( 0.0f, FMath::Min( 255.0f, colour[0] * 255.0f ) )+0.5f);
+            c[1] = (uint8_t)(FMath::Max( 0.0f, FMath::Min( 255.0f, colour[1] * 255.0f ) )+0.5f);
+            c[2] = (uint8_t)(FMath::Max( 0.0f, FMath::Min( 255.0f, colour[2] * 255.0f ) )+0.5f);
 
 			for ( int i=0; i<pixelCount; ++i )
 			{
@@ -64,9 +64,9 @@ namespace mu
         case EImageFormat::IF_RGBA_UBYTE:
         {
             uint8_t c[3];
-            c[0] = (uint8_t)(std::max( 0.0f, std::min( 255.0f, colour[0] * 255.0f ) )+0.5f);
-            c[1] = (uint8_t)(std::max( 0.0f, std::min( 255.0f, colour[1] * 255.0f ) )+0.5f);
-            c[2] = (uint8_t)(std::max( 0.0f, std::min( 255.0f, colour[2] * 255.0f ) )+0.5f);
+            c[0] = (uint8_t)(FMath::Max( 0.0f, FMath::Min( 255.0f, colour[0] * 255.0f ) )+0.5f);
+            c[1] = (uint8_t)(FMath::Max( 0.0f, FMath::Min( 255.0f, colour[1] * 255.0f ) )+0.5f);
+            c[2] = (uint8_t)(FMath::Max( 0.0f, FMath::Min( 255.0f, colour[2] * 255.0f ) )+0.5f);
 
             for ( int i=0; i<pixelCount; ++i )
             {
@@ -85,9 +85,9 @@ namespace mu
         case EImageFormat::IF_BGRA_UBYTE:
         {
             uint8_t c[3];
-            c[0] = (uint8_t)(std::max( 0.0f, std::min( 255.0f, colour[2] * 255.0f ) )+0.5f);
-            c[1] = (uint8_t)(std::max( 0.0f, std::min( 255.0f, colour[1] * 255.0f ) )+0.5f);
-            c[2] = (uint8_t)(std::max( 0.0f, std::min( 255.0f, colour[0] * 255.0f ) )+0.5f);
+            c[0] = (uint8_t)(FMath::Max( 0.0f, FMath::Min( 255.0f, colour[2] * 255.0f ) )+0.5f);
+            c[1] = (uint8_t)(FMath::Max( 0.0f, FMath::Min( 255.0f, colour[1] * 255.0f ) )+0.5f);
+            c[2] = (uint8_t)(FMath::Max( 0.0f, FMath::Min( 255.0f, colour[0] * 255.0f ) )+0.5f);
 
             for ( int i=0; i<pixelCount; ++i )
             {
diff --git a/Engine/Plugins/Experimental/Mutable/Source/MutableRuntime/Private/MuR/OpImageSoftLight.h b/Engine/Plugins/Experimental/Mutable/Source/MutableRuntime/Private/MuR/OpImageSoftLight.h
--- a/Engine/Plugins/Experimental/Mutable/Source/MutableRuntime/Private/MuR/OpImageSoftLight.h
+++ b/Engine/Plugins/Experimental/Mutable/Source/MutableRuntime/Private/MuR/OpImageSoftLight.h
@@ -30,7 +30,7 @@ namespace mu
                                 const Image* pBase,
                                 const Image* pMask,
                                 vec3<float> col,
-                                vector< std::pair<uint8_t,uint8_t*> >* pTables )
+                                TArray< TPair<uint8_t,uint8_t*> >* pTables )
 	{
         check( pResult->GetFormat() == pBase->GetFormat() );
         check( pResult->GetSizeX() == pBase->GetSizeX() );
@@ -54,15 +54,15 @@ namespace mu
 
 				for ( int t=0; t<3; ++t )
 				{
-                    uint8_t c = uint8_t( std::min( 255.0f, std::max( 0.0f, col[t]*255.0f ) ) );
+                    uint8_t c = uint8_t( FMath::Min( 255.0f, FMath::Max( 0.0f, col[t]*255.0f ) ) );
 
 					// Try to find an existing table
 					tables[t] = 0;
-					for ( size_t j=0; !tables[t] && j<pTables->size(); ++j )
+					for ( size_t j=0; !tables[t] && j<pTables->Num(); ++j )
 					{
-						if ( (*pTables)[j].first==c )
+						if ( (*pTables)[j].Key==c )
 						{
-							tables[t] = (*pTables)[j].second;
+							tables[t] = (*pTables)[j].Value;
 						}
 					}
 
@@ -73,7 +73,7 @@ namespace mu
 						{
                             tables[t][i] = (uint8_t)SoftLightChannel( i, c );
 						}
-                        pTables->push_back(std::pair<uint8_t,uint8_t*>( c, tables[t] ) );
+                        pTables->Add(TPair<uint8_t,uint8_t*>( c, tables[t] ) );
 					}
 				}
 
@@ -95,7 +95,7 @@ namespace mu
                 uint8_t tables[256*3];
                 for ( int t=0; t<3; ++t )
                 {
-                    uint8_t c = uint8_t( std::min( 255.0f, std::max( 0.0f, col[t]*255.0f ) ) );
+                    uint8_t c = uint8_t( FMath::Min( 255.0f, FMath::Max( 0.0f, col[t]*255.0f ) ) );
                     for ( int i=0; i<256;++i )
                     {
                         tables[t*256+i] = (uint8_t)SoftLightChannel( i, c );
diff --git a/Engine/Plugins/Experimental/Mutable/Source/MutableRuntime/Private/MuR/OpImageSwizzle.h b/Engine/Plugins/Experimental/Mutable/Source/MutableRuntime/Private/MuR/OpImageSwizzle.h
--- a/Engine/Plugins/Experimental/Mutable/Source/MutableRuntime/Private/MuR/OpImageSwizzle.h
+++ b/Engine/Plugins/Experimental/Mutable/Source/MutableRuntime/Private/MuR/OpImageSwizzle.h
@@ -27,8 +27,8 @@ namespace mu
 
         // Pixelcount should already match, but due to bugs it may not be the case. Try to detect it,
         // but avoid crashing below:
-        uint16_t numChannels = GetImageFormatData(format).m_channels;
-        for (uint16_t c=0;c<numChannels; ++c)
+        uint16 numChannels = GetImageFormatData(format).m_channels;
+        for (uint16 c=0;c<numChannels; ++c)
         {
             if (pSources[c])
             {
diff --git a/Engine/Plugins/Experimental/Mutable/Source/MutableRuntime/Private/MuR/OpLayoutPack.h b/Engine/Plugins/Experimental/Mutable/Source/MutableRuntime/Private/MuR/OpLayoutPack.h
--- a/Engine/Plugins/Experimental/Mutable/Source/MutableRuntime/Private/MuR/OpLayoutPack.h
+++ b/Engine/Plugins/Experimental/Mutable/Source/MutableRuntime/Private/MuR/OpLayoutPack.h
@@ -28,8 +28,8 @@ namespace mu
 	// 		box< vec2<int> > b;
 	// 		pSource->GetBlock( index, &b.min[0], &b.min[1], &b.size[0], &b.size[1] );
 
-	// 		maxX = std::max( maxX, b.size[0] );
-	// 		maxY = std::max( maxY, b.size[1] );
+	// 		maxX = FMath::Max( maxX, b.size[0] );
+	// 		maxY = FMath::Max( maxY, b.size[1] );
 
 	// 		area += b.size[0] * b.size[1];
 
@@ -93,7 +93,7 @@ namespace mu
 	// 			int level = 0;
 	// 			for( int xs=x; xs<x+blocks[index][0]; ++xs )
 	// 			{
-	// 				level = std::max( level, horizon[xs] );
+	// 				level = FMath::Max( level, horizon[xs] );
 	// 			}
 
 	// 			if (level<currentLevel)
@@ -114,7 +114,7 @@ namespace mu
 
 	// 		// Store
 	// 		positions[ index ] = vec2<int>( currentX, currentLevel );
-	// 		maxY = std::max( maxY, currentLevel+blocks[index][1] );
+	// 		maxY = FMath::Max( maxY, currentLevel+blocks[index][1] );
 	// 	}
 
     //     // Set data in the result
@@ -247,8 +247,8 @@ namespace mu
     //         box< vec2<int> > b;
     //         pSourceLayout->GetBlock( index, &b.min[0], &b.min[1], &b.size[0], &b.size[1] );
 
-    //         maxX = std::max( maxX, b.size[0] );
-    //         maxY = std::max( maxY, b.size[1] );
+    //         maxX = FMath::Max( maxX, b.size[0] );
+    //         maxY = FMath::Max( maxY, b.size[1] );
 
     //         area += b.size[0] * b.size[1];
 
@@ -304,7 +304,7 @@ namespace mu
     //                 int level = 0;
     //                 for( int xs=x; xs<x+scratch->sorted[p].size[0]; ++xs )
     //                 {
-    //                     level = std::max( level, (int)horizon[xs] );
+    //                     level = FMath::Max( level, (int)horizon[xs] );
     //                 }
 
     //                 if (level<currentLevel)
@@ -320,12 +320,12 @@ namespace mu
     //             // Update horizon
     //             for( int xs=currentX; xs<currentX+scratch->sorted[p].size[0]; ++xs )
     //             {
-    //                 horizon[xs] = (uint16_t)(currentLevel+scratch->sorted[p].size[1]);
+    //                 horizon[xs] = (uint16)(currentLevel+scratch->sorted[p].size[1]);
     //             }
 
     //             // Store
     //             scratch->positions[ scratch->sorted[p].index ] = vec2<int>( currentX, currentLevel );
-    //             maxY = std::max( maxY, currentLevel+scratch->sorted[p].size[1] );
+    //             maxY = FMath::Max( maxY, currentLevel+scratch->sorted[p].size[1] );
     //         }
 
     //         maxY = ceilPow2( maxY );
@@ -442,7 +442,8 @@ namespace mu
 
 		// Number of blocks alrady packed
 		size_t packed = 0;
-		vector<uint8_t> packedFlag(scratch->sorted.Num(), 0);
+		TArray<uint8_t> packedFlag;
+		packedFlag.SetNumZeroed(scratch->sorted.Num());
 
 		// Pack with fixed horizontal size
 		check(*maxX < 256);
@@ -482,7 +483,7 @@ namespace mu
 					int level = 0;
 					for (int xs = x; xs < x + candidateSizeX; ++xs)
 					{
-						level = std::max(level, (int)horizon[xs]);
+						level = FMath::Max(level, (int)horizon[xs]);
 					}
 
 					if (level < currentLevel)
@@ -566,7 +567,7 @@ namespace mu
 				// Update horizon
 				for (int xs = bestX; xs < bestX + scratch->sorted[best].size[0]; ++xs)
 				{
-					horizon[xs] = (uint16_t)(bestLevel + scratch->sorted[best].size[1]);
+					horizon[xs] = (uint16)(bestLevel + scratch->sorted[best].size[1]);
 				}
 			}
 
diff --git a/Engine/Plugins/Experimental/Mutable/Source/MutableRuntime/Private/MuR/OpLayoutRemoveBlocks.h b/Engine/Plugins/Experimental/Mutable/Source/MutableRuntime/Private/MuR/OpLayoutRemoveBlocks.h
--- a/Engine/Plugins/Experimental/Mutable/Source/MutableRuntime/Private/MuR/OpLayoutRemoveBlocks.h
+++ b/Engine/Plugins/Experimental/Mutable/Source/MutableRuntime/Private/MuR/OpLayoutRemoveBlocks.h
@@ -14,18 +14,19 @@ namespace mu
     inline LayoutPtr LayoutRemoveBlocks( const Layout* pSource, const Mesh* pMesh, int layoutIndex )
     {
         // Create the list of blocks in the mesh
-        vector<bool> blocksFound(1024, false);
+		TArray<bool> blocksFound;
+		blocksFound.SetNumZeroed(1024);
 
         UntypedMeshBufferIteratorConst itBlocks( pMesh->GetVertexBuffers(), MBS_LAYOUTBLOCK,
                                                  layoutIndex );
         if ( itBlocks.GetFormat()==MBF_UINT16 )
         {
-            const uint16_t* pBlocks = reinterpret_cast<const uint16_t*>( itBlocks.ptr() );
+            const uint16* pBlocks = reinterpret_cast<const uint16*>( itBlocks.ptr() );
             for ( int i=0; i<pMesh->GetVertexCount(); ++i )
             {
-                if (pBlocks[i]>=blocksFound.size())
+                if (pBlocks[i]>=blocksFound.Num())
                 {
-                    blocksFound.resize(pBlocks[i]+1024,false);
+                    blocksFound.SetNumZeroed(pBlocks[i]+1024);
                 }
 
                 blocksFound[ pBlocks[i] ] = true;
@@ -49,7 +50,7 @@ namespace mu
 		for ( int32 b=0; b<pResult->m_blocks.Num(); ++b )
 		{
             int blockIndex = pResult->m_blocks[b].m_id;
-            if ( blockIndex<(int)blocksFound.size() && blocksFound[blockIndex] )
+            if ( blockIndex<(int)blocksFound.Num() && blocksFound[blockIndex] )
 			{
 				// keep
 				pResult->m_blocks[dest] = pResult->m_blocks[b];
diff --git a/Engine/Plugins/Experimental/Mutable/Source/MutableRuntime/Private/MuR/OpMeshApplyLayout.h b/Engine/Plugins/Experimental/Mutable/Source/MutableRuntime/Private/MuR/OpMeshApplyLayout.h
--- a/Engine/Plugins/Experimental/Mutable/Source/MutableRuntime/Private/MuR/OpMeshApplyLayout.h
+++ b/Engine/Plugins/Experimental/Mutable/Source/MutableRuntime/Private/MuR/OpMeshApplyLayout.h
@@ -54,7 +54,8 @@ namespace mu
 		{
 			FVector2f min, size;
 		};
-		vector< Box > transforms( pLayout->GetBlockCount() );
+		TArray< Box > transforms;
+		transforms.SetNum(pLayout->GetBlockCount());
 		for ( int b=0; b<pLayout->GetBlockCount(); ++b )
 		{
 			FIntPoint grid = pLayout->GetGridSize();
@@ -135,7 +136,7 @@ namespace mu
                     uint64_t u_32 = pUV[c];
                     uint64_t u_48 = u_32 * ((uint64_t)(((float)0xffff) * transforms[ relBlock ].size[c]));
                     u_48 += (uint64_t)(((float)0xffffffffffULL) * transforms[ relBlock ].min[c]);
-                    pUV[c] = (uint16_t)( u_48 >> 16 );
+                    pUV[c] = (uint16)( u_48 >> 16 );
 				}
 			}
 
@@ -146,14 +147,14 @@ namespace mu
 						)
 			{
 				// TODO: Optimise
-                uint16_t* pUV = reinterpret_cast<uint16_t*>( pVertices );
+                uint16* pUV = reinterpret_cast<uint16*>( pVertices );
 
 				for ( int c=0; c<2; ++c )
 				{
                     uint32_t u_16 = pUV[c];
                     uint32_t u_32 = u_16 * ((uint8_t)(((float)0xffff) * transforms[ relBlock ].size[c]));
                     u_32 += (uint32_t)(((float)0xffffffff) * transforms[ relBlock ].min[c]);
-                    pUV[c] = (uint16_t)( u_32 >> 16 );
+                    pUV[c] = (uint16)( u_32 >> 16 );
 				}
 			}
 
diff --git a/Engine/Plugins/Experimental/Mutable/Source/MutableRuntime/Private/MuR/OpMeshApplyPose.h b/Engine/Plugins/Experimental/Mutable/Source/MutableRuntime/Private/MuR/OpMeshApplyPose.h
--- a/Engine/Plugins/Experimental/Mutable/Source/MutableRuntime/Private/MuR/OpMeshApplyPose.h
+++ b/Engine/Plugins/Experimental/Mutable/Source/MutableRuntime/Private/MuR/OpMeshApplyPose.h
@@ -32,7 +32,8 @@ namespace mu
 
         // Prepare the skin matrices. They may be in different order, and we only need the ones
         // relevant for the base mesh deformation.
-        vector<FTransform3f> skinTransforms(pSkeleton->GetBoneCount());
+		TArray<FTransform3f> skinTransforms;
+		skinTransforms.SetNum(pSkeleton->GetBoneCount());
         for ( int32 b=0; b< pSkeleton->GetBoneCount(); ++b )
         {
             int32 poseBoneIndex = pPose->FindBonePose(pSkeleton->GetBoneName(b));
diff --git a/Engine/Plugins/Experimental/Mutable/Source/MutableRuntime/Private/MuR/OpMeshChartDifference.h b/Engine/Plugins/Experimental/Mutable/Source/MutableRuntime/Private/MuR/OpMeshChartDifference.h
--- a/Engine/Plugins/Experimental/Mutable/Source/MutableRuntime/Private/MuR/OpMeshChartDifference.h
+++ b/Engine/Plugins/Experimental/Mutable/Source/MutableRuntime/Private/MuR/OpMeshChartDifference.h
@@ -93,7 +93,7 @@ namespace mu
 		// Classify the target faces in buckets along the Y axis
 #define NUM_BUCKETS	128
 #define AXIS		1
-		vector<int> buckets[ NUM_BUCKETS ];
+		TArray<int> buckets[ NUM_BUCKETS ];
 		float bucketStart = aabbox.min[AXIS];
 		float bucketSize = aabbox.size[AXIS] / NUM_BUCKETS;
 
@@ -107,17 +107,17 @@ namespace mu
             uint32_t index2 = iti.GetAsUINT32(); ++iti;
 			float y = ( (*(itp+index0))[AXIS] + (*(itp+index1))[AXIS] + (*(itp+index2))[AXIS] ) / 3;
 			float fbucket = (y-bucketStart) / bucketSize;
-			int bucket = std::min( NUM_BUCKETS-1, std::max( 0, (int)fbucket ) );
-			buckets[bucket].push_back(tf);
-			int hibucket = std::min( NUM_BUCKETS-1, std::max( 0, (int)(fbucket+bucketThreshold) ) );
+			int bucket = FMath::Min( NUM_BUCKETS-1, FMath::Max( 0, (int)fbucket ) );
+			buckets[bucket].Add(tf);
+			int hibucket = FMath::Min( NUM_BUCKETS-1, FMath::Max( 0, (int)(fbucket+bucketThreshold) ) );
 			if (hibucket!=bucket)
 			{
-				buckets[hibucket].push_back(tf);
+				buckets[hibucket].Add(tf);
 			}
-			int lobucket = std::min( NUM_BUCKETS-1, std::max( 0, (int)(fbucket-bucketThreshold) ) );
+			int lobucket = FMath::Min( NUM_BUCKETS-1, FMath::Max( 0, (int)(fbucket-bucketThreshold) ) );
 			if (lobucket!=bucket)
 			{
-				buckets[lobucket].push_back(tf);
+				buckets[lobucket].Add(tf);
 			}
 		}
 
@@ -143,9 +143,9 @@ namespace mu
 			// find the bucket for this face
 			float y = ( (*(itop+ov[0]))[AXIS] + (*(itop+ov[1]))[AXIS] + (*(itop+ov[2]))[AXIS] ) / 3;
 			float fbucket = (y-bucketStart) / bucketSize;
-			int bucket = std::min( NUM_BUCKETS-1, std::max( 0, (int)fbucket ) );
+			int bucket = FMath::Min( NUM_BUCKETS-1, FMath::Max( 0, (int)fbucket ) );
 
-			for ( std::size_t btf=0; !hasFace && btf<buckets[bucket].size(); btf++ )
+			for ( int32 btf=0; !hasFace && btf<buckets[bucket].Num(); btf++ )
 			{
 				int tf =  buckets[bucket][btf];
 
diff --git a/Engine/Plugins/Experimental/Mutable/Source/MutableRuntime/Private/MuR/OpMeshClipMorphPlane.h b/Engine/Plugins/Experimental/Mutable/Source/MutableRuntime/Private/MuR/OpMeshClipMorphPlane.h
--- a/Engine/Plugins/Experimental/Mutable/Source/MutableRuntime/Private/MuR/OpMeshClipMorphPlane.h
+++ b/Engine/Plugins/Experimental/Mutable/Source/MutableRuntime/Private/MuR/OpMeshClipMorphPlane.h
@@ -39,22 +39,22 @@ namespace mu
 
 	struct vertex_bone_info
 	{
-		mu::vector<int32_t> bone_indices;
-		mu::vector<int32_t> bone_weights;
+		TArray<int32> bone_indices;
+		TArray<int32> bone_weights;
 	};
 
-	inline bool VertexIsAffectedByBone(int32_t vertex_idx, const mu::vector<bool>& bone_is_affected, const mu::vector<vertex_bone_info>& vertex_info)
+	inline bool VertexIsAffectedByBone(int32 vertex_idx, const TArray<bool>& bone_is_affected, const TArray<vertex_bone_info>& vertex_info)
 	{
-		if (vertex_idx >= (int)vertex_info.size())
+		if (vertex_idx >= (int)vertex_info.Num())
 		{
 			return false;
 		}
 
-		check(vertex_info[vertex_idx].bone_indices.size() == vertex_info[vertex_idx].bone_weights.size());		
+		check(vertex_info[vertex_idx].bone_indices.Num() == vertex_info[vertex_idx].bone_weights.Num());		
 
-        for (size_t i = 0; i < vertex_info[vertex_idx].bone_indices.size(); ++i)
+        for (size_t i = 0; i < vertex_info[vertex_idx].bone_indices.Num(); ++i)
 		{			
-            check(vertex_info[vertex_idx].bone_indices[i] <= (int)bone_is_affected.size());
+            check(vertex_info[vertex_idx].bone_indices[i] <= (int)bone_is_affected.Num());
 
 			if (bone_is_affected[vertex_info[vertex_idx].bone_indices[i]] && vertex_info[vertex_idx].bone_weights[i] > 0)
 			{
@@ -88,24 +88,24 @@ namespace mu
 
 		MeshPtr pDest = pBase->Clone();
 
-        uint32_t vcount = pBase->GetVertexBuffers().GetElementCount();
+        uint32 vcount = pBase->GetVertexBuffers().GetElementCount();
 
 		if (!vcount)
 		{
 			return pDest;
 		}
 
-		mu::vector<bool> bone_is_affected;
-		mu::vector<vertex_bone_info> vertex_info;
+		TArray<bool> bone_is_affected;
+		TArray<vertex_bone_info> vertex_info;
 
         auto pBaseSkeleton = pBase->GetSkeleton();
         if (!boneName.empty() && pBaseSkeleton )
 		{
-            bone_is_affected.resize(pBaseSkeleton->GetBoneCount());
+            bone_is_affected.SetNum(pBaseSkeleton->GetBoneCount());
 
-            for (int32_t bone_idx = 0; bone_idx < pBaseSkeleton->GetBoneCount(); ++bone_idx)
+            for (int32 bone_idx = 0; bone_idx < pBaseSkeleton->GetBoneCount(); ++bone_idx)
 			{
-				int32_t current_idx = bone_idx;
+				int32 current_idx = bone_idx;
 				bool found_bone = false;
 
 				while (current_idx >= 0)
@@ -125,7 +125,7 @@ namespace mu
 			}
 
 			// Look for affected vertex indices
-			vertex_info.resize(vcount);
+			vertex_info.SetNum(vcount);
 			//int firstCount = pBase->GetVertexBuffers().GetElementCount();
 
 			for (int32 vb = 0; vb < pBase->GetVertexBuffers().m_buffers.Num(); ++vb)
@@ -161,7 +161,7 @@ namespace mu
 								
 								for (int j = 0; j < components; ++j)
 								{
-									vertex_info[i].bone_indices.push_back(pD[j]);
+									vertex_info[i].bone_indices.Add(pD[j]);
 								}
 
 								secondOffset += elemSize;
@@ -171,11 +171,11 @@ namespace mu
 							case MBF_INT16:
 							case MBF_UINT16:
 							{
-								const uint16_t* pD = reinterpret_cast<const uint16_t*>(&result.m_data[secondOffset]);
+								const uint16* pD = reinterpret_cast<const uint16*>(&result.m_data[secondOffset]);
 
 								for (int j = 0; j < components; ++j)
 								{
-									vertex_info[i].bone_indices.push_back(pD[j]);
+									vertex_info[i].bone_indices.Add(pD[j]);
 								}
 
 								secondOffset += elemSize;
@@ -185,11 +185,11 @@ namespace mu
 							case MBF_INT32:
 							case MBF_UINT32:
 							{
-								const uint32_t* pD = reinterpret_cast<const uint32_t*>(&result.m_data[secondOffset]);
+								const uint32* pD = reinterpret_cast<const uint32*>(&result.m_data[secondOffset]);
 
 								for (int j = 0; j < components; ++j)
 								{
-									vertex_info[i].bone_indices.push_back(pD[j]);
+									vertex_info[i].bone_indices.Add(pD[j]);
 								}
 
 								secondOffset += elemSize;
@@ -215,7 +215,7 @@ namespace mu
 
 								for (int j = 0; j < components; ++j)
 								{
-                                    vertex_info[i].bone_weights.push_back(pD[j]);
+                                    vertex_info[i].bone_weights.Add(pD[j]);
 								}
 
 								secondOffset += elemSize;
@@ -226,11 +226,11 @@ namespace mu
 							case MBF_UINT16:
 							case MBF_NUINT16:
 							{
-								const uint16_t* pD = reinterpret_cast<const uint16_t*>(&result.m_data[secondOffset]);
+								const uint16* pD = reinterpret_cast<const uint16*>(&result.m_data[secondOffset]);
 
 								for (int j = 0; j < components; ++j)
 								{
-									vertex_info[i].bone_weights.push_back(pD[j]);
+									vertex_info[i].bone_weights.Add(pD[j]);
 								}
 
 								secondOffset += elemSize;
@@ -241,11 +241,11 @@ namespace mu
 							case MBF_UINT32:
 							case MBF_NUINT32:
 							{
-								const uint32_t* pD = reinterpret_cast<const uint32_t*>(&result.m_data[secondOffset]);
+								const uint32* pD = reinterpret_cast<const uint32*>(&result.m_data[secondOffset]);
 
 								for (int j = 0; j < components; ++j)
 								{
-									vertex_info[i].bone_weights.push_back(pD[j]);
+									vertex_info[i].bone_weights.Add(pD[j]);
 								}
 
 								secondOffset += elemSize;
@@ -258,7 +258,7 @@ namespace mu
 
                                 for (int j = 0; j < components; ++j)
                                 {
-                                    vertex_info[i].bone_weights.push_back(pD[j]>0.0f);
+                                    vertex_info[i].bone_weights.Add(pD[j]>0.0f);
                                 }
 
                                 secondOffset += elemSize;
@@ -274,7 +274,9 @@ namespace mu
 			}
 		}
 
-        mu::map<uint32_t, bool> removed_vertexs;
+		// TODO: Replace with an array of bools?
+        TArray<bool> RemovedVertices;
+		RemovedVertices.Init(false,pDest->GetVertexCount());
 
 		const FMeshBufferSet& MBSPriv = pDest->GetVertexBuffers();
 		for (int32 b = 0; b < MBSPriv.m_buffers.Num(); ++b)
@@ -289,7 +291,7 @@ namespace mu
 				switch (sem)
 				{
 				case MBS_POSITION:
-                    for (uint32_t v = 0; v < vcount; ++v)
+                    for (uint32 v = 0; v < vcount; ++v)
 					{
 						vec3f vertex(0.0f, 0.0f, 0.0f);
 						for (int i = 0; i < 3; ++i)
@@ -343,7 +345,7 @@ namespace mu
 								{
 									vec3f vert_displ = normal * -dot_cut;
 									vertex = vertex + vert_displ;
-									removed_vertexs[v] = true;
+									RemovedVertices[v] = true;
 								}
 							}
 						}
@@ -371,12 +373,12 @@ namespace mu
 		UntypedMeshBufferIteratorConst ito(pDest->GetIndexBuffers(), MBS_VERTEXINDEX);
 		for (int f = 0; f < aFaceCount; ++f)
 		{
-            vec3<uint32_t> ov;
+            vec3<uint32> ov;
 			ov[0] = ito.GetAsUINT32(); ++ito;
 			ov[1] = ito.GetAsUINT32(); ++ito;
 			ov[2] = ito.GetAsUINT32(); ++ito;
-
-			bool all_vertexs_removed = removed_vertexs[ov[0]] && removed_vertexs[ov[1]] && removed_vertexs[ov[2]];
+		
+			bool all_vertexs_removed = RemovedVertices[ov[0]] && RemovedVertices[ov[1]] && RemovedVertices[ov[2]];
 
 			if (!all_vertexs_removed)
 			{
@@ -405,7 +407,7 @@ namespace mu
             // We assume there will be only one.
             check(pDest->m_surfaces.Num()==1);
 
-            pDest->m_surfaces[0].m_indexCount -= (int32_t)removedIndices;
+            pDest->m_surfaces[0].m_indexCount -= (int32)removedIndices;
         }
 
 		return pDest;
diff --git a/Engine/Plugins/Experimental/Mutable/Source/MutableRuntime/Private/MuR/OpMeshClipWithMesh.h b/Engine/Plugins/Experimental/Mutable/Source/MutableRuntime/Private/MuR/OpMeshClipWithMesh.h
--- a/Engine/Plugins/Experimental/Mutable/Source/MutableRuntime/Private/MuR/OpMeshClipWithMesh.h
+++ b/Engine/Plugins/Experimental/Mutable/Source/MutableRuntime/Private/MuR/OpMeshClipWithMesh.h
@@ -19,77 +19,6 @@
 
 namespace mu
 {
-    const static float vert_collapse_eps = 0.0001f;
-
-    //---------------------------------------------------------------------------------------------
-    //! Create a map from vertices into vertices, collapsing vertices that have the same position
-    //---------------------------------------------------------------------------------------------
-    inline void MeshCreateCollapsedVertexMap( const Mesh* pMesh,
-                                       mu::vector<int>& collapsedVertexMap,
-                                       mu::vector<vec3f>& vertices
-                                       )
-    {
-        MUTABLE_CPUPROFILER_SCOPE(MeshCreateCollapsedVertexMap);
-    	
-        int vcount = pMesh->GetVertexCount();
-        collapsedVertexMap.resize(vcount);
-        vertices.resize(vcount);
-
-        const FMeshBufferSet& MBSPriv = pMesh->GetVertexBuffers();
-        for (int32 b = 0; b < MBSPriv.m_buffers.Num(); ++b)
-        {
-            for (int32 c = 0; c < MBSPriv.m_buffers[b].m_channels.Num(); ++c)
-            {
-                MESH_BUFFER_SEMANTIC sem = MBSPriv.m_buffers[b].m_channels[c].m_semantic;
-                int semIndex = MBSPriv.m_buffers[b].m_channels[c].m_semanticIndex;
-
-                UntypedMeshBufferIteratorConst it(pMesh->GetVertexBuffers(), sem, semIndex);
-
-                switch (sem)
-                {
-                case MBS_POSITION:
-                    // First create a cache of the vertices in the vertices array
-                    for (int v = 0; v < vcount; ++v)
-                    {
-                        vec3f vertex(0.0f, 0.0f, 0.0f);
-                        for (int i = 0; i < 3; ++i)
-                        {
-                            ConvertData(i, &vertex[0], MBF_FLOAT32, it.ptr(), it.GetFormat());
-                        }
-
-                        vertices[v] = vertex;
-
-                        ++it;
-                    }
-
-                    // Create map to store which vertices are the same (collapse nearby vertices)
-                    for (std::size_t v = 0; v < vertices.size(); ++v)
-                    {
-                        collapsedVertexMap[v] = (int32_t)v;
-
-                        for (std::size_t candidate_v = 0; candidate_v < v; ++candidate_v)
-                        {
-                            int collapsed_candidate_v_idx = collapsedVertexMap[candidate_v];
-
-                            // Test whether v and the collapsed candidate are close enough to be collapsed
-                            vec3f r = vertices[collapsed_candidate_v_idx] - vertices[v];
-
-                            if (dot(r, r) <= vert_collapse_eps * vert_collapse_eps)
-                            {
-                                collapsedVertexMap[v] = collapsed_candidate_v_idx;
-                                break;
-                            }
-                        }
-                    }
-
-                    break;
-
-                default:
-                    break;
-                }
-            }
-        }
-    }
 
     //---------------------------------------------------------------------------------------------
     //! Create a map from vertices into vertices, collapsing vertices that have the same position, 
@@ -172,16 +101,15 @@ namespace mu
         int vcount = pMesh->GetVertexCount();
         int fcount = pMesh->GetFaceCount();
 
-        mu::vector<vec3f> vertices(vcount);
+        TArray<FVector3f> vertices;
 
         // Map in ClipMesh from vertices to the one they are collapsed to because they are very similar, if they aren't collapsed then they are mapped to themselves
-        mu::vector<int> collapsedVertexMap(vcount);
+		TArray<int> collapsedVertexMap;
 
         MeshCreateCollapsedVertexMap( pMesh, collapsedVertexMap, vertices );
 
-
         // Acumulate edges
-        mu::map< mu::pair<int,int>, int > faceCountPerEdge;
+        TMap< TPair<int,int>, int > faceCountPerEdge;
 
         UntypedMeshBufferIteratorConst itClipMesh(pMesh->GetIndexBuffers(), MBS_VERTEXINDEX);
         for (int f = 0; f < fcount; ++f)
@@ -202,11 +130,18 @@ namespace mu
                     return false;
                 }
 
-                mu::pair<int,int> edge;
-                edge.first = std::min( v0, v1 );
-                edge.second = std::max( v0, v1 );
+                TPair<int,int> edge;
+                edge.Key = FMath::Min( v0, v1 );
+                edge.Value = FMath::Max( v0, v1 );
 
-                faceCountPerEdge[edge]++;
+				if (faceCountPerEdge.Contains(edge))
+				{
+					faceCountPerEdge[edge]++;
+				}
+				else
+				{
+					faceCountPerEdge.Add(edge,1);
+				}
             }
         }
 
@@ -214,7 +149,7 @@ namespace mu
         // See if every edge has 2 faces
         for( const auto& e: faceCountPerEdge)
         {
-            if (e.second!=2)
+            if (e.Value!=2)
             {
                 return false;
             }
@@ -232,20 +167,22 @@ namespace mu
         MUTABLE_CPUPROFILER_SCOPE(MeshRemoveUnusedVertices);
 
         // Mark used vertices
-        mu::vector<uint8_t> used( pMesh->GetVertexCount(), false );
+		TArray<uint8> used;
+		used.SetNumZeroed(pMesh->GetVertexCount());
         UntypedMeshBufferIteratorConst iti(pMesh->GetIndexBuffers(), MBS_VERTEXINDEX);
         int IndexCount = pMesh->GetIndexCount();
         for (int i = 0; i < IndexCount; ++i)
         {
-            uint32_t Index = iti.GetAsUINT32();
+            uint32 Index = iti.GetAsUINT32();
             ++iti;
             used[Index] = true;
         }
 
         // Build vertex map
-        mu::vector<int32_t> oldToNewVertex(pMesh->GetVertexCount());
-        int totalNewVertices = 0;
-        for (int v = 0; v<pMesh->GetVertexCount(); ++v)
+		TArray<int32> oldToNewVertex;
+		oldToNewVertex.SetNumUninitialized(pMesh->GetVertexCount());
+        int32 totalNewVertices = 0;
+        for (int32 v = 0; v<pMesh->GetVertexCount(); ++v)
         {
             if (used[v])
             {
@@ -262,14 +199,14 @@ namespace mu
         for (int b = 0; b<pMesh->GetVertexBuffers().GetBufferCount(); ++b)
         {
             int elemSize = pMesh->GetVertexBuffers().GetElementSize(b);
-            const uint8_t* pSourceData = pMesh->GetVertexBuffers().GetBufferData(b);
-            uint8_t* pData = pMesh->GetVertexBuffers().GetBufferData(b);
+            const uint8* pSourceData = pMesh->GetVertexBuffers().GetBufferData(b);
+            uint8* pData = pMesh->GetVertexBuffers().GetBufferData(b);
             for (int v = 0; v<pMesh->GetVertexCount(); ++v)
             {
                 if (oldToNewVertex[v]!=-1)
                 {
-                    uint8_t* pElemData = pData + elemSize*oldToNewVertex[v];
-                    const uint8_t* pElemSourceData = pSourceData + elemSize*v;
+                    uint8* pElemData = pData + elemSize*oldToNewVertex[v];
+                    const uint8* pElemSourceData = pSourceData + elemSize*v;
                     // Avoid warning for overlapping memcpy in valgrind
                     if (pElemData != pElemSourceData)
                     {
@@ -286,10 +223,10 @@ namespace mu
         {
             for (int i = 0; i < IndexCount; ++i)
             {
-                uint32_t Index = *(uint32_t*)ito.ptr();
-                int32_t NewIndex = oldToNewVertex[Index];
+                uint32 Index = *(uint32*)ito.ptr();
+                int32 NewIndex = oldToNewVertex[Index];
                 check(NewIndex >= 0);
-                *(uint32_t*)ito.ptr() = (uint32_t)NewIndex;
+                *(uint32*)ito.ptr() = (uint32)NewIndex;
                 ++ito;
             }
         }
@@ -297,10 +234,10 @@ namespace mu
         {
             for (int i = 0; i < IndexCount; ++i)
             {
-                uint16_t Index = *(uint16_t*)ito.ptr();
-                int32_t NewIndex = oldToNewVertex[Index];
+                uint16 Index = *(uint16*)ito.ptr();
+                int32 NewIndex = oldToNewVertex[Index];
                 check(NewIndex >= 0);
-                *(uint16_t*)ito.ptr() = (uint16_t)NewIndex;
+                *(uint16*)ito.ptr() = (uint16)NewIndex;
                 ++ito;
             }
         }
@@ -319,11 +256,11 @@ namespace mu
     //---------------------------------------------------------------------------------------------
     inline int get_num_intersections( const vec3f& vertex,
                                       const vec3f& ray,
-                                      const mu::vector<vec3f>& vertices,
-                                      const mu::vector<uint32_t>& faces,
-                                      mu::vector<int>& collapsedVertexMap,
-                                      mu::vector<uint8_t>& vertex_already_intersected,
-                                      mu::map<mu::pair<int, int>, bool>& edge_already_intersected,
+                                      const TArray<vec3f>& vertices,
+                                      const TArray<uint32>& faces,
+                                      TArray<int>& collapsedVertexMap,
+									  TArray<uint8>& vertex_already_intersected,
+                                      TMap<TPair<int, int>, bool>& edge_already_intersected,
 									  float dynamic_epsilon
                                       )
     {
@@ -333,18 +270,18 @@ namespace mu
         int num_intersections = 0;
         vec3f intersection;
 
-		FMemory::Memzero( vertex_already_intersected.data(), vertex_already_intersected.size() );
+		FMemory::Memzero( vertex_already_intersected.GetData(), vertex_already_intersected.Num());
 
-        edge_already_intersected.clear();
+        edge_already_intersected.Empty();
 
     #define get_collapsed_vertex(vertex_idx) (vertices[collapsedVertexMap[vertex_idx]])
 
-        size_t fcount = faces.size()/3;
+        size_t fcount = faces.Num()/3;
 
         // Check vertex against all ClipMorph faces
-        for (uint32_t face = 0; face < fcount; ++face)
+        for (uint32 face = 0; face < fcount; ++face)
         {
-            uint32_t vertex_indexs[3] = { faces[3 * face], faces[3 * face + 1], faces[3 * face + 2] };
+            uint32 vertex_indexs[3] = { faces[3 * face], faces[3 * face + 1], faces[3 * face + 2] };
 
             vec3f v0 = get_collapsed_vertex(vertex_indexs[0]);
             vec3f v1 = get_collapsed_vertex(vertex_indexs[1]);
@@ -370,14 +307,14 @@ namespace mu
                     int collapsed_edge_vert_index0 = collapsedVertexMap[vertex_indexs[out_intersected_edge_v0]];
                     int collapsed_edge_vert_index1 = collapsedVertexMap[vertex_indexs[out_intersected_edge_v1]];
 
-                    mu::pair<int, int> key;
-                    key.first = collapsed_edge_vert_index0 <= collapsed_edge_vert_index1 ? collapsed_edge_vert_index0 : collapsed_edge_vert_index1;
-                    key.second = collapsed_edge_vert_index1 > collapsed_edge_vert_index0 ? collapsed_edge_vert_index1 : collapsed_edge_vert_index0;
-                    edge_not_intersected_before = edge_already_intersected.find(key) == edge_already_intersected.end();
+                    TPair<int, int> key;
+                    key.Key = collapsed_edge_vert_index0 <= collapsed_edge_vert_index1 ? collapsed_edge_vert_index0 : collapsed_edge_vert_index1;
+                    key.Value = collapsed_edge_vert_index1 > collapsed_edge_vert_index0 ? collapsed_edge_vert_index1 : collapsed_edge_vert_index0;
+                    edge_not_intersected_before = edge_already_intersected.Find(key) != nullptr;
 
                     if (edge_not_intersected_before)
                     {
-                        edge_already_intersected[key] = true;
+                        edge_already_intersected.Add( key, true );
                     }
                 }
 
@@ -392,190 +329,6 @@ namespace mu
     }
 
 
-    //---------------------------------------------------------------------------------------------
-    //! Reference version
-    //---------------------------------------------------------------------------------------------
-    inline void MeshClipMeshClassifyVertices(mu::vector<uint8_t>& vertex_in_clip_mesh,
-                                             const Mesh* pBase, const Mesh* pClipMesh)
-    {
-        MUTABLE_CPUPROFILER_SCOPE(MeshClipMeshClassifyVertices);
-
-        // TODO: this fails in bro for some object. do it at graph construction time to report it 
-        // nicely.
-        //check( MeshIsClosed(pClipMesh) );
-
-        int vcount = pClipMesh->GetVertexBuffers().GetElementCount();
-        int fcount = pClipMesh->GetFaceCount();
-
-        int orig_vcount = pBase->GetVertexBuffers().GetElementCount();
-
-        // Stores whether each vertex in the original mesh in the clip mesh volume
-        vertex_in_clip_mesh.resize(orig_vcount);
-        std::fill(vertex_in_clip_mesh.begin(), vertex_in_clip_mesh.end(), false);
-
-        if (!vcount)
-        {
-            return;
-        }
-
-        mu::vector<vec3f> vertices(vcount); // ClipMesh vertex cache
-        mu::vector<uint32_t> faces(fcount * 3); // ClipMesh face cache
-
-        // Map in ClipMesh from vertices to the one they are collapsed to because they are very 
-        // similar, if they aren't collapsed then they are mapped to themselves
-        mu::vector<int> collapsedVertexMap(vcount);
-
-        MeshCreateCollapsedVertexMap( pClipMesh, collapsedVertexMap, vertices );
-
-
-        // Create cache of the faces
-        UntypedMeshBufferIteratorConst itClipMesh(pClipMesh->GetIndexBuffers(), MBS_VERTEXINDEX);
-
-        for (int f = 0; f < fcount; ++f)
-        {
-            faces[3 * f]     = itClipMesh.GetAsUINT32(); ++itClipMesh;
-            faces[3 * f + 1] = itClipMesh.GetAsUINT32(); ++itClipMesh;
-            faces[3 * f + 2] = itClipMesh.GetAsUINT32(); ++itClipMesh;
-        }
-
-
-        // Create a bounding box of the clip mesh
-        box<vec3f> clipMeshBoundingBox;
-        clipMeshBoundingBox.min = vertices[0];
-        clipMeshBoundingBox.size = vec3f(0.0f,0.0f,0.0f);
-        for (size_t i=1; i<vertices.size(); ++i)
-        {
-            clipMeshBoundingBox.Bound(vertices[i]);
-        }
-
-		// Dynamic distance epsilon to support different engines
-		float maxDimensionBoundingBox = length(clipMeshBoundingBox.size);
-		// 0.000001 is the value that helps to achieve the dynamic epsilon, do not change it
-		float dynamicEpsilon = 0.000001f * maxDimensionBoundingBox * (maxDimensionBoundingBox < 1.0f ? maxDimensionBoundingBox : 1.0f);
-
-        // Z-direction. Don't change this without reviewing the acceleration structures below.
-        const vec3f ray = vec3f(0.f, 0.f, 203897203.f);
-
-        // Create an acceleration grid to avoid testing all clip-mesh triangles.
-        // This assumes that the testing ray direction is Z
-        const int GRID_SIZE = 8;
-        std::unique_ptr<mu::vector<uint32_t>[]> gridFaces( new mu::vector<uint32_t>[GRID_SIZE*GRID_SIZE] );
-        vec2f gridCellSize = clipMeshBoundingBox.size.xy() / (float)GRID_SIZE;
-        for ( int i=0; i<GRID_SIZE; ++i)
-        {
-            for ( int j=0; j<GRID_SIZE; ++j)
-            {
-                box<vec2f> cellBox;
-                cellBox.size = gridCellSize;
-                cellBox.min = clipMeshBoundingBox.min.xy() + gridCellSize * vec2f((float)i,(float)j);
-
-                mu::vector<uint32_t>& cellFaces = gridFaces[i+j*GRID_SIZE];
-                cellFaces.reserve(fcount/GRID_SIZE);
-                for (int f = 0; f < fcount; ++f)
-                {
-                    // Imprecise, conservative classification of faces.
-                    box<vec2f> faceBox;
-                    faceBox.size = vec2f(0.0f,0.0f);
-                    faceBox.min = vertices[ faces[3 * f + 0] ].xy();
-                    faceBox.Bound(vertices[ faces[3 * f + 1] ].xy());
-                    faceBox.Bound(vertices[ faces[3 * f + 2] ].xy());
-
-                    if (cellBox.Intersects(faceBox))
-                    {
-                        cellFaces.push_back(faces[3 * f + 0]);
-                        cellFaces.push_back(faces[3 * f + 1]);
-                        cellFaces.push_back(faces[3 * f + 2]);
-                    }
-                }
-            }
-        }
-
-
-        // Now go through all vertices in the mesh and record whether they are inside or outside of the ClipMesh
-        uint32_t dest_vertex_count = pBase->GetVertexCount();
-
-        const FMeshBufferSet& MBSPriv2 = pBase->GetVertexBuffers();
-        for (int32 b = 0; b < MBSPriv2.m_buffers.Num(); ++b)
-        {
-            for (int32 c = 0; c < MBSPriv2.m_buffers[b].m_channels.Num(); ++c)
-            {
-                MESH_BUFFER_SEMANTIC sem = MBSPriv2.m_buffers[b].m_channels[c].m_semantic;
-                int semIndex = MBSPriv2.m_buffers[b].m_channels[c].m_semanticIndex;
-
-                UntypedMeshBufferIteratorConst it(pBase->GetVertexBuffers(), sem, semIndex);
-
-                mu::vector<uint8_t> vertex_already_intersected(vcount);
-                mu::map<mu::pair<int, int>, bool> edge_already_intersected;
-
-                switch (sem)
-                {
-                case MBS_POSITION:
-                    for (uint32_t v = 0; v < dest_vertex_count; ++v)
-                    {
-                        vec3f vertex(0.0f, 0.0f, 0.0f);
-                        for (int i = 0; i < 3; ++i)
-                        {
-                            ConvertData(i, &vertex[0], MBF_FLOAT32, it.ptr(), it.GetFormat());
-                        }
-
-                        // Find out grid cell
-                        vec2f hPos = (vertex.xy()-clipMeshBoundingBox.min.xy()) / clipMeshBoundingBox.size.xy() * (float)GRID_SIZE;
-                        int i = std::min( std::max( (int)hPos.x(), 0 ), GRID_SIZE-1 );
-                        int j = std::min( std::max( (int)hPos.y(), 0 ), GRID_SIZE-1 );
-
-                        // Early discard test: if the vertex is not inside the bounding box of the clip mesh, it won't be clipped.
-                        bool contains = (clipMeshBoundingBox.ContainsInclusive(vertex));
-
-                        if (contains)
-                        {
-                            // Optimised test
-                            int cellIndex = i+j*GRID_SIZE;
-                            mu::vector<uint32_t>& thisCell = gridFaces[cellIndex];
-                            int num_intersections = get_num_intersections(
-                                        vertex,
-                                        ray,
-                                        vertices,
-                                        thisCell,
-                                        collapsedVertexMap,
-                                        vertex_already_intersected,
-                                        edge_already_intersected,
-										dynamicEpsilon
-                                        );
-
-                            // Full test BLEH, debug
-//                            int full_num_intersections = get_num_intersections(
-//                                        vertex,
-//                                        ray,
-//                                        vertices,
-//                                        faces,
-//                                        collapsedVertexMap,
-//                                        vertex_already_intersected,
-//                                        edge_already_intersected
-//                                        );
-
-
-                            vertex_in_clip_mesh[v] = num_intersections % 2 == 1;
-
-                            // This may be used to debug degenerated cases if the conditional above is also removed.
-                            // \todo: make sure it works well
-//                            if (!contains && vertex_in_clip_mesh[v])
-//                            {
-//                                assert(false);
-//                            }
-                        }
-
-                        ++it;
-                    }
-                    break;
-
-                default:
-                    break;
-                }
-            }
-        }
-    }
-
-
     //---------------------------------------------------------------------------------------------
     //! Core Geometry version
     //---------------------------------------------------------------------------------------------
@@ -862,13 +615,13 @@ namespace mu
         
 		MeshPtr pDest = pBase->Clone();
 
-        uint32_t vcount = pClipMesh->GetVertexBuffers().GetElementCount();
+        uint32 vcount = pClipMesh->GetVertexBuffers().GetElementCount();
         if (!vcount)
         {
             return pDest; // Since there's nothing to clip against return a copy of the original base mesh
         }
 
-        mu::vector<uint8_t> vertex_in_clip_mesh;  // Stores whether each vertex in the original mesh in in the clip mesh volume
+        TArray<uint8> vertex_in_clip_mesh;  // Stores whether each vertex in the original mesh in in the clip mesh volume
         MeshClipMeshClassifyVertices( vertex_in_clip_mesh, pBase, pClipMesh );
 
         // Now remove all the faces from the result mesh that have all the vertices outside the clip volume
@@ -879,7 +632,7 @@ namespace mu
         UntypedMeshBufferIteratorConst ito(pDest->GetIndexBuffers(), MBS_VERTEXINDEX);
         for (int f = 0; f < aFaceCount; ++f)
         {
-            vec3<uint32_t> ov;
+            vec3<uint32> ov;
             ov[0] = ito.GetAsUINT32(); ++ito;
             ov[1] = ito.GetAsUINT32(); ++ito;
             ov[2] = ito.GetAsUINT32(); ++ito;
@@ -899,7 +652,7 @@ namespace mu
             itBase += 3;
         }
 
-        std::size_t removedIndices = itBase - itDest;
+        SIZE_T removedIndices = itBase - itDest;
         check(removedIndices % 3 == 0);
 
         pDest->GetFaceBuffers().SetElementCount(aFaceCount - (int)removedIndices / 3);
@@ -929,7 +682,7 @@ namespace mu
 
         MeshPtr pDest = pBase->Clone();
 
-        uint32_t VCount = pClipMesh->GetVertexBuffers().GetElementCount();
+        uint32 VCount = pClipMesh->GetVertexBuffers().GetElementCount();
         if (!VCount)
         {
             return pDest; // Since there's nothing to clip against return a copy of the original base mesh
@@ -991,11 +744,11 @@ namespace mu
     //---------------------------------------------------------------------------------------------
     //!
     //---------------------------------------------------------------------------------------------
-    inline MeshPtr CreateMask( MeshPtrConst pBase, const mu::vector<uint8_t>& excludedVertices )
+    inline MeshPtr CreateMask( MeshPtrConst pBase, const TArray<uint8>& excludedVertices )
     {
 
         int maskVertexCount = 0;
-        for( uint8_t b: excludedVertices )
+        for( uint8 b: excludedVertices )
         {
             if (!b) ++maskVertexCount;
         }
@@ -1007,36 +760,30 @@ namespace mu
             pMask->GetVertexBuffers().SetElementCount( maskVertexCount );
             pMask->GetVertexBuffers().SetBufferCount( 1 );
 
-            vector<MESH_BUFFER_SEMANTIC> semantic;
-            vector<int> semanticIndex;
-            vector<MESH_BUFFER_FORMAT> format;
-            vector<int> components;
-            vector<int> offsets;
-
-            // Vertex index channel
-            semantic.push_back( MBS_VERTEXINDEX );
-            semanticIndex.push_back( 0 );
-            format.push_back( MBF_UINT32 );
-            components.push_back( 1 );
-            offsets.push_back( 0 );
+			// Vertex index channel
+			MESH_BUFFER_SEMANTIC semantic = MBS_VERTEXINDEX;
+			int semanticIndex = 0;
+			MESH_BUFFER_FORMAT format = MBF_UINT32;
+			int components = 1;
+			int offsets = 0;
 
             pMask->GetVertexBuffers().SetBuffer
                 (
                     0,
                     4,
                     1,
-                    &semantic[0],
-                    &semanticIndex[0],
-                    &format[0],
-                    &components[0],
-                    &offsets[0]
+                    &semantic,
+                    &semanticIndex,
+                    &format,
+                    &components,
+                    &offsets
                 );
         }
 
-        MeshBufferIterator<MBF_UINT32,uint32_t,1> itMask( pMask->GetVertexBuffers(), MBS_VERTEXINDEX, 0 );
+        MeshBufferIterator<MBF_UINT32,uint32,1> itMask( pMask->GetVertexBuffers(), MBS_VERTEXINDEX, 0 );
         UntypedMeshBufferIteratorConst itBase( pBase->GetVertexBuffers(), MBS_VERTEXINDEX, 0 );
 
-        for (size_t v = 0; v<excludedVertices.size(); ++v)
+        for (size_t v = 0; v<excludedVertices.Num(); ++v)
         {
             if (!excludedVertices[v])
             {
@@ -1050,69 +797,6 @@ namespace mu
         return pMask;
     }
 
-    //---------------------------------------------------------------------------------------------
-    //!
-    //---------------------------------------------------------------------------------------------
-    inline MeshPtr CreateMask( MeshPtrConst pBase, const TArray<uint8>& ExcludedVertices )
-    {
-
-        int32 MaskVertexCount = 0;
-        for( uint8 B: ExcludedVertices )
-        {
-            MaskVertexCount += (int32)(!B);
-        }
-
-        MeshPtr pMask = new Mesh();
-
-        // Create the vertex buffer
-        {
-            pMask->GetVertexBuffers().SetElementCount( MaskVertexCount );
-            pMask->GetVertexBuffers().SetBufferCount( 1 );
-
-            vector<MESH_BUFFER_SEMANTIC> semantic;
-            vector<int> semanticIndex;
-            vector<MESH_BUFFER_FORMAT> format;
-            vector<int> components;
-            vector<int> offsets;
-
-            // Vertex index channel
-            semantic.push_back( MBS_VERTEXINDEX );
-            semanticIndex.push_back( 0 );
-            format.push_back( MBF_UINT32 );
-            components.push_back( 1 );
-            offsets.push_back( 0 );
-
-            pMask->GetVertexBuffers().SetBuffer
-                (
-                    0,
-                    4,
-                    1,
-                    &semantic[0],
-                    &semanticIndex[0],
-                    &format[0],
-                    &components[0],
-                    &offsets[0]
-                );
-        }
-
-        MeshBufferIterator<MBF_UINT32,uint32_t,1> itMask( pMask->GetVertexBuffers(), MBS_VERTEXINDEX, 0 );
-        UntypedMeshBufferIteratorConst itBase( pBase->GetVertexBuffers(), MBS_VERTEXINDEX, 0 );
-
-		const int32 ExcludedVerticesNum = ExcludedVertices.Num();
-        for ( int32 V = 0; V < ExcludedVerticesNum; ++V )
-        {
-            if (!ExcludedVertices[V])
-            {
-                (*itMask)[0] = itBase.GetAsUINT32();
-                ++itMask;
-            }
-
-            ++itBase;
-        }
-
-        return pMask;
-    }
-
 
     //---------------------------------------------------------------------------------------------
     //! Generate a mask mesh with the faces of the base mesh inside the clip mesh.
@@ -1167,23 +851,24 @@ namespace mu
     {
         MUTABLE_CPUPROFILER_SCOPE(MeshMaskClipMesh_Reference);
 
-        uint32_t vcount = pClipMesh->GetVertexBuffers().GetElementCount();
+        uint32 vcount = pClipMesh->GetVertexBuffers().GetElementCount();
         if (!vcount)
         {
             return nullptr; // Since there's nothing to clip against return a copy of the original base mesh
         }
 
-        mu::vector<uint8_t> vertex_in_clip_mesh;  // Stores whether each vertex in the original mesh in in the clip mesh volume
+        TArray<uint8> vertex_in_clip_mesh;  // Stores whether each vertex in the original mesh in in the clip mesh volume
         MeshClipMeshClassifyVertices( vertex_in_clip_mesh, pBase, pClipMesh );
 
         // We only remove vertices if all their faces are clipped
-        mu::vector<uint8_t> vertex_with_face_not_clipped(vertex_in_clip_mesh.size(),0);
+		TArray<uint8> vertex_with_face_not_clipped;
+		vertex_with_face_not_clipped.SetNumZeroed(vertex_in_clip_mesh.Num());
 
         UntypedMeshBufferIteratorConst ito(pBase->GetIndexBuffers(), MBS_VERTEXINDEX);
         int aFaceCount = pBase->GetFaceCount();
         for (int f = 0; f < aFaceCount; ++f)
         {
-            vec3<uint32_t> ov;
+            vec3<uint32> ov;
             ov[0] = ito.GetAsUINT32(); ++ito;
             ov[1] = ito.GetAsUINT32(); ++ito;
             ov[2] = ito.GetAsUINT32(); ++ito;
@@ -1212,7 +897,7 @@ namespace mu
     {
         MUTABLE_CPUPROFILER_SCOPE(MeshMaskDiff);
 
-        uint32_t vcount = pFragment->GetVertexBuffers().GetElementCount();
+        uint32 vcount = pFragment->GetVertexBuffers().GetElementCount();
         if (!vcount)
         {
             return nullptr;
@@ -1247,7 +932,7 @@ namespace mu
         // Classify the target faces in buckets along the Y axis
 #define NUM_BUCKETS	128
 #define AXIS		1
-        vector<int> buckets[ NUM_BUCKETS ];
+        TArray<int> buckets[ NUM_BUCKETS ];
         float bucketStart = aabbox.min[AXIS];
         float bucketSize = aabbox.size[AXIS] / NUM_BUCKETS;
 
@@ -1256,22 +941,22 @@ namespace mu
         MeshBufferIteratorConst<MBF_FLOAT32,float,3> itp( pFragment->GetVertexBuffers(), MBS_POSITION );
         for ( int tf=0; tf<fragmentFaceCount; tf++ )
         {
-            uint32_t index0 = iti.GetAsUINT32(); ++iti;
-            uint32_t index1 = iti.GetAsUINT32(); ++iti;
-            uint32_t index2 = iti.GetAsUINT32(); ++iti;
+            uint32 index0 = iti.GetAsUINT32(); ++iti;
+            uint32 index1 = iti.GetAsUINT32(); ++iti;
+            uint32 index2 = iti.GetAsUINT32(); ++iti;
             float y = ( (*(itp+index0))[AXIS] + (*(itp+index1))[AXIS] + (*(itp+index2))[AXIS] ) / 3;
             float fbucket = (y-bucketStart) / bucketSize;
-            int bucket = std::min( NUM_BUCKETS-1, std::max( 0, (int)fbucket ) );
-            buckets[bucket].push_back(tf);
-            int hibucket = std::min( NUM_BUCKETS-1, std::max( 0, (int)(fbucket+bucketThreshold) ) );
+            int bucket = FMath::Min( NUM_BUCKETS-1, FMath::Max( 0, (int)fbucket ) );
+            buckets[bucket].Add(tf);
+            int hibucket = FMath::Min( NUM_BUCKETS-1, FMath::Max( 0, (int)(fbucket+bucketThreshold) ) );
             if (hibucket!=bucket)
             {
-                buckets[hibucket].push_back(tf);
+                buckets[hibucket].Add(tf);
             }
-            int lobucket = std::min( NUM_BUCKETS-1, std::max( 0, (int)(fbucket-bucketThreshold) ) );
+            int lobucket = FMath::Min( NUM_BUCKETS-1, FMath::Max( 0, (int)(fbucket-bucketThreshold) ) );
             if (lobucket!=bucket)
             {
-                buckets[lobucket].push_back(tf);
+                buckets[lobucket].Add(tf);
             }
         }
 
@@ -1283,7 +968,8 @@ namespace mu
 //			LogDebug("bucket : %d\n", buckets[b].size() );
 //		}
 
-        vector<uint8_t> faceClipped(sourceFaceCount,false);
+		TArray<uint8> faceClipped;
+		faceClipped.SetNumZeroed(sourceFaceCount);
 
         UntypedMeshBufferIteratorConst ito( pBase->GetIndexBuffers(), MBS_VERTEXINDEX );
         MeshBufferIteratorConst<MBF_FLOAT32,float,3> itop( pBase->GetVertexBuffers(), MBS_POSITION );
@@ -1291,7 +977,7 @@ namespace mu
         for ( int f=0; f<sourceFaceCount; ++f )
         {
             bool hasFace = false;
-            vec3<uint32_t> ov;
+            vec3<uint32> ov;
             ov[0] = ito.GetAsUINT32(); ++ito;
             ov[1] = ito.GetAsUINT32(); ++ito;
             ov[2] = ito.GetAsUINT32(); ++ito;
@@ -1299,13 +985,13 @@ namespace mu
             // find the bucket for this face
             float y = ( (*(itop+ov[0]))[AXIS] + (*(itop+ov[1]))[AXIS] + (*(itop+ov[2]))[AXIS] ) / 3;
             float fbucket = (y-bucketStart) / bucketSize;
-            int bucket = std::min( NUM_BUCKETS-1, std::max( 0, (int)fbucket ) );
+            int bucket = FMath::Min( NUM_BUCKETS-1, FMath::Max( 0, (int)fbucket ) );
 
-            for ( std::size_t btf=0; !hasFace && btf<buckets[bucket].size(); btf++ )
+            for ( int32 btf=0; !hasFace && btf<buckets[bucket].Num(); btf++ )
             {
                 int tf =  buckets[bucket][btf];
 
-                vec3<uint32_t> v;
+                vec3<uint32> v;
                 v[0] = (itti+3*tf+0).GetAsUINT32();
                 v[1] = (itti+3*tf+1).GetAsUINT32();
                 v[2] = (itti+3*tf+2).GetAsUINT32();
@@ -1326,13 +1012,14 @@ namespace mu
         }
 
         // We only remove vertices if all their faces are clipped
-        mu::vector<uint8_t> vertex_with_face_not_clipped(sourceVertexCount,0);
+		TArray<uint8> vertex_with_face_not_clipped;
+		vertex_with_face_not_clipped.SetNumZeroed(sourceVertexCount);
 
         UntypedMeshBufferIteratorConst itoi(pBase->GetIndexBuffers(), MBS_VERTEXINDEX);
         int aFaceCount = pBase->GetFaceCount();
         for (int f = 0; f < aFaceCount; ++f)
         {
-            vec3<uint32_t> ov;
+            vec3<uint32> ov;
             ov[0] = itoi.GetAsUINT32(); ++itoi;
             ov[1] = itoi.GetAsUINT32(); ++itoi;
             ov[2] = itoi.GetAsUINT32(); ++itoi;
diff --git a/Engine/Plugins/Experimental/Mutable/Source/MutableRuntime/Private/MuR/OpMeshDifference.h b/Engine/Plugins/Experimental/Mutable/Source/MutableRuntime/Private/MuR/OpMeshDifference.h
--- a/Engine/Plugins/Experimental/Mutable/Source/MutableRuntime/Private/MuR/OpMeshDifference.h
+++ b/Engine/Plugins/Experimental/Mutable/Source/MutableRuntime/Private/MuR/OpMeshDifference.h
@@ -53,8 +53,8 @@ struct MORPH_VERTEX
 		uint32_t vcount = pBase->GetVertexBuffers().GetElementCount();		
 
 		// If no channels were specified, get them all
-		vector<MESH_BUFFER_SEMANTIC> allSemantics;		
-		vector<int> allSemanticIndices;
+		TArray<MESH_BUFFER_SEMANTIC> allSemantics;		
+		TArray<int> allSemanticIndices;
 		if ( !numChannels )
 		{
 			for ( int vb=0; vb<pBase->GetVertexBuffers().GetBufferCount(); ++vb )
@@ -74,8 +74,8 @@ struct MORPH_VERTEX
 					     sem != MBS_OTHER && 
 					     ( !ignoreTexCoords || sem!=MBS_TEXCOORDS ) )
 					{
-						allSemantics.push_back( sem );
-						allSemanticIndices.push_back( semIndex );
+						allSemantics.Add( sem );
+						allSemanticIndices.Add( semIndex );
 						++numChannels;
 					}
 				}
@@ -90,8 +90,10 @@ struct MORPH_VERTEX
 		// TODO: Not always 4 components
 		// TODO: Not always floats
 		int differentVertexCount = 0;
-		vector<bool> isVertexDifferent( vcount, false );
-		vector< vec4<float> > deltas( vcount*numChannels );
+		TArray<bool> isVertexDifferent;
+		isVertexDifferent.SetNumZeroed(vcount);
+		TArray< vec4<float> > deltas;
+		deltas.SetNum(vcount * numChannels);
 
 		for ( int c=0; c<numChannels; ++c )
 		{
@@ -129,28 +131,28 @@ struct MORPH_VERTEX
 			pDest->GetVertexBuffers().SetElementCount( differentVertexCount );			
 			pDest->GetVertexBuffers().SetBufferCount( 1 );								
 
-			vector<MESH_BUFFER_SEMANTIC> semantic;				
-			vector<int> semanticIndex;							
-			vector<MESH_BUFFER_FORMAT> format;					
-			vector<int> components;								
-			vector<int> offsets;								
+			TArray<MESH_BUFFER_SEMANTIC> semantic;
+			TArray<int> semanticIndex;
+			TArray<MESH_BUFFER_FORMAT> format;
+			TArray<int> components;
+			TArray<int> offsets;
 			int offset = 0;										
 
 			// Vertex index channel
-			semantic.push_back( MBS_VERTEXINDEX );					
-			semanticIndex.push_back( 0 );
-			format.push_back( MBF_UINT32 );		
-			components.push_back( 1 );
-			offsets.push_back( offset );
+			semantic.Add( MBS_VERTEXINDEX );					
+			semanticIndex.Add( 0 );
+			format.Add( MBF_UINT32 );
+			components.Add( 1 );
+			offsets.Add( offset );
 			offset += 4;						
 
 			for ( int c=0; c<numChannels; ++c )
 			{
-				semantic.push_back( semantics[c] );
-				semanticIndex.push_back( semanticIndices[c] );
-				format.push_back( MBF_FLOAT32 );
-				components.push_back( 4 );
-				offsets.push_back( offset );
+				semantic.Add( semantics[c] );
+				semanticIndex.Add( semanticIndices[c] );
+				format.Add( MBF_FLOAT32 );
+				components.Add( 4 );
+				offsets.Add( offset );
 				offset += 16;				
 			}
 
diff --git a/Engine/Plugins/Experimental/Mutable/Source/MutableRuntime/Private/MuR/OpMeshExtractLayoutBlock.h b/Engine/Plugins/Experimental/Mutable/Source/MutableRuntime/Private/MuR/OpMeshExtractLayoutBlock.h
--- a/Engine/Plugins/Experimental/Mutable/Source/MutableRuntime/Private/MuR/OpMeshExtractLayoutBlock.h
+++ b/Engine/Plugins/Experimental/Mutable/Source/MutableRuntime/Private/MuR/OpMeshExtractLayoutBlock.h
@@ -13,10 +13,10 @@ namespace mu
     //---------------------------------------------------------------------------------------------
     inline void MeshExtractFromVertices( const Mesh* pSource,
                                          Mesh* pResult,
-                                         const vector<int>& oldToNew,
-                                         const vector<int>& newToOld )
+                                         const TArray<int>& oldToNew,
+                                         const TArray<int>& newToOld )
     {
-        int resultVertices = (int)newToOld.size();
+        int resultVertices = newToOld.Num();
 
         // Assemble the new vertex buffer
         pResult->GetVertexBuffers().SetElementCount( resultVertices );
@@ -34,7 +34,8 @@ namespace mu
         }
 
         // Assemble the new index buffers
-        vector<bool> usedSourceFaces( pSource->GetFaceCount(), false );
+		TArray<bool> usedSourceFaces;
+		usedSourceFaces.SetNumZeroed(pSource->GetFaceCount());
         UntypedMeshBufferIteratorConst itIndex( pSource->GetIndexBuffers(), MBS_VERTEXINDEX );
         UntypedMeshBufferIterator itResultIndex( pResult->GetIndexBuffers(), MBS_VERTEXINDEX );
         int indexCount = 0;
@@ -53,16 +54,16 @@ namespace mu
                     usedSourceFaces[i] = true;
 
                     // Clamp in case triangles go across blocks
-                    pDestIndices[ indexCount++ ] = std::max( 0, oldToNew[ pIndices[i*3+0] ] );
-                    pDestIndices[ indexCount++ ] = std::max( 0, oldToNew[ pIndices[i*3+1] ] );
-                    pDestIndices[ indexCount++ ] = std::max( 0, oldToNew[ pIndices[i*3+2] ] );
+                    pDestIndices[ indexCount++ ] = FMath::Max( 0, oldToNew[ pIndices[i*3+0] ] );
+                    pDestIndices[ indexCount++ ] = FMath::Max( 0, oldToNew[ pIndices[i*3+1] ] );
+                    pDestIndices[ indexCount++ ] = FMath::Max( 0, oldToNew[ pIndices[i*3+2] ] );
                 }
             }
         }
         else if ( itIndex.GetFormat()==MBF_UINT16 )
         {
-            const uint16_t* pIndices = reinterpret_cast<const uint16_t*>( itIndex.ptr() );
-            uint16_t* pDestIndices = reinterpret_cast<uint16_t*>( itResultIndex.ptr() );
+            const uint16* pIndices = reinterpret_cast<const uint16*>( itIndex.ptr() );
+            uint16* pDestIndices = reinterpret_cast<uint16*>( itResultIndex.ptr() );
             for ( int i=0; i<pSource->GetIndexCount()/3; ++i )
             {
                 if ( oldToNew[ pIndices[i*3+0] ]>=0
@@ -74,9 +75,9 @@ namespace mu
                     usedSourceFaces[i] = true;
 
                     // Clamp in case triangles go across blocks
-                    pDestIndices[ indexCount++ ] = (uint16_t)std::max( 0, oldToNew[ pIndices[i*3+0] ] );
-                    pDestIndices[ indexCount++ ] = (uint16_t)std::max( 0, oldToNew[ pIndices[i*3+1] ] );
-                    pDestIndices[ indexCount++ ] = (uint16_t)std::max( 0, oldToNew[ pIndices[i*3+2] ] );
+                    pDestIndices[ indexCount++ ] = (uint16)FMath::Max( 0, oldToNew[ pIndices[i*3+0] ] );
+                    pDestIndices[ indexCount++ ] = (uint16)FMath::Max( 0, oldToNew[ pIndices[i*3+1] ] );
+                    pDestIndices[ indexCount++ ] = (uint16)FMath::Max( 0, oldToNew[ pIndices[i*3+2] ] );
                 }
             }
         }
@@ -114,7 +115,7 @@ namespace mu
 	//---------------------------------------------------------------------------------------------
     inline MeshPtr MeshExtractLayoutBlock( const Mesh* pSource,
                                            uint32_t layout,
-                                           uint16_t blockCount,
+                                           uint16 blockCount,
                                            const uint32_t* pExtractBlocks )
 	{
 		// TODO: Optimise
@@ -129,13 +130,14 @@ namespace mu
         if (itBlocks.GetFormat()!=MBF_NONE)
         {
             int resultVertices = 0;
-            vector<int> oldToNew( pSource->GetVertexCount(), -1 );
-            vector<int> newToOld;
-            newToOld.reserve( pSource->GetVertexCount() );
+			TArray<int> oldToNew;
+			oldToNew.Init(-1,pSource->GetVertexCount());
+			TArray<int> newToOld;
+            newToOld.Reserve( pSource->GetVertexCount() );
 
             if ( itBlocks.GetFormat()==MBF_UINT16 )
             {
-                const uint16_t* pBlocks = reinterpret_cast<const uint16_t*>( itBlocks.ptr() );
+                const uint16* pBlocks = reinterpret_cast<const uint16*>( itBlocks.ptr() );
                 for ( int i=0; i<pSource->GetVertexCount(); ++i )
                 {
                     uint32_t vertexBlock = pBlocks[i];
@@ -153,7 +155,7 @@ namespace mu
                     if ( found )
                     {
                         oldToNew[i] = resultVertices++;
-                        newToOld.push_back( i );
+                        newToOld.Add( i );
                     }
                 }
             }
@@ -184,9 +186,10 @@ namespace mu
         MeshPtr pResult = pSource->Clone();
 
         int resultVertices = 0;
-        vector<int> oldToNew( pSource->GetVertexCount(), -1 );
-        vector<int> newToOld;
-        newToOld.reserve( pSource->GetVertexCount() );
+		TArray<int> oldToNew;
+		oldToNew.Init(-1,pSource->GetVertexCount());
+		TArray<int> newToOld;
+        newToOld.Reserve( pSource->GetVertexCount() );
 
         UntypedMeshBufferIteratorConst itIndex( pSource->GetIndexBuffers(), MBS_VERTEXINDEX );
         for ( int32 f=0; f<pSource->m_faceGroups[group].m_faces.Num(); ++f )
@@ -199,7 +202,7 @@ namespace mu
                 if ( oldToNew[vertexIndex]<0 )
                 {
                     oldToNew[vertexIndex] = resultVertices++;
-                    newToOld.push_back( vertexIndex );
+                    newToOld.Add( vertexIndex );
                 }
             }
         }
