[mutable] Optimize mesh clip versus switching of clipping mesh.

diff --git a/Engine/Plugins/Experimental/Mutable/Source/MutableRuntime/Private/MuR/OpMeshFormat.cpp b/Engine/Plugins/Experimental/Mutable/Source/MutableRuntime/Private/MuR/OpMeshFormat.cpp
--- a/Engine/Plugins/Experimental/Mutable/Source/MutableRuntime/Private/MuR/OpMeshFormat.cpp
+++ b/Engine/Plugins/Experimental/Mutable/Source/MutableRuntime/Private/MuR/OpMeshFormat.cpp
@@ -153,206 +153,6 @@ namespace mu
 
 		};
 
-
-		void RebuildTangents
-		(
-			const FMeshBufferSet& IndexBuffers,
-			FMeshBufferSet& VertexBuffers
-		)
-		{
-			// Because of these, this method only works for a specific mesh format
-			UntypedMeshBufferIteratorConst vertex(VertexBuffers, MBS_POSITION, 0);
-			UntypedMeshBufferIteratorConst normal(VertexBuffers, MBS_NORMAL, 0);
-			UntypedMeshBufferIteratorConst texcoord(VertexBuffers, MBS_TEXCOORDS, 0);
-
-			int vertexCount = VertexBuffers.GetElementCount();
-
-			// Group similar vertices
-			vector< TVertex > vertexData;
-			vertexData.reserve(vertexCount);
-			vector< uint32_t > lVertexGroups(vertexCount, 0);
-
-			//Agrupamos los vertices por posicion, normal, texcoord
-			for (int32_t i = 0; i < vertexCount; ++i)
-			{
-				TVertex v((vertex + i).GetAsVec4f().xyz(),
-					(normal + i).GetAsVec4f().xyz(),
-					(texcoord + i).GetAsVec4f().xy());
-				vector<TVertex>::const_iterator it = std::find(vertexData.begin(),
-					vertexData.end(),
-					v);
-
-				size_t pos = 0;
-				if (it != vertexData.end())
-				{
-					pos = it - vertexData.begin();
-				}
-				else
-				{
-					pos = vertexData.size();
-					vertexData.push_back(v);
-				}
-
-				lVertexGroups[i] = (int)pos;
-			}
-
-			// Calculate the tangent space
-			MeshBufferIterator<MBF_FLOAT32, float, 3> tangent(VertexBuffers, MBS_TANGENT, 0);
-			MeshBufferIterator<MBF_FLOAT32, float, 3> bitangent(VertexBuffers, MBS_BINORMAL, 0);
-			UntypedMeshBufferIteratorConst indices(IndexBuffers, MBS_VERTEXINDEX, 0);
-			int indexCount = IndexBuffers.GetElementCount();
-
-			vector<TFace> vertexFaces(indexCount);
-
-			// Worst case: a group for every vertex
-			vector< vector< TFace > > grp_faces(indexCount);
-
-			for (long faceidx = 0, count = indexCount / 3; faceidx < count; faceidx++)
-			{
-				long i1 = 0;
-				long i2 = 0;
-				long i3 = 0;
-
-				if (indices.GetFormat() == MBF_UINT16)
-				{
-					auto pIndices = reinterpret_cast<const uint16_t*>(indices.ptr());
-					i1 = pIndices[faceidx * 3];
-					i2 = pIndices[faceidx * 3 + 1];
-					i3 = pIndices[faceidx * 3 + 2];
-				}
-				else if (indices.GetFormat() == MBF_UINT32)
-				{
-					auto pIndices = reinterpret_cast<const uint32_t*>(indices.ptr());
-					i1 = pIndices[faceidx * 3];
-					i2 = pIndices[faceidx * 3 + 1];
-					i3 = pIndices[faceidx * 3 + 2];
-				}
-				else
-				{
-					check(false);
-				}
-
-				TFace face;
-
-				const vec3<float> v1 = (vertex + (int)i1).GetAsVec4f().xyz();
-				const vec3<float> v2 = (vertex + (int)i2).GetAsVec4f().xyz();
-				const vec3<float> v3 = (vertex + (int)i3).GetAsVec4f().xyz();
-
-				// Compute face t & b
-				const vec2<float> w1 = (texcoord + (int)i1).GetAsVec4f().xy();
-				const vec2<float> w2 = (texcoord + (int)i2).GetAsVec4f().xy();
-				const vec2<float> w3 = (texcoord + (int)i3).GetAsVec4f().xy();
-				face = TFace(v1, v2, v3, w1, w2, w3);
-
-				if (length(face.N) != 0)
-				{
-					vertexFaces[i1] = vertexFaces[i2] = vertexFaces[i3] = face;
-
-					grp_faces[lVertexGroups[i1]].push_back(face);
-					grp_faces[lVertexGroups[i2]].push_back(face);
-					grp_faces[lVertexGroups[i3]].push_back(face);
-				}
-			}
-
-			for (long vtxidx = 0; vtxidx < vertexCount; vtxidx++)
-			{
-				const vec3<float> n = (normal + (int)vtxidx).GetAsVec4f().xyz();
-				const uint32_t group = lVertexGroups[vtxidx];
-
-				vec3<float> lTangent(0, 0, 0);
-				vec3<float> lBiTangent(0, 0, 0);
-
-				const vector< TFace >& faces = grp_faces[group];
-
-				const TFace& vFace = vertexFaces[vtxidx];
-
-				// Ignore mirrors
-				for (unsigned i = 0; i < faces.size(); ++i)
-				{
-					const float DN = dot(vFace.N, faces[i].N);
-					const float DT = dot(vFace.T, faces[i].T);
-					const float DB = dot(vFace.B, faces[i].B);
-
-					if (DN > 0 && DT > 0 && DB > 0)
-					{
-						lTangent += faces[i].T;
-						lBiTangent += faces[i].B;
-					}
-				}
-
-				vec3<float> ortogonalizedTangent
-					= normalise(lTangent
-						- (n * dot(n, lTangent)));
-
-				vec3<float> ortogonalizedBiTangent
-					= normalise(lBiTangent
-						- (n * dot(n, lBiTangent))
-						- (ortogonalizedTangent * dot(ortogonalizedTangent, lBiTangent)));
-
-				// this is where we can do a final check for zero length vectors
-				// and set them to something appropriate
-				float lenTan = length(ortogonalizedTangent);
-				float lenBin = length(ortogonalizedBiTangent);
-
-				if ((lenTan <= MUTABLE_TANGENT_GENERATION_EPSILON_2)
-					|| (lenBin <= MUTABLE_TANGENT_GENERATION_EPSILON_2)) //should be approx 1.0f
-				{
-					// the tangent space is ill defined at this vertex
-					// so we can generate a valid one based on the normal vector,
-					// which I'm assuming is valid!
-					if (lenTan > 0.5f)
-					{
-						// the tangent is valid, so we can just use that to calculate the binormal
-						ortogonalizedBiTangent = cross(n, ortogonalizedTangent);
-					}
-					else if (lenBin > 0.5)
-					{
-						// the binormal is good and we can use it to calculate the tangent
-						ortogonalizedTangent = cross(ortogonalizedBiTangent, n);
-					}
-					else
-					{
-						// both vectors are invalid, so we should create something
-						// that is at least valid if not correct
-						// I'm checking two possible axis, because the normal could be one of them,
-						// and we want to chose a different one to start making our valid basis.
-						// I can find out which is further away from it by checking the dot product
-						vec3<float> startAxis;
-
-						if (dot(vec3<float>(1, 0, 0), n) < dot(vec3<float>(0, 1, 0), n))
-						{
-							//the xAxis is more different than the yAxis when compared to the normal
-							startAxis = vec3<float>(1, 0, 0);
-						}
-						else
-						{
-							// the yAxis is more different than the xAxis when compared to the normal
-							startAxis = vec3<float>(0, 1, 0);
-						}
-
-						ortogonalizedTangent = cross(n, startAxis);
-						ortogonalizedBiTangent = cross(n, ortogonalizedTangent);
-					}
-				}
-				else
-				{
-					// one final sanity check, make sure that they tangent and binormal are different
-					// enough
-					if (dot(ortogonalizedTangent, ortogonalizedBiTangent)
-						>
-						MUTABLE_TANGENT_MIN_AXIS_DIFFERENCE)
-					{
-						// then they are too similar lets make them more different
-						ortogonalizedBiTangent = cross(n, ortogonalizedTangent);
-					}
-				}
-
-				tangent[vtxidx] = ortogonalizedTangent;
-				bitangent[vtxidx] = ortogonalizedBiTangent;
-			}
-
-		}
-
 	}
 
 
@@ -481,7 +281,7 @@ namespace mu
 						{
 							for (int v = 0; v < vCount; ++v)
 							{
-								auto pTypedResultBuf = (uint16_t*)pResultBuf;
+								auto pTypedResultBuf = (uint16*)pResultBuf;
 								for (int i = 0; i < resultComponents; ++i)
 								{
 									pTypedResultBuf[i] = 65535;
@@ -640,8 +440,8 @@ namespace mu
 
 								else if (resultFormat == MBF_NUINT16)
 								{
-									uint16_t* pData = (uint16_t*)pResultBuf;
-									uint16_t accum = 0;
+									uint16* pData = (uint16*)pResultBuf;
+									uint16 accum = 0;
 									for (int i = 0; i < resultComponents; ++i)
 									{
 										accum += pData[i];
@@ -676,11 +476,11 @@ namespace mu
 			// offsets.
 			for (int b = 0; b < Result.GetBufferCount(); ++b)
 			{
-				vector<MESH_BUFFER_SEMANTIC> resultSemantics;
-				vector<int> resultSemanticIndexs;
-				vector<MESH_BUFFER_FORMAT> resultFormats;
-				vector<int> resultComponentss;
-				vector<int> resultOffsets;
+				TArray<MESH_BUFFER_SEMANTIC> resultSemantics;
+				TArray<int> resultSemanticIndexs;
+				TArray<MESH_BUFFER_FORMAT> resultFormats;
+				TArray<int> resultComponentss;
+				TArray<int> resultOffsets;
 				int offset = 0;
 
 				// For every channel in this buffer
@@ -707,23 +507,23 @@ namespace mu
 
 					if (sourceBuffer >= 0)
 					{
-						resultSemantics.push_back(resultSemantic);
-						resultSemanticIndexs.push_back(resultSemanticIndex);
-						resultFormats.push_back(resultFormat);
-						resultComponentss.push_back(resultComponents);
-						resultOffsets.push_back(offset);
+						resultSemantics.Add(resultSemantic);
+						resultSemanticIndexs.Add(resultSemanticIndex);
+						resultFormats.Add(resultFormat);
+						resultComponentss.Add(resultComponents);
+						resultOffsets.Add(offset);
 
 						offset += GetMeshFormatData(resultFormat).m_size * resultComponents;
 					}
 				}
 
-				if (resultSemantics.empty())
+				if (resultSemantics.IsEmpty())
 				{
 					Result.SetBuffer(b, 0, 0, nullptr, nullptr, nullptr, nullptr, nullptr);
 				}
 				else
 				{
-					Result.SetBuffer(b, offset, (int)resultSemantics.size(),
+					Result.SetBuffer(b, offset, resultSemantics.Num(),
 						&resultSemantics[0],
 						&resultSemanticIndexs[0],
 						&resultFormats[0],
@@ -788,7 +588,6 @@ namespace mu
 		bool formatVertices,
 		bool formatIndices,
 		bool formatFaces,
-		bool rebuildTangents,
 		bool ignoreMissingChannels
 	)
 	{
@@ -799,52 +598,6 @@ namespace mu
 
 		MeshPtrConst pSource = pPureSource;
 
-		// Rebuild the tangent space if necessary
-		if (rebuildTangents)
-		{
-			MeshPtr pClonedSource = pPureSource->Clone();
-			pSource = pClonedSource;
-
-			// Make sure there are the tangents and binormal channels.
-			int tanBuf, tanChan, binBuf, binChan;
-			pSource->GetVertexBuffers().FindChannel(MBS_TANGENT, 0, &tanBuf, &tanChan);
-			pSource->GetVertexBuffers().FindChannel(MBS_BINORMAL, 0, &binBuf, &binChan);
-
-			if (tanBuf <= 0 || binBuf <= 0)
-			{
-				MESH_BUFFER buffer;
-
-				if (tanBuf <= 0)
-				{
-					MESH_BUFFER_CHANNEL channel;
-					channel.m_semantic = MBS_TANGENT;
-					channel.m_format = MBF_FLOAT32;
-					channel.m_componentCount = 3;
-					channel.m_semanticIndex = 0;
-					channel.m_offset = 0;
-					buffer.m_channels.Add(channel);
-					buffer.m_elementSize += 3 * sizeof(float);
-				}
-
-				if (binBuf <= 0)
-				{
-					MESH_BUFFER_CHANNEL channel;
-					channel.m_semantic = MBS_BINORMAL;
-					channel.m_format = MBF_FLOAT32;
-					channel.m_componentCount = 3;
-					channel.m_semanticIndex = 0;
-					channel.m_offset = (uint8)buffer.m_elementSize;
-					buffer.m_channels.Add(channel);
-					buffer.m_elementSize += 3 * sizeof(float);
-				}
-
-				buffer.m_data.SetNum(buffer.m_elementSize * pClonedSource->GetVertexCount());
-				pClonedSource->GetVertexBuffers().m_buffers.Add(buffer);
-			}
-
-			RebuildTangents(pClonedSource->GetIndexBuffers(), pClonedSource->GetVertexBuffers());
-		}
-
 		MeshPtr pResult = pFormat->Clone();
 
 		// Make sure that the bone indices will fit in this format, or extend it.
@@ -874,7 +627,7 @@ namespace mu
 						auto va = it.GetAsVec8i();
 						for (int c = 0; c < it.GetComponents(); ++c)
 						{
-							maxBoneIndex = std::max(maxBoneIndex, va[c]);
+							maxBoneIndex = FMath::Max(maxBoneIndex, va[c]);
 						}
 						++it;
 					}
diff --git a/Engine/Plugins/Experimental/Mutable/Source/MutableRuntime/Private/MuR/OpMeshFormat.h b/Engine/Plugins/Experimental/Mutable/Source/MutableRuntime/Private/MuR/OpMeshFormat.h
--- a/Engine/Plugins/Experimental/Mutable/Source/MutableRuntime/Private/MuR/OpMeshFormat.h
+++ b/Engine/Plugins/Experimental/Mutable/Source/MutableRuntime/Private/MuR/OpMeshFormat.h
@@ -23,7 +23,6 @@ class FMeshBufferSet;
 			bool formatVertices,
 			bool formatIndices,
 			bool formatFaces,
-			bool rebuildTangents,
 			bool ignoreMissingChannels
 		);
 
diff --git a/Engine/Plugins/Experimental/Mutable/Source/MutableRuntime/Private/MuR/OpMeshMerge.h b/Engine/Plugins/Experimental/Mutable/Source/MutableRuntime/Private/MuR/OpMeshMerge.h
--- a/Engine/Plugins/Experimental/Mutable/Source/MutableRuntime/Private/MuR/OpMeshMerge.h
+++ b/Engine/Plugins/Experimental/Mutable/Source/MutableRuntime/Private/MuR/OpMeshMerge.h
@@ -136,7 +136,7 @@ namespace mu
                             {
                                 for ( int v=0; v<firstCount; ++v )
                                 {
-                                    *(uint32_t*)pDest = *(const uint16_t*)pSource;
+                                    *(uint32_t*)pDest = *(const uint16*)pSource;
                                     pSource+=first.m_elementSize;
                                     pDest+=result.m_elementSize;
                                 }
@@ -170,7 +170,7 @@ namespace mu
                             {
                                 for ( int v=0; v<firstCount; ++v )
                                 {
-                                    *(uint16_t*)pDest = *(const uint8_t*)pSource;
+                                    *(uint16*)pDest = *(const uint8_t*)pSource;
                                     pSource+=first.m_elementSize;
                                     pDest+=result.m_elementSize;
                                 }
@@ -225,7 +225,7 @@ namespace mu
                         {
                             for ( int v=0; v<secondCount; ++v )
                             {
-                                *(uint16_t*)pDest = uint16_t(firstVertexCount) + *(const uint16_t*)pSource;
+                                *(uint16*)pDest = uint16(firstVertexCount) + *(const uint16*)pSource;
                                 pSource+=second.m_elementSize;
                                 pDest+=result.m_elementSize;
                             }
@@ -268,7 +268,7 @@ namespace mu
                             {
                                 for ( int v=0; v<secondCount; ++v )
                                 {
-                                    *(uint32_t*)pDest = uint32_t(firstVertexCount) + *(const uint16_t*)pSource;
+                                    *(uint32_t*)pDest = uint32_t(firstVertexCount) + *(const uint16*)pSource;
                                     pSource+=second.m_elementSize;
                                     pDest+=result.m_elementSize;
                                 }
@@ -308,7 +308,7 @@ namespace mu
                             {
                                 for ( int v=0; v<secondCount; ++v )
                                 {
-                                    *(uint16_t*)pDest = uint16_t(firstVertexCount) + *(const uint8_t*)pSource;
+                                    *(uint16*)pDest = uint16(firstVertexCount) + *(const uint8_t*)pSource;
                                     pSource+=second.m_elementSize;
                                     pDest+=result.m_elementSize;
                                 }
@@ -500,7 +500,7 @@ namespace mu
                 int16_t secondMeshIndex = pResultSkeleton->m_boneParents[ob];
                 if (secondMeshIndex != INDEX_NONE)
                 {
-                    pResultSkeleton->m_boneParents[ob] = (uint16_t)secondToFirstBones[secondMeshIndex];
+                    pResultSkeleton->m_boneParents[ob] = (uint16)secondToFirstBones[secondMeshIndex];
                 }
             }
         }
@@ -788,7 +788,7 @@ namespace mu
                 }
                 else
                 {
-                    pVFirst = MeshFormat( pFirst, pResult.get(), false, true, false, false, false, false );
+                    pVFirst = MeshFormat( pFirst, pResult.get(), false, true, false, false, false );
                 }
 
 
@@ -798,7 +798,7 @@ namespace mu
                 }
                 else
                 {
-                    pVSecond = MeshFormat( pSecond, pResult.get(), false, true, false, false, false, false );
+                    pVSecond = MeshFormat( pSecond, pResult.get(), false, true, false, false, false );
                 }
 
 				check( pVFirst->GetVertexBuffers().HasSameFormat
@@ -897,18 +897,18 @@ namespace mu
 								case MBF_INT16:
 								case MBF_UINT16:
 								{
-                                    uint16_t* pD = reinterpret_cast<uint16_t*>
+                                    uint16* pD = reinterpret_cast<uint16*>
 											( &result.m_data[resultOffset] );
 
                                     int comp=0;
                                     for ( ; comp<components; ++comp )
 									{
-                                        uint16_t bone = pD[comp];
+                                        uint16 bone = pD[comp];
 
                                         // be defensive
                                         if (bone<secondToFirstBones.Num())
                                         {
-                                            pD[comp] = (uint16_t)secondToFirstBones[ bone ];
+                                            pD[comp] = (uint16)secondToFirstBones[ bone ];
                                         }
                                         else
                                         {
@@ -986,7 +986,7 @@ namespace mu
     //---------------------------------------------------------------------------------------------
     inline void ExtendSkeleton( Skeleton* pBase, const Skeleton* pOther )
     {
-        map<int,int> otherToResult;
+        TMap<int,int> otherToResult;
 
         int initialBones = pBase->GetBoneCount();
         for ( int b=0; pOther && b<pOther->GetBoneCount(); ++b)
@@ -995,7 +995,7 @@ namespace mu
             if ( resultBoneIndex<0 )
             {
                 int32 newIndex = pBase->m_bones.Num();
-                otherToResult[b] = int(newIndex);
+                otherToResult.Add(b,newIndex);
                 pBase->m_bones.Add( pOther->m_bones[b] );
 
                 // Will be remapped below
@@ -1006,7 +1006,7 @@ namespace mu
             }
             else
             {
-                otherToResult[b] = resultBoneIndex;
+                otherToResult.Add(b,resultBoneIndex);
             }
         }
 
@@ -1042,10 +1042,10 @@ namespace mu
         pResult->SetSkeleton( pSkeleton );
 
         // Build a skeleton map
-        vector<int> sourceToSkeleton;
+        TArray<int> sourceToSkeleton;
         for ( const auto& boneName: pSourceSkeleton->m_bones )
         {
-            sourceToSkeleton.push_back( pResult->GetSkeleton()->FindBone(boneName.c_str()) );
+            sourceToSkeleton.Add( pResult->GetSkeleton()->FindBone(boneName.c_str()) );
         }
 
         bool somethingRemapped = false;
@@ -1073,7 +1073,7 @@ namespace mu
                             {
                                 size_t boneIndex = size_t( pTypedData[c] );
                                 //check( boneIndex < sourceToSkeleton.Num() );
-                                if ( boneIndex < sourceToSkeleton.size() )
+                                if ( boneIndex < sourceToSkeleton.Num() )
                                 {
                                     int bone = sourceToSkeleton[pTypedData[c]];
                                     if ( pTypedData[c] != uint8_t( bone ) )
@@ -1098,18 +1098,18 @@ namespace mu
                     {
                         for (uint32_t v=0;v<BufSet.m_elementCount; ++v)
                         {
-                            uint16_t* pTypedData = (uint16_t*)pData;
+                            uint16* pTypedData = (uint16*)pData;
 
                             for (int c=0; c<chan.m_componentCount; ++c)
                             {
                                 size_t boneIndex = size_t(pTypedData[c]);
                                 //check( boneIndex < sourceToSkeleton.size() );
-                                if ( boneIndex < sourceToSkeleton.size() )
+                                if ( boneIndex < sourceToSkeleton.Num() )
                                 {
                                     int bone = sourceToSkeleton[boneIndex];
-                                    if ( pTypedData[c] != uint16_t(bone) )
+                                    if ( pTypedData[c] != uint16(bone) )
                                     {
-                                        pTypedData[c] = uint16_t(bone);
+                                        pTypedData[c] = uint16(bone);
                                         somethingRemapped = true;
                                     }
                                 }
@@ -1135,7 +1135,7 @@ namespace mu
                             {
                                 size_t boneIndex = size_t( pTypedData[c] );
                                 //check( boneIndex < sourceToSkeleton.size() );
-                                if ( boneIndex < sourceToSkeleton.size() )
+                                if ( boneIndex < sourceToSkeleton.Num() )
                                 {
                                     int bone = sourceToSkeleton[pTypedData[c]];
                                     if ( pTypedData[c] != uint32_t( bone ) )
diff --git a/Engine/Plugins/Experimental/Mutable/Source/MutableRuntime/Private/MuR/OpMeshMorph.h b/Engine/Plugins/Experimental/Mutable/Source/MutableRuntime/Private/MuR/OpMeshMorph.h
--- a/Engine/Plugins/Experimental/Mutable/Source/MutableRuntime/Private/MuR/OpMeshMorph.h
+++ b/Engine/Plugins/Experimental/Mutable/Source/MutableRuntime/Private/MuR/OpMeshMorph.h
@@ -10,299 +10,6 @@
 
 namespace mu
 {
-
-    //---------------------------------------------------------------------------------------------
-    //! Reference Linear factor version
-    //---------------------------------------------------------------------------------------------
-    inline MeshPtr MeshMorphReference( const Mesh* pBase, const Mesh* pMorph, float factor )
-    {
-        MUTABLE_CPUPROFILER_SCOPE(MeshMorph);
-
-        if (!pBase) return nullptr;
-
-        MeshPtr pDest = pBase->Clone();
-
-        if (!pMorph) return pDest;
-
-        // Number of vertices to modify
-        uint32_t vcountMorph = pMorph->GetVertexBuffers().GetElementCount();
-        uint32_t vcountBase = pBase->GetVertexBuffers().GetElementCount();
-        if ( !vcountMorph || !vcountBase )
-        {
-            return pDest;
-        }
-
-        uint32_t ccount = 0;
-        if ( pMorph->GetVertexBuffers().GetBufferCount()>0 )
-        {
-            ccount = pMorph->GetVertexBuffers().GetBufferChannelCount(0);
-        }
-
-        // Iterator of the vertex ids of the base vertices
-        UntypedMeshBufferIteratorConst itBaseId( pBase->GetVertexBuffers(), MBS_VERTEXINDEX );
-
-        vector< UntypedMeshBufferIterator > itBaseChannels( ccount );
-        vector< UntypedMeshBufferIteratorConst > itMorphChannels( ccount );
-        for ( size_t c=1; c<ccount; ++c )
-        {
-            const FMeshBufferSet& MBSPriv = pMorph->GetVertexBuffers();
-            MESH_BUFFER_SEMANTIC sem = MBSPriv.m_buffers[0].m_channels[c].m_semantic;
-            int semIndex = MBSPriv.m_buffers[0].m_channels[c].m_semanticIndex;
-
-            itBaseChannels[c] = UntypedMeshBufferIterator
-                    ( pDest->GetVertexBuffers(), sem, semIndex );
-            itMorphChannels[c] = UntypedMeshBufferIteratorConst
-                    ( pMorph->GetVertexBuffers(), sem, semIndex );
-        }
-
-        // Morph
-
-        // Number of vertices consumed from the morph mesh
-        uint32_t processedMorphVertices = 0;
-
-        // Morph mesh channels are always this friendly types
-        MeshBufferIteratorConst<MBF_UINT32,uint32_t,1>
-                itFirstMorphVertex( pMorph->GetVertexBuffers(), MBS_VERTEXINDEX );
-
-        // Number of vertices at the beginning of the morph already consumed (optimisation)
-        int morphVerticesFromStartAlreadyConsumed = 0;
-
-        // Number of vertices we need to advance the base mesh iterators to get to the current base
-        // mesh vertex
-        int step = 0;
-
-        // Iterate all the base mesh vertices
-        for ( uint32_t vi=0; vi<vcountBase && processedMorphVertices<vcountMorph; ++vi )
-        {
-            uint32_t baseMeshVertexIndex = itBaseId.GetAsUINT32();
-            int32_t morphVertexIndex = -1;
-
-            // Look for a morph vertex to apply to this base mesh vertex
-            auto itMorphVer = itFirstMorphVertex+morphVerticesFromStartAlreadyConsumed;
-            for ( int32_t morphVertexCandidate=morphVerticesFromStartAlreadyConsumed;
-                  morphVertexCandidate<int32_t(vcountMorph);
-                  ++morphVertexCandidate )
-            {
-                if ((*itMorphVer)[0]==baseMeshVertexIndex)
-                {
-                    morphVertexIndex = morphVertexCandidate;
-                    if (morphVertexCandidate==morphVerticesFromStartAlreadyConsumed)
-                    {
-                        ++morphVerticesFromStartAlreadyConsumed;
-                    }
-                    break;
-                }
-                ++itMorphVer;
-            }
-
-            if( morphVertexIndex>=0 )
-            {
-                // Morph one vertex
-                for ( size_t c=1; c<ccount; ++c )
-                {
-                    if ( itBaseChannels[c].ptr() )
-                    {
-                        itBaseChannels[c] += step;
-
-                        vec4<float> value = itBaseChannels[c].GetAsVec4f();
-                        vec4<float> delta = (itMorphChannels[c]+morphVertexIndex).GetAsVec4f();
-
-                        value += delta * factor;
-
-                        for( int i=0; i<itBaseChannels[c].GetComponents(); ++i )
-                        {
-                            ConvertData( i,
-                                         itBaseChannels[c].ptr(),
-                                         itBaseChannels[c].GetFormat(),
-                                         &value, MBF_FLOAT32 );
-                        }
-                    }
-                }
-
-                // Reset the step to the next vertex
-                step = 1;
-            }
-            else
-            {
-                // Vertex was not morphed, the step to the next vertex will be bigger
-                ++step;
-            }
-
-            ++itBaseId;
-        }
-
-        return pDest;
-    }
-
-
-
-    //---------------------------------------------------------------------------------------------
-    //! Reference Factor-less version
-    //---------------------------------------------------------------------------------------------
-    inline MeshPtr MeshMorphReference( const Mesh* pBase, const Mesh* pMorph )
-    {
-        // Trust the compiler to remove the factor
-        return MeshMorphReference( pBase, pMorph, 1.0f );
-    }
-
-
-    //---------------------------------------------------------------------------------------------
-    //! Reference Linear factor version for morphing 2 targets
-    //---------------------------------------------------------------------------------------------
-    inline MeshPtr MeshMorph2Reference( const Mesh* pBase, const Mesh* pMin, const Mesh* pMax, float factor )
-    {
-        MUTABLE_CPUPROFILER_SCOPE(MeshMorph2Reference);
-
-        MeshPtr pDest = pBase->Clone();
-
-        // Number of vertices to modify
-        uint32_t vmincount = pMin->GetVertexBuffers().GetElementCount();
-        uint32_t vmaxcount = pMax->GetVertexBuffers().GetElementCount();
-
-        uint32_t vcountBase = pBase->GetVertexBuffers().GetElementCount();
-
-        if ( (vmincount+vmaxcount)==0 || !vcountBase )
-        {
-            return pDest;
-        }
-
-
-        auto refTarget = vmincount > 0 ? pMin : pMax;
-
-        uint32_t ccount = refTarget->GetVertexBuffers().GetBufferChannelCount(0);
-
-        // Iterator of the vertex ids of the base vertices
-        UntypedMeshBufferIteratorConst itBaseId( pBase->GetVertexBuffers(), MBS_VERTEXINDEX );
-
-        vector< UntypedMeshBufferIterator > itBaseChannels( ccount );
-        vector< UntypedMeshBufferIteratorConst > itMinChannels( ccount );
-        vector< UntypedMeshBufferIteratorConst > itMaxChannels( ccount );
-        for ( size_t c=1; c<ccount; ++c )
-        {
-            const FMeshBufferSet& MBSPriv = refTarget->GetVertexBuffers();
-            MESH_BUFFER_SEMANTIC sem = MBSPriv.m_buffers[0].m_channels[c].m_semantic;
-            int semIndex = MBSPriv.m_buffers[0].m_channels[c].m_semanticIndex;
-
-            itBaseChannels[c] = UntypedMeshBufferIterator
-                    ( pDest->GetVertexBuffers(), sem, semIndex );
-            itMinChannels[c] = UntypedMeshBufferIteratorConst
-                    ( pMin->GetVertexBuffers(), sem, semIndex );
-            itMaxChannels[c] = UntypedMeshBufferIteratorConst
-                    ( pMax->GetVertexBuffers(), sem, semIndex );
-        }
-
-        // Morph
-        uint32_t minProcessed=0;
-        uint32_t maxProcessed=0;
-        float maxFactor = factor;
-        float minFactor = 1.0f-factor;
-
-        // Morph mesh channels are always this friendly types
-        MeshBufferIteratorConst<MBF_UINT32,uint32_t,1>
-                itFirstMinVer( pMin->GetVertexBuffers(), MBS_VERTEXINDEX );
-        MeshBufferIteratorConst<MBF_UINT32,uint32_t,1>
-                itFirstMaxVer( pMax->GetVertexBuffers(), MBS_VERTEXINDEX );
-
-        // Number of vertices at the beginning of the morph already consumed (optimisation)
-        int minMorphVerticesFromStartAlreadyConsumed = 0;
-        int maxMorphVerticesFromStartAlreadyConsumed = 0;
-
-        // Number of vertices we need to advance the base mesh iterators to get to the current base
-        // mesh vertex
-        int step = 0;
-
-        // Iterate all the base mesh vertices
-        for ( uint32_t i=0; i<vcountBase && (minProcessed<vmincount || maxProcessed<vmaxcount); ++i)
-        {
-            uint32_t baseMeshVertexIndex = itBaseId.GetAsUINT32();
-
-            int32_t minMorphVertexIndex = -1;
-            int32_t maxMorphVertexIndex = -1;
-
-            // Look for a min morph vertex to apply to this base mesh vertex
-            auto itMorphVer = itFirstMinVer+minMorphVerticesFromStartAlreadyConsumed;
-            for ( int32_t morphVertexCandidate=minMorphVerticesFromStartAlreadyConsumed;
-                  morphVertexCandidate<int32_t(vmincount);
-                  ++morphVertexCandidate )
-            {
-                if ((*itMorphVer)[0]==baseMeshVertexIndex)
-                {
-                    minMorphVertexIndex = morphVertexCandidate;
-                    if (morphVertexCandidate==minMorphVerticesFromStartAlreadyConsumed)
-                    {
-                        ++minMorphVerticesFromStartAlreadyConsumed;
-                    }
-                    break;
-                }
-                ++itMorphVer;
-            }
-
-            // Look for a max morph vertex to apply to this base mesh vertex
-            itMorphVer = itFirstMaxVer+maxMorphVerticesFromStartAlreadyConsumed;
-            for ( int32_t morphVertexCandidate=maxMorphVerticesFromStartAlreadyConsumed;
-                  morphVertexCandidate<int32_t(vmaxcount);
-                  ++morphVertexCandidate )
-            {
-                if ((*itMorphVer)[0]==baseMeshVertexIndex)
-                {
-                    maxMorphVertexIndex = morphVertexCandidate;
-                    if (morphVertexCandidate==maxMorphVerticesFromStartAlreadyConsumed)
-                    {
-                        ++maxMorphVerticesFromStartAlreadyConsumed;
-                    }
-                    break;
-                }
-                ++itMorphVer;
-            }
-            
-            // Morph one vertex
-            if ( minMorphVertexIndex>=0 || maxMorphVertexIndex>=0 )
-            {
-                for ( size_t c=1; c<ccount; ++c )
-                {
-                    if ( itBaseChannels[c].ptr() )
-                    {
-                        itBaseChannels[c] += step;
-
-                        vec4<float> value = itBaseChannels[c].GetAsVec4f();
-
-                        if ( minMorphVertexIndex>=0 )
-                        {
-                            vec4<float> minvalue = (itMinChannels[c]+minMorphVertexIndex).GetAsVec4f();
-                            value += minvalue * minFactor;
-                        }
-
-                        if ( maxMorphVertexIndex>=0 )
-                        {
-                            vec4<float> maxvalue = (itMaxChannels[c]+maxMorphVertexIndex).GetAsVec4f();
-                            value += maxvalue * maxFactor;
-                        }
-
-						const int32_t componentCount = itBaseChannels[c].GetComponents();
-                        for( int component=0; component<componentCount; ++component )
-                        {
-                            ConvertData( component,
-                                         itBaseChannels[c].ptr(),
-                                         itBaseChannels[c].GetFormat(),
-                                         &value, MBF_FLOAT32 );
-                        }
-                    }
-                }
-
-                step = 1;
-            }
-            else
-            {
-                // Vertex was not morphed, the step to the next vertex will be bigger
-                ++step;
-            }
-
-            ++itBaseId;
-        }
-
-        return pDest;
-    }
-
 	
 	//---------------------------------------------------------------------------------------------
 	//! Optimized linear factor version for morphing 2 targets
@@ -314,8 +21,8 @@ namespace mu
 		if (!pBase) return nullptr;
 
         const auto ApplyMorph = []
-            ( auto BaseIdIter, const std::vector< UntypedMeshBufferIterator >& baseChannelsIters, const int32_t baseSize,
-              auto MorphIdIter, const std::vector< UntypedMeshBufferIteratorConst >& morphChannelsIters, const int32_t morphSize,
+            ( auto BaseIdIter, const TArray< UntypedMeshBufferIterator >& baseChannelsIters, const int32_t baseSize,
+              auto MorphIdIter, const TArray< UntypedMeshBufferIteratorConst >& morphChannelsIters, const int32_t morphSize,
               const float factor )
         -> void
         {
@@ -327,8 +34,8 @@ namespace mu
 				for ( int32_t i = 0; i < baseSize; ++i, ++limitsBaseIdIter )
 				{
 					const uint32_t id = limitsBaseIdIter.GetAsUINT32();
-					minBaseId = std::min( id, minBaseId );
-					maxBaseId = std::max( id, maxBaseId );
+					minBaseId = FMath::Min( id, minBaseId );
+					maxBaseId = FMath::Max( id, maxBaseId );
 				}
 			}
 
@@ -362,7 +69,7 @@ namespace mu
                 for (  ; v + runSize < baseSize && m + runSize < morphSize && RunBaseIter.GetAsUINT32() == (*RunMorphIter)[0];
                         ++runSize, ++RunBaseIter, ++RunMorphIter );
 
-                const size_t channelCount = morphChannelsIters.size();
+                const size_t channelCount = morphChannelsIters.Num();
                 for ( size_t c = 1; c < channelCount; ++c )
                 {
                 
@@ -391,7 +98,7 @@ namespace mu
                     }
                 }
 
-				v += std::max(runSize - 1, 0);
+				v += FMath::Max(runSize - 1, 0);
             }
         };
 
@@ -415,9 +122,12 @@ namespace mu
 			// Iterator of the vertex ids of the base vertices
 			UntypedMeshBufferIteratorConst itBaseId(pBase->GetVertexBuffers(), MBS_VERTEXINDEX);
 
-			std::vector< UntypedMeshBufferIterator > itBaseChannels(ccount);
-			std::vector< UntypedMeshBufferIteratorConst > itMinChannels(ccount);
-			std::vector< UntypedMeshBufferIteratorConst > itMaxChannels(ccount);
+			TArray< UntypedMeshBufferIterator > itBaseChannels;
+			itBaseChannels.SetNum(ccount);
+			TArray< UntypedMeshBufferIteratorConst > itMinChannels;
+			itMinChannels.SetNum(ccount);
+			TArray< UntypedMeshBufferIteratorConst > itMaxChannels;
+			itMaxChannels.SetNum(ccount);
 
 			for (size_t c = 1; c < ccount; ++c)
 			{
diff --git a/Engine/Plugins/Experimental/Mutable/Source/MutableRuntime/Private/MuR/OpMeshRemapIndices.h b/Engine/Plugins/Experimental/Mutable/Source/MutableRuntime/Private/MuR/OpMeshRemapIndices.h
--- a/Engine/Plugins/Experimental/Mutable/Source/MutableRuntime/Private/MuR/OpMeshRemapIndices.h
+++ b/Engine/Plugins/Experimental/Mutable/Source/MutableRuntime/Private/MuR/OpMeshRemapIndices.h
@@ -63,11 +63,11 @@ namespace mu
                 case MBF_UINT16:
                     for ( uint32_t v = 0; v < vcountBase;  ++v)
                     {
-                        uint16_t* pIndex = (uint16_t*)pData;
-                        uint16_t oldIndex = *pIndex;
+                        uint16* pIndex = (uint16*)pData;
+                        uint16 oldIndex = *pIndex;
                         check( *pIndex < vcountReference );
                         uint32_t newIndex = ( itReferenceId + oldIndex ).GetAsUINT32();
-                        *pIndex = uint16_t(newIndex);
+                        *pIndex = uint16(newIndex);
                         pData += buf.m_elementSize;
                     }
                     break;
diff --git a/Engine/Plugins/Experimental/Mutable/Source/MutableRuntime/Private/MuR/OpMeshRemoveChart.h b/Engine/Plugins/Experimental/Mutable/Source/MutableRuntime/Private/MuR/OpMeshRemoveChart.h
--- a/Engine/Plugins/Experimental/Mutable/Source/MutableRuntime/Private/MuR/OpMeshRemoveChart.h
+++ b/Engine/Plugins/Experimental/Mutable/Source/MutableRuntime/Private/MuR/OpMeshRemoveChart.h
@@ -18,7 +18,7 @@ namespace mu
     };
 
 
-    inline void ExtractVertexIndexIntervals( vector<ID_INTERVAL>& intervals, const Mesh* pSource )
+    inline void ExtractVertexIndexIntervals( TArray<ID_INTERVAL>& intervals, const Mesh* pSource )
     {
         MeshBufferIteratorConst<MBF_UINT32,uint32_t,1> itVI( pSource->GetVertexBuffers(), MBS_VERTEXINDEX );
         ID_INTERVAL current;
@@ -44,7 +44,7 @@ namespace mu
                 }
                 else
                 {
-                    intervals.push_back(current);
+                    intervals.Add(current);
                     current.idStart = id;
                     current.idPosition = sv;
                     current.size = 1;
@@ -54,12 +54,12 @@ namespace mu
 
         if (current.idStart>=0)
         {
-            intervals.push_back(current);
+            intervals.Add(current);
         }
     }
 
 
-    inline int FindPositionInIntervals( const vector<ID_INTERVAL>& intervals, int id )
+    inline int FindPositionInIntervals( const TArray<ID_INTERVAL>& intervals, int id )
     {
         for( const auto& interval: intervals )
         {
@@ -98,7 +98,7 @@ namespace mu
         uint8_t* removedVertices = (uint8_t*)mutable_malloc_aligned( resultVertexCount, 16 );
 		FMemory::Memzero(removedVertices,resultVertexCount);
         {
-            vector<ID_INTERVAL> intervals;
+			TArray<ID_INTERVAL> intervals;
             ExtractVertexIndexIntervals( intervals, pResult.get() );
 
             for ( int mv=0; mv<pMask->GetVertexBuffers().GetElementCount(); ++mv )
@@ -119,7 +119,8 @@ namespace mu
         // Map from source vertex index, to new vertex index for used vertices.
         // These are indices as in the index buffer, not the absoulte vertex index as in the
         // vertexbuffer MBS_VERTEXINDEX buffers.
-        vector<int> usedVertices( pResult->GetVertexCount(), -1 );
+		TArray<int> usedVertices;
+		usedVertices.Init( -1, pResult->GetVertexCount() );
         {
             size_t removedIndices = 0;
 
@@ -168,13 +169,13 @@ namespace mu
 
             else if ( pResult->GetIndexBuffers().GetElementSize(0)==2 )
             {
-                MeshBufferIteratorConst<MBF_UINT16,uint16_t,1> itSource( pResult->GetIndexBuffers(), MBS_VERTEXINDEX );
-                MeshBufferIterator<MBF_UINT16,uint16_t,1> itDest( pResult->GetIndexBuffers(), MBS_VERTEXINDEX );
+                MeshBufferIteratorConst<MBF_UINT16,uint16,1> itSource( pResult->GetIndexBuffers(), MBS_VERTEXINDEX );
+                MeshBufferIterator<MBF_UINT16,uint16,1> itDest( pResult->GetIndexBuffers(), MBS_VERTEXINDEX );
 
                 int indexCount = pResult->GetIndexCount();
                 for ( int f=0; f<indexCount/3; ++f )
                 {
-                    uint16_t sourceIndices[3];
+                    uint16 sourceIndices[3];
                     sourceIndices[0] = itSource[0][0];
                     sourceIndices[1] = itSource[1][0];
                     sourceIndices[2] = itSource[2][0];
@@ -188,7 +189,7 @@ namespace mu
                     {
                         for (int i=0;i<3;++i)
                         {
-                            uint16_t sourceIndex = sourceIndices[i];
+                            uint16 sourceIndex = sourceIndices[i];
 
                             if ( usedVertices[ sourceIndex ] < 0 )
                             {
@@ -196,8 +197,8 @@ namespace mu
                                 firstFreeVertex++;
                             }
 
-                            uint16_t destIndex = (uint16_t)usedVertices[ sourceIndex ];
-                            *(uint16_t*)itDest.ptr() = destIndex;
+                            uint16 destIndex = (uint16)usedVertices[ sourceIndex ];
+                            *(uint16*)itDest.ptr() = destIndex;
 
                             itDest++;
                         }
diff --git a/Engine/Plugins/Experimental/Mutable/Source/MutableRuntime/Private/MuR/Operations.cpp b/Engine/Plugins/Experimental/Mutable/Source/MutableRuntime/Private/MuR/Operations.cpp
--- a/Engine/Plugins/Experimental/Mutable/Source/MutableRuntime/Private/MuR/Operations.cpp
+++ b/Engine/Plugins/Experimental/Mutable/Source/MutableRuntime/Private/MuR/Operations.cpp
@@ -153,7 +153,7 @@ namespace mu
 
 
     //---------------------------------------------------------------------------------------------
-    void ForEachReference( OP& op, const std::function<void(OP::ADDRESS*)>& f )
+    void ForEachReference( OP& op, const TFunctionRef<void(OP::ADDRESS*)> f )
     {
         switch ( op.type )
         {
@@ -462,7 +462,7 @@ namespace mu
 
 
     //---------------------------------------------------------------------------------------------
-    void ForEachReference( const PROGRAM& program, OP::ADDRESS at, const std::function<void(OP::ADDRESS)>& f )
+    void ForEachReference( const PROGRAM& program, OP::ADDRESS at, const TFunctionRef<void(OP::ADDRESS)> f )
     {
         OP_TYPE type = program.GetOpType(at);
         switch ( type )
@@ -995,9 +995,9 @@ namespace mu
             memcpy( &source, data, sizeof(OP::ADDRESS) ); data+=sizeof(OP::ADDRESS);
             f(source);
 
-            uint16_t removes = 0;
-            memcpy( &removes, data, sizeof(uint16_t) ); data+=sizeof(uint16_t);
-            for (uint16_t r=0; r<removes; ++r)
+            uint16 removes = 0;
+            memcpy( &removes, data, sizeof(uint16) ); data+=sizeof(uint16);
+            for (uint16 r=0; r<removes; ++r)
             {
                 mu::OP::ADDRESS condition;
                 memcpy( &condition, data, sizeof(OP::ADDRESS) ); data+=sizeof(OP::ADDRESS);
diff --git a/Engine/Plugins/Experimental/Mutable/Source/MutableRuntime/Private/MuR/Operations.h b/Engine/Plugins/Experimental/Mutable/Source/MutableRuntime/Private/MuR/Operations.h
--- a/Engine/Plugins/Experimental/Mutable/Source/MutableRuntime/Private/MuR/Operations.h
+++ b/Engine/Plugins/Experimental/Mutable/Source/MutableRuntime/Private/MuR/Operations.h
@@ -2,14 +2,14 @@
 
 #pragma once
 
-#include "HAL/PlatformMath.h"
-#include "HAL/UnrealMemory.h"
-#include "Misc/AssertionMacros.h"
 #include "MuR/Image.h"
 #include "MuR/Platform.h"
 #include "MuR/SerialisationPrivate.h"
+#include "HAL/PlatformMath.h"
+#include "HAL/UnrealMemory.h"
+#include "Misc/AssertionMacros.h"
+#include "Templates/Function.h"
 
-#include <functional>
 
 #define MUTABLE_OP_MAX_INTERPOLATE_COUNT	6
 #define MUTABLE_OP_MAX_ADD_COUNT			7
@@ -21,7 +21,7 @@ namespace mu
 {
 
     //!
-    enum class OP_TYPE : uint16_t
+    enum class OP_TYPE : uint16
     {
         //-----------------------------------------------------------------------------------------
         // No operation
@@ -401,11 +401,11 @@ namespace mu
 
             // If not negative, index of the skeleton to set to the mesh from the skeleton
             // constant array.
-            int32_t skeleton;
+            int32 skeleton;
 
             // If not negative, index of the physics body to set to the mesh from the physics body
             // constant array.
-			int32_t physicsBody;
+			int32 physicsBody;
         };
 
         struct ParameterArgs
@@ -427,7 +427,7 @@ namespace mu
         struct BoolEqualScalarConstArgs
         {
             ADDRESS value;
-            int16_t constant;
+            int16 constant;
         };
 
         struct BoolBinaryArgs
@@ -457,12 +457,12 @@ namespace mu
         {
             ADDRESS image;
             ADDRESS x,y;
-            uint8_t filter;
+            uint8 filter;
         };
 
         struct ColourSwizzleArgs
         {
-            uint8_t sourceChannels[ MUTABLE_OP_MAX_SWIZZLE_CHANNELS ];
+            uint8 sourceChannels[ MUTABLE_OP_MAX_SWIZZLE_CHANNELS ];
             ADDRESS sources[ MUTABLE_OP_MAX_SWIZZLE_CHANNELS ];
         };
 
@@ -474,7 +474,7 @@ namespace mu
         struct ColourLayoutBlockTransformArgs
         {
             ADDRESS layout;
-            uint16_t block;
+            uint16 block;
         };
 
         struct ColourMultiplyArgs
@@ -497,7 +497,7 @@ namespace mu
 				MULTIPLY,
 				DIVIDE
 			} OPERATION;
-            uint8_t operation;
+            uint8 operation;
 
 			ADDRESS a, b;
 		};
@@ -509,8 +509,8 @@ namespace mu
             ADDRESS base;
             ADDRESS mask;
             ADDRESS blended;
-			uint8_t blendType;	// One of BLEND_TYPE
-			uint8_t flags;
+			uint8 blendType;	// One of BLEND_TYPE
+			uint8 flags;
 			typedef enum
             {
                 F_NONE              =	0,
@@ -527,8 +527,8 @@ namespace mu
             ADDRESS mask;
             ADDRESS blended;
             ADDRESS rangeSize;
-            uint16_t rangeId;
-            uint16_t blendType; // One of BLEND_TYPE
+            uint16 rangeId;
+            uint16 blendType; // One of BLEND_TYPE
         };
 
         struct ImageLayerColourArgs
@@ -537,7 +537,7 @@ namespace mu
             ADDRESS mask;
             ADDRESS colour;
             EImageFormat reformat;
-			uint8_t blendType;	// One of BLEND_TYPE
+			uint8 blendType;	// One of BLEND_TYPE
         };
 
         struct ImagePixelFormatArgs
@@ -552,10 +552,10 @@ namespace mu
             ADDRESS source;
 
             //! Number of mipmaps to build. If zero, it means all.
-            uint8_t levels;
+            uint8 levels;
 
             //! Number of mipmaps that can be generated for a single layout block.
-            uint8_t blockLevels;
+            uint8 blockLevels;
 
             //! This is true if this operation is supposed to build only the tail mipmaps.
             //! It is used during the code optimisation phase, and to validate the code.
@@ -571,9 +571,9 @@ namespace mu
         struct ImageResizeArgs
         {
             ADDRESS source;
-            uint16_t size[2];
+            uint16 size[2];
 
-            inline uint16_t GetSize(int i) const { return size[i]; }
+            inline uint16 GetSize(int i) const { return size[i]; }
         };
 
         struct ImageResizeLikeArgs
@@ -608,21 +608,21 @@ namespace mu
             ADDRESS layout;
 
             //! Size of a layout block in pixels.
-            uint16_t blockSize[2];
+            uint16 blockSize[2];
 			EImageFormat format;
 
             //! If true, generate all the mipmaps except the ones in skipMipmaps
-            uint8_t generateMipmaps;
+            uint8 generateMipmaps;
 
             //! Mipmaps to generate if mipmaps are to be generated. 0 means all.
-            uint8_t mipmapCount;
+            uint8 mipmapCount;
         };
 
         struct ImageComposeArgs
         {
             ADDRESS layout, base, blockImage;
             ADDRESS mask;
-            uint32_t blockIndex;
+            uint32 blockIndex;
         };
 
         struct ImageDifferenceArgs
@@ -662,7 +662,7 @@ namespace mu
         struct ImageSwizzleArgs
         {
 			EImageFormat format;
-            uint8_t sourceChannels[ MUTABLE_OP_MAX_SWIZZLE_CHANNELS ];
+            uint8 sourceChannels[ MUTABLE_OP_MAX_SWIZZLE_CHANNELS ];
             ADDRESS sources[ MUTABLE_OP_MAX_SWIZZLE_CHANNELS ];
         };
 
@@ -683,7 +683,7 @@ namespace mu
         {
             ADDRESS colour0;
             ADDRESS colour1;
-            uint16_t size[2];
+            uint16 size[2];
         };
 
         struct ImageBinariseArgs
@@ -696,7 +696,7 @@ namespace mu
         {
             ADDRESS colour;
 			EImageFormat format;
-            uint16_t size[2];
+            uint16 size[2];
         };
 
         struct ImageGPUArgs
@@ -707,21 +707,21 @@ namespace mu
         struct ImageCropArgs
         {
             ADDRESS source;
-            uint16_t minX, minY, sizeX, sizeY;
+            uint16 minX, minY, sizeX, sizeY;
         };
 
         struct ImagePatchArgs
         {
             ADDRESS base;
             ADDRESS patch;
-            uint16_t minX, minY;
+            uint16 minX, minY;
         };
 
         struct ImageRasterMeshArgs
         {
             ADDRESS mesh;
-            uint16_t sizeX, sizeY;
-            int32_t blockIndex;
+            uint16 sizeX, sizeY;
+            int32 blockIndex;
 
             // These are used in case of projected mesh raster.
             ADDRESS image;
@@ -738,7 +738,7 @@ namespace mu
         struct ImageMakeGrowMapArgs
         {
             ADDRESS mask;
-            int32_t border;
+            int32 border;
         };
 
         struct ImageDisplaceArgs
@@ -776,7 +776,7 @@ namespace mu
         {
             ADDRESS mesh;
             ADDRESS layout;
-            uint16_t channel;
+            uint16 channel;
         };
 
         struct MeshDifferenceArgs
@@ -786,12 +786,12 @@ namespace mu
 
             //! If true the texture coordinates will not be morphed. This is only relevant if the list
             //! in channelSemantic is empty.
-            uint8_t ignoreTextureCoords;
+            uint8 ignoreTextureCoords;
 
             //! Channels to diff. Meaningful until the first null channel. If all are null, it means
             //! that all the channels should be diff-ed
-            uint8_t channelSemantic[MUTABLE_OP_MAX_MORPH_CHANNELS];
-            uint8_t channelSemanticIndex[MUTABLE_OP_MAX_MORPH_CHANNELS];
+            uint8 channelSemantic[MUTABLE_OP_MAX_MORPH_CHANNELS];
+            uint8 channelSemanticIndex[MUTABLE_OP_MAX_MORPH_CHANNELS];
         };
 
         struct MeshMergeArgs
@@ -801,7 +801,7 @@ namespace mu
 
             // If 0, it merges the surfaces, otherwise, add a new surface
             // for the added mesh.
-            uint32_t newSurfaceID;
+            uint32 newSurfaceID;
         };
 
         struct MeshInterpolateArgs
@@ -852,7 +852,6 @@ namespace mu
                 BT_VERTEX			= 1,
                 BT_INDEX			= 2,
                 BT_FACE				= 4,
-                BT_REBUILD_TANGENTS = 8,
                 //! This flag will not add blank channels for the channels in the format mesh but not
                 //! in the source mesh.
                 BT_IGNORE_MISSING	= 16,
@@ -861,7 +860,7 @@ namespace mu
             } BUFFER_TYPE;
 
             //! Flag combination, selecting the buffers to reformat. The rest are left untouched.
-            uint8_t buffers;
+            uint8 buffers;
 
         };
 
@@ -869,7 +868,7 @@ namespace mu
         struct MeshExtractFaceGroupArgs
         {
             ADDRESS source;
-            uint32_t group;
+            uint32 group;
         };
 
 		struct MeshTransformArgs
@@ -891,7 +890,7 @@ namespace mu
 				VS_SHAPE,
 				VS_BONE_HIERARCHY
 			} VERTEX_SELECTION_TYPE;
-            uint8_t vertexSelectionType;
+            uint8 vertexSelectionType;
 
 			float dist, factor, maxBoneRadius;
 		};
@@ -928,7 +927,7 @@ namespace mu
 			ADDRESS scalarB;
 		};
 
-		enum class EMeshBindShapeFlags : uint32_t
+		enum class EMeshBindShapeFlags : uint32
 		{
 			None				   = 0,
 			ReshapeSkeleton		   = 1 << 0,
@@ -944,7 +943,7 @@ namespace mu
 		{
 			ADDRESS mesh;
 			ADDRESS shape;
-			uint32_t flags;
+			uint32 flags;
 			uint32 bindingMethod;
 		};
 
@@ -952,7 +951,7 @@ namespace mu
 		{
 			ADDRESS mesh;
 			ADDRESS shape;
-			uint32_t flags;
+			uint32 flags;
 		};
 
 		struct MeshMorphReshapeArgs
@@ -976,8 +975,8 @@ namespace mu
         {
             ADDRESS instance;
             ADDRESS value;
-            uint32_t id;
-            uint32_t externalId;
+            uint32 id;
+            uint32 externalId;
             ADDRESS name;
 
             // Index in the PROGRAM::m_parameterLists with the parameters that are relevant
@@ -1010,7 +1009,7 @@ namespace mu
 
             // Source mesh to scan for active blocks
             ADDRESS mesh;            
-            uint8_t meshLayoutIndex;
+            uint8 meshLayoutIndex;
         };
 
 
@@ -1053,7 +1052,7 @@ namespace mu
         } FLAGS;
 
         //! Bitmask of FLAGS values
-        uint16_t unused;
+        uint16 unused;
 
         //! Constant arguments per operation type
 		//! Everything in this union needs to be converted to their own AST class eventually.
@@ -1133,7 +1132,7 @@ namespace mu
         ARGS args;
     };
 
-    typedef t_OP<uint32_t> OP;
+    typedef t_OP<uint32> OP;
 
     // OP is not serialisable anymore, we link code now.
     MUTABLE_DEFINE_POD_SERIALISABLE( OP );
@@ -1199,10 +1198,10 @@ namespace mu
     }
 
     //! Utility function to apply a function to all operation references to other operations.
-	MUTABLERUNTIME_API extern void ForEachReference( OP& op, const std::function<void(OP::ADDRESS*)>& );
+	MUTABLERUNTIME_API extern void ForEachReference( OP& op, const TFunctionRef<void(OP::ADDRESS*)> );
 
     //! Utility function to apply a function to all operation references to other operations.
-	MUTABLERUNTIME_API extern void ForEachReference( const struct PROGRAM& program, OP::ADDRESS at, const std::function<void(OP::ADDRESS)>& );
+	MUTABLERUNTIME_API extern void ForEachReference( const struct PROGRAM& program, OP::ADDRESS at, const TFunctionRef<void(OP::ADDRESS)> );
 
 	//!
 	MUTABLERUNTIME_API inline OP_TYPE GetSwitchForType( DATATYPE d )
diff --git a/Engine/Plugins/Experimental/Mutable/Source/MutableRuntime/Private/MuR/Parameters.cpp b/Engine/Plugins/Experimental/Mutable/Source/MutableRuntime/Private/MuR/Parameters.cpp
--- a/Engine/Plugins/Experimental/Mutable/Source/MutableRuntime/Private/MuR/Parameters.cpp
+++ b/Engine/Plugins/Experimental/Mutable/Source/MutableRuntime/Private/MuR/Parameters.cpp
@@ -66,7 +66,7 @@ namespace mu
         const int32_t ver = 0;
         arch << ver;
 
-        uint32_t valueCount = uint32_t(params->GetPrivate()->m_values.size());
+        uint32 valueCount = params->GetPrivate()->m_values.Num();
         arch << valueCount;
 
         for (size_t p=0; p<valueCount; ++p)
@@ -86,7 +86,7 @@ namespace mu
 		LLM_SCOPE_BYNAME(TEXT("MutableRuntime"));
 
         ParametersPtr pResult = pModel->NewParameters();
-        size_t modelParameters = pModel->GetPrivate()->m_program.m_parameters.size();
+        size_t modelParameters = pModel->GetPrivate()->m_program.m_parameters.Num();
 
         int32_t ver;
         arch >> ver;
@@ -190,7 +190,7 @@ namespace mu
 	//---------------------------------------------------------------------------------------------
 	int Parameters::GetCount() const
 	{
-		return (int)m_pD->m_values.size();
+		return (int)m_pD->m_values.Num();
 	}
 
 
@@ -198,7 +198,7 @@ namespace mu
 	const char* Parameters::GetName( int index ) const
 	{
 		const PROGRAM& program = m_pD->m_pModel->GetPrivate()->m_program;
-		check( index>=0 && index<(int)program.m_parameters.size() );
+		check( index>=0 && index<(int)program.m_parameters.Num() );
 
 		return program.m_parameters[index].m_name.c_str();
 	}
@@ -208,7 +208,7 @@ namespace mu
 	const char* Parameters::GetUid( int index ) const
 	{
 		const PROGRAM& program = m_pD->m_pModel->GetPrivate()->m_program;
-		check( index>=0 && index<(int)program.m_parameters.size() );
+		check( index>=0 && index<(int)program.m_parameters.Num() );
 
 		return program.m_parameters[index].m_uid.c_str();
 	}
@@ -225,7 +225,7 @@ namespace mu
 	PARAMETER_TYPE Parameters::GetType( int index ) const
 	{
 		const PROGRAM& program = m_pD->m_pModel->GetPrivate()->m_program;
-		check( index>=0 && index<(int)program.m_parameters.size() );
+		check( index>=0 && index<(int)program.m_parameters.Num() );
 
 		return program.m_parameters[index].m_type;
 	}
@@ -235,7 +235,7 @@ namespace mu
 	PARAMETER_DETAILED_TYPE Parameters::GetDetailedType( int index ) const
 	{
 		const PROGRAM& program = m_pD->m_pModel->GetPrivate()->m_program;
-		check( index>=0 && index<(int)program.m_parameters.size() );
+		check( index>=0 && index<(int)program.m_parameters.Num() );
 
 		return program.m_parameters[index].m_detailedType;
 	}
@@ -247,18 +247,18 @@ namespace mu
 		LLM_SCOPE_BYNAME(TEXT("MutableRuntime"));
 
         const PROGRAM& program = m_pD->m_pModel->GetPrivate()->m_program;
-        check( paramIndex>=0 && paramIndex<int(program.m_parameters.size()) );
+        check( paramIndex>=0 && paramIndex<int(program.m_parameters.Num()) );
         Ptr<RangeIndex> range;
 
-        if ( paramIndex>=0 && paramIndex<int(program.m_parameters.size()) )
+        if ( paramIndex>=0 && paramIndex<int(program.m_parameters.Num()) )
         {
-            size_t rangeCount = program.m_parameters[paramIndex].m_ranges.size();
+            size_t rangeCount = program.m_parameters[paramIndex].m_ranges.Num();
             if ( rangeCount>0 )
             {
                 range = new RangeIndex;
                 range->m_pD->m_pParameters = this;
                 range->m_pD->m_parameter = paramIndex;
-                range->m_pD->m_values.resize( rangeCount, 0 );
+                range->m_pD->m_values.SetNumZeroed( rangeCount );
             }
         }
 
@@ -269,24 +269,24 @@ namespace mu
     //---------------------------------------------------------------------------------------------
     int Parameters::GetValueCount( int paramIndex ) const
     {
-        if (paramIndex<0 || paramIndex >= int(m_pD->m_multiValues.size()) )
+        if (paramIndex<0 || paramIndex >= int(m_pD->m_multiValues.Num()) )
         {
             return 0;
         }
 
-        return int( m_pD->m_multiValues[paramIndex].size() );
+        return int( m_pD->m_multiValues[paramIndex].Num() );
     }
 
 
     //---------------------------------------------------------------------------------------------
     Ptr<RangeIndex> Parameters::GetValueIndex( int paramIndex, int valueIndex ) const
     {
-        if (paramIndex<0 || paramIndex >= int(m_pD->m_multiValues.size()) )
+        if (paramIndex<0 || paramIndex >= int(m_pD->m_multiValues.Num()) )
         {
             return nullptr;
         }
 
-        if (valueIndex<0 || valueIndex >= int(m_pD->m_multiValues[paramIndex].size()) )
+        if (valueIndex<0 || valueIndex >= int(m_pD->m_multiValues[paramIndex].Num()) )
         {
             return nullptr;
         }
@@ -298,7 +298,7 @@ namespace mu
         }
 
         Ptr<RangeIndex> result = NewRangeIndex( paramIndex );
-        result->m_pD->m_values = it->first;
+        result->m_pD->m_values = it->Key;
 
         return result;
     }
@@ -307,12 +307,12 @@ namespace mu
     //---------------------------------------------------------------------------------------------
     void Parameters::ClearAllValues( int paramIndex )
     {
-        if (paramIndex<0 || paramIndex >= int(m_pD->m_multiValues.size()) )
+        if (paramIndex<0 || paramIndex >= int(m_pD->m_multiValues.Num()) )
         {
             return;
         }
 
-        m_pD->m_multiValues[paramIndex].clear();
+        m_pD->m_multiValues[paramIndex].Empty();
     }
 
 
@@ -320,8 +320,8 @@ namespace mu
 	int Parameters::GetAdditionalImageCount( int index ) const
 	{
 		const PROGRAM& program = m_pD->m_pModel->GetPrivate()->m_program;
-		check( index>=0 && index<(int)program.m_parameters.size() );
-		return (int)program.m_parameters[index].m_descImages.size();
+		check( index>=0 && index<(int)program.m_parameters.Num() );
+		return (int)program.m_parameters[index].m_descImages.Num();
 	}
 
 
@@ -329,13 +329,13 @@ namespace mu
     bool Parameters::GetBoolValue( int index,
                                    const Ptr<const RangeIndex>& pos ) const
 	{
-		check( index>=0 && index<(int)m_pD->m_values.size() );
+		check( index>=0 && index<(int)m_pD->m_values.Num() );
         check( GetType(index)==PARAMETER_TYPE::T_BOOL );
 
         // Early out in case of invalid parameters
         if ( index < 0
              ||
-             index >= (int)m_pD->m_values.size()
+             index >= (int)m_pD->m_values.Num()
              ||
              GetType(index) != PARAMETER_TYPE::T_BOOL )
         {
@@ -352,13 +352,13 @@ namespace mu
         // Multivalue case
         check( pos->m_pD->m_parameter==index );
 
-        if ( index<int(m_pD->m_multiValues.size()))
+        if ( index<int(m_pD->m_multiValues.Num()))
         {
             const auto& m = m_pD->m_multiValues[index];
-            auto it = m.find(pos->m_pD->m_values);
-            if (it!=m.end())
+            auto it = m.Find(pos->m_pD->m_values);
+            if (it)
             {
-                return it->second.m_bool;
+                return it->m_bool;
             }
         }
 
@@ -371,13 +371,13 @@ namespace mu
     void Parameters::SetBoolValue( int index, bool value,
                                    const Ptr<const RangeIndex>& pos )
 	{
-		check( index>=0 && index<(int)m_pD->m_values.size() );
+		check( index>=0 && index<(int)m_pD->m_values.Num() );
 		check( GetType(index)== PARAMETER_TYPE::T_BOOL );
 
         // Early out in case of invalid parameters
         if ( index < 0
              ||
-             index >= (int)m_pD->m_values.size()
+             index >= (int)m_pD->m_values.Num()
              ||
              GetType(index) != PARAMETER_TYPE::T_BOOL )
         {
@@ -388,9 +388,9 @@ namespace mu
         if (!pos)
         {
             // Clear multivalue, if set.
-            if (index<int(m_pD->m_multiValues.size()))
+            if (index<int(m_pD->m_multiValues.Num()))
             {
-                m_pD->m_multiValues[index].clear();
+                m_pD->m_multiValues[index].Empty();
             }
 
             m_pD->m_values[index].m_bool = value;
@@ -401,22 +401,22 @@ namespace mu
         {
             check( pos->m_pD->m_parameter==index );
 
-            if ( index>=int(m_pD->m_multiValues.size()))
+            if ( index>=int(m_pD->m_multiValues.Num()))
             {
-                m_pD->m_multiValues.resize(index+1);
+                m_pD->m_multiValues.SetNum(index+1);
             }
 
             auto& m = m_pD->m_multiValues[index];
-            auto it = m.find(pos->m_pD->m_values);
-            if (it!=m.end())
+            auto it = m.Find(pos->m_pD->m_values);
+            if (it)
             {
-                it->second.m_bool = value;
+                it->m_bool = value;
             }
             else
             {
                 PARAMETER_VALUE valueData;
                 valueData.m_bool = value;
-                m.insert( std::make_pair<>( pos->m_pD->m_values, valueData ) );
+                m.Add( pos->m_pD->m_values, valueData );
             }
         }
 	}
@@ -426,8 +426,8 @@ namespace mu
 	int Parameters::GetIntPossibleValueCount( int paramIndex ) const
 	{
 		const PROGRAM& program = m_pD->m_pModel->GetPrivate()->m_program;
-		check( paramIndex>=0 && paramIndex<(int)program.m_parameters.size() );
-		return (int)program.m_parameters[paramIndex].m_possibleValues.size();
+		check( paramIndex>=0 && paramIndex<(int)program.m_parameters.Num() );
+		return (int)program.m_parameters[paramIndex].m_possibleValues.Num();
 	}
 
 
@@ -436,9 +436,9 @@ namespace mu
 	{
 		const PROGRAM& program = m_pD->m_pModel->GetPrivate()->m_program;
 		check( paramIndex>=0
-				&& paramIndex<(int)program.m_parameters.size() );
+				&& paramIndex<(int)program.m_parameters.Num() );
 		check( valueIndex>=0
-				&& valueIndex<(int)program.m_parameters[paramIndex].m_possibleValues.size() );
+				&& valueIndex<(int)program.m_parameters[paramIndex].m_possibleValues.Num() );
 
 		return (int)program.m_parameters[paramIndex].m_possibleValues[valueIndex].m_value;
 	}
@@ -449,10 +449,10 @@ namespace mu
 	{
 		const PROGRAM& program = m_pD->m_pModel->GetPrivate()->m_program;
 		check(paramIndex >= 0
-			&& paramIndex < (int)program.m_parameters.size());
+			&& paramIndex < (int)program.m_parameters.Num());
 
 		int result = -1;
-		for (size_t v = 0; v < program.m_parameters[paramIndex].m_possibleValues.size(); ++v)
+		for (size_t v = 0; v < program.m_parameters[paramIndex].m_possibleValues.Num(); ++v)
 		{
 			if (program.m_parameters[paramIndex].m_possibleValues[v].m_name == valueName)
 			{
@@ -469,9 +469,9 @@ namespace mu
 	{
 		const PROGRAM& program = m_pD->m_pModel->GetPrivate()->m_program;
 		check(paramIndex >= 0
-			&& paramIndex < (int)program.m_parameters.size());
+			&& paramIndex < (int)program.m_parameters.Num());
 
-		for (size_t v = 0; v < program.m_parameters[paramIndex].m_possibleValues.size(); ++v)
+		for (size_t v = 0; v < program.m_parameters[paramIndex].m_possibleValues.Num(); ++v)
 		{
 			if (program.m_parameters[paramIndex].m_possibleValues[v].m_value == Value)
 			{
@@ -487,9 +487,9 @@ namespace mu
 	{
 		const PROGRAM& program = m_pD->m_pModel->GetPrivate()->m_program;
 		check( paramIndex>=0
-				&& paramIndex<(int)program.m_parameters.size() );
+				&& paramIndex<(int)program.m_parameters.Num() );
 		check( valueIndex>=0
-				&& valueIndex<(int)program.m_parameters[paramIndex].m_possibleValues.size() );
+				&& valueIndex<(int)program.m_parameters[paramIndex].m_possibleValues.Num() );
 
 		return program.m_parameters[paramIndex].m_possibleValues[valueIndex].m_name.c_str();
 	}
@@ -499,13 +499,13 @@ namespace mu
     int Parameters::GetIntValue( int index,
                                  const Ptr<const RangeIndex>& pos ) const
 	{
-		check( index>=0 && index<(int)m_pD->m_values.size() );
+		check( index>=0 && index<(int)m_pD->m_values.Num() );
 		check( GetType(index)== PARAMETER_TYPE::T_INT );
 
         // Early out in case of invalid parameters
         if ( index < 0
              ||
-             index >= (int)m_pD->m_values.size()
+             index >= (int)m_pD->m_values.Num()
              ||
              GetType(index) != PARAMETER_TYPE::T_INT )
         {
@@ -522,13 +522,13 @@ namespace mu
         // Multivalue case
         check( pos->m_pD->m_parameter==index );
 
-        if ( index<int(m_pD->m_multiValues.size()))
+        if ( index<int(m_pD->m_multiValues.Num()))
         {
             const auto& m = m_pD->m_multiValues[index];
-            auto it = m.find(pos->m_pD->m_values);
-            if (it!=m.end())
+            auto it = m.Find(pos->m_pD->m_values);
+            if (it)
             {
-                return it->second.m_int;
+                return it->m_int;
             }
         }
 
@@ -541,13 +541,13 @@ namespace mu
     void Parameters::SetIntValue( int index, int value,
                                   const Ptr<const RangeIndex>& pos )
 	{
-		check( index>=0 && index<(int)m_pD->m_values.size() );
+		check( index>=0 && index<(int)m_pD->m_values.Num() );
 		check( GetType(index)== PARAMETER_TYPE::T_INT );
 
         // Early out in case of invalid parameters
         if ( index < 0
              ||
-             index >= (int)m_pD->m_values.size()
+             index >= (int)m_pD->m_values.Num()
              ||
              GetType(index) != PARAMETER_TYPE::T_INT )
         {
@@ -558,9 +558,9 @@ namespace mu
         if (!pos)
         {
             // Clear multivalue, if set.
-            if (index<int(m_pD->m_multiValues.size()))
+            if (index<int(m_pD->m_multiValues.Num()))
             {
-                m_pD->m_multiValues[index].clear();
+                m_pD->m_multiValues[index].Empty();
             }
 
             m_pD->m_values[index].m_int = value;
@@ -571,22 +571,22 @@ namespace mu
         {
             check( pos->m_pD->m_parameter==index );
 
-            if ( index>=int(m_pD->m_multiValues.size()))
+            if ( index>=int(m_pD->m_multiValues.Num()))
             {
-                m_pD->m_multiValues.resize(index+1);
+                m_pD->m_multiValues.SetNum(index+1);
             }
 
             auto& m = m_pD->m_multiValues[index];
-            auto it = m.find(pos->m_pD->m_values);
-            if (it!=m.end())
+            auto it = m.Find(pos->m_pD->m_values);
+            if (it)
             {
-                it->second.m_int = value;
+                it->m_int = value;
             }
             else
             {
                 PARAMETER_VALUE valueData;
                 valueData.m_int = value;
-                m.insert( std::make_pair<>( pos->m_pD->m_values, valueData ) );
+                m.Add( pos->m_pD->m_values, valueData );
             }
         }
     }
@@ -596,13 +596,13 @@ namespace mu
     float Parameters::GetFloatValue( int index,
                                      const Ptr<const RangeIndex>& pos ) const
 	{
-		check( index>=0 && index<(int)m_pD->m_values.size() );
+		check( index>=0 && index<(int)m_pD->m_values.Num() );
 		check( GetType(index)== PARAMETER_TYPE::T_FLOAT );
 
         // Early out in case of invalid parameters
         if ( index < 0
              ||
-             index >= (int)m_pD->m_values.size()
+             index >= (int)m_pD->m_values.Num()
              ||
              GetType(index) != PARAMETER_TYPE::T_FLOAT )
         {
@@ -619,13 +619,13 @@ namespace mu
         // Multivalue case
         check( pos->m_pD->m_parameter==index );
 
-        if ( index<int(m_pD->m_multiValues.size()))
+        if ( index<int(m_pD->m_multiValues.Num()))
         {
             const auto& m = m_pD->m_multiValues[index];
-            auto it = m.find(pos->m_pD->m_values);
-            if (it!=m.end())
+            auto it = m.Find(pos->m_pD->m_values);
+            if (it)
             {
-                return it->second.m_float;
+                return it->m_float;
             }
         }
 
@@ -638,13 +638,13 @@ namespace mu
     void Parameters::SetFloatValue( int index, float value,
                                     const Ptr<const RangeIndex>& pos )
 	{
-		check( index>=0 && index<(int)m_pD->m_values.size() );
+		check( index>=0 && index<(int)m_pD->m_values.Num() );
 		check( GetType(index)== PARAMETER_TYPE::T_FLOAT );
 
         // Early out in case of invalid parameters
         if ( index < 0
              ||
-             index >= (int)m_pD->m_values.size()
+             index >= (int)m_pD->m_values.Num()
              ||
              GetType(index) != PARAMETER_TYPE::T_FLOAT )
         {
@@ -655,9 +655,9 @@ namespace mu
         if (!pos)
         {
             // Clear multivalue, if set.
-            if (index<int(m_pD->m_multiValues.size()))
+            if (index<int(m_pD->m_multiValues.Num()))
             {
-                m_pD->m_multiValues[index].clear();
+                m_pD->m_multiValues[index].Empty();
             }
 
             m_pD->m_values[index].m_float = value;
@@ -668,22 +668,22 @@ namespace mu
         {
             check( pos->m_pD->m_parameter==index );
 
-            if ( index>=int(m_pD->m_multiValues.size()))
+            if ( index>=int(m_pD->m_multiValues.Num()))
             {
-                m_pD->m_multiValues.resize(index+1);
+                m_pD->m_multiValues.SetNum(index+1);
             }
 
             auto& m = m_pD->m_multiValues[index];
-            auto it = m.find(pos->m_pD->m_values);
-            if (it!=m.end())
+            auto it = m.Find(pos->m_pD->m_values);
+            if (it)
             {
-                it->second.m_float = value;
+                it->m_float = value;
             }
             else
             {
                 PARAMETER_VALUE valueData;
                 valueData.m_float = value;
-                m.insert( std::make_pair<>( pos->m_pD->m_values, valueData ) );
+                m.Add( pos->m_pD->m_values, valueData );
             }
         }
     }
@@ -693,13 +693,13 @@ namespace mu
     void Parameters::GetColourValue( int index, float* pR, float* pG, float* pB,
                                      const Ptr<const RangeIndex>& pos ) const
     {
-        check( index>=0 && index<(int)m_pD->m_values.size() );
+        check( index>=0 && index<(int)m_pD->m_values.Num() );
         check( GetType(index)== PARAMETER_TYPE::T_COLOUR );
 
         // Early out in case of invalid parameters
         if ( index < 0
              ||
-             index >= (int)m_pD->m_values.size()
+             index >= (int)m_pD->m_values.Num()
              ||
              GetType(index) != PARAMETER_TYPE::T_COLOUR )
         {
@@ -719,15 +719,15 @@ namespace mu
         // Multivalue case
         check( pos->m_pD->m_parameter==index );
 
-        if ( index<int(m_pD->m_multiValues.size()))
+        if ( index<int(m_pD->m_multiValues.Num()))
         {
             const auto& m = m_pD->m_multiValues[index];
-            auto it = m.find(pos->m_pD->m_values);
-            if (it!=m.end())
+            auto it = m.Find(pos->m_pD->m_values);
+            if (it)
             {
-                if (pR) *pR = it->second.m_colour[0];
-                if (pG) *pG = it->second.m_colour[1];
-                if (pB) *pB = it->second.m_colour[2];
+                if (pR) *pR = it->m_colour[0];
+                if (pG) *pG = it->m_colour[1];
+                if (pB) *pB = it->m_colour[2];
                 return;
             }
         }
@@ -744,13 +744,13 @@ namespace mu
     void Parameters::SetColourValue( int index, float r, float g, float b,
                                      const Ptr<const RangeIndex>& pos )
     {
-        check( index>=0 && index<(int)m_pD->m_values.size() );
+        check( index>=0 && index<(int)m_pD->m_values.Num() );
         check( GetType(index)== PARAMETER_TYPE::T_COLOUR );
 
         // Early out in case of invalid parameters
         if ( index < 0
              ||
-             index >= (int)m_pD->m_values.size()
+             index >= (int)m_pD->m_values.Num()
              ||
              GetType(index) != PARAMETER_TYPE::T_COLOUR )
         {
@@ -761,9 +761,9 @@ namespace mu
         if (!pos)
         {
             // Clear multivalue, if set.
-            if (index<int(m_pD->m_multiValues.size()))
+            if (index<int(m_pD->m_multiValues.Num()))
             {
-                m_pD->m_multiValues[index].clear();
+                m_pD->m_multiValues[index].Empty();
             }
 
             m_pD->m_values[index].m_colour[0] = r;
@@ -776,18 +776,18 @@ namespace mu
         {
             check( pos->m_pD->m_parameter==index );
 
-            if ( index>=int(m_pD->m_multiValues.size()))
+            if ( index>=int(m_pD->m_multiValues.Num()))
             {
-                m_pD->m_multiValues.resize(index+1);
+                m_pD->m_multiValues.SetNum(index+1);
             }
 
             auto& m = m_pD->m_multiValues[index];
-            auto it = m.find(pos->m_pD->m_values);
-            if (it!=m.end())
+            auto it = m.Find(pos->m_pD->m_values);
+            if (it)
             {
-                it->second.m_colour[0] = r;
-                it->second.m_colour[1] = g;
-                it->second.m_colour[2] = b;
+                it->m_colour[0] = r;
+                it->m_colour[1] = g;
+                it->m_colour[2] = b;
             }
             else
             {
@@ -795,7 +795,7 @@ namespace mu
                 valueData.m_colour[0] = r;
                 valueData.m_colour[1] = g;
                 valueData.m_colour[2] = b;
-                m.insert( std::make_pair<>( pos->m_pD->m_values, valueData ) );
+                m.Add( pos->m_pD->m_values, valueData );
             }
         }
     }
@@ -804,7 +804,7 @@ namespace mu
     //---------------------------------------------------------------------------------------------
     EXTERNAL_IMAGE_ID Parameters::GetImageValue( int index ) const
     {
-        check( index >= 0 && index < (int)m_pD->m_values.size() );
+        check( index >= 0 && index < (int)m_pD->m_values.Num() );
         check( GetType( index ) == PARAMETER_TYPE::T_IMAGE );
 
         return m_pD->m_values[index].m_image;
@@ -814,7 +814,7 @@ namespace mu
     //---------------------------------------------------------------------------------------------
     void Parameters::SetImageValue( int index, EXTERNAL_IMAGE_ID id )
     {
-        check( index >= 0 && index < (int)m_pD->m_values.size() );
+        check( index >= 0 && index < (int)m_pD->m_values.Num() );
         check( GetType( index ) == PARAMETER_TYPE::T_IMAGE );
 
         m_pD->m_values[index].m_image = id;
@@ -824,11 +824,11 @@ namespace mu
     //---------------------------------------------------------------------------------------------
     const char* Parameters::GetStringValue( int index, const Ptr<const RangeIndex>& pos ) const
     {
-        check( index >= 0 && index < (int)m_pD->m_values.size() );
+        check( index >= 0 && index < (int)m_pD->m_values.Num() );
         check( GetType( index ) == PARAMETER_TYPE::T_STRING );
 
         // Early out in case of invalid parameters
-        if ( index < 0 || index >= (int)m_pD->m_values.size() ||
+        if ( index < 0 || index >= (int)m_pD->m_values.Num() ||
              GetType( index ) != PARAMETER_TYPE::T_STRING )
         {
             return "";
@@ -844,13 +844,13 @@ namespace mu
         // Multivalue case
         check( pos->m_pD->m_parameter == index );
 
-        if ( index < int( m_pD->m_multiValues.size() ) )
+        if ( index < int( m_pD->m_multiValues.Num() ) )
         {
             const auto& m = m_pD->m_multiValues[index];
-            auto it = m.find( pos->m_pD->m_values );
-            if ( it != m.end() )
+            auto it = m.Find( pos->m_pD->m_values );
+            if ( it )
             {
-                return it->second.m_text;
+                return it->m_text;
             }
         }
 
@@ -864,26 +864,26 @@ namespace mu
     {
 		LLM_SCOPE_BYNAME(TEXT("MutableRuntime"));
 
-        check( index >= 0 && index < (int)m_pD->m_values.size() );
+        check( index >= 0 && index < (int)m_pD->m_values.Num() );
         check( GetType( index ) == PARAMETER_TYPE::T_FLOAT );
 
         // Early out in case of invalid parameters
-        if ( index < 0 || index >= (int)m_pD->m_values.size() ||
+        if ( index < 0 || index >= (int)m_pD->m_values.Num() ||
              GetType( index ) != PARAMETER_TYPE::T_FLOAT )
         {
             return;
         }
 
-        int len = std::min( MUTABLE_MAX_STRING_PARAM_LENGTH,
+        int len = FMath::Min( MUTABLE_MAX_STRING_PARAM_LENGTH,
                             int( strnlen( value, MUTABLE_MAX_STRING_PARAM_LENGTH ) ) );
 
         // Single value case
         if ( !pos )
         {
             // Clear multivalue, if set.
-            if ( index < int( m_pD->m_multiValues.size() ) )
+            if ( index < int( m_pD->m_multiValues.Num() ) )
             {
-                m_pD->m_multiValues[index].clear();
+                m_pD->m_multiValues[index].Empty();
             }
 
             //m_pD->m_values[index].m_float = value;
@@ -896,18 +896,18 @@ namespace mu
         {
             check( pos->m_pD->m_parameter == index );
 
-            if ( index >= int( m_pD->m_multiValues.size() ) )
+            if ( index >= int( m_pD->m_multiValues.Num() ) )
             {
-                m_pD->m_multiValues.resize( index + 1 );
+                m_pD->m_multiValues.SetNum( index + 1 );
             }
 
             auto& m = m_pD->m_multiValues[index];
-            auto it = m.find( pos->m_pD->m_values );
-            if ( it != m.end() )
+            auto it = m.Find( pos->m_pD->m_values );
+            if ( it )
             {
                 //it->second.m_float = value;
-				FMemory::Memcpy( it->second.m_text, value, len );
-                it->second.m_text[len] = 0;
+				FMemory::Memcpy( it->m_text, value, len );
+                it->m_text[len] = 0;
             }
             else
             {
@@ -915,7 +915,7 @@ namespace mu
                 //valueData.m_float = value;
 				FMemory::Memcpy( valueData.m_text, value, len );
                 valueData.m_text[len] = 0;
-                m.insert( std::make_pair<>( pos->m_pD->m_values, valueData ) );
+                m.Add( pos->m_pD->m_values, valueData );
             }
         }
     }
@@ -931,9 +931,9 @@ namespace mu
         // Early out in case of invalid parameters
         if ( index < 0
              ||
-             index >= (int)m_values.size()
+             index >= (int)m_values.Num()
 			 || 
-			 index >= (int)program.m_parameters.size()
+			 index >= (int)program.m_parameters.Num()
              ||
              program.m_parameters[index].m_type != PARAMETER_TYPE::T_PROJECTOR )
         {
@@ -955,13 +955,12 @@ namespace mu
         {
             check( pos->m_pD->m_parameter==index );
 
-            if ( index<int(m_multiValues.size()))
+            if ( index<m_multiValues.Num())
             {
-                const auto& m = m_multiValues[index];
-                auto it = m.find(pos->m_pD->m_values);
-                if (it!=m.end())
+                const PARAMETER_VALUE* it = m_multiValues[index].Find(pos->m_pD->m_values);
+                if (it)
                 {
-                    result = &it->second.m_projector;
+                    result = &it->m_projector;
                 }
             }
         }
@@ -986,13 +985,13 @@ namespace mu
                                         float* pProjectionAngle,
                                         const Ptr<const RangeIndex>& pos ) const
 	{
-		check( index>=0 && index<(int)m_pD->m_values.size() );
+		check( index>=0 && index<(int)m_pD->m_values.Num() );
 		check( GetType(index)== PARAMETER_TYPE::T_PROJECTOR );
 
         // Early out in case of invalid parameters
         if ( index < 0
              ||
-             index >= (int)m_pD->m_values.size()
+             index >= (int)m_pD->m_values.Num()
              ||
              GetType(index) != PARAMETER_TYPE::T_PROJECTOR )
         {
@@ -1033,13 +1032,13 @@ namespace mu
                                         float projectionAngle,
                                         const Ptr<const RangeIndex>& pos )
 	{
-		check( index>=0 && index<(int)m_pD->m_values.size() );
+		check( index>=0 && index<(int)m_pD->m_values.Num() );
 		check( GetType(index)== PARAMETER_TYPE::T_PROJECTOR );
 
         // Early out in case of invalid parameters
         if ( index < 0
              ||
-             index >= (int)m_pD->m_values.size()
+             index >= (int)m_pD->m_values.Num()
              ||
              GetType(index) != PARAMETER_TYPE::T_PROJECTOR )
         {
@@ -1052,9 +1051,9 @@ namespace mu
         if (!pos)
         {
             // Clear multivalue, if set.
-            if (index<int(m_pD->m_multiValues.size()))
+            if (index<int(m_pD->m_multiValues.Num()))
             {
-                m_pD->m_multiValues[index].clear();
+                m_pD->m_multiValues[index].Empty();
             }
 
             result = &m_pD->m_values[index].m_projector;
@@ -1065,21 +1064,21 @@ namespace mu
         {
             check( pos->m_pD->m_parameter==index );
 
-            if ( index>=int(m_pD->m_multiValues.size()))
+            if ( index>=int(m_pD->m_multiValues.Num()))
             {
-                m_pD->m_multiValues.resize(index+1);
+                m_pD->m_multiValues.SetNum(index+1);
             }
 
-            auto& m = m_pD->m_multiValues[index];
-            auto it = m.find(pos->m_pD->m_values);
-            if (it!=m.end())
+			TMap< TArray<int32_t>, PARAMETER_VALUE >& m = m_pD->m_multiValues[index];
+			PARAMETER_VALUE* it = m.Find(pos->m_pD->m_values);
+            if (it)
             {
-                result = &it->second.m_projector;
+                result = &it->m_projector;
             }
             else
             {
-                auto it2 = m.insert( std::make_pair<>( pos->m_pD->m_values, PARAMETER_VALUE() ) );
-                result = &it2.first->second.m_projector;
+                PARAMETER_VALUE& it2 = m.Add( pos->m_pD->m_values, PARAMETER_VALUE() );
+                result = &it2.m_projector;
             }
         }
 
@@ -1089,7 +1088,7 @@ namespace mu
         result->type = PROJECTOR_TYPE::COUNT;
         if (m_pD->m_pModel)
         {
-            const auto& program = m_pD->m_pModel->GetPrivate()->m_program;
+            const PROGRAM& program = m_pD->m_pModel->GetPrivate()->m_program;
             result->type = program.m_parameters[index].m_defaultValue.m_projector.type;
         }
 
@@ -1129,17 +1128,17 @@ namespace mu
         }
 
         size_t thisNumMultiValues = 0;
-        bool thisHasMultiValues = int(m_pD->m_multiValues.size()) > thisParamIndex;
+        bool thisHasMultiValues = int(m_pD->m_multiValues.Num()) > thisParamIndex;
         if (thisHasMultiValues)
         {
-            thisNumMultiValues = m_pD->m_multiValues[thisParamIndex].size();
+            thisNumMultiValues = m_pD->m_multiValues[thisParamIndex].Num();
         }
 
         size_t otherNumMultiValues = 0;
-        bool otherHasMultiValues = int(other->m_pD->m_multiValues.size()) > otherParamIndex;
+        bool otherHasMultiValues = int(other->m_pD->m_multiValues.Num()) > otherParamIndex;
         if (otherHasMultiValues)
         {
-            otherNumMultiValues = other->m_pD->m_multiValues[otherParamIndex].size();
+            otherNumMultiValues = other->m_pD->m_multiValues[otherParamIndex].Num();
         }
 
         if ( thisNumMultiValues != otherNumMultiValues )
@@ -1151,7 +1150,7 @@ namespace mu
              &&
              otherHasMultiValues
              &&
-             m_pD->m_multiValues[thisParamIndex]!=other->m_pD->m_multiValues[otherParamIndex] )
+             !(m_pD->m_multiValues[thisParamIndex]==other->m_pD->m_multiValues[otherParamIndex] ))
         {
             return false;
         }
@@ -1167,7 +1166,7 @@ namespace mu
 
 		int result = -1;
 
-		for( int i=0; result<0 && i<(int)program.m_parameters.size(); ++i )
+		for( int i=0; result<0 && i<(int)program.m_parameters.Num(); ++i )
 		{
 			const PARAMETER_DESC& p = program.m_parameters[i];
 
@@ -1209,7 +1208,7 @@ namespace mu
     //---------------------------------------------------------------------------------------------
     int RangeIndex::GetRangeCount() const
     {
-        return int( m_pD->m_values.size() );
+        return int( m_pD->m_values.Num() );
     }
 
 
@@ -1217,10 +1216,10 @@ namespace mu
     const char* RangeIndex::GetRangeName( int index ) const
     {
         check( index >= 0 && index < GetRangeCount() );
-        check( index < int(m_pD->m_pParameters->GetPrivate()->m_values.size()) );
+        check( index < int(m_pD->m_pParameters->GetPrivate()->m_values.Num()) );
         ModelPtrConst m_pModel = m_pD->m_pParameters->GetPrivate()->m_pModel;
         check( m_pModel );
-        check( index < int(m_pModel->GetPrivate()->m_program.m_ranges.size()) );
+        check( index < int(m_pModel->GetPrivate()->m_program.m_ranges.Num()) );
 
         return m_pModel->GetPrivate()->m_program.m_ranges[index].m_name.c_str();
     }
@@ -1230,10 +1229,10 @@ namespace mu
     const char* RangeIndex::GetRangeUid( int index ) const
     {
         check( index >= 0 && index < GetRangeCount() );
-        check( index < int(m_pD->m_pParameters->GetPrivate()->m_values.size()) );
+        check( index < int(m_pD->m_pParameters->GetPrivate()->m_values.Num()) );
         ModelPtrConst m_pModel = m_pD->m_pParameters->GetPrivate()->m_pModel;
         check( m_pModel );
-        check( index < int(m_pModel->GetPrivate()->m_program.m_ranges.size()) );
+        check( index < int(m_pModel->GetPrivate()->m_program.m_ranges.Num()) );
 
         return m_pModel->GetPrivate()->m_program.m_ranges[index].m_uid.c_str();
     }
diff --git a/Engine/Plugins/Experimental/Mutable/Source/MutableRuntime/Private/MuR/ParametersPrivate.h b/Engine/Plugins/Experimental/Mutable/Source/MutableRuntime/Private/MuR/ParametersPrivate.h
--- a/Engine/Plugins/Experimental/Mutable/Source/MutableRuntime/Private/MuR/ParametersPrivate.h
+++ b/Engine/Plugins/Experimental/Mutable/Source/MutableRuntime/Private/MuR/ParametersPrivate.h
@@ -4,10 +4,10 @@
 
 #include "MuR/Parameters.h"
 
-#include "../Private/MuR/SerialisationPrivate.h"
-#include "../Private/MuR/MemoryPrivate.h"
-#include "../Private/MuR/Operations.h"
-#include "../Private/MuR/ImagePrivate.h"
+#include "MuR/SerialisationPrivate.h"
+#include "MuR/MemoryPrivate.h"
+#include "MuR/Operations.h"
+#include "MuR/ImagePrivate.h"
 
 #include "MuR/Model.h"
 
@@ -253,7 +253,7 @@ namespace mu
 
 	struct Curve
 	{
-		vector<CurveKeyFrame> keyFrames;
+		TArray<CurveKeyFrame> keyFrames;
 		float defaultValue = 0.f;
 
         //!
@@ -326,10 +326,10 @@ namespace mu
 
         //! Ranges, if the parameter is multi-dimensional. The indices refer to the Model's program
         //! vector of range descriptors.
-        vector<uint32_t> m_ranges;
+		TArray<uint32_t> m_ranges;
 
         //! Additional parameter description information
-        vector<OP::ADDRESS> m_descImages;
+		TArray<OP::ADDRESS> m_descImages;
 
         //! Possible values of the parameter in case of being an integer, and its names
         struct INT_VALUE_DESC
@@ -361,7 +361,7 @@ namespace mu
 
         //! For integer parameters, this contains the description of the possible values.
         //! If empty, the integer may have any value.
-        vector<INT_VALUE_DESC> m_possibleValues;
+		TArray<INT_VALUE_DESC> m_possibleValues;
 
         //!
         bool operator==( const PARAMETER_DESC& other ) const
@@ -452,7 +452,7 @@ namespace mu
         int m_parameter = -1;
 
         //! Position in the several dimension of the range, as defined in m_pParameters
-        vector<int> m_values;
+		TArray<int> m_values;
     };
 
 
@@ -466,11 +466,11 @@ namespace mu
         ModelPtrConst m_pModel;
 
         //! Values for the parameters if they are not multidimensional.
-        vector<PARAMETER_VALUE> m_values;
+		TArray<PARAMETER_VALUE> m_values;
 
         //! If the parameter is multidemensional, the values are stored here.
         //! The key of the map is the vector of values stored in a RangeIndex
-        vector< map< vector<int32_t>, PARAMETER_VALUE > > m_multiValues;
+		TArray< TMap< TArray<int32_t>, PARAMETER_VALUE > > m_multiValues;
 
 
         //!
@@ -491,7 +491,7 @@ namespace mu
             check( ver == 1 );
 
             arch >> m_values;
-            arch>>m_multiValues;
+            arch >> m_multiValues;
         }
 
         //!
diff --git a/Engine/Plugins/Experimental/Mutable/Source/MutableRuntime/Private/MuR/Serialisation.cpp b/Engine/Plugins/Experimental/Mutable/Source/MutableRuntime/Private/MuR/Serialisation.cpp
--- a/Engine/Plugins/Experimental/Mutable/Source/MutableRuntime/Private/MuR/Serialisation.cpp
+++ b/Engine/Plugins/Experimental/Mutable/Source/MutableRuntime/Private/MuR/Serialisation.cpp
@@ -203,7 +203,7 @@ namespace mu
     class InputArchiveWithProxies::Private : public Base
     {
     public:
-        vector< Ptr<ResourceProxy<Image>> > m_proxyHistory;
+        TArray< Ptr<ResourceProxy<Image>> > m_proxyHistory;
         ProxyFactory* m_pFactory = nullptr;
     };
 
@@ -250,7 +250,7 @@ namespace mu
             }
             else
             {
-                if ( (std::size_t)id < m_pD->m_proxyHistory.size() )
+                if ( (std::size_t)id < m_pD->m_proxyHistory.Num() )
                 {
                     p = m_pD->m_proxyHistory[id];
 
@@ -261,7 +261,7 @@ namespace mu
                 else
                 {
                     // Ids come in order.
-                    m_pD->m_proxyHistory.resize(id+1);
+                    m_pD->m_proxyHistory.SetNum(id+1);
 
                     p = m_pD->m_pFactory->NewImageProxy(*this);
                     m_pD->m_proxyHistory[id] = p;
diff --git a/Engine/Plugins/Experimental/Mutable/Source/MutableRuntime/Private/MuR/SerialisationPrivate.h b/Engine/Plugins/Experimental/Mutable/Source/MutableRuntime/Private/MuR/SerialisationPrivate.h
--- a/Engine/Plugins/Experimental/Mutable/Source/MutableRuntime/Private/MuR/SerialisationPrivate.h
+++ b/Engine/Plugins/Experimental/Mutable/Source/MutableRuntime/Private/MuR/SerialisationPrivate.h
@@ -11,13 +11,14 @@
 #include "MuR/RefCounted.h"
 #include "MuR/MutableMemory.h"
 #include "MuR/Types.h"
-#include "../Private/MuR/MemoryPrivate.h"
+#include "MuR/MemoryPrivate.h"
 
-#include <vector>
+#include <string>
 
 
 namespace mu
 {
+	typedef std::string string;
 
     //!
     class MUTABLERUNTIME_API InputMemoryStream : public InputStream
@@ -213,30 +214,6 @@ namespace mu
 			t = (T)v;													\
 		}
 
-	//---------------------------------------------------------------------------------------------
-	template< typename T >
-	inline void operator<< ( OutputArchive& arch, const vector<T>& v )
-	{
-		// TODO: Optimise for vectors of PODs
-        arch << (uint32)v.size();
-        for( std::size_t i=0; i<v.size(); ++i )
-		{
-			arch << v[i];
-		}
-	}
-
-	template< typename T >
-	inline void operator>> ( InputArchive& arch, vector<T>& v )
-	{
-		// TODO: Optimise for vectors of PODs
-        uint32 size;
-		arch >> size;
-		v.resize( size );
-        for( uint32 i=0; i<size; ++i )
-		{
-			arch >> v[i];
-		}
-	}
 
     //---------------------------------------------------------------------------------------------
     template< typename T0, typename T1 >
@@ -254,40 +231,12 @@ namespace mu
     }
 
 
-    //---------------------------------------------------------------------------------------------
-    template< typename K, typename T >
-    inline void operator<< ( OutputArchive& arch, const map<K,T>& v )
-    {
-        arch << (uint32_t)v.size();
-        for( const auto& p: v )
-        {
-            arch << p.first;
-            arch << p.second;
-        }
-    }
-
-    template< typename K, typename T >
-    inline void operator>> ( InputArchive& arch, map<K,T>& v )
-    {
-        uint32_t size;
-        arch >> size;
-
-        for( uint32_t i=0; i<size; ++i )
-        {
-            K k;
-            T t;
-            arch >> k;
-            arch >> t;
-            v.insert( std::make_pair<>(k,t) );
-        }
-    }
-
 	
 	//---------------------------------------------------------------------------------------------
 	template<typename T> inline void operator<<(OutputArchive& arch, const TArray<T>& v)
 	{
 		// TODO: Optimise for vectors of PODs
-		arch << (uint32_t)v.Num();
+		arch << (uint32)v.Num();
 		for (std::size_t i = 0; i < v.Num(); ++i)
 		{
 			arch << v[i];
@@ -297,7 +246,7 @@ namespace mu
 	template<typename T> inline void operator>>(InputArchive& arch, TArray<T>& v)
 	{
 		// TODO: Optimise for vectors of PODs
-		uint32_t size;
+		uint32 size;
 		arch >> size;
 		v.SetNum(size);
 		for (std::size_t i = 0; i < size; ++i)
@@ -306,31 +255,37 @@ namespace mu
 		}
 	}
 
+	//---------------------------------------------------------------------------------------------
+	template< typename K, typename T >
+	inline void operator<< (OutputArchive& arch, const TMap<K, T>& v)
+	{
+		arch << (uint32)v.Num();
+		for (const auto& p : v)
+		{
+			arch << p.Key;
+			arch << p.Value;
+		}
+	}
+
+	template< typename K, typename T >
+	inline void operator>> (InputArchive& arch, TMap<K, T>& v)
+	{
+		uint32 size;
+		arch >> size;
+
+		for (uint32 i = 0; i < size; ++i)
+		{
+			K k;
+			T t;
+			arch >> k;
+			arch >> t;
+			v.Add(k, t);
+		}
+	}
+
 
 	//---------------------------------------------------------------------------------------------
 #define MUTABLE_DEFINE_POD_VECTOR_SERIALISABLE(T)										\
-	template<>																			\
-	inline void operator<< <T>( OutputArchive& arch, const vector<T>& v )				\
-	{																					\
-        arch << (uint32_t)v.size();														\
-		if (v.size())																	\
-		{																				\
-            arch.GetPrivate()->m_pStream->Write( &v[0], (uint32_t)v.size()*sizeof(T) );	\
-		}																				\
-	}																					\
-																						\
-	template<>																			\
-	inline void operator>> <T>( InputArchive& arch, vector<T>& v )						\
-	{																					\
-        uint32_t size;																	\
-		arch >> size;																	\
-		v.resize( size );																\
-		if (size)																		\
-		{																				\
-            arch.GetPrivate()->m_pStream->Read( &v[0], (uint32_t)size*sizeof(T) );		\
-		}																				\
-	}																					\
-																						\
 	template<> inline void operator<<<T>(OutputArchive& arch, const TArray<T>& v)		\
 	{                                                                                   \
 		uint32 Num = uint32(v.Num());													\
@@ -366,9 +321,9 @@ namespace mu
 
 	//---------------------------------------------------------------------------------------------
 	template<>
-	inline void operator<< <string>( OutputArchive& arch, const string& v )
+	inline void operator<< <std::string>( OutputArchive& arch, const std::string& v )
 	{
-        arch << (uint32_t)v.size();
+        arch << (uint32)v.size();
 		if ( v.size() )
 		{
 			arch.GetPrivate()->m_pStream->Write( &v[0], (unsigned)v.size()*sizeof(char) );
@@ -376,9 +331,9 @@ namespace mu
 	}
 
 	template<>
-	inline void operator>> <string>( InputArchive& arch, string& v )
+	inline void operator>> <std::string>( InputArchive& arch, std::string& v )
 	{
-        uint32_t size;
+        uint32 size;
 		arch >> size;
 		v.resize( size );
 		if (size)
