[mutable] Optimize mesh clip versus switching of clipping mesh.

diff --git a/Engine/Plugins/Experimental/Mutable/Source/MutableRuntime/Private/MuR/System.cpp b/Engine/Plugins/Experimental/Mutable/Source/MutableRuntime/Private/MuR/System.cpp
--- a/Engine/Plugins/Experimental/Mutable/Source/MutableRuntime/Private/MuR/System.cpp
+++ b/Engine/Plugins/Experimental/Mutable/Source/MutableRuntime/Private/MuR/System.cpp
@@ -239,7 +239,7 @@ namespace mu
 
 		PROGRAM& program = pLiveInstance->m_pModel->GetPrivate()->m_program;
 
-		bool validState = stateIndex >= 0 && stateIndex < (int)program.m_states.size();
+		bool validState = stateIndex >= 0 && stateIndex < (int)program.m_states.Num();
 		if (!validState)
 		{
 			UE_LOG(LogMutableCore, Error, TEXT("Invalid state in mutable update."));
@@ -365,7 +365,7 @@ namespace mu
 					m_pD->m_memory->m_descCache.Reset();
 				}
 
-				m_pD->m_memory->m_descCache.SetNum(program.m_opAddress.size());
+				m_pD->m_memory->m_descCache.SetNum(program.m_opAddress.Num());
 
 				OP::ADDRESS at = res.m_rootAddress;
 				mu::OP_TYPE opType = program.GetOpType(at);
@@ -432,7 +432,7 @@ namespace mu
 		MUTABLE_CPUPROFILER_SCOPE(EndUpdate);
 
 		// Reduce the cache until it fits the limit.
-		uint64 totalMemory = m_pD->m_modelCache.EnsureCacheBelowBudget(0);
+		uint64 totalMemory = m_pD->m_modelCache.EnsureCacheBelowBudget(0, [](const Model*, int) {return false;});
 		m_pD->m_profileMetrics[size_t(System::ProfileMetric::StreamingCacheBytes)] = totalMemory;
 
 
@@ -471,11 +471,11 @@ namespace mu
 
 		check(parameter >= 0
 			&&
-			parameter < (int)pModel->GetPrivate()->m_program.m_parameters.size());
+			parameter < (int)pModel->GetPrivate()->m_program.m_parameters.Num());
 
 		const PARAMETER_DESC& param =
 			pModel->GetPrivate()->m_program.m_parameters[parameter];
-		check(image >= 0 && image < (int)param.m_descImages.size());
+		check(image >= 0 && image < (int)param.m_descImages.Num());
 
 		OP::ADDRESS imageAddress = param.m_descImages[image];
 
@@ -581,27 +581,21 @@ namespace mu
         // check what parameters have changed
         updatedParameters = 0;
         const PROGRAM& program = pLiveInstance->m_pModel->GetPrivate()->m_program;
-        const vector<int>& runtimeParams = program.m_states[ pLiveInstance->m_state ].m_runtimeParameters;
+        const TArray<int>& runtimeParams = program.m_states[ pLiveInstance->m_state ].m_runtimeParameters;
 
-        check( pParams->GetCount() == (int)program.m_parameters.size() );
+        check( pParams->GetCount() == (int)program.m_parameters.Num() );
         check( !pLiveInstance->m_pOldParameters
                         ||
                         pParams->GetCount() == pLiveInstance->m_pOldParameters->GetCount() );
 
-        for ( int p=0; p<(int)program.m_parameters.size() && !fullBuild; ++p )
+        for ( int p=0; p<(int)program.m_parameters.Num() && !fullBuild; ++p )
         {
-            bool isRuntime = std::find( runtimeParams.begin(), runtimeParams.end(), p )
-                    !=
-                    runtimeParams.end() ;
-
+            bool isRuntime = runtimeParams.Contains( p );
             bool changed = !pParams->HasSameValue( p, pLiveInstance->m_pOldParameters, p );
 
             if (changed && isRuntime)
             {
-				uint64 runtimeIndex =
-                        std::find( runtimeParams.begin(), runtimeParams.end(), p )
-                        -
-                        runtimeParams.begin();
+				uint64 runtimeIndex = runtimeParams.IndexOfByKey(p);
                 updatedParameters |= uint64(1) << runtimeIndex;
             }
             else if (changed)
@@ -631,7 +625,7 @@ namespace mu
 		// We don't have a FLiveInstance, let's create the memory
 		// \TODO: There is no clear moment to remove this... EndBuild?
 		m_memory = MakeShared<FProgramCache>();
-		m_memory->Init(pModel->GetPrivate()->m_program.m_opAddress.size());
+		m_memory->Init(pModel->GetPrivate()->m_program.m_opAddress.Num());
 
 		PrepareCache(pModel, -1);
 	}
@@ -840,7 +834,7 @@ namespace mu
 		MUTABLE_CPUPROFILER_SCOPE(PrepareCache);
 
 		PROGRAM& program = pModel->GetPrivate()->m_program;
-		size_t opCount = program.m_opAddress.size();
+		size_t opCount = program.m_opAddress.Num();
 		m_memory->m_opHitCount.clear();
 		m_memory->Init(opCount);
 
@@ -890,7 +884,7 @@ namespace mu
 
     //---------------------------------------------------------------------------------------------
     uint64 FModelCache::EnsureCacheBelowBudget( uint64 additionalMemory,
-                                              std::function<bool(const Model*,int)> isRomLockedFunc )
+												TFunctionRef<bool(const Model*,int)> isRomLockedFunc )
     {
 
 		uint64 totalMemory = 0;
@@ -935,7 +929,7 @@ namespace mu
 
                             if (bIsLoaded
                                 &&
-                                (!isRomLockedFunc || !isRomLockedFunc(pCacheModel.get(),RomIndex)) )
+                                (!isRomLockedFunc(pCacheModel.get(),RomIndex)) )
                             {
                                 constexpr float factorWeight = 100.0f;
                                 constexpr float factorTime = -1.0f;
@@ -1009,7 +1003,7 @@ namespace mu
 
     //---------------------------------------------------------------------------------------------
     void FModelCache::UpdateForLoad( int romIndex, const Model* pModel,
-                                     std::function<bool(const Model*,int)> isRomLockedFunc )
+                                     TFunctionRef<bool(const Model*,int)> isRomLockedFunc )
     {
 		MarkRomUsed( romIndex, pModel);
 
diff --git a/Engine/Plugins/Experimental/Mutable/Source/MutableRuntime/Private/MuR/SystemPrivate.h b/Engine/Plugins/Experimental/Mutable/Source/MutableRuntime/Private/MuR/SystemPrivate.h
--- a/Engine/Plugins/Experimental/Mutable/Source/MutableRuntime/Private/MuR/SystemPrivate.h
+++ b/Engine/Plugins/Experimental/Mutable/Source/MutableRuntime/Private/MuR/SystemPrivate.h
@@ -73,7 +73,7 @@ namespace mu
 	{
 	public:
 		//! Set or add a value to the index
-		void SetFromModelRangeIndex(uint16_t rangeIndex, int rangeValue)
+		void SetFromModelRangeIndex(uint16 rangeIndex, int rangeValue)
 		{
 			auto Index = IndexOfByPredicate([=](const ElementType& v) { return v.Key >= rangeIndex; });
 			if (Index != INDEX_NONE && (*this)[Index].Key == rangeIndex)
@@ -861,12 +861,12 @@ namespace mu
 
         //! Make sure the cached memory is below the internal budget, even counting with the
         //! passed additional memory.
-        uint64 EnsureCacheBelowBudget( uint64 additionalMemory,
-                                         std::function<bool(const Model*,int)> isRomLockedFunc = nullptr );
+		uint64 EnsureCacheBelowBudget(uint64 AdditionalMemory, 
+			TFunctionRef<bool(const Model*, int)> IsRomLockedFunc = [](const Model*, int) {return false;});
 
         //!
         void UpdateForLoad( int romIndex, const Model* pModel,
-                            std::function<bool(const Model*,int)> isRomLockedFunc = nullptr );
+                            TFunctionRef<bool(const Model*,int)> isRomLockedFunc );
         void MarkRomUsed( int romIndex, const Model* pModel );
 
         //! Private helper
diff --git a/Engine/Plugins/Experimental/Mutable/Source/MutableRuntime/Public/MuR/Image.h b/Engine/Plugins/Experimental/Mutable/Source/MutableRuntime/Public/MuR/Image.h
--- a/Engine/Plugins/Experimental/Mutable/Source/MutableRuntime/Public/MuR/Image.h
+++ b/Engine/Plugins/Experimental/Mutable/Source/MutableRuntime/Public/MuR/Image.h
@@ -2,16 +2,16 @@
 
 #pragma once
 
-#include "../Private/MuR/MutableMath.h"
-#include "Containers/Array.h"
-#include "HAL/LowLevelMemTracker.h"
-#include "HAL/Platform.h"
-#include "Misc/AssertionMacros.h"
+#include "MuR/MutableMath.h"
 #include "MuR/MutableTrace.h"
 #include "MuR/Ptr.h"
 #include "MuR/RefCounted.h"
 #include "MuR/Serialisation.h"
 #include "MuR/SerialisationPrivate.h"
+#include "Containers/Array.h"
+#include "HAL/LowLevelMemTracker.h"
+#include "HAL/Platform.h"
+#include "Misc/AssertionMacros.h"
 
 namespace mu
 {
diff --git a/Engine/Plugins/Experimental/Mutable/Source/MutableRuntime/Public/MuR/Layout.h b/Engine/Plugins/Experimental/Mutable/Source/MutableRuntime/Public/MuR/Layout.h
--- a/Engine/Plugins/Experimental/Mutable/Source/MutableRuntime/Public/MuR/Layout.h
+++ b/Engine/Plugins/Experimental/Mutable/Source/MutableRuntime/Public/MuR/Layout.h
@@ -2,14 +2,14 @@
 
 #pragma once
 
-#include "../Private/MuR/MutableMath.h"
-#include "Containers/Array.h"
-#include "Math/UnrealMathSSE.h"
-#include "Misc/AssertionMacros.h"
+#include "MuR/MutableMath.h"
 #include "MuR/Ptr.h"
 #include "MuR/RefCounted.h"
 #include "MuR/Serialisation.h"
 #include "MuR/SerialisationPrivate.h"
+#include "Containers/Array.h"
+#include "Math/UnrealMathSSE.h"
+#include "Misc/AssertionMacros.h"
 
 namespace mu
 {
diff --git a/Engine/Plugins/Experimental/Mutable/Source/MutableRuntime/Public/MuR/Mesh.h b/Engine/Plugins/Experimental/Mutable/Source/MutableRuntime/Public/MuR/Mesh.h
--- a/Engine/Plugins/Experimental/Mutable/Source/MutableRuntime/Public/MuR/Mesh.h
+++ b/Engine/Plugins/Experimental/Mutable/Source/MutableRuntime/Public/MuR/Mesh.h
@@ -531,10 +531,10 @@ namespace mu
 		struct VERTEX_MATCH_MAP
 		{
 			//! One for every vertex
-			vector<int> m_firstMatch;
+			TArray<int> m_firstMatch;
 
 			//! The matches of every vertex in a sequence
-			vector<int> m_matches;
+			TArray<int> m_matches;
 
 			//!
 			bool Matches(int v, int ov) const;
diff --git a/Engine/Plugins/Experimental/Mutable/Source/MutableRuntime/Public/MuR/MeshBufferSet.h b/Engine/Plugins/Experimental/Mutable/Source/MutableRuntime/Public/MuR/MeshBufferSet.h
--- a/Engine/Plugins/Experimental/Mutable/Source/MutableRuntime/Public/MuR/MeshBufferSet.h
+++ b/Engine/Plugins/Experimental/Mutable/Source/MutableRuntime/Public/MuR/MeshBufferSet.h
@@ -2,7 +2,7 @@
 
 #pragma once
 
-#include "../Private/MuR/SerialisationPrivate.h"
+#include "MuR/SerialisationPrivate.h"
 #include "Containers/Array.h"
 #include "HAL/PlatformMath.h"
 #include "MuR/MutableMemory.h"
diff --git a/Engine/Plugins/Experimental/Mutable/Source/MutableRuntime/Public/MuR/Model.h b/Engine/Plugins/Experimental/Mutable/Source/MutableRuntime/Public/MuR/Model.h
--- a/Engine/Plugins/Experimental/Mutable/Source/MutableRuntime/Public/MuR/Model.h
+++ b/Engine/Plugins/Experimental/Mutable/Source/MutableRuntime/Public/MuR/Model.h
@@ -5,8 +5,7 @@
 #include "MuR/Parameters.h"
 #include "MuR/Ptr.h"
 #include "MuR/RefCounted.h"
-
-#include <functional>
+#include "Templates/Function.h"
 
 namespace mu { class InputArchive; }
 namespace mu { class ModelStreamer; }
@@ -159,7 +158,7 @@ namespace mu
         //! Return the parameters of one of a random instance of the model.
         //! \param randomGenerator This function will be used to generate random values
         //! for the continuous parameters of the instance.
-        ParametersPtr GetRandomInstance( std::function<float()> randomGenerator );
+        ParametersPtr GetRandomInstance(TFunctionRef<float()> randomGenerator );
 
         //-----------------------------------------------------------------------------------------
         // Interface pattern
diff --git a/Engine/Plugins/Experimental/Mutable/Source/MutableRuntime/Public/MuR/MutableMemory.h b/Engine/Plugins/Experimental/Mutable/Source/MutableRuntime/Public/MuR/MutableMemory.h
--- a/Engine/Plugins/Experimental/Mutable/Source/MutableRuntime/Public/MuR/MutableMemory.h
+++ b/Engine/Plugins/Experimental/Mutable/Source/MutableRuntime/Public/MuR/MutableMemory.h
@@ -5,6 +5,7 @@
 #include "CoreTypes.h"
 
 #include <cstring>	// for size_t in some platforms
+#include <new>
 
 namespace mu
 {
diff --git a/Engine/Plugins/Experimental/Mutable/Source/MutableRuntime/Public/MuR/PhysicsBody.h b/Engine/Plugins/Experimental/Mutable/Source/MutableRuntime/Public/MuR/PhysicsBody.h
--- a/Engine/Plugins/Experimental/Mutable/Source/MutableRuntime/Public/MuR/PhysicsBody.h
+++ b/Engine/Plugins/Experimental/Mutable/Source/MutableRuntime/Public/MuR/PhysicsBody.h
@@ -2,7 +2,11 @@
 
 #pragma once
 
-#include "../Private/MuR/SerialisationPrivate.h"
+#include "MuR/SerialisationPrivate.h"
+#include "MuR/MemoryPrivate.h"
+#include "MuR/Ptr.h"
+#include "MuR/RefCounted.h"
+#include "MuR/Serialisation.h"
 #include "Containers/Array.h"
 #include "Math/Quat.h"
 #include "Math/Transform.h"
@@ -10,10 +14,6 @@
 #include "Math/Vector.h"
 #include "Math/VectorRegister.h"
 #include "Misc/AssertionMacros.h"
-#include "MuR/MemoryPrivate.h"
-#include "MuR/Ptr.h"
-#include "MuR/RefCounted.h"
-#include "MuR/Serialisation.h"
 
 
 namespace mu
diff --git a/Engine/Plugins/Experimental/Mutable/Source/MutableRuntime/Public/MuR/Types.h b/Engine/Plugins/Experimental/Mutable/Source/MutableRuntime/Public/MuR/Types.h
--- a/Engine/Plugins/Experimental/Mutable/Source/MutableRuntime/Public/MuR/Types.h
+++ b/Engine/Plugins/Experimental/Mutable/Source/MutableRuntime/Public/MuR/Types.h
@@ -3,6 +3,8 @@
 #pragma once
 
 #include "HAL/Platform.h"
+#include "Containers/Map.h"
+#include "Containers/Array.h"
 
 namespace mu
 {
@@ -12,5 +14,28 @@ namespace mu
     //!
     typedef uint64            EXTERNAL_IMAGE_ID;
 
+
+}
+
+
+//
+template<typename T>
+inline uint32 GetTypeHash(const TArray<T>& Key)
+{
+	uint32 Hash = 0;
+	for (const T& e : Key)
+	{
+		Hash = HashCombine(Hash, GetTypeHash(e));
+	}
+	return Hash;
 }
 
+
+//
+template<typename K, typename V>
+inline bool operator==(const TMap<K,V>& A, const TMap<K,V>& B )
+{	
+	return A.OrderIndependentCompareEqual(B);
+}
+
+
diff --git a/Engine/Plugins/Experimental/Mutable/Source/MutableTools/Private/MuT/AST.cpp b/Engine/Plugins/Experimental/Mutable/Source/MutableTools/Private/MuT/AST.cpp
--- a/Engine/Plugins/Experimental/Mutable/Source/MutableTools/Private/MuT/AST.cpp
+++ b/Engine/Plugins/Experimental/Mutable/Source/MutableTools/Private/MuT/AST.cpp
@@ -3,6 +3,7 @@
 #include "MuT/AST.h"
 
 #include "Containers/Set.h"
+#include "Containers/Queue.h"
 #include "Logging/LogCategory.h"
 #include "Logging/LogMacros.h"
 #include "MuR/Mesh.h"
@@ -10,134 +11,19 @@
 #include "MuR/Platform.h"
 #include "MuT/ASTOpConstantResource.h"
 #include "MuT/ASTOpMeshRemoveMask.h"
-#include "MuT/ErrorLogPrivate.h"
 #include "MuT/StreamsPrivate.h"
+#include "MuT/Platform.h"
 #include "Trace/Detail/Channel.h"
-#include "array"
-#include "atomic"
-#include "cstdint"
-#include "iosfwd"
-#include "unordered_map"
-#include "unordered_set"
 
-#include <queue>
+#include <atomic>
 
 
-std::atomic<uint32_t> mu::ASTOp::s_lastTraverseIndex(1);
+std::atomic<uint32> mu::ASTOp::s_lastTraverseIndex(1);
 
 
 namespace mu
 {
 
-void DebugLogAST(const Ptr<ASTOp>& at, int indent, ASTOpList* done, const char* label )
-{
-    if (!at) return;
-
-    static int s_logTree = 0;
-
-    if (indent==0)
-    {
-		UE_LOG(LogMutableCore, Warning, TEXT("ASTOp tree [%d]:"), s_logTree );
-        ++s_logTree;
-    }
-
-    ASTOpList localDone;
-    if (!done) done = &localDone;
-
-    int index = -1;
-    for ( size_t s=0; s<done->size(); ++s )
-    {
-        if ((*done)[s]==at)
-        {
-            index = int(s);
-            break;
-        }
-    }
-
-    std::string pre;
-    for (int i=0; i<indent; ++i) pre +=" ";
-    if (label)
-    {
-        pre += label;
-    }
-
-    auto opType = at->GetOpType();
-    if (index>=0)
-    {
-        if ( GetOpDataType(opType)==DT_MESH)
-        {
-			UE_LOG(LogMutableCore, Warning, TEXT("%srepeated : %d"),
-                     pre.c_str(), index );
-        }
-        return;
-    }
-    index = (int)done->size();
-    done->push_back(at);
-
-    bool childrenAdded = false;
-    if ( GetOpDataType(opType)==DT_MESH)
-    {
-        if (opType==OP_TYPE::ME_CONSTANT)
-        {
-            ASTOpConstantResource* opc = dynamic_cast<ASTOpConstantResource*>(at.get());
-            auto value = opc->GetValue();
-            auto mesh = (const Mesh*)( value.get());
-            if (!mesh)
-            {
-                UE_LOG(LogMutableCore, Warning, TEXT("%s%s [%d]"),
-                         pre.c_str(),
-                         "mesh-null", index );
-            }
-            else if (mesh->GetIndexCount()==0)
-            {
-                UE_LOG(LogMutableCore, Warning, TEXT("%s%s [%d]"),
-                         pre.c_str(),
-                         "mesh-mask", index );
-            }
-            else
-            {
-                UE_LOG(LogMutableCore, Warning, TEXT("%s%s [%d]"),
-                         pre.c_str(),
-                         "mesh", index );
-            }
-        }
-
-        else if (opType==OP_TYPE::ME_REMOVEMASK)
-        {
-            UE_LOG(LogMutableCore, Warning, TEXT("%s%s [%d]"),
-                     pre.c_str(),
-                     mu::s_opNames[int(opType)], index );
-
-            ASTOpMeshRemoveMask* opc = dynamic_cast<ASTOpMeshRemoveMask*>(at.get());
-
-            DebugLogAST( opc->source.child(), indent+2, done, "source : " );
-
-            for ( const auto& m : opc->removes )
-            {
-                DebugLogAST( m.first.child(), indent+2, done, "condition : " );
-                DebugLogAST( m.second.child(), indent+2, done, "mask : ");
-            }
-            childrenAdded = true;
-        }
-
-        else
-        {
-            UE_LOG(LogMutableCore, Warning, TEXT("%s%s [%d]"),
-                     pre.c_str(),
-                     mu::s_opNames[int(opType)], index );
-        }
-    }
-
-    if (!childrenAdded)
-    {
-        at->ForEachChild([indent,done]( ASTChild& c )
-        {
-            DebugLogAST( c.m_child, indent+2, done );
-        });
-    }
-}
-
-
 //-------------------------------------------------------------------------------------------------
 //-------------------------------------------------------------------------------------------------
 //-------------------------------------------------------------------------------------------------
@@ -200,7 +86,7 @@ ASTChild& ASTChild::operator=( ASTChild&& rhs )
 //-------------------------------------------------------------------------------------------------
 //-------------------------------------------------------------------------------------------------
 //-------------------------------------------------------------------------------------------------
-void ASTOp::ForEachParent(const std::function<void(ASTOp*)>& f )
+void ASTOp::ForEachParent(const TFunctionRef<void(ASTOp*)> f )
 {
     for( auto& p: m_parents )
     {
@@ -216,7 +102,7 @@ void ASTOp::ForEachParent(const std::function<void(ASTOp*)>& f )
 void ASTOp::RemoveChildren()
 {
     // Actually destroyed when running out of scope
-    vector<Ptr<ASTOp>> toDestroy;
+    TArray<Ptr<ASTOp>> toDestroy;
 
     // Try to make children destruction iterative
     TArray<ASTOp*> pending;
@@ -234,7 +120,7 @@ void ASTOp::RemoveChildren()
                 // Are we clearing the last reference?
                 if (c.child()->GetRefCount()==1)
                 {
-                    toDestroy.push_back(c.child());
+                    toDestroy.Add(c.child());
                     pending.Add(c.child().get());
                 }
 
@@ -301,7 +187,7 @@ bool ASTOp::operator==( const ASTOp& other ) const
 
 
 //---------------------------------------------------------------------------------------------
-void ASTOp::FullAssert( const vector<Ptr<ASTOp>>& roots )
+void ASTOp::FullAssert( const TArray<Ptr<ASTOp>>& roots )
 {
     MUTABLE_CPUPROFILER_SCOPE(AST_FullAssert);
     Traverse_TopDown_Unique_Imprecise( roots, [](const Ptr<ASTOp>& n)
@@ -312,7 +198,7 @@ void ASTOp::FullAssert( const vector<Ptr<ASTOp>>& roots )
 }
 
 //-------------------------------------------------------------------------------------------------
-size_t ASTOp::CountNodes( const vector<Ptr<ASTOp>>& roots )
+size_t ASTOp::CountNodes( const TArray<Ptr<ASTOp>>& roots )
 {
     MUTABLE_CPUPROFILER_SCOPE(AST_CountNodes);
     size_t count=0;
@@ -368,7 +254,7 @@ void ASTOp::ClearLinkData( Ptr<ASTOp>& root )
 {
     MUTABLE_CPUPROFILER_SCOPE(AST_ClearLinkData);
     ASTOpList roots;
-    roots.push_back(root);
+    roots.Add(root);
     Traverse_TopDown_Unique_Imprecise( roots,
                              [&](const Ptr<ASTOp>& n){ n->linkedAddress = 0; return true; });
 }
@@ -417,13 +303,16 @@ void ASTOp::LogHistogram( ASTOpList& roots )
 
 
 //-------------------------------------------------------------------------------------------------
-void ASTOp::Traverse_TopDown_Unique( const vector<Ptr<ASTOp>>& roots,
-                                     std::function<bool(Ptr<ASTOp>&)> f )
+void ASTOp::Traverse_TopDown_Unique( const TArray<Ptr<ASTOp>>& roots,
+                                     TFunctionRef<bool(Ptr<ASTOp>&)> f )
 {
-    std::queue<Ptr<ASTOp>> pending;
-    for (auto& r:roots) pending.push( r );
+    TQueue<Ptr<ASTOp>> pending;
+	for (auto& r : roots)
+	{
+		pending.Enqueue(r);
+	}
 
-    std::unordered_set<Ptr<const ASTOp>> traversed;
+    TSet<Ptr<const ASTOp>> traversed;
 
     // We record the parents of all roots as traversed
     for ( const auto& r: roots )
@@ -431,17 +320,17 @@ void ASTOp::Traverse_TopDown_Unique( const vector<Ptr<ASTOp>>& roots,
         r->ForEachParent( [&]( const ASTOp* parent )
         {
             // If the parent is also a root, we will want to process it.
-            if ( std::find(roots.begin(), roots.end(), parent) == roots.end() )
+            if ( !roots.Contains(parent) )
             {
-                traversed.insert( parent );
+                traversed.Add( parent );
             }
         });
     }
 
-    while (pending.size())
+    while (!pending.IsEmpty())
     {
-        Ptr<ASTOp> pCurrent = pending.front();
-        pending.pop();
+		Ptr<ASTOp> pCurrent;
+		pending.Dequeue(pCurrent);
         if (!pCurrent)
         {
             continue;
@@ -452,7 +341,7 @@ void ASTOp::Traverse_TopDown_Unique( const vector<Ptr<ASTOp>>& roots,
 
         pCurrent->ForEachParent( [&]( const ASTOp* parent )
         {
-            if (traversed.find(parent)==traversed.end())
+            if (!traversed.Contains(parent))
             {
                 // \todo Is the parent in the relevant subtree?
 
@@ -463,11 +352,11 @@ void ASTOp::Traverse_TopDown_Unique( const vector<Ptr<ASTOp>>& roots,
 
         if (!parentsTraversed)
         {
-            pending.push(pCurrent);
+            pending.Enqueue(pCurrent);
         }
-        else if (traversed.find(pCurrent)==traversed.end())
+        else if (!traversed.Contains(pCurrent))
         {
-            traversed.insert(pCurrent);
+            traversed.Add(pCurrent);
 
             // Process
             bool recurse = f(pCurrent);
@@ -477,9 +366,9 @@ void ASTOp::Traverse_TopDown_Unique( const vector<Ptr<ASTOp>>& roots,
             {
                 pCurrent->ForEachChild([&]( ASTChild& c )
                 {
-                    if (c && traversed.find(c.m_child)==traversed.end())
+                    if (c && !traversed.Contains(c.m_child))
                     {
-                        pending.push( c.m_child );
+                        pending.Enqueue( c.m_child );
                     }
                 });
             }
@@ -489,23 +378,26 @@ void ASTOp::Traverse_TopDown_Unique( const vector<Ptr<ASTOp>>& roots,
 
 
 //-------------------------------------------------------------------------------------------------
-void ASTOp::Traverse_TopDown_Unique_Imprecise( const vector<Ptr<ASTOp>>& roots,
-                                     std::function<bool(Ptr<ASTOp>&)> f )
+void ASTOp::Traverse_TopDown_Unique_Imprecise( const TArray<Ptr<ASTOp>>& roots,
+                                     TFunctionRef<bool(Ptr<ASTOp>&)> f )
 {
-    std::queue<Ptr<ASTOp>> pending;
-    for (auto& r:roots) pending.push( r );
+    TQueue<Ptr<ASTOp>> pending;
+	for (auto& r : roots)
+	{
+		pending.Enqueue(r);
+	}
 
-    std::unordered_set<Ptr<const ASTOp>> traversed;
+    TSet<Ptr<const ASTOp>> traversed;
 
-    while (pending.size())
+    while (!pending.IsEmpty())
     {
-        Ptr<ASTOp> pCurrent = pending.front();
-        pending.pop();
+		Ptr<ASTOp> pCurrent;
+		pending.Dequeue(pCurrent);
 
         // It could have been completed in another branch
-        if (pCurrent && traversed.find(pCurrent)==traversed.end())
+        if (pCurrent && !traversed.Contains(pCurrent))
         {
-            traversed.insert(pCurrent);
+            traversed.Add(pCurrent);
 
             // Process
             bool recurse = f(pCurrent);
@@ -515,9 +407,9 @@ void ASTOp::Traverse_TopDown_Unique_Imprecise( const vector<Ptr<ASTOp>>& roots,
             {
                 pCurrent->ForEachChild([&]( ASTChild& c )
                 {
-                    if (c && traversed.find(c.m_child)==traversed.end())
+                    if (c && !traversed.Contains(c.m_child))
                     {
-                        pending.push( c.m_child );
+                        pending.Enqueue( c.m_child );
                     }
                 });
             }
@@ -527,19 +419,19 @@ void ASTOp::Traverse_TopDown_Unique_Imprecise( const vector<Ptr<ASTOp>>& roots,
 
 
 //-------------------------------------------------------------------------------------------------
-void ASTOp::Traverse_TopRandom_Unique_NonReentrant( const vector<Ptr<ASTOp>>& roots,
-                                                  std::function<bool(Ptr<ASTOp>&)> f )
+void ASTOp::Traverse_TopRandom_Unique_NonReentrant( const TArray<Ptr<ASTOp>>& roots,
+	TFunctionRef<bool(Ptr<ASTOp>&)> f )
 {
     ASTOpList pending;
 
-    uint32_t traverseIndex = s_lastTraverseIndex++;
+    uint32 traverseIndex = s_lastTraverseIndex++;
 
     for (auto& r:roots)
     {
         if (r && r->m_traverseIndex!=traverseIndex )
         {
             r->m_traverseIndex = traverseIndex;
-            pending.push_back( r );
+            pending.Add( r );
         }
     }
     for( auto& p : pending)
@@ -547,10 +439,9 @@ void ASTOp::Traverse_TopRandom_Unique_NonReentrant( const vector<Ptr<ASTOp>>& ro
         p->m_traverseIndex = traverseIndex-1;
     }
 
-    while (pending.size())
+    while (pending.Num())
     {
-        Ptr<ASTOp> pCurrent = pending.back();
-        pending.pop_back();
+        Ptr<ASTOp> pCurrent = pending.Pop();
 
         // It could have been completed in another branch
         if (pCurrent->m_traverseIndex!=traverseIndex)
@@ -567,7 +458,7 @@ void ASTOp::Traverse_TopRandom_Unique_NonReentrant( const vector<Ptr<ASTOp>>& ro
                 {
                     if (c && c.m_child->m_traverseIndex!=traverseIndex)
                     {
-                        pending.push_back( c.m_child );
+                        pending.Add( c.m_child );
                     }
                 });
             }
@@ -585,7 +476,7 @@ void Visitor_TopDown_Unique_Cloning::Traverse( Ptr<ASTOp>& root )
     // Visit the given root
     if (root)
     {
-        m_pending.push_back( std::make_pair(false,root) );
+		m_pending.Add({ false,root });
 
         Process();
 
@@ -597,16 +488,15 @@ void Visitor_TopDown_Unique_Cloning::Traverse( Ptr<ASTOp>& root )
 
 void Visitor_TopDown_Unique_Cloning::Process()
 {
-    while ( m_pending.size() )
+    while ( m_pending.Num() )
     {
-        std::pair<bool,Ptr<ASTOp>> item = m_pending.back();
-        m_pending.pop_back();
+        TPair<bool,Ptr<ASTOp>> item = m_pending.Pop();
 
-        Ptr<ASTOp> at = item.second;
+        Ptr<ASTOp> at = item.Value;
 
 		ASTOp::MapChildFunc Identity = [](const Ptr<ASTOp>& o) {return o; };
 		
-		if (item.first)
+		if (item.Key)
         {
             // Item indicating we finished with all the children of this instruction
 			auto cop = at->Clone(Identity);
@@ -680,13 +570,13 @@ void Visitor_TopDown_Unique_Cloning::Process()
                 if (processChildren)
                 {
                     // TODO: Shouldn't we recurse newAt?
-                    m_pending.push_back( std::make_pair(true,at) );
+					m_pending.Add({ true,at });
 
                     at->ForEachChild( [&](ASTChild& ref)
                     {
                         if (ref && m_oldToNew.find(ref.m_child)==m_oldToNew.end())
                         {
-                            m_pending.push_back( std::make_pair(false,ref.m_child) );
+							m_pending.Add({ false,ref.m_child });
                         }
                     });
                 }
@@ -700,15 +590,14 @@ void Visitor_TopDown_Unique_Cloning::Process()
 //-------------------------------------------------------------------------------------------------
 //-------------------------------------------------------------------------------------------------
 //-------------------------------------------------------------------------------------------------
-void ASTOp::Traverse_TopDown_Repeat( const vector<Ptr<ASTOp>>& roots,
-                                     std::function<bool(Ptr<ASTOp>& node)> f )
+void ASTOp::Traverse_TopDown_Repeat( const TArray<Ptr<ASTOp>>& roots,
+                                     TFunctionRef<bool(Ptr<ASTOp>& node)> f )
 {
     ASTOpList pending = roots;
 
-    while (pending.size())
+    while (pending.Num())
     {
-        Ptr<ASTOp> pCurrent = pending.back();
-        pending.pop_back();
+        Ptr<ASTOp> pCurrent = pending.Pop();
 
         if (pCurrent)
         {
@@ -722,7 +611,7 @@ void ASTOp::Traverse_TopDown_Repeat( const vector<Ptr<ASTOp>>& roots,
                 {
                     if (c)
                     {
-                        pending.push_back( c.m_child );
+                        pending.Add( c.m_child );
                     }
                 });
             }
@@ -737,18 +626,18 @@ void ASTOp::Traverse_TopDown_Repeat( const vector<Ptr<ASTOp>>& roots,
 void ASTOp::Traverse_BottomUp_Unique_NonReentrant
 (
         ASTOpList& roots,
-        std::function<void(Ptr<ASTOp>&)> f
+        TFunctionRef<void(Ptr<ASTOp>&)> f
 )
 {
-    uint32_t traverseIndex = s_lastTraverseIndex++;
+    uint32 traverseIndex = s_lastTraverseIndex++;
 
-    vector< std::pair<Ptr<ASTOp>,int> > pending;
+    TArray< std::pair<Ptr<ASTOp>,int> > pending;
     for (auto& r:roots)
     {
         if (r && r->m_traverseIndex!=traverseIndex )
         {
             r->m_traverseIndex=traverseIndex;
-            pending.push_back( std::make_pair<>(r,0) );
+            pending.Add( std::make_pair<>(r,0) );
         }
     }
     for(auto& p : pending)
@@ -756,11 +645,11 @@ void ASTOp::Traverse_BottomUp_Unique_NonReentrant
         p.first->m_traverseIndex = traverseIndex-1;
     }
 
-    while (pending.size())
+    while (pending.Num())
     {
-        int phase = pending.back().second;
-        Ptr<ASTOp> pCurrent = pending.back().first;
-        pending.pop_back();
+        int phase = pending.Last().second;
+        Ptr<ASTOp> pCurrent = pending.Last().first;
+        pending.Pop();
 
         // It could have been completed in another branch
         if (pCurrent->m_traverseIndex!=traverseIndex)
@@ -768,7 +657,7 @@ void ASTOp::Traverse_BottomUp_Unique_NonReentrant
             if (phase==0)
             {
                 // Process this again...
-                pending.push_back( std::make_pair<>(pCurrent,1) );
+                pending.Add( std::make_pair<>(pCurrent,1) );
 
                 // ...after the children are processed
                 pCurrent->ForEachChild([&]( ASTChild& c )
@@ -776,7 +665,7 @@ void ASTOp::Traverse_BottomUp_Unique_NonReentrant
                     if (c && c.m_child->m_traverseIndex!=traverseIndex )
                     {
                         auto e = std::make_pair<>(c.m_child,0);
-                        pending.push_back( e );
+                        pending.Add( e );
                     }
                 });
             }
@@ -798,19 +687,19 @@ void ASTOp::Traverse_BottomUp_Unique_NonReentrant
 void ASTOp::Traverse_BottomUp_Unique_NonReentrant
 (
         ASTOpList& roots,
-        std::function<void(Ptr<ASTOp>&)> f,
-        std::function<bool(const ASTOp*)> accept
+        TFunctionRef<void(Ptr<ASTOp>&)> f,
+        TFunctionRef<bool(const ASTOp*)> accept
 )
 {
-    uint32_t traverseIndex = s_lastTraverseIndex++;
+    uint32 traverseIndex = s_lastTraverseIndex++;
 
-    vector< std::pair<Ptr<ASTOp>,int> > pending;
+    TArray< std::pair<Ptr<ASTOp>,int> > pending;
     for (auto& r:roots)
     {
         if (r && r->m_traverseIndex!=traverseIndex)
         {
             r->m_traverseIndex=traverseIndex;
-            pending.push_back( std::make_pair<>(r,0) );
+            pending.Add( std::make_pair<>(r,0) );
         }
     }
     for(auto& p : pending)
@@ -818,11 +707,11 @@ void ASTOp::Traverse_BottomUp_Unique_NonReentrant
         p.first->m_traverseIndex = traverseIndex-1;
     }
 
-    while (pending.size())
+    while (pending.Num())
     {
-        int phase = pending.back().second;
-        Ptr<ASTOp> pCurrent = pending.back().first;
-        pending.pop_back();
+        int phase = pending.Last().second;
+        Ptr<ASTOp> pCurrent = pending.Last().first;
+        pending.Pop();
 
         // It could have been completed in another branch
         if (pCurrent->m_traverseIndex!=traverseIndex && accept(pCurrent.get()))
@@ -830,14 +719,14 @@ void ASTOp::Traverse_BottomUp_Unique_NonReentrant
             if (phase==0)
             {
                 // Process this again...
-                pending.push_back( std::make_pair<>(pCurrent,1) );
+                pending.Add( std::make_pair<>(pCurrent,1) );
 
                 // ...after the children are processed
                 pCurrent->ForEachChild([&]( ASTChild& c )
                 {
                     if (c && accept(c.m_child.get()) && c.m_child->m_traverseIndex!=traverseIndex )
                     {
-                        pending.push_back( std::make_pair<>(c.m_child,0) );
+                        pending.Add( std::make_pair<>(c.m_child,0) );
                     }
                 });
             }
@@ -859,8 +748,8 @@ void ASTOp::Traverse_BottomUp_Unique_NonReentrant
 void ASTOp::Traverse_BottomUp_Unique
 (
         ASTOpList& roots,
-        std::function<void(Ptr<ASTOp>&)> f,
-        std::function<bool(const ASTOp*)> accept
+        TFunctionRef<void(Ptr<ASTOp>&)> f,
+        TFunctionRef<bool(const ASTOp*)> accept
 )
 {
     TSet<Ptr<ASTOp>> Traversed;
@@ -921,14 +810,14 @@ void ASTOp::Traverse_BottomUp_Unique
 void ASTOp::Traverse_BottomUp_Unique
 (
         Ptr<ASTOp>& root,
-        std::function<void(Ptr<ASTOp>&)> f,
-        std::function<bool(const ASTOp*)> accept
+        TFunctionRef<void(Ptr<ASTOp>&)> f,
+        TFunctionRef<bool(const ASTOp*)> accept
 )
 {
     if (root)
     {
         ASTOpList roots;
-        roots.push_back(root);
+        roots.Add(root);
         Traverse_BottomUp_Unique(roots,f,accept);
     }
 }
@@ -1046,7 +935,7 @@ ASTOpFixed::~ASTOpFixed()
 }
 
 
-void ASTOpFixed::ForEachChild( const std::function<void(ASTChild&)>& f )
+void ASTOpFixed::ForEachChild( const TFunctionRef<void(ASTChild&)> f )
 {
     // ugly but temporary while ASTOpFixed exists
     ForEachReference( op, [&](OP::ADDRESS* pAt)
@@ -1080,8 +969,8 @@ void ASTOpFixed::Link( PROGRAM& program, const FLinkerOptions* )
             }
         });
 
-        linkedAddress = (OP::ADDRESS)program.m_opAddress.size();        
-        program.m_opAddress.push_back((uint32_t)program.m_byteCode.size());
+        linkedAddress = (OP::ADDRESS)program.m_opAddress.Num();        
+        program.m_opAddress.Add((uint32)program.m_byteCode.Num());
 
         AppendCode(program.m_byteCode,lop.type);
         if (lop.type==OP_TYPE::ME_MERGE)
@@ -1107,7 +996,7 @@ bool ASTOpFixed::IsEqual(const ASTOp& otherUntyped) const
 }
 
 
-mu::Ptr<ASTOp> ASTOpFixed::Clone( MapChildFunc& mapChild ) const
+mu::Ptr<ASTOp> ASTOpFixed::Clone( MapChildFuncRef mapChild ) const
 {
     Ptr<ASTOpFixed> n = new ASTOpFixed();
     n->op = op;
@@ -1232,8 +1121,8 @@ FImageDesc ASTOpFixed::GetImageDesc( bool returnBestOption, GetImageDescContext*
 
     case OP_TYPE::IM_RESIZEREL:
         res = GetImageDesc( op.args.ImageResizeRel.source, returnBestOption, context );
-        res.m_size[0] = (uint16_t)( res.m_size[0]*op.args.ImageResizeRel.factor[0] );
-        res.m_size[1] = (uint16_t)( res.m_size[1]*op.args.ImageResizeRel.factor[1] );
+        res.m_size[0] = (uint16)( res.m_size[0]*op.args.ImageResizeRel.factor[0] );
+        res.m_size[1] = (uint16)( res.m_size[1]*op.args.ImageResizeRel.factor[1] );
         break;
 
     case OP_TYPE::IM_RESIZELIKE:
@@ -1510,7 +1399,7 @@ ASTOp::BOOL_EVAL_RESULT ASTOpFixed::EvaluateBool( ASTOpList& facts, EVALUATE_BOO
         const auto& b = children[op.args.BoolBinary.b].child();
         BOOL_EVAL_RESULT resultA = BET_UNKNOWN;
         BOOL_EVAL_RESULT resultB = BET_UNKNOWN;
-        for ( size_t f=0; f<facts.size(); ++f )
+        for ( size_t f=0; f<facts.Num(); ++f )
         {
             if ( a && resultA == BET_UNKNOWN )
             {
@@ -1552,7 +1441,7 @@ ASTOp::BOOL_EVAL_RESULT ASTOpFixed::EvaluateBool( ASTOpList& facts, EVALUATE_BOO
         const auto& b = children[op.args.BoolBinary.b].child();
         BOOL_EVAL_RESULT resultA = BET_UNKNOWN;
         BOOL_EVAL_RESULT resultB = BET_UNKNOWN;
-        for ( size_t f=0; f<facts.size(); ++f )
+        for ( size_t f=0; f<facts.Num(); ++f )
         {
             if ( a && resultA == BET_UNKNOWN )
             {
@@ -1616,7 +1505,7 @@ int ASTOpFixed::EvaluateInt( ASTOpList& /*facts*/, bool &unknown ) const
 
         // TODO
 //    case IntExpression::IET_PARAMETER:
-//        for ( size_t f=0; (result==BET_UNKNOWN) && f<facts.size(); ++f )
+//        for ( size_t f=0; (result==BET_UNKNOWN) && f<facts.Num(); ++f )
 //        {
 //            result = EvaluateIntEquality(facts[f].get(),intExp,value);
 //        }
@@ -1763,8 +1652,8 @@ mu::Ptr<ImageSizeExpression> ASTOpFixed::GetImageSizeExpression() const
             pRes = children[op.args.ImageResizeRel.source].child()->GetImageSizeExpression();
             if (pRes->type == ImageSizeExpression::ISET_CONSTANT)
             {
-                pRes->size[0] = uint16_t( pRes->size[0] * op.args.ImageResizeRel.factor[0] );
-                pRes->size[1] = uint16_t( pRes->size[1] * op.args.ImageResizeRel.factor[1] );
+                pRes->size[0] = uint16( pRes->size[0] * op.args.ImageResizeRel.factor[0] );
+                pRes->size[1] = uint16( pRes->size[1] * op.args.ImageResizeRel.factor[1] );
             }
             else
             {
diff --git a/Engine/Plugins/Experimental/Mutable/Source/MutableTools/Private/MuT/AST.h b/Engine/Plugins/Experimental/Mutable/Source/MutableTools/Private/MuT/AST.h
--- a/Engine/Plugins/Experimental/Mutable/Source/MutableTools/Private/MuT/AST.h
+++ b/Engine/Plugins/Experimental/Mutable/Source/MutableTools/Private/MuT/AST.h
@@ -2,12 +2,7 @@
 
 #pragma once
 
-#include "Containers/Array.h"
-#include "Containers/ContainerAllocationPolicies.h"
-#include "Containers/Map.h"
-#include "HAL/PlatformCrt.h"
-#include "HAL/PlatformMath.h"
-#include "Misc/AssertionMacros.h"
+#include "MuT/Platform.h"
 #include "MuR/Image.h"
 #include "MuR/MemoryPrivate.h"
 #include "MuR/ModelPrivate.h"
@@ -17,15 +12,23 @@
 #include "MuR/ParametersPrivate.h"
 #include "MuR/Ptr.h"
 #include "MuR/RefCounted.h"
+
+#include "Containers/Array.h"
+#include "Containers/ContainerAllocationPolicies.h"
+#include "Containers/Map.h"
+#include "HAL/PlatformCrt.h"
+#include "HAL/PlatformMath.h"
+#include "Misc/AssertionMacros.h"
 #include "Templates/TypeHash.h"
+#include "Templates/Function.h"
 
 #include <array>
 #include <atomic>
-#include <functional>
 #include <memory>
 #include <unordered_map>
 #include <unordered_set>
 #include <utility>
+#include <set>
 
 namespace mu { class ASTOp; }
 namespace mu { class ASTOpFixed; }
@@ -81,7 +84,7 @@ namespace mu
 
         // For layout factor sizes
         Ptr<class ASTOp> layout;
-        uint16_t factor[2];
+        uint16 factor[2];
 
         // For conditionals
         Ptr<class ASTOp> condition;
@@ -171,8 +174,8 @@ namespace mu
     };
 
 
-    typedef vector<Ptr<ASTOp>> ASTOpList;
-    typedef std::unordered_set<Ptr<ASTOp>> ASTOpSet;
+    typedef TArray<Ptr<ASTOp>> ASTOpList;
+    typedef TSet<Ptr<ASTOp>> ASTOpSet;
 
 
     //! Detailed optimization flags
@@ -220,8 +223,8 @@ namespace mu
 		const ASTOp* m_root = nullptr;
 		Ptr<ASTOp> m_initialSource;
 		//! For each operation we sink, the map from old instructions to new instructions.
-		map<Ptr<const class ASTOpFixed>, std::unordered_map<Ptr<ASTOp>, Ptr<ASTOp>>> m_oldToNew;
-		vector<Ptr<ASTOp>> m_newOps;
+		TMap<Ptr<const class ASTOpFixed>, std::unordered_map<Ptr<ASTOp>, Ptr<ASTOp>>> m_oldToNew;
+		TArray<Ptr<ASTOp>> m_newOps;
 
 		Ptr<ASTOp> Visit(Ptr<ASTOp> at, const ASTOpFixed* currentCropOp);
 	};
@@ -241,8 +244,8 @@ namespace mu
 		const class ASTOpImagePixelFormat* m_root = nullptr;
 		Ptr<ASTOp> m_initialSource;
 		//! For each operation we sink, the map from old instructions to new instructions.
-		map<Ptr<const ASTOp>, std::unordered_map<Ptr<ASTOp>, Ptr<ASTOp>>> m_oldToNew;
-		vector<Ptr<ASTOp>> m_newOps;
+		TMap<Ptr<const ASTOp>, std::unordered_map<Ptr<ASTOp>, Ptr<ASTOp>>> m_oldToNew;
+		TArray<Ptr<ASTOp>> m_newOps;
 
 		Ptr<ASTOp> Visit(Ptr<ASTOp> at, const class ASTOpImagePixelFormat* currentFormatOp);
 	};
@@ -262,8 +265,8 @@ namespace mu
 
 		const class ASTOpMeshFormat* m_root = nullptr;
 		mu::Ptr<ASTOp> m_initialSource;
-		std::unordered_map<mu::Ptr<ASTOp>, mu::Ptr<ASTOp>> m_oldToNew;
-		vector<mu::Ptr<ASTOp>> m_newOps;
+		TMap<mu::Ptr<ASTOp>, mu::Ptr<ASTOp>> m_oldToNew;
+		TArray<mu::Ptr<ASTOp>> m_newOps;
 
 		mu::Ptr<ASTOp> Visit(const mu::Ptr<ASTOp>& at, const class ASTOpMeshFormat* currentFormatOp);
 	};
@@ -284,8 +287,8 @@ namespace mu
 		Ptr<ASTOp> m_initialSource;
 
 		//! For each operation we sink, the map from old instructions to new instructions.
-		std::unordered_map<Ptr<ASTOp>, Ptr<ASTOp>> m_oldToNew;
-		vector<Ptr<ASTOp>> m_newOps;
+		TMap<Ptr<ASTOp>, Ptr<ASTOp>> m_oldToNew;
+		TArray<Ptr<ASTOp>> m_newOps;
 
 		Ptr<ASTOp> Visit(Ptr<ASTOp> at, const class ASTOpImageMipmap* currentMipmapOp);
 
@@ -373,7 +376,7 @@ namespace mu
     private:
 
         //! Operations referring to this one. They may be null: elements are never removed
-        //! from this vector.
+        //! from this TArray.
 		TArray<ASTOp*, TInlineAllocator<4> > m_parents;
 
     public:
@@ -395,10 +398,10 @@ namespace mu
         virtual void Assert();
 
         //! Run something for each child operation, with a chance to modify it.
-        virtual void ForEachChild( const std::function<void(ASTChild&)>& ) = 0;
+        virtual void ForEachChild( const TFunctionRef<void(ASTChild&)> ) = 0;
 
         //! Run something for each parent operation+.
-        void ForEachParent( const std::function<void(ASTOp*)>& );
+        void ForEachParent( const TFunctionRef<void(ASTOp*)> );
 
         //! Run something for each child operation, with a chance to modify it.
         virtual bool operator==( const ASTOp& other ) const;
@@ -407,8 +410,9 @@ namespace mu
         virtual uint64 Hash() const = 0;
 
         //! Shallow clone. New node will have no parents but reference to the same children.
-        using MapChildFunc=std::function<Ptr<ASTOp>(const Ptr<ASTOp>&)>;
-		virtual Ptr<ASTOp> Clone(MapChildFunc& mapChild) const = 0;
+		using MapChildFunc = TFunction<Ptr<ASTOp>(const Ptr<ASTOp>&)>;
+		using MapChildFuncRef = TFunctionRef<Ptr<ASTOp>(const Ptr<ASTOp>&)>;
+		virtual Ptr<ASTOp> Clone(MapChildFuncRef mapChild) const = 0;
 
     protected:
 
@@ -419,9 +423,9 @@ namespace mu
     public:
 
         //---------------------------------------------------------------------------------------------
-        static void FullAssert( const vector<Ptr<ASTOp>>& roots );
+        static void FullAssert( const TArray<Ptr<ASTOp>>& roots );
 
-        static size_t CountNodes( const vector<Ptr<ASTOp>>& roots );
+        static size_t CountNodes( const TArray<Ptr<ASTOp>>& roots );
 
         //! Deep clone. New node will have no parents and reference new children
         static Ptr<ASTOp> DeepClone( const Ptr<ASTOp>& );
@@ -511,21 +515,21 @@ namespace mu
         static void LinkRange( PROGRAM& program,
                                const RANGE_DATA& range,
                                OP::ADDRESS& rangeSize,
-                               uint16_t& rangeId )
+                               uint16& rangeId )
         {
             if (range.rangeSize)
             {
                 if (range.rangeSize->linkedRange<0)
                 {
-                    check( program.m_ranges.size()<255 );
-                    range.rangeSize->linkedRange = int8_t( program.m_ranges.size() );
+                    check( program.m_ranges.Num()<255 );
+                    range.rangeSize->linkedRange = int8( program.m_ranges.Num() );
                     RANGE_DESC rangeData;
                     rangeData.m_name = range.rangeName;
                     rangeData.m_uid = range.rangeUID;
-                    program.m_ranges.push_back( rangeData );
+                    program.m_ranges.Add( rangeData );
                 }
                 rangeSize = range.rangeSize->linkedAddress;
-                rangeId = uint16_t(range.rangeSize->linkedRange);
+                rangeId = uint16(range.rangeSize->linkedRange);
             }
         }
 
@@ -535,20 +539,18 @@ namespace mu
         //---------------------------------------------------------------------------------------------
 
         //!
-        static void Traverse_TopDown_Unique( const vector<Ptr<ASTOp>>& roots,
-                                             std::function<bool(Ptr<ASTOp>&)> f );
+        static void Traverse_TopDown_Unique( const TArray<Ptr<ASTOp>>& roots, TFunctionRef<bool(Ptr<ASTOp>&)> f );
 
         //! \todo: it is not strictly top down.
-        static void Traverse_TopDown_Unique_Imprecise( const vector<Ptr<ASTOp>>& roots,
-                                             std::function<bool(Ptr<ASTOp>&)> f );
+        static void Traverse_TopDown_Unique_Imprecise( const TArray<Ptr<ASTOp>>& roots, TFunctionRef<bool(Ptr<ASTOp>&)> f );
 
         //! Kind of top-down, but really not.
         //! This version is slighlty faster, but doesn't support recursive traversals so
         //! use it only in controlled cases.
         static void Traverse_TopRandom_Unique_NonReentrant
             (
-                const vector<Ptr<ASTOp>>& roots,
-                std::function<bool(Ptr<ASTOp>&)> f
+                const TArray<Ptr<ASTOp>>& roots,
+				TFunctionRef<bool(Ptr<ASTOp>&)> f
             );
 
         //! Kind of top-down, but really not.
@@ -558,48 +560,45 @@ namespace mu
         static inline void Traverse_TopDown_Unique_Imprecise_WithState
             (
                 Ptr<ASTOp>& root, const STATE& initialState,
-                std::function<bool(Ptr<ASTOp>&, STATE&, vector<pair<Ptr<ASTOp>,STATE>>&)> f
+				TFunctionRef<bool(Ptr<ASTOp>&, STATE&, TArray<TPair<Ptr<ASTOp>,STATE>>&)> f
             )
         {
             if (!root) { return; }
 
-            vector<pair<Ptr<ASTOp>,STATE>> pending;
-            pending.push_back( std::make_pair<>( root, initialState ) );
+            TArray<TPair<Ptr<ASTOp>,STATE>> pending;
+			pending.Emplace( root, initialState );
 
 
             struct custom_partial_hash
             {
-				uint64 operator()(const pair<Ptr<ASTOp>,STATE>& k) const
+				uint64 operator()(const std::pair<Ptr<ASTOp>,STATE>& k) const
                 {
                     return std::hash<const void*>()(k.first.get());
                 }
             };
 
-            std::unordered_set<pair<Ptr<ASTOp>,STATE>, custom_partial_hash> traversed;
+            std::unordered_set<std::pair<Ptr<ASTOp>,STATE>, custom_partial_hash> traversed;
 
-            while (pending.size())
+            while (!pending.IsEmpty())
             {
-                pair<Ptr<ASTOp>,STATE> current = pending.back();
-                pending.pop_back();
+				TPair<Ptr<ASTOp>,STATE> current = pending.Pop();
 
                 // It could have been completed in another branch
-                auto iti = traversed.insert( current );
+				auto iti = traversed.insert({ current.Key,current.Value });
                 if (iti.second)
                 {
-                    // It has been inserted, so it wasn't there before.
-
-                    // Process. State in current.second may change
-                    bool recurse = f(current.first,current.second,pending);
+                    // Process. State in current.Value may change
+                    bool recurse = f(current.Key,current.Value,pending);
 
                     // Recurse children
                     if (recurse)
                     {
-                        current.first->ForEachChild([&]( ASTChild& c )
+                        current.Key->ForEachChild([&]( ASTChild& c )
                         {
-                            auto child = std::make_pair<>(c.m_child, current.second);
-                            if ( c.m_child && !traversed.count(child) )
+							std::pair<Ptr<ASTOp>, STATE> TraverseKey(c.m_child, current.Value);
+							if (c.m_child && !traversed.count(TraverseKey))
                             {
-                                pending.push_back( child );
+								pending.Emplace( c.m_child, current.Value );
                             }
                         });
                     }
@@ -609,40 +608,35 @@ namespace mu
 
 
         //! Kind of top-down, but really not.
-        static void Traverse_TopDown_Repeat( const vector<Ptr<ASTOp>>& roots,
-                                             std::function<bool(Ptr<ASTOp>& node)> f );
+        static void Traverse_TopDown_Repeat( const TArray<Ptr<ASTOp>>& roots, TFunctionRef<bool(Ptr<ASTOp>& node)> f );
 
         //! This version is slighlty faster, but doesn't support recursive traversals so
         //! use it only in controlled cases.
-        static void Traverse_BottomUp_Unique_NonReentrant
-        (
-                ASTOpList& roots,
-                std::function<void(Ptr<ASTOp>&)> f
-        );
+        static void Traverse_BottomUp_Unique_NonReentrant( ASTOpList& roots, TFunctionRef<void(Ptr<ASTOp>&)> f );
 
         //! This version is slighlty faster, but doesn't support recursive traversals so
         //! use it only in controlled cases.
         static void Traverse_BottomUp_Unique_NonReentrant
         (
-                ASTOpList& roots,
-                std::function<void(Ptr<ASTOp>&)> f,
-                std::function<bool(const ASTOp*)> accept
+            ASTOpList& roots,
+			TFunctionRef<void(Ptr<ASTOp>&)> f,
+			TFunctionRef<bool(const ASTOp*)> accept
         );
 
         //!
         static void Traverse_BottomUp_Unique
         (
-                ASTOpList& roots,
-                std::function<void(Ptr<ASTOp>&)> f,
-                std::function<bool(const ASTOp*)> accept = [](const ASTOp*){return true;}
+            ASTOpList& roots,
+			TFunctionRef<void(Ptr<ASTOp>&)> f,
+			TFunctionRef<bool(const ASTOp*)> accept = [](const ASTOp*){return true;}
         );
 
         //!
         static void Traverse_BottomUp_Unique
         (
-                Ptr<ASTOp>& root,
-                std::function<void(Ptr<ASTOp>&)> f,
-                std::function<bool(const ASTOp*)> accept = [](const ASTOp*){return true;}
+            Ptr<ASTOp>& root,
+			TFunctionRef<void(Ptr<ASTOp>&)> f,
+			TFunctionRef<bool(const ASTOp*)> accept = [](const ASTOp*){return true;}
         );
 
         //!
@@ -650,20 +644,20 @@ namespace mu
 
         //! Generic traverse control counter. It should always be left to 0 after any process for all
         //! nodes in the hierarchy.
-        uint32_t m_traverseIndex = 0;
-        static std::atomic<uint32_t> s_lastTraverseIndex;
+        uint32 m_traverseIndex = 0;
+        static std::atomic<uint32> s_lastTraverseIndex;
 
         //!
-        int8_t linkedRange = -1;
+        int8 linkedRange = -1;
 
         //! special flag for Traverse_TopDown_Adaptative traversal.
-        uint8_t m_traversedTopDownAdapative : 1;
+        uint8 m_traversedTopDownAdapative : 1;
 
         //! Embedded node data for the constant subtree detection. This flag is only valid if the
         //! constant detection process has been executed and no relevant AST transformations have
         //! happened.
-        uint8_t m_constantSubtree : 1;
-        uint8_t m_hasSpecialOpInSubtree : 1;
+        uint8 m_constantSubtree : 1;
+        uint8 m_hasSpecialOpInSubtree : 1;
 
     private:
         friend class ASTChild;
@@ -772,9 +766,9 @@ namespace mu
     private:
 
         //! States found so far
-        vector<STATE> m_states;
+        TArray<STATE> m_states;
 
-        //! Index of the current state, from the m_states vector.
+        //! Index of the current state, from the m_states TArray.
         int m_currentState;
 
         struct PENDING
@@ -793,10 +787,10 @@ namespace mu
             int stateIndex;
         };
 
-        std::vector<PENDING> m_pending;
+        TArray<PENDING> m_pending;
 
         //! List of traversed nodes with the state in which they were traversed.
-        std::unordered_multimap<Ptr<ASTOp>,int> m_traversed;
+        std::multimap<Ptr<ASTOp>,int> m_traversed;
 
     protected:
 
@@ -817,15 +811,14 @@ namespace mu
         {
             if(at)
             {
-                auto it = std::find(m_states.begin(),m_states.end(),newState);
-                if (it==m_states.end())
+                auto it = m_states.Find(newState);
+                if (it==INDEX_NONE)
                 {
-                    m_states.push_back(newState);
-                    it = m_states.end() - 1;
+                    m_states.Add(newState);
                 }
-                int stateIndex = (int)(it - m_states.begin());
+                int stateIndex = m_states.Find(newState);
 
-                m_pending.push_back( PENDING(at,stateIndex) );
+                m_pending.Add( PENDING(at,stateIndex) );
             }
         }
 
@@ -834,20 +827,19 @@ namespace mu
         {
             if(at)
             {
-                m_pending.emplace_back( at, m_currentState );
+                m_pending.Emplace( at, m_currentState );
             }
         }
 
         //! Can be called from visit to set the state to visit all children ops
         void SetCurrentState(const STATE& newState)
         {
-            auto it = std::find(m_states.begin(),m_states.end(),newState);
-            if (it==m_states.end())
+            auto it = m_states.Find(newState);
+            if (it==INDEX_NONE)
             {
-                m_states.push_back(newState);
-                it = m_states.end() - 1;
+                m_states.Add(newState);
             }
-            m_currentState = (int)(it - m_states.begin());
+            m_currentState = m_states.Find(newState);
         }
 
     public:
@@ -858,23 +850,22 @@ namespace mu
         //!
         void Traverse(const ASTOpList& roots, const STATE& initialState)
         {
-            m_pending.clear();
-            m_states.clear();
-            m_states.push_back(initialState);
+            m_pending.Empty();
+            m_states.Empty();
+            m_states.Add(initialState);
             m_currentState = 0;
 
             for( const auto& r: roots)
             {
                 if (r)
                 {
-                    m_pending.push_back( PENDING(r,m_currentState) );
+                    m_pending.Add( PENDING(r,m_currentState) );
                 }
             }
 
-            while (m_pending.size())
+            while (m_pending.Num())
             {
-                PENDING item = m_pending.back();
-                m_pending.pop_back();
+                PENDING item = m_pending.Pop();
 
                 auto thisNodeRange = m_traversed.equal_range(item.at);
                 bool visitedInThisState = false;
@@ -903,7 +894,7 @@ namespace mu
                         {
                             if (c)
                             {
-                                m_pending.push_back( PENDING(c.child(),m_currentState) );
+                                m_pending.Add( PENDING(c.child(),m_currentState) );
                             }
                         });
                     }
@@ -936,7 +927,7 @@ namespace mu
     private:
 
         //! Operations to be processed
-        vector< std::pair<bool,Ptr<ASTOp>> > m_pending;
+        TArray< TPair<bool,Ptr<ASTOp>> > m_pending;
 
         //! Map for visited operations
         std::unordered_map<Ptr<ASTOp>,Ptr<ASTOp>> m_oldToNew;
@@ -974,13 +965,13 @@ namespace mu
     {
     public:
 
-        //! Fixed size op structure. Any code ADDRESS in it, refers to the children vector below.
+        //! Fixed size op structure. Any code ADDRESS in it, refers to the children TArray below.
         OP op;
 
         //! Operations used by this one. To be indexed with the ADDRESSes in the m_op.
-        //vector<ASTChild> children;
+        //TArray<ASTChild> children;
         std::array<ASTChild,8> children;
-        uint8_t childCount = 1;
+        uint8 childCount = 1;
 
     public:
         ASTOpFixed();
@@ -989,8 +980,8 @@ namespace mu
 
         OP_TYPE GetOpType() const override { return op.type; }
 
-        Ptr<ASTOp> Clone( MapChildFunc& mapChild ) const override;
-        void ForEachChild( const std::function<void(ASTChild&)>& f ) override;
+        Ptr<ASTOp> Clone( MapChildFuncRef mapChild ) const override;
+        void ForEachChild( const TFunctionRef<void(ASTChild&)> f ) override;
         void Link( PROGRAM& program, const FLinkerOptions* Options) override;
         bool IsEqual(const ASTOp& otherUntyped) const override;
 		uint64 Hash() const override;
@@ -1013,9 +1004,9 @@ namespace mu
         inline void SetChild( OP::ADDRESS& at, const Ptr<ASTOp>& child )
         {
             // hack check
-            check( ((uint8_t*)&at)>=((uint8_t*)&op.args)
+            check( ((uint8*)&at)>=((uint8*)&op.args)
                             &&
-                            ((uint8_t*)&at)<=((uint8_t*)&op.args)+sizeof(op.args));
+                            ((uint8*)&at)<=((uint8*)&op.args)+sizeof(op.args));
 
             if (!at)
             {
diff --git a/Engine/Plugins/Experimental/Mutable/Source/MutableTools/Private/MuT/ASTOpAddLOD.cpp b/Engine/Plugins/Experimental/Mutable/Source/MutableTools/Private/MuT/ASTOpAddLOD.cpp
--- a/Engine/Plugins/Experimental/Mutable/Source/MutableTools/Private/MuT/ASTOpAddLOD.cpp
+++ b/Engine/Plugins/Experimental/Mutable/Source/MutableTools/Private/MuT/ASTOpAddLOD.cpp
@@ -51,18 +51,18 @@ namespace mu
 	}
 
 
-	mu::Ptr<ASTOp> ASTOpAddLOD::Clone(MapChildFunc& mapChild) const
+	mu::Ptr<ASTOp> ASTOpAddLOD::Clone(MapChildFuncRef mapChild) const
 	{
 		Ptr<ASTOpAddLOD> n = new ASTOpAddLOD();
 		for (auto& c : lods)
 		{
-			n->lods.emplace_back(n, mapChild(c.child()));
+			n->lods.Emplace(n, mapChild(c.child()));
 		}
 		return n;
 	}
 
 
-	void ASTOpAddLOD::ForEachChild(const std::function<void(ASTChild&)>& f)
+	void ASTOpAddLOD::ForEachChild(const TFunctionRef<void(ASTChild&)> f)
 	{
 		for (auto& l : lods)
 		{
@@ -93,8 +93,8 @@ namespace mu
 				}
 			}
 
-			linkedAddress = (OP::ADDRESS)program.m_opAddress.size();
-			program.m_opAddress.push_back((uint32_t)program.m_byteCode.size());
+			linkedAddress = (OP::ADDRESS)program.m_opAddress.Num();
+			program.m_opAddress.Add((uint32_t)program.m_byteCode.Num());
 			AppendCode(program.m_byteCode, OP_TYPE::IN_ADDLOD);
 			AppendCode(program.m_byteCode, args);
 		}
diff --git a/Engine/Plugins/Experimental/Mutable/Source/MutableTools/Private/MuT/ASTOpAddLOD.h b/Engine/Plugins/Experimental/Mutable/Source/MutableTools/Private/MuT/ASTOpAddLOD.h
--- a/Engine/Plugins/Experimental/Mutable/Source/MutableTools/Private/MuT/ASTOpAddLOD.h
+++ b/Engine/Plugins/Experimental/Mutable/Source/MutableTools/Private/MuT/ASTOpAddLOD.h
@@ -8,8 +8,6 @@
 #include "MuR/Ptr.h"
 #include "MuT/AST.h"
 
-#include <functional>
-
 
 namespace mu
 {
@@ -23,7 +21,7 @@ struct PROGRAM;
 	{
 	public:
 
-		vector< ASTChild > lods;
+		TArray< ASTChild > lods;
 
 	public:
 
@@ -33,9 +31,9 @@ struct PROGRAM;
 
 		OP_TYPE GetOpType() const override { return OP_TYPE::IN_ADDLOD; }
 		uint64 Hash() const override;
-		void ForEachChild(const std::function<void(ASTChild&)>&) override;
+		void ForEachChild(const TFunctionRef<void(ASTChild&)>) override;
 		bool IsEqual(const ASTOp& otherUntyped) const override;
-		Ptr<ASTOp> Clone(MapChildFunc& mapChild) const override;
+		Ptr<ASTOp> Clone(MapChildFuncRef mapChild) const override;
 		void Link(PROGRAM& program, const FLinkerOptions* Options) override;
 	};
 
diff --git a/Engine/Plugins/Experimental/Mutable/Source/MutableTools/Private/MuT/ASTOpConditional.cpp b/Engine/Plugins/Experimental/Mutable/Source/MutableTools/Private/MuT/ASTOpConditional.cpp
--- a/Engine/Plugins/Experimental/Mutable/Source/MutableTools/Private/MuT/ASTOpConditional.cpp
+++ b/Engine/Plugins/Experimental/Mutable/Source/MutableTools/Private/MuT/ASTOpConditional.cpp
@@ -13,9 +13,6 @@
 #include "MuT/ASTOpConstantBool.h"
 #include "MuT/StreamsPrivate.h"
 
-#include <memory>
-#include <utility>
-
 
 using namespace mu;
 
@@ -52,7 +49,7 @@ bool ASTOpConditional::IsEqual(const ASTOp& otherUntyped) const
 
 
 //-------------------------------------------------------------------------------------------------
-mu::Ptr<ASTOp> ASTOpConditional::Clone(MapChildFunc& mapChild) const
+mu::Ptr<ASTOp> ASTOpConditional::Clone(MapChildFuncRef mapChild) const
 {
     Ptr<ASTOpConditional> n = new ASTOpConditional();
     n->type = type;
@@ -97,7 +94,7 @@ void ASTOpConditional::Assert()
 
 
 //-------------------------------------------------------------------------------------------------
-void ASTOpConditional::ForEachChild(const std::function<void(ASTChild&)>& f )
+void ASTOpConditional::ForEachChild(const TFunctionRef<void(ASTChild&)> f )
 {
     f(condition);
     f(yes);
@@ -118,9 +115,9 @@ void ASTOpConditional::Link( PROGRAM& program, const FLinkerOptions*)
         if (yes) args.yes = yes->linkedAddress;
         if (no) args.no = no->linkedAddress;
 
-        linkedAddress = (OP::ADDRESS)program.m_opAddress.size();
+        linkedAddress = (OP::ADDRESS)program.m_opAddress.Num();
         //program.m_code.push_back(op);
-        program.m_opAddress.push_back((uint32_t)program.m_byteCode.size());
+        program.m_opAddress.Add((uint32_t)program.m_byteCode.Num());
         AppendCode(program.m_byteCode,type);
         AppendCode(program.m_byteCode,args);
 
diff --git a/Engine/Plugins/Experimental/Mutable/Source/MutableTools/Private/MuT/ASTOpConditional.h b/Engine/Plugins/Experimental/Mutable/Source/MutableTools/Private/MuT/ASTOpConditional.h
--- a/Engine/Plugins/Experimental/Mutable/Source/MutableTools/Private/MuT/ASTOpConditional.h
+++ b/Engine/Plugins/Experimental/Mutable/Source/MutableTools/Private/MuT/ASTOpConditional.h
@@ -8,8 +8,6 @@
 #include "MuR/Ptr.h"
 #include "MuT/AST.h"
 
-#include <functional>
-
 
 namespace mu
 {
@@ -41,10 +39,10 @@ struct PROGRAM;
 		OP_TYPE GetOpType() const override { return type; }
 
 		bool IsEqual(const ASTOp& otherUntyped) const override;
-		Ptr<ASTOp> Clone(MapChildFunc& mapChild) const override;
+		Ptr<ASTOp> Clone(MapChildFuncRef mapChild) const override;
 		uint64 Hash() const override;
 		void Assert() override;
-		void ForEachChild(const std::function<void(ASTChild&)>& f) override;
+		void ForEachChild(const TFunctionRef<void(ASTChild&)> f) override;
 		void Link(PROGRAM& program, const FLinkerOptions* Options) override;
 		FImageDesc GetImageDesc(bool returnBestOption, class GetImageDescContext* context) override;
 		void GetLayoutBlockSize(int* pBlockX, int* pBlockY) override;
diff --git a/Engine/Plugins/Experimental/Mutable/Source/MutableTools/Private/MuT/ASTOpConstantBool.cpp b/Engine/Plugins/Experimental/Mutable/Source/MutableTools/Private/MuT/ASTOpConstantBool.cpp
--- a/Engine/Plugins/Experimental/Mutable/Source/MutableTools/Private/MuT/ASTOpConstantBool.cpp
+++ b/Engine/Plugins/Experimental/Mutable/Source/MutableTools/Private/MuT/ASTOpConstantBool.cpp
@@ -9,9 +9,6 @@
 #include "MuR/Types.h"
 #include "MuT/StreamsPrivate.h"
 
-#include <memory>
-#include <utility>
-
 
 using namespace mu;
 
@@ -25,7 +22,7 @@ ASTOpConstantBool::ASTOpConstantBool( bool v )
 }
 
 
-void ASTOpConstantBool::ForEachChild(const std::function<void(ASTChild&)>& )
+void ASTOpConstantBool::ForEachChild(const TFunctionRef<void(ASTChild&)> )
 {
 }
 
@@ -48,7 +45,7 @@ uint64 ASTOpConstantBool::Hash() const
 }
 
 
-mu::Ptr<ASTOp> ASTOpConstantBool::Clone(MapChildFunc&) const
+mu::Ptr<ASTOp> ASTOpConstantBool::Clone(MapChildFuncRef) const
 {
     Ptr<ASTOpConstantBool> n = new ASTOpConstantBool();
     n->value = value;
@@ -64,9 +61,9 @@ void ASTOpConstantBool::Link( PROGRAM& program, const FLinkerOptions*)
         memset( &args,0, sizeof(args) );
         args.value = value;
 
-        linkedAddress = (OP::ADDRESS)program.m_opAddress.size();
+        linkedAddress = (OP::ADDRESS)program.m_opAddress.Num();
 //        program.m_code.push_back(op);
-        program.m_opAddress.push_back((uint32_t)program.m_byteCode.size());
+        program.m_opAddress.Add((uint32_t)program.m_byteCode.Num());
         AppendCode(program.m_byteCode,OP_TYPE::BO_CONSTANT);
         AppendCode(program.m_byteCode,args);
     }
diff --git a/Engine/Plugins/Experimental/Mutable/Source/MutableTools/Private/MuT/ASTOpConstantBool.h b/Engine/Plugins/Experimental/Mutable/Source/MutableTools/Private/MuT/ASTOpConstantBool.h
--- a/Engine/Plugins/Experimental/Mutable/Source/MutableTools/Private/MuT/ASTOpConstantBool.h
+++ b/Engine/Plugins/Experimental/Mutable/Source/MutableTools/Private/MuT/ASTOpConstantBool.h
@@ -7,8 +7,6 @@
 #include "MuR/Ptr.h"
 #include "MuT/AST.h"
 
-#include <functional>
-
 
 namespace mu
 {
@@ -31,9 +29,9 @@ struct PROGRAM;
 
 		OP_TYPE GetOpType() const override { return OP_TYPE::BO_CONSTANT; }
 		uint64 Hash() const override;
-		void ForEachChild(const std::function<void(ASTChild&)>&) override;
+		void ForEachChild(const TFunctionRef<void(ASTChild&)>) override;
 		bool IsEqual(const ASTOp& otherUntyped) const override;
-		Ptr<ASTOp> Clone(MapChildFunc& mapChild) const override;
+		Ptr<ASTOp> Clone(MapChildFuncRef mapChild) const override;
 		void Link(PROGRAM& program, const FLinkerOptions* Options) override;
 		BOOL_EVAL_RESULT EvaluateBool(ASTOpList& facts, EVALUATE_BOOL_CACHE* cache) const override;
 	};
diff --git a/Engine/Plugins/Experimental/Mutable/Source/MutableTools/Private/MuT/ASTOpConstantProjector.cpp b/Engine/Plugins/Experimental/Mutable/Source/MutableTools/Private/MuT/ASTOpConstantProjector.cpp
--- a/Engine/Plugins/Experimental/Mutable/Source/MutableTools/Private/MuT/ASTOpConstantProjector.cpp
+++ b/Engine/Plugins/Experimental/Mutable/Source/MutableTools/Private/MuT/ASTOpConstantProjector.cpp
@@ -8,9 +8,6 @@
 #include "MuR/RefCounted.h"
 #include "MuT/StreamsPrivate.h"
 
-#include <memory>
-#include <utility>
-
 
 namespace mu
 {
@@ -19,7 +16,7 @@ namespace mu
 	//-------------------------------------------------------------------------------------------------
 	//-------------------------------------------------------------------------------------------------
 	//-------------------------------------------------------------------------------------------------
-	void ASTOpConstantProjector::ForEachChild(const std::function<void(ASTChild&)>&)
+	void ASTOpConstantProjector::ForEachChild(const TFunctionRef<void(ASTChild&)>)
 	{
 	}
 
@@ -42,7 +39,7 @@ namespace mu
 	}
 
 
-	mu::Ptr<ASTOp> ASTOpConstantProjector::Clone(MapChildFunc&) const
+	mu::Ptr<ASTOp> ASTOpConstantProjector::Clone(MapChildFuncRef) const
 	{
 		Ptr<ASTOpConstantProjector> n = new ASTOpConstantProjector();
 		n->value = value;
@@ -58,9 +55,9 @@ namespace mu
 			memset(&args, 0, sizeof(args));
 			args.value = program.AddConstant(value);
 
-			linkedAddress = (OP::ADDRESS)program.m_opAddress.size();
+			linkedAddress = (OP::ADDRESS)program.m_opAddress.Num();
 			//program.m_code.push_back(op);
-			program.m_opAddress.push_back((uint32_t)program.m_byteCode.size());
+			program.m_opAddress.Add((uint32_t)program.m_byteCode.Num());
 			AppendCode(program.m_byteCode, OP_TYPE::PR_CONSTANT);
 			AppendCode(program.m_byteCode, args);
 		}
diff --git a/Engine/Plugins/Experimental/Mutable/Source/MutableTools/Private/MuT/ASTOpConstantProjector.h b/Engine/Plugins/Experimental/Mutable/Source/MutableTools/Private/MuT/ASTOpConstantProjector.h
--- a/Engine/Plugins/Experimental/Mutable/Source/MutableTools/Private/MuT/ASTOpConstantProjector.h
+++ b/Engine/Plugins/Experimental/Mutable/Source/MutableTools/Private/MuT/ASTOpConstantProjector.h
@@ -8,8 +8,6 @@
 #include "MuR/Ptr.h"
 #include "MuT/AST.h"
 
-#include <functional>
-
 
 namespace mu
 {
@@ -30,9 +28,9 @@ struct PROGRAM;
 
 		OP_TYPE GetOpType() const override { return OP_TYPE::PR_CONSTANT; }
 		uint64 Hash() const override;
-		void ForEachChild(const std::function<void(ASTChild&)>&) override;
+		void ForEachChild(const TFunctionRef<void(ASTChild&)>) override;
 		bool IsEqual(const ASTOp& otherUntyped) const override;
-		Ptr<ASTOp> Clone(MapChildFunc& mapChild) const override;
+		Ptr<ASTOp> Clone(MapChildFuncRef mapChild) const override;
 		void Link(PROGRAM& program, const FLinkerOptions*) override;
 	};
 
diff --git a/Engine/Plugins/Experimental/Mutable/Source/MutableTools/Private/MuT/ASTOpConstantResource.cpp b/Engine/Plugins/Experimental/Mutable/Source/MutableTools/Private/MuT/ASTOpConstantResource.cpp
--- a/Engine/Plugins/Experimental/Mutable/Source/MutableTools/Private/MuT/ASTOpConstantResource.cpp
+++ b/Engine/Plugins/Experimental/Mutable/Source/MutableTools/Private/MuT/ASTOpConstantResource.cpp
@@ -27,7 +27,7 @@ using namespace mu;
 //-------------------------------------------------------------------------------------------------
 //-------------------------------------------------------------------------------------------------
 //-------------------------------------------------------------------------------------------------
-void ASTOpConstantResource::ForEachChild(const std::function<void (ASTChild&)>& )
+void ASTOpConstantResource::ForEachChild(const TFunctionRef<void (ASTChild&)> )
 {
 }
 
@@ -45,7 +45,7 @@ bool ASTOpConstantResource::IsEqual(const ASTOp& otherUntyped) const
 
 
 //-------------------------------------------------------------------------------------------------
-mu::Ptr<ASTOp> ASTOpConstantResource::Clone(MapChildFunc&) const
+mu::Ptr<ASTOp> ASTOpConstantResource::Clone(MapChildFuncRef) const
 {
     Ptr<ASTOpConstantResource> n = new ASTOpConstantResource();
     n->type = type;
@@ -115,8 +115,8 @@ void ASTOpConstantResource::Link( PROGRAM& program, const FLinkerOptions* Option
 
             args.value = program.AddConstant(pTyped.get());
 
-            linkedAddress = (OP::ADDRESS)program.m_opAddress.size();
-            program.m_opAddress.push_back((uint32_t)program.m_byteCode.size());
+            linkedAddress = (OP::ADDRESS)program.m_opAddress.Num();
+            program.m_opAddress.Add((uint32_t)program.m_byteCode.Num());
             AppendCode(program.m_byteCode,type);
             AppendCode(program.m_byteCode,args);
         }
@@ -160,8 +160,8 @@ void ASTOpConstantResource::Link( PROGRAM& program, const FLinkerOptions* Option
 
 			if (bValidData)
 			{
-				linkedAddress = (OP::ADDRESS)program.m_opAddress.size();
-				program.m_opAddress.push_back((uint32_t)program.m_byteCode.size());
+				linkedAddress = (OP::ADDRESS)program.m_opAddress.Num();
+				program.m_opAddress.Add((uint32_t)program.m_byteCode.Num());
 				AppendCode(program.m_byteCode, type);
 				AppendCode(program.m_byteCode, args);
 			}
diff --git a/Engine/Plugins/Experimental/Mutable/Source/MutableTools/Private/MuT/ASTOpConstantResource.h b/Engine/Plugins/Experimental/Mutable/Source/MutableTools/Private/MuT/ASTOpConstantResource.h
--- a/Engine/Plugins/Experimental/Mutable/Source/MutableTools/Private/MuT/ASTOpConstantResource.h
+++ b/Engine/Plugins/Experimental/Mutable/Source/MutableTools/Private/MuT/ASTOpConstantResource.h
@@ -9,8 +9,6 @@
 #include "MuR/RefCounted.h"
 #include "MuT/AST.h"
 
-#include <functional>
-
 
 namespace mu
 {
@@ -57,15 +55,14 @@ template <class SCALAR> class vec4;
 		Ptr<const RefCounted> GetValue() const;
 
 		//! Set the value to store in this op
-		void SetValue(const Ptr<const RefCounted>& v,
-			bool useDiskCache);
+		void SetValue(const Ptr<const RefCounted>& v, bool useDiskCache);
 
 
 		// ASTOp interface
 		OP_TYPE GetOpType() const override { return type; }
-		void ForEachChild(const std::function<void(ASTChild&)>&) override;
+		void ForEachChild(const TFunctionRef<void(ASTChild&)>) override;
 		bool IsEqual(const ASTOp& otherUntyped) const override;
-		Ptr<ASTOp> Clone(MapChildFunc& mapChild) const override;
+		Ptr<ASTOp> Clone(MapChildFuncRef mapChild) const override;
 		uint64 Hash() const override;
 		void Link(PROGRAM& program, const FLinkerOptions*) override;
 		FImageDesc GetImageDesc(bool, class GetImageDescContext*) override;
diff --git a/Engine/Plugins/Experimental/Mutable/Source/MutableTools/Private/MuT/ASTOpConstantString.cpp b/Engine/Plugins/Experimental/Mutable/Source/MutableTools/Private/MuT/ASTOpConstantString.cpp
--- a/Engine/Plugins/Experimental/Mutable/Source/MutableTools/Private/MuT/ASTOpConstantString.cpp
+++ b/Engine/Plugins/Experimental/Mutable/Source/MutableTools/Private/MuT/ASTOpConstantString.cpp
@@ -7,9 +7,7 @@
 #include "MuR/RefCounted.h"
 #include "MuR/Types.h"
 
-#include <memory>
 #include <string>
-#include <utility>
 
 
 using namespace mu;
@@ -18,7 +16,7 @@ using namespace mu;
 //-------------------------------------------------------------------------------------------------
 //-------------------------------------------------------------------------------------------------
 //-------------------------------------------------------------------------------------------------
-void ASTOpConstantString::ForEachChild(const std::function<void(ASTChild&)>& )
+void ASTOpConstantString::ForEachChild(const TFunctionRef<void(ASTChild&)> )
 {
 }
 
@@ -40,7 +38,7 @@ bool ASTOpConstantString::IsEqual(const ASTOp& otherUntyped) const
 }
 
 
-mu::Ptr<ASTOp> ASTOpConstantString::Clone(MapChildFunc&) const
+mu::Ptr<ASTOp> ASTOpConstantString::Clone(MapChildFuncRef) const
 {
     Ptr<ASTOpConstantString> n = new ASTOpConstantString();
     n->value = value;
@@ -56,8 +54,8 @@ void ASTOpConstantString::Link( PROGRAM& program, const FLinkerOptions* )
         memset( &args,0, sizeof(args) );
         args.value = program.AddConstant(value);
 
-        linkedAddress = (OP::ADDRESS)program.m_opAddress.size();
-        program.m_opAddress.push_back((uint32_t)program.m_byteCode.size());
+        linkedAddress = (OP::ADDRESS)program.m_opAddress.Num();
+        program.m_opAddress.Add((uint32_t)program.m_byteCode.Num());
         AppendCode(program.m_byteCode,OP_TYPE::ST_CONSTANT);
         AppendCode(program.m_byteCode,args);
     }
diff --git a/Engine/Plugins/Experimental/Mutable/Source/MutableTools/Private/MuT/ASTOpImageCompose.cpp b/Engine/Plugins/Experimental/Mutable/Source/MutableTools/Private/MuT/ASTOpImageCompose.cpp
--- a/Engine/Plugins/Experimental/Mutable/Source/MutableTools/Private/MuT/ASTOpImageCompose.cpp
+++ b/Engine/Plugins/Experimental/Mutable/Source/MutableTools/Private/MuT/ASTOpImageCompose.cpp
@@ -72,7 +72,7 @@ uint64 ASTOpImageCompose::Hash() const
 
 
 //-------------------------------------------------------------------------------------------------
-mu::Ptr<ASTOp> ASTOpImageCompose::Clone(MapChildFunc& mapChild) const
+mu::Ptr<ASTOp> ASTOpImageCompose::Clone(MapChildFuncRef mapChild) const
 {
 	mu::Ptr<ASTOpImageCompose> n = new ASTOpImageCompose();
     n->Layout = mapChild(Layout.child());
@@ -85,7 +85,7 @@ mu::Ptr<ASTOp> ASTOpImageCompose::Clone(MapChildFunc& mapChild) const
 
 
 //-------------------------------------------------------------------------------------------------
-void ASTOpImageCompose::ForEachChild(const std::function<void(ASTChild&)>& f )
+void ASTOpImageCompose::ForEachChild(const TFunctionRef<void(ASTChild&)> f )
 {
 	f(Layout);
 	f(Base);
@@ -109,8 +109,8 @@ void ASTOpImageCompose::Link( PROGRAM& program, const FLinkerOptions*)
 		if (Mask) args.mask = Mask->linkedAddress;
         args.blockIndex = BlockIndex;
 
-        linkedAddress = (OP::ADDRESS)program.m_opAddress.size();
-        program.m_opAddress.push_back((uint32_t)program.m_byteCode.size());
+        linkedAddress = (OP::ADDRESS)program.m_opAddress.Num();
+        program.m_opAddress.Add((uint32_t)program.m_byteCode.Num());
         AppendCode(program.m_byteCode,OP_TYPE::IM_COMPOSE);
         AppendCode(program.m_byteCode,args);
     }
@@ -237,7 +237,7 @@ mu::Ptr<ASTOp> ASTOpImageCompose::OptimiseSemantic(const MODEL_OPTIMIZATION_OPTI
 		// Constant single-block full layout?
 		if (pLayout->GetBlockCount() == 1
 			&&
-			pLayout->m_blocks[0].m_min == vec2<uint16_t>(0, 0)
+			pLayout->m_blocks[0].m_min == vec2<uint16>(0, 0)
 			&&
 			pLayout->m_blocks[0].m_size == pLayout->m_size
 			&&
diff --git a/Engine/Plugins/Experimental/Mutable/Source/MutableTools/Private/MuT/ASTOpImageCompose.h b/Engine/Plugins/Experimental/Mutable/Source/MutableTools/Private/MuT/ASTOpImageCompose.h
--- a/Engine/Plugins/Experimental/Mutable/Source/MutableTools/Private/MuT/ASTOpImageCompose.h
+++ b/Engine/Plugins/Experimental/Mutable/Source/MutableTools/Private/MuT/ASTOpImageCompose.h
@@ -8,8 +8,6 @@
 #include "MuR/Ptr.h"
 #include "MuT/AST.h"
 
-#include <functional>
-
 
 namespace mu
 {
@@ -35,8 +33,8 @@ template <class SCALAR> class vec4;
 		OP_TYPE GetOpType() const override { return OP_TYPE::IM_COMPOSE; }
 		uint64 Hash() const override;
 		bool IsEqual(const ASTOp& otherUntyped) const override;
-		Ptr<ASTOp> Clone(MapChildFunc& mapChild) const override;
-		void ForEachChild(const std::function<void(ASTChild&)>&) override;
+		Ptr<ASTOp> Clone(MapChildFuncRef mapChild) const override;
+		void ForEachChild(const TFunctionRef<void(ASTChild&)>) override;
 		void Link(PROGRAM& program, const FLinkerOptions* Options) override;
 		//TODO: Ptr<ASTOp> OptimiseSink(const MODEL_OPTIMIZATION_OPTIONS& options, OPTIMIZE_SINK_CONTEXT& context) const override;
 		Ptr<ASTOp> OptimiseSemantic(const MODEL_OPTIMIZATION_OPTIONS& options) const;
diff --git a/Engine/Plugins/Experimental/Mutable/Source/MutableTools/Private/MuT/ASTOpImageMipmap.cpp b/Engine/Plugins/Experimental/Mutable/Source/MutableTools/Private/MuT/ASTOpImageMipmap.cpp
--- a/Engine/Plugins/Experimental/Mutable/Source/MutableTools/Private/MuT/ASTOpImageMipmap.cpp
+++ b/Engine/Plugins/Experimental/Mutable/Source/MutableTools/Private/MuT/ASTOpImageMipmap.cpp
@@ -67,7 +67,7 @@ uint64 ASTOpImageMipmap::Hash() const
 }
 
 
-mu::Ptr<ASTOp> ASTOpImageMipmap::Clone(MapChildFunc& mapChild) const
+mu::Ptr<ASTOp> ASTOpImageMipmap::Clone(MapChildFuncRef mapChild) const
 {
 	mu::Ptr<ASTOpImageMipmap> n = new ASTOpImageMipmap();
     n->Source = mapChild(Source.child());
@@ -83,7 +83,7 @@ mu::Ptr<ASTOp> ASTOpImageMipmap::Clone(MapChildFunc& mapChild) const
 }
 
 
-void ASTOpImageMipmap::ForEachChild(const std::function<void(ASTChild&)>& f )
+void ASTOpImageMipmap::ForEachChild(const TFunctionRef<void(ASTChild&)> f )
 {
     f( Source );
 }
@@ -106,9 +106,9 @@ void ASTOpImageMipmap::Link( PROGRAM& program, const FLinkerOptions* )
 		args.ditherMipmapAlpha = DitherMipmapAlpha;
 		if (Source) args.source = Source->linkedAddress;
 
-        linkedAddress = (OP::ADDRESS)program.m_opAddress.size();
+        linkedAddress = (OP::ADDRESS)program.m_opAddress.Num();
         //program.m_code.push_back(op);
-        program.m_opAddress.push_back((uint32_t)program.m_byteCode.size());
+        program.m_opAddress.Add((uint32_t)program.m_byteCode.Num());
         AppendCode(program.m_byteCode,OP_TYPE::IM_MIPMAP);
         AppendCode(program.m_byteCode,args);
     }
@@ -260,8 +260,8 @@ mu::Ptr<ASTOp> Sink_ImageMipmapAST::Apply(const ASTOp* root)
 	check(root->GetOpType() == OP_TYPE::IM_MIPMAP);
 
 	m_root = dynamic_cast<const ASTOpImageMipmap*>(root);
-	m_oldToNew.clear();
-	m_newOps.clear();
+	m_oldToNew.Empty();
+	m_newOps.Empty();
 
 	if (m_root->bOnlyTail)
 	{
@@ -321,16 +321,16 @@ mu::Ptr<ASTOp> Sink_ImageMipmapAST::Visit(mu::Ptr<ASTOp> at, const ASTOpImageMip
 	if (!at) return nullptr;
 
 	// Newly created?
-	if (std::find(m_newOps.begin(), m_newOps.end(), at) != m_newOps.end())
+	if (m_newOps.Contains(at))
 	{
 		return at;
 	}
 
 	// Already visited?
-	auto cacheIt = m_oldToNew.find(at);
-	if (cacheIt != m_oldToNew.end())
+	auto cacheIt = m_oldToNew.Find(at);
+	if (cacheIt)
 	{
-		return cacheIt->second;
+		return *cacheIt;
 	}
 
 	mu::Ptr<ASTOp> newAt = at;
@@ -400,10 +400,10 @@ mu::Ptr<ASTOp> Sink_ImageMipmapAST::Visit(mu::Ptr<ASTOp> at, const ASTOpImageMip
 
 			uint8_t blockLevels = 0;
 			{
-				uint16_t minX = typedSource->location[0];
-				uint16_t minY = typedSource->location[1];
-				uint16_t sizeX = patchDesc.m_size[0];
-				uint16_t sizeY = patchDesc.m_size[1];
+				uint16 minX = typedSource->location[0];
+				uint16 minY = typedSource->location[1];
+				uint16 sizeX = patchDesc.m_size[0];
+				uint16 sizeY = patchDesc.m_size[1];
 				while (minX && minY && sizeX && sizeY
 					&&
 					minX % 2 == 0 && minY % 2 == 0 && sizeX % 2 == 0 && sizeY % 2 == 0)
@@ -453,7 +453,7 @@ mu::Ptr<ASTOp> Sink_ImageMipmapAST::Visit(mu::Ptr<ASTOp> at, const ASTOpImageMip
 		newAt = newOp;
 	}
 
-	m_oldToNew[at] = newAt;
+	m_oldToNew.Add(at, newAt);
 
 	return newAt;
 }
diff --git a/Engine/Plugins/Experimental/Mutable/Source/MutableTools/Private/MuT/ASTOpImageMipmap.h b/Engine/Plugins/Experimental/Mutable/Source/MutableTools/Private/MuT/ASTOpImageMipmap.h
--- a/Engine/Plugins/Experimental/Mutable/Source/MutableTools/Private/MuT/ASTOpImageMipmap.h
+++ b/Engine/Plugins/Experimental/Mutable/Source/MutableTools/Private/MuT/ASTOpImageMipmap.h
@@ -8,8 +8,6 @@
 #include "MuR/Ptr.h"
 #include "MuT/AST.h"
 
-#include <functional>
-
 
 namespace mu
 {
@@ -46,8 +44,8 @@ template <class SCALAR> class vec4;
 		OP_TYPE GetOpType() const override { return OP_TYPE::IM_MIPMAP; }
 		uint64 Hash() const override;
 		bool IsEqual(const ASTOp& otherUntyped) const override;
-		Ptr<ASTOp> Clone(MapChildFunc& mapChild) const override;
-		void ForEachChild(const std::function<void(ASTChild&)>&) override;
+		Ptr<ASTOp> Clone(MapChildFuncRef mapChild) const override;
+		void ForEachChild(const TFunctionRef<void(ASTChild&)>) override;
 		void Link(PROGRAM& program, const FLinkerOptions* Options) override;
 		Ptr<ASTOp> OptimiseSink(const MODEL_OPTIMIZATION_OPTIONS& options, OPTIMIZE_SINK_CONTEXT& context) const override;
 		FImageDesc GetImageDesc(bool returnBestOption, GetImageDescContext* context) override;
diff --git a/Engine/Plugins/Experimental/Mutable/Source/MutableTools/Private/MuT/ASTOpImageMultiLayer.cpp b/Engine/Plugins/Experimental/Mutable/Source/MutableTools/Private/MuT/ASTOpImageMultiLayer.cpp
--- a/Engine/Plugins/Experimental/Mutable/Source/MutableTools/Private/MuT/ASTOpImageMultiLayer.cpp
+++ b/Engine/Plugins/Experimental/Mutable/Source/MutableTools/Private/MuT/ASTOpImageMultiLayer.cpp
@@ -65,7 +65,7 @@ uint64 ASTOpImageMultiLayer::Hash() const
 
 
 //-------------------------------------------------------------------------------------------------
-mu::Ptr<ASTOp> ASTOpImageMultiLayer::Clone(MapChildFunc& mapChild) const
+mu::Ptr<ASTOp> ASTOpImageMultiLayer::Clone(MapChildFuncRef mapChild) const
 {
     Ptr<ASTOpImageMultiLayer> n = new ASTOpImageMultiLayer();
     n->base = mapChild(base.child());
@@ -80,7 +80,7 @@ mu::Ptr<ASTOp> ASTOpImageMultiLayer::Clone(MapChildFunc& mapChild) const
 
 
 //-------------------------------------------------------------------------------------------------
-void ASTOpImageMultiLayer::ForEachChild( const std::function<void(ASTChild&)>& f )
+void ASTOpImageMultiLayer::ForEachChild( const TFunctionRef<void(ASTChild&)> f )
 {
     f( base );
     f( blend );
@@ -98,7 +98,7 @@ void ASTOpImageMultiLayer::Link( PROGRAM& program, const FLinkerOptions*)
         OP::ImageMultiLayerArgs args;
         memset( &args,0, sizeof(args) );
 
-        args.blendType = (uint16_t)blendType;
+        args.blendType = (uint16)blendType;
 
         if (base) args.base = base->linkedAddress;
         if (blend) args.blended = blend->linkedAddress;
@@ -108,8 +108,8 @@ void ASTOpImageMultiLayer::Link( PROGRAM& program, const FLinkerOptions*)
             LinkRange( program, range, args.rangeSize, args.rangeId );
         }
 
-        linkedAddress = (OP::ADDRESS)program.m_opAddress.size();
-        program.m_opAddress.push_back((uint32_t)program.m_byteCode.size());
+        linkedAddress = (OP::ADDRESS)program.m_opAddress.Num();
+        program.m_opAddress.Add((uint32_t)program.m_byteCode.Num());
         AppendCode(program.m_byteCode,OP_TYPE::IM_MULTILAYER);
         AppendCode(program.m_byteCode,args);
     }
diff --git a/Engine/Plugins/Experimental/Mutable/Source/MutableTools/Private/MuT/ASTOpImageNormalComposite.cpp b/Engine/Plugins/Experimental/Mutable/Source/MutableTools/Private/MuT/ASTOpImageNormalComposite.cpp
--- a/Engine/Plugins/Experimental/Mutable/Source/MutableTools/Private/MuT/ASTOpImageNormalComposite.cpp
+++ b/Engine/Plugins/Experimental/Mutable/Source/MutableTools/Private/MuT/ASTOpImageNormalComposite.cpp
@@ -11,8 +11,6 @@
 #include "MuR/Types.h"
 #include "MuT/StreamsPrivate.h"
 
-#include <memory>
-#include <utility>
 
 using namespace mu;
 
@@ -53,7 +51,7 @@ uint64 ASTOpImageNormalComposite::Hash() const
 }
 
 
-mu::Ptr<ASTOp> ASTOpImageNormalComposite::Clone(MapChildFunc& mapChild) const
+mu::Ptr<ASTOp> ASTOpImageNormalComposite::Clone(MapChildFuncRef mapChild) const
 {
     Ptr<ASTOpImageNormalComposite> N = new ASTOpImageNormalComposite();
 	N->Base = mapChild(Base.child());
@@ -65,7 +63,7 @@ mu::Ptr<ASTOp> ASTOpImageNormalComposite::Clone(MapChildFunc& mapChild) const
 }
 
 
-void ASTOpImageNormalComposite::ForEachChild(const std::function<void(ASTChild&)>& f )
+void ASTOpImageNormalComposite::ForEachChild(const TFunctionRef<void(ASTChild&)> f )
 {
 	f( Base );
 	f( Normal );
@@ -93,8 +91,8 @@ void ASTOpImageNormalComposite::Link( PROGRAM& program, const FLinkerOptions* )
 		args.power = Power;
 		args.mode = Mode;
 
-        linkedAddress = (OP::ADDRESS)program.m_opAddress.size();
-        program.m_opAddress.push_back((uint32_t)program.m_byteCode.size());
+        linkedAddress = (OP::ADDRESS)program.m_opAddress.Num();
+        program.m_opAddress.Add((uint32_t)program.m_byteCode.Num());
         AppendCode(program.m_byteCode,OP_TYPE::IM_NORMALCOMPOSITE);
         AppendCode(program.m_byteCode,args);
     }
diff --git a/Engine/Plugins/Experimental/Mutable/Source/MutableTools/Private/MuT/ASTOpImagePatch.cpp b/Engine/Plugins/Experimental/Mutable/Source/MutableTools/Private/MuT/ASTOpImagePatch.cpp
--- a/Engine/Plugins/Experimental/Mutable/Source/MutableTools/Private/MuT/ASTOpImagePatch.cpp
+++ b/Engine/Plugins/Experimental/Mutable/Source/MutableTools/Private/MuT/ASTOpImagePatch.cpp
@@ -9,9 +9,6 @@
 #include "MuR/Types.h"
 #include "MuT/StreamsPrivate.h"
 
-#include <memory>
-#include <utility>
-
 
 using namespace mu;
 
@@ -58,7 +55,7 @@ uint64 ASTOpImagePatch::Hash() const
 
 
 //-------------------------------------------------------------------------------------------------
-mu::Ptr<ASTOp> ASTOpImagePatch::Clone(MapChildFunc& mapChild) const
+mu::Ptr<ASTOp> ASTOpImagePatch::Clone(MapChildFuncRef mapChild) const
 {
     Ptr<ASTOpImagePatch> n = new ASTOpImagePatch();
     n->base = mapChild(base.child());
@@ -69,7 +66,7 @@ mu::Ptr<ASTOp> ASTOpImagePatch::Clone(MapChildFunc& mapChild) const
 
 
 //-------------------------------------------------------------------------------------------------
-void ASTOpImagePatch::ForEachChild(const std::function<void(ASTChild&)>& f )
+void ASTOpImagePatch::ForEachChild(const TFunctionRef<void(ASTChild&)> f )
 {
     f( base );
     f( patch );
@@ -90,8 +87,8 @@ void ASTOpImagePatch::Link( PROGRAM& program, const FLinkerOptions*)
         args.minX = location[0];
         args.minY = location[1];
 
-        linkedAddress = (OP::ADDRESS)program.m_opAddress.size();
-        program.m_opAddress.push_back((uint32_t)program.m_byteCode.size());
+        linkedAddress = (OP::ADDRESS)program.m_opAddress.Num();
+        program.m_opAddress.Add((uint32_t)program.m_byteCode.Num());
         AppendCode(program.m_byteCode,OP_TYPE::IM_PATCH);
         AppendCode(program.m_byteCode,args);
     }
diff --git a/Engine/Plugins/Experimental/Mutable/Source/MutableTools/Private/MuT/ASTOpImagePatch.h b/Engine/Plugins/Experimental/Mutable/Source/MutableTools/Private/MuT/ASTOpImagePatch.h
--- a/Engine/Plugins/Experimental/Mutable/Source/MutableTools/Private/MuT/ASTOpImagePatch.h
+++ b/Engine/Plugins/Experimental/Mutable/Source/MutableTools/Private/MuT/ASTOpImagePatch.h
@@ -9,8 +9,6 @@
 #include "MuR/Ptr.h"
 #include "MuT/AST.h"
 
-#include <functional>
-
 
 namespace mu
 {
@@ -22,7 +20,7 @@ struct PROGRAM;
 
 		ASTChild base;
 		ASTChild patch;
-		vec2<uint16_t> location;
+		vec2<uint16> location;
 
 	public:
 
@@ -33,8 +31,8 @@ struct PROGRAM;
 		OP_TYPE GetOpType() const override { return OP_TYPE::IM_PATCH; }
 		uint64 Hash() const override;
 		bool IsEqual(const ASTOp& otherUntyped) const override;
-		Ptr<ASTOp> Clone(MapChildFunc& mapChild) const override;
-		void ForEachChild(const std::function<void(ASTChild&)>&) override;
+		Ptr<ASTOp> Clone(MapChildFuncRef mapChild) const override;
+		void ForEachChild(const TFunctionRef<void(ASTChild&)>) override;
 		void Link(PROGRAM& program, const FLinkerOptions* Options) override;
 		FImageDesc GetImageDesc(bool returnBestOption, GetImageDescContext* context) override;
 		Ptr<ImageSizeExpression> GetImageSizeExpression() const override;
diff --git a/Engine/Plugins/Experimental/Mutable/Source/MutableTools/Private/MuT/ASTOpImagePixelFormat.cpp b/Engine/Plugins/Experimental/Mutable/Source/MutableTools/Private/MuT/ASTOpImagePixelFormat.cpp
--- a/Engine/Plugins/Experimental/Mutable/Source/MutableTools/Private/MuT/ASTOpImagePixelFormat.cpp
+++ b/Engine/Plugins/Experimental/Mutable/Source/MutableTools/Private/MuT/ASTOpImagePixelFormat.cpp
@@ -61,7 +61,7 @@ uint64 ASTOpImagePixelFormat::Hash() const
 }
 
 
-mu::Ptr<ASTOp> ASTOpImagePixelFormat::Clone(MapChildFunc& mapChild) const
+mu::Ptr<ASTOp> ASTOpImagePixelFormat::Clone(MapChildFuncRef mapChild) const
 {
 	mu::Ptr<ASTOpImagePixelFormat> n = new ASTOpImagePixelFormat();
     n->Source = mapChild(Source.child());
@@ -71,7 +71,7 @@ mu::Ptr<ASTOp> ASTOpImagePixelFormat::Clone(MapChildFunc& mapChild) const
 }
 
 
-void ASTOpImagePixelFormat::ForEachChild(const std::function<void(ASTChild&)>& f )
+void ASTOpImagePixelFormat::ForEachChild(const TFunctionRef<void(ASTChild&)> f )
 {
     f( Source );
 }
@@ -89,9 +89,9 @@ void ASTOpImagePixelFormat::Link( PROGRAM& program, const FLinkerOptions* )
 		args.formatIfAlpha = FormatIfAlpha;
 		if (Source) args.source = Source->linkedAddress;
 
-        linkedAddress = (OP::ADDRESS)program.m_opAddress.size();
+        linkedAddress = (OP::ADDRESS)program.m_opAddress.Num();
         //program.m_code.push_back(op);
-        program.m_opAddress.push_back((uint32_t)program.m_byteCode.size());
+        program.m_opAddress.Add((uint32_t)program.m_byteCode.Num());
         AppendCode(program.m_byteCode,OP_TYPE::IM_PIXELFORMAT);
         AppendCode(program.m_byteCode,args);
     }
@@ -279,8 +279,8 @@ mu::Ptr<ImageSizeExpression> ASTOpImagePixelFormat::GetImageSizeExpression() con
 mu::Ptr<ASTOp> Sink_ImagePixelFormatAST::Apply(const ASTOp* root)
 {
 	m_root = dynamic_cast<const ASTOpImagePixelFormat*>(root);
-	m_oldToNew.clear();
-	m_newOps.clear();
+	m_oldToNew.Empty();
+	m_newOps.Empty();
 
 	check(root->GetOpType() == OP_TYPE::IM_PIXELFORMAT);
 
@@ -305,7 +305,7 @@ mu::Ptr<ASTOp> Sink_ImagePixelFormatAST::Visit(mu::Ptr<ASTOp> at, const ASTOpIma
 	if (!at) return nullptr;
 
 	// Newly created?
-	if (std::find(m_newOps.begin(), m_newOps.end(), at) != m_newOps.end())
+	if (m_newOps.Find(at)!=INDEX_NONE)
 	{
 		return at;
 	}
@@ -315,10 +315,13 @@ mu::Ptr<ASTOp> Sink_ImagePixelFormatAST::Visit(mu::Ptr<ASTOp> at, const ASTOpIma
 	bool isBlockFormat = GetImageFormatData(format).m_pixelsPerBlockX != 0;
 
 	// Already visited?
-	auto cacheIt = m_oldToNew[currentFormatOp].find(at);
-	if (cacheIt != m_oldToNew[currentFormatOp].end())
 	{
-		return cacheIt->second;
+		auto& CurrentSet = m_oldToNew.FindOrAdd(currentFormatOp);
+		auto cacheIt = CurrentSet.find(at);
+		if (cacheIt != CurrentSet.end())
+		{
+			return cacheIt->second;
+		}
 	}
 
 	mu::Ptr<ASTOp> newAt = at;
diff --git a/Engine/Plugins/Experimental/Mutable/Source/MutableTools/Private/MuT/ASTOpImageTransform.cpp b/Engine/Plugins/Experimental/Mutable/Source/MutableTools/Private/MuT/ASTOpImageTransform.cpp
--- a/Engine/Plugins/Experimental/Mutable/Source/MutableTools/Private/MuT/ASTOpImageTransform.cpp
+++ b/Engine/Plugins/Experimental/Mutable/Source/MutableTools/Private/MuT/ASTOpImageTransform.cpp
@@ -70,7 +70,7 @@ uint64 ASTOpImageTransform::Hash() const
 
 
 //-------------------------------------------------------------------------------------------------
-mu::Ptr<ASTOp> ASTOpImageTransform::Clone(MapChildFunc& mapChild) const
+mu::Ptr<ASTOp> ASTOpImageTransform::Clone(MapChildFuncRef mapChild) const
 {
     Ptr<ASTOpImageTransform> n = new ASTOpImageTransform();
     n->base = mapChild(base.child());
@@ -84,7 +84,7 @@ mu::Ptr<ASTOp> ASTOpImageTransform::Clone(MapChildFunc& mapChild) const
 
 
 //-------------------------------------------------------------------------------------------------
-void ASTOpImageTransform::ForEachChild( const std::function<void(ASTChild&)>& f )
+void ASTOpImageTransform::ForEachChild( const TFunctionRef<void(ASTChild&)> f )
 {
     f( base );
     f( offsetX );
@@ -110,8 +110,8 @@ void ASTOpImageTransform::Link( PROGRAM& program, const FLinkerOptions*)
 		Args.scaleY = scaleY ? scaleY->linkedAddress : 0;
 		Args.rotation = rotation ? rotation->linkedAddress : 0;
 
-        linkedAddress = (OP::ADDRESS)program.m_opAddress.size();
-        program.m_opAddress.push_back((uint32_t)program.m_byteCode.size());
+        linkedAddress = (OP::ADDRESS)program.m_opAddress.Num();
+        program.m_opAddress.Add((uint32_t)program.m_byteCode.Num());
         AppendCode(program.m_byteCode, OP_TYPE::IM_TRANSFORM);
 		AppendCode(program.m_byteCode, Args);
     }
diff --git a/Engine/Plugins/Experimental/Mutable/Source/MutableTools/Private/MuT/ASTOpInstanceAdd.cpp b/Engine/Plugins/Experimental/Mutable/Source/MutableTools/Private/MuT/ASTOpInstanceAdd.cpp
--- a/Engine/Plugins/Experimental/Mutable/Source/MutableTools/Private/MuT/ASTOpInstanceAdd.cpp
+++ b/Engine/Plugins/Experimental/Mutable/Source/MutableTools/Private/MuT/ASTOpInstanceAdd.cpp
@@ -53,7 +53,7 @@ bool ASTOpInstanceAdd::IsEqual(const ASTOp& otherUntyped) const
 
 
 //-------------------------------------------------------------------------------------------------
-mu::Ptr<ASTOp> ASTOpInstanceAdd::Clone(MapChildFunc& mapChild) const
+mu::Ptr<ASTOp> ASTOpInstanceAdd::Clone(MapChildFuncRef mapChild) const
 {
     Ptr<ASTOpInstanceAdd> n = new ASTOpInstanceAdd();
     n->type = type;
@@ -101,7 +101,7 @@ void ASTOpInstanceAdd::Assert()
 
 
 //-------------------------------------------------------------------------------------------------
-void ASTOpInstanceAdd::ForEachChild(const std::function<void(ASTChild&)>& f )
+void ASTOpInstanceAdd::ForEachChild(const TFunctionRef<void(ASTChild&)> f )
 {
     f( instance );
     f( value );
@@ -132,39 +132,37 @@ void ASTOpInstanceAdd::Link( PROGRAM& program, const FLinkerOptions* )
             SubtreeRelevantParametersVisitorAST visitor;
             visitor.Run( value.child() );
 
-            vector<uint16_t> params;
+            TArray<uint16> params;
             for(const string& paramName: visitor.m_params)
             {
-                for( size_t i=0; i<program.m_parameters.size(); ++i )
+                for( int32 i=0; i<program.m_parameters.Num(); ++i )
                 {
                     const auto& param = program.m_parameters[i];
                     if (param.m_name==paramName)
                     {
-                        params.push_back( uint16_t(i) );
+                        params.Add( uint16(i) );
                         break;
                     }
                 }
             }
 
-            std::sort(params.begin(),params.end());
+            params.Sort();
 
-            auto it = std::find( program.m_parameterLists.begin(),
-                                 program.m_parameterLists.end(),
-                                 params );
+            auto it = program.m_parameterLists.Find( params );
 
-            if (it!=program.m_parameterLists.end())
+            if (it!=INDEX_NONE)
             {
-                args.relevantParametersListIndex = uint32_t(it - program.m_parameterLists.begin());
+                args.relevantParametersListIndex = it;
             }
             else
             {
-                args.relevantParametersListIndex = uint32_t(program.m_parameterLists.size());
-                program.m_parameterLists.push_back( params );
+                args.relevantParametersListIndex = uint32_t(program.m_parameterLists.Num());
+                program.m_parameterLists.Add( params );
             }
         }
 
-        linkedAddress = (OP::ADDRESS)program.m_opAddress.size();
-        program.m_opAddress.push_back((uint32_t)program.m_byteCode.size());
+        linkedAddress = (OP::ADDRESS)program.m_opAddress.Num();
+        program.m_opAddress.Add((uint32_t)program.m_byteCode.Num());
         AppendCode(program.m_byteCode,type);
         AppendCode(program.m_byteCode,args);
     }
diff --git a/Engine/Plugins/Experimental/Mutable/Source/MutableTools/Private/MuT/ASTOpMeshApplyPose.cpp b/Engine/Plugins/Experimental/Mutable/Source/MutableTools/Private/MuT/ASTOpMeshApplyPose.cpp
--- a/Engine/Plugins/Experimental/Mutable/Source/MutableTools/Private/MuT/ASTOpMeshApplyPose.cpp
+++ b/Engine/Plugins/Experimental/Mutable/Source/MutableTools/Private/MuT/ASTOpMeshApplyPose.cpp
@@ -53,7 +53,7 @@ uint64 ASTOpMeshApplyPose::Hash() const
 }
 
 
-mu::Ptr<ASTOp> ASTOpMeshApplyPose::Clone(MapChildFunc& mapChild) const
+mu::Ptr<ASTOp> ASTOpMeshApplyPose::Clone(MapChildFuncRef mapChild) const
 {
     Ptr<ASTOpMeshApplyPose> n = new ASTOpMeshApplyPose();
     n->base = mapChild(base.child());
@@ -62,7 +62,7 @@ mu::Ptr<ASTOp> ASTOpMeshApplyPose::Clone(MapChildFunc& mapChild) const
 }
 
 
-void ASTOpMeshApplyPose::ForEachChild(const std::function<void(ASTChild&)>& f )
+void ASTOpMeshApplyPose::ForEachChild(const TFunctionRef<void(ASTChild&)> f )
 {
     f( base );
     f( pose );
@@ -80,9 +80,9 @@ void ASTOpMeshApplyPose::Link( PROGRAM& program, const FLinkerOptions* )
         if (base) args.base = base->linkedAddress;
         if (pose) args.pose = pose->linkedAddress;
 
-        linkedAddress = (OP::ADDRESS)program.m_opAddress.size();
+        linkedAddress = (OP::ADDRESS)program.m_opAddress.Num();
         //program.m_code.push_back(op);
-        program.m_opAddress.push_back((uint32_t)program.m_byteCode.size());
+        program.m_opAddress.Add((uint32_t)program.m_byteCode.Num());
         AppendCode(program.m_byteCode,OP_TYPE::ME_APPLYPOSE);
         AppendCode(program.m_byteCode,args);
     }
