[mutable] Optimize mesh clip versus switching of clipping mesh.

diff --git a/Engine/Plugins/Experimental/Mutable/Source/MutableTools/Private/MuT/ASTOpMeshApplyShape.cpp b/Engine/Plugins/Experimental/Mutable/Source/MutableTools/Private/MuT/ASTOpMeshApplyShape.cpp
--- a/Engine/Plugins/Experimental/Mutable/Source/MutableTools/Private/MuT/ASTOpMeshApplyShape.cpp
+++ b/Engine/Plugins/Experimental/Mutable/Source/MutableTools/Private/MuT/ASTOpMeshApplyShape.cpp
@@ -55,7 +55,7 @@ uint64 ASTOpMeshApplyShape::Hash() const
 }
 
 
-mu::Ptr<ASTOp> ASTOpMeshApplyShape::Clone(MapChildFunc& mapChild) const
+mu::Ptr<ASTOp> ASTOpMeshApplyShape::Clone(MapChildFuncRef mapChild) const
 {
     Ptr<ASTOpMeshApplyShape> n = new ASTOpMeshApplyShape();
 	n->Mesh = mapChild(Mesh.child());
@@ -67,7 +67,7 @@ mu::Ptr<ASTOp> ASTOpMeshApplyShape::Clone(MapChildFunc& mapChild) const
 }
 
 
-void ASTOpMeshApplyShape::ForEachChild(const std::function<void(ASTChild&)>& f )
+void ASTOpMeshApplyShape::ForEachChild(const TFunctionRef<void(ASTChild&)> f )
 {
 	f(Mesh);
 	f(Shape);
@@ -101,8 +101,8 @@ void ASTOpMeshApplyShape::Link( PROGRAM& program, const FLinkerOptions* )
 		if (Mesh) args.mesh = Mesh->linkedAddress;
 		if (Shape) args.shape = Shape->linkedAddress;
 
-        linkedAddress = (OP::ADDRESS)program.m_opAddress.size();
-        program.m_opAddress.push_back((uint32_t)program.m_byteCode.size());
+        linkedAddress = (OP::ADDRESS)program.m_opAddress.Num();
+        program.m_opAddress.Add((uint32_t)program.m_byteCode.Num());
         AppendCode(program.m_byteCode,OP_TYPE::ME_APPLYSHAPE);
         AppendCode(program.m_byteCode,args);
     }
diff --git a/Engine/Plugins/Experimental/Mutable/Source/MutableTools/Private/MuT/ASTOpMeshBindShape.cpp b/Engine/Plugins/Experimental/Mutable/Source/MutableTools/Private/MuT/ASTOpMeshBindShape.cpp
--- a/Engine/Plugins/Experimental/Mutable/Source/MutableTools/Private/MuT/ASTOpMeshBindShape.cpp
+++ b/Engine/Plugins/Experimental/Mutable/Source/MutableTools/Private/MuT/ASTOpMeshBindShape.cpp
@@ -82,7 +82,7 @@ uint64 ASTOpMeshBindShape::Hash() const
 }
 
 
-mu::Ptr<ASTOp> ASTOpMeshBindShape::Clone(MapChildFunc& mapChild) const
+mu::Ptr<ASTOp> ASTOpMeshBindShape::Clone(MapChildFuncRef mapChild) const
 {
     Ptr<ASTOpMeshBindShape> n = new ASTOpMeshBindShape();
 	n->Mesh = mapChild(Mesh.child());
@@ -101,7 +101,7 @@ mu::Ptr<ASTOp> ASTOpMeshBindShape::Clone(MapChildFunc& mapChild) const
 }
 
 
-void ASTOpMeshBindShape::ForEachChild(const std::function<void(ASTChild&)>& f )
+void ASTOpMeshBindShape::ForEachChild(const TFunctionRef<void(ASTChild&)> f )
 {
 	f(Mesh);
 	f(Shape);
@@ -163,8 +163,8 @@ void ASTOpMeshBindShape::Link( PROGRAM& program, const FLinkerOptions*)
 			args.shape = Shape->linkedAddress;
 		}
 		
-		linkedAddress = (OP::ADDRESS)program.m_opAddress.size();
-		program.m_opAddress.push_back((uint32_t)program.m_byteCode.size());
+		linkedAddress = (OP::ADDRESS)program.m_opAddress.Num();
+		program.m_opAddress.Add((uint32_t)program.m_byteCode.Num());
 		AppendCode(program.m_byteCode, OP_TYPE::ME_BINDSHAPE);
 		AppendCode(program.m_byteCode, args);
 
@@ -234,7 +234,7 @@ mu::Ptr<ASTOp> ASTOpMeshBindShape::OptimiseSink(const MODEL_OPTIMIZATION_OPTIONS
 			}
 
 			// We need to copy the options because we change them
-			for (size_t v = 0; v < NewSwitch->cases.size(); ++v)
+			for (int32 v = 0; v < NewSwitch->cases.Num(); ++v)
 			{
 				if (NewSwitch->cases[v].branch)
 				{
@@ -276,7 +276,7 @@ mu::Ptr<ASTOp> ASTOpMeshBindShape::OptimiseSink(const MODEL_OPTIMIZATION_OPTIONS
 			}
 
 			// We need to copy the options because we change them
-			for (size_t v = 0; v < NewSwitch->cases.size(); ++v)
+			for (int32 v = 0; v < NewSwitch->cases.Num(); ++v)
 			{
 				if (NewSwitch->cases[v].branch)
 				{
@@ -316,7 +316,7 @@ mu::Ptr<ASTOp> ASTOpMeshBindShape::OptimiseSink(const MODEL_OPTIMIZATION_OPTIONS
 			}
 
 			// We need to copy the options because we change them
-			for (size_t v = 0; v < NewSwitch->cases.size(); ++v)
+			for (int32 v = 0; v < NewSwitch->cases.Num(); ++v)
 			{
 				if (NewSwitch->cases[v].branch)
 				{
diff --git a/Engine/Plugins/Experimental/Mutable/Source/MutableTools/Private/MuT/ASTOpMeshClipDeform.cpp b/Engine/Plugins/Experimental/Mutable/Source/MutableTools/Private/MuT/ASTOpMeshClipDeform.cpp
--- a/Engine/Plugins/Experimental/Mutable/Source/MutableTools/Private/MuT/ASTOpMeshClipDeform.cpp
+++ b/Engine/Plugins/Experimental/Mutable/Source/MutableTools/Private/MuT/ASTOpMeshClipDeform.cpp
@@ -48,7 +48,7 @@ uint64 ASTOpMeshClipDeform::Hash() const
 }
 
 
-mu::Ptr<ASTOp> ASTOpMeshClipDeform::Clone(MapChildFunc& mapChild) const
+mu::Ptr<ASTOp> ASTOpMeshClipDeform::Clone(MapChildFuncRef mapChild) const
 {
     Ptr<ASTOpMeshClipDeform> n = new ASTOpMeshClipDeform();
 	n->Mesh = mapChild(Mesh.child());
@@ -57,7 +57,7 @@ mu::Ptr<ASTOp> ASTOpMeshClipDeform::Clone(MapChildFunc& mapChild) const
 }
 
 
-void ASTOpMeshClipDeform::ForEachChild(const std::function<void(ASTChild&)>& f )
+void ASTOpMeshClipDeform::ForEachChild(const TFunctionRef<void(ASTChild&)> f )
 {
 	f(Mesh);
 	f(ClipShape);
@@ -82,8 +82,8 @@ void ASTOpMeshClipDeform::Link( PROGRAM& program, const FLinkerOptions*)
 			 args.clipShape = ClipShape->linkedAddress;
 		}	
 
-        linkedAddress = (OP::ADDRESS)program.m_opAddress.size();
-        program.m_opAddress.push_back((uint32_t)program.m_byteCode.size());
+        linkedAddress = (OP::ADDRESS)program.m_opAddress.Num();
+        program.m_opAddress.Add((uint32_t)program.m_byteCode.Num());
         AppendCode(program.m_byteCode,OP_TYPE::ME_CLIPDEFORM);
         AppendCode(program.m_byteCode,args);
     }
diff --git a/Engine/Plugins/Experimental/Mutable/Source/MutableTools/Private/MuT/ASTOpMeshClipMorphPlane.cpp b/Engine/Plugins/Experimental/Mutable/Source/MutableTools/Private/MuT/ASTOpMeshClipMorphPlane.cpp
--- a/Engine/Plugins/Experimental/Mutable/Source/MutableTools/Private/MuT/ASTOpMeshClipMorphPlane.cpp
+++ b/Engine/Plugins/Experimental/Mutable/Source/MutableTools/Private/MuT/ASTOpMeshClipMorphPlane.cpp
@@ -56,7 +56,7 @@ uint64 ASTOpMeshClipMorphPlane::Hash() const
 }
 
 
-mu::Ptr<ASTOp> ASTOpMeshClipMorphPlane::Clone(MapChildFunc& mapChild) const
+mu::Ptr<ASTOp> ASTOpMeshClipMorphPlane::Clone(MapChildFuncRef mapChild) const
 {
     Ptr<ASTOpMeshClipMorphPlane> n = new ASTOpMeshClipMorphPlane();
     n->source = mapChild(source.child());
@@ -70,7 +70,7 @@ mu::Ptr<ASTOp> ASTOpMeshClipMorphPlane::Clone(MapChildFunc& mapChild) const
 }
 
 
-void ASTOpMeshClipMorphPlane::ForEachChild(const std::function<void(ASTChild&)>& f )
+void ASTOpMeshClipMorphPlane::ForEachChild(const TFunctionRef<void(ASTChild&)> f )
 {
     f( source );
 }
@@ -86,8 +86,8 @@ void ASTOpMeshClipMorphPlane::Link( PROGRAM& program, const FLinkerOptions* )
 
         if (source) args.source = source->linkedAddress;
 
-        args.morphShape = (OP::ADDRESS)program.m_constantShapes.size();
-        program.m_constantShapes.push_back(morphShape);
+        args.morphShape = (OP::ADDRESS)program.m_constantShapes.Num();
+        program.m_constantShapes.Add(morphShape);
 
         args.vertexSelectionType = (uint8_t)vertexSelectionType;
         if (vertexSelectionType==OP::MeshClipMorphPlaneArgs::VS_BONE_HIERARCHY)
@@ -103,9 +103,9 @@ void ASTOpMeshClipMorphPlane::Link( PROGRAM& program, const FLinkerOptions* )
         args.factor = factor;
 		args.maxBoneRadius = vertexSelectionBoneMaxRadius;
 
-        linkedAddress = (OP::ADDRESS)program.m_opAddress.size();
+        linkedAddress = (OP::ADDRESS)program.m_opAddress.Num();
         //program.m_code.push_back(op);
-        program.m_opAddress.push_back((uint32_t)program.m_byteCode.size());
+        program.m_opAddress.Add((uint32_t)program.m_byteCode.Num());
         AppendCode(program.m_byteCode,OP_TYPE::ME_CLIPMORPHPLANE);
         AppendCode(program.m_byteCode,args);
     }
diff --git a/Engine/Plugins/Experimental/Mutable/Source/MutableTools/Private/MuT/ASTOpMeshExtractLayoutBlocks.cpp b/Engine/Plugins/Experimental/Mutable/Source/MutableTools/Private/MuT/ASTOpMeshExtractLayoutBlocks.cpp
--- a/Engine/Plugins/Experimental/Mutable/Source/MutableTools/Private/MuT/ASTOpMeshExtractLayoutBlocks.cpp
+++ b/Engine/Plugins/Experimental/Mutable/Source/MutableTools/Private/MuT/ASTOpMeshExtractLayoutBlocks.cpp
@@ -40,7 +40,7 @@ bool ASTOpMeshExtractLayoutBlocks::IsEqual(const ASTOp& otherUntyped) const
 }
 
 
-mu::Ptr<ASTOp> ASTOpMeshExtractLayoutBlocks::Clone(MapChildFunc& mapChild) const
+mu::Ptr<ASTOp> ASTOpMeshExtractLayoutBlocks::Clone(MapChildFuncRef mapChild) const
 {
     Ptr<ASTOpMeshExtractLayoutBlocks> n = new ASTOpMeshExtractLayoutBlocks();
     n->source = mapChild(source.child());
@@ -52,12 +52,12 @@ mu::Ptr<ASTOp> ASTOpMeshExtractLayoutBlocks::Clone(MapChildFunc& mapChild) const
 
 void ASTOpMeshExtractLayoutBlocks::Assert()
 {
-    check( blocks.size()<std::numeric_limits<uint16_t>::max() );
+    check( blocks.Num()<std::numeric_limits<uint16>::max() );
     ASTOp::Assert();
 }
 
 
-void ASTOpMeshExtractLayoutBlocks::ForEachChild(const std::function<void(ASTChild&)>& f )
+void ASTOpMeshExtractLayoutBlocks::ForEachChild(const TFunctionRef<void(ASTChild&)> f )
 {
     f(source);
 }
@@ -75,18 +75,18 @@ void ASTOpMeshExtractLayoutBlocks::Link( PROGRAM& program, const FLinkerOptions*
     // Already linked?
     if (!linkedAddress)
     {
-        linkedAddress = (OP::ADDRESS)program.m_opAddress.size();
+        linkedAddress = (OP::ADDRESS)program.m_opAddress.Num();
 
-        program.m_opAddress.push_back((uint32_t)program.m_byteCode.size());
+        program.m_opAddress.Add((uint32)program.m_byteCode.Num());
         AppendCode(program.m_byteCode, OP_TYPE::ME_EXTRACTLAYOUTBLOCK);
         OP::ADDRESS sourceAt = source ? source->linkedAddress : 0;
         AppendCode(program.m_byteCode, sourceAt );
-        AppendCode(program.m_byteCode, (uint16_t)layout );
-        AppendCode(program.m_byteCode, (uint16_t)blocks.size() );
+        AppendCode(program.m_byteCode, (uint16)layout );
+        AppendCode(program.m_byteCode, (uint16)blocks.Num() );
 
         for ( auto b: blocks )
         {
-            AppendCode(program.m_byteCode, (uint32_t)b );
+            AppendCode(program.m_byteCode, (uint32)b );
         }
     }
 }
diff --git a/Engine/Plugins/Experimental/Mutable/Source/MutableTools/Private/MuT/ASTOpMeshExtractLayoutBlocks.h b/Engine/Plugins/Experimental/Mutable/Source/MutableTools/Private/MuT/ASTOpMeshExtractLayoutBlocks.h
--- a/Engine/Plugins/Experimental/Mutable/Source/MutableTools/Private/MuT/ASTOpMeshExtractLayoutBlocks.h
+++ b/Engine/Plugins/Experimental/Mutable/Source/MutableTools/Private/MuT/ASTOpMeshExtractLayoutBlocks.h
@@ -8,8 +8,6 @@
 #include "MuR/Ptr.h"
 #include "MuT/AST.h"
 
-#include <functional>
-
 
 namespace mu
 {
@@ -27,10 +25,10 @@ struct PROGRAM;
 		ASTChild source;
 
 		//! Layout to use to select the blocks.
-		uint16_t layout = 0;
+		uint16 layout = 0;
 
 		//! Blocks
-		vector<uint32_t> blocks;
+		TArray<uint32> blocks;
 
 	public:
 
@@ -41,9 +39,9 @@ struct PROGRAM;
 		// ASTOp interface
 		OP_TYPE GetOpType() const override { return OP_TYPE::ME_EXTRACTLAYOUTBLOCK; }
 		uint64 Hash() const override;
-		void ForEachChild(const std::function<void(ASTChild&)>&) override;
+		void ForEachChild(const TFunctionRef<void(ASTChild&)>) override;
 		bool IsEqual(const ASTOp& otherUntyped) const override;
-		Ptr<ASTOp> Clone(MapChildFunc& mapChild) const override;
+		Ptr<ASTOp> Clone(MapChildFuncRef mapChild) const override;
 		void Assert() override;
 		void Link(PROGRAM& program, const FLinkerOptions* Options) override;
 
diff --git a/Engine/Plugins/Experimental/Mutable/Source/MutableTools/Private/MuT/ASTOpMeshFormat.cpp b/Engine/Plugins/Experimental/Mutable/Source/MutableTools/Private/MuT/ASTOpMeshFormat.cpp
--- a/Engine/Plugins/Experimental/Mutable/Source/MutableTools/Private/MuT/ASTOpMeshFormat.cpp
+++ b/Engine/Plugins/Experimental/Mutable/Source/MutableTools/Private/MuT/ASTOpMeshFormat.cpp
@@ -61,7 +61,7 @@ uint64 ASTOpMeshFormat::Hash() const
 }
 
 
-mu::Ptr<ASTOp> ASTOpMeshFormat::Clone(MapChildFunc& mapChild) const
+mu::Ptr<ASTOp> ASTOpMeshFormat::Clone(MapChildFuncRef mapChild) const
 {
 	mu::Ptr<ASTOpMeshFormat> n = new ASTOpMeshFormat();
     n->Source = mapChild(Source.child());
@@ -71,7 +71,7 @@ mu::Ptr<ASTOp> ASTOpMeshFormat::Clone(MapChildFunc& mapChild) const
 }
 
 
-void ASTOpMeshFormat::ForEachChild(const std::function<void(ASTChild&)>& f )
+void ASTOpMeshFormat::ForEachChild(const TFunctionRef<void(ASTChild&)> f )
 {
     f( Source );
     f( Format );
@@ -90,9 +90,9 @@ void ASTOpMeshFormat::Link( PROGRAM& program, const FLinkerOptions* )
         if (Source) args.source = Source->linkedAddress;
         if (Format) args.format = Format->linkedAddress;
 
-        linkedAddress = (OP::ADDRESS)program.m_opAddress.size();
+        linkedAddress = (OP::ADDRESS)program.m_opAddress.Num();
         //program.m_code.push_back(op);
-        program.m_opAddress.push_back((uint32_t)program.m_byteCode.size());
+        program.m_opAddress.Add((uint32_t)program.m_byteCode.Num());
         AppendCode(program.m_byteCode,OP_TYPE::ME_FORMAT);
         AppendCode(program.m_byteCode,args);
     }
@@ -115,7 +115,7 @@ mu::Ptr<ASTOp> Sink_MeshFormatAST::Apply(const ASTOp* root)
 	m_root = dynamic_cast<const ASTOpMeshFormat*>(root);
 	check(m_root);
 
-	m_oldToNew.clear();
+	m_oldToNew.Empty();
 
 	m_initialSource = m_root->Source.child();
 	mu::Ptr<ASTOp> newSource = Visit(m_initialSource, m_root);
@@ -165,18 +165,18 @@ namespace
 
 		int offset = 0;
 		int numChannels = 0;
-		vector<MESH_BUFFER_SEMANTIC> semantics;
-		vector<int> semanticIndices;
-		vector<MESH_BUFFER_FORMAT> formats;
-		vector<int> components;
-		vector<int> offsets;
+		TArray<MESH_BUFFER_SEMANTIC> semantics;
+		TArray<int> semanticIndices;
+		TArray<MESH_BUFFER_FORMAT> formats;
+		TArray<int> components;
+		TArray<int> offsets;
 
 		// Add the vertex index channel
-		semantics.push_back(MBS_VERTEXINDEX);
-		semanticIndices.push_back(0);
-		formats.push_back(MBF_UINT32);
-		components.push_back(1);
-		offsets.push_back(offset);
+		semantics.Add(MBS_VERTEXINDEX);
+		semanticIndices.Add(0);
+		formats.Add(MBF_UINT32);
+		components.Add(1);
+		offsets.Add(offset);
 		offset += components[numChannels] * GetMeshFormatData(formats[numChannels]).m_size;
 		numChannels++;
 
@@ -194,11 +194,11 @@ namespace
 				(vb, c, &semantic, &semanticIndex, &format, &component, nullptr);
 
 				// TODO: Filter useless semantics for morphing.
-				semantics.push_back(semantic);
-				semanticIndices.push_back(semanticIndex);
-				formats.push_back(format);
-				components.push_back(component);
-				offsets.push_back(offset);
+				semantics.Add(semantic);
+				semanticIndices.Add(semanticIndex);
+				formats.Add(format);
+				components.Add(component);
+				offsets.Add(offset);
 				offset += components[numChannels]
 					* GetMeshFormatData(formats[numChannels]).m_size;
 				numChannels++;
@@ -230,16 +230,16 @@ mu::Ptr<ASTOp> Sink_MeshFormatAST::Visit(const mu::Ptr<ASTOp>& at, const ASTOpMe
 	if (!at) return nullptr;
 
 	// Newly created?
-	if (std::find(m_newOps.begin(), m_newOps.end(), at) != m_newOps.end())
+	if (m_newOps.Find(at) != INDEX_NONE)
 	{
 		return at;
 	}
 
 	// Already visited?
-	auto cacheIt = m_oldToNew.find(at);
-	if (cacheIt != m_oldToNew.end())
+	auto cacheIt = m_oldToNew.Find(at);
+	if (cacheIt)
 	{
-		return cacheIt->second;
+		return *cacheIt;
 	}
 
 	mu::Ptr<ASTOp> newAt = at;
@@ -406,7 +406,7 @@ mu::Ptr<ASTOp> Sink_MeshFormatAST::Visit(const mu::Ptr<ASTOp>& at, const ASTOpMe
 
 	}
 
-	m_oldToNew[at] = newAt;
+	m_oldToNew.Add(at, newAt);
 
 	return newAt;
 }
diff --git a/Engine/Plugins/Experimental/Mutable/Source/MutableTools/Private/MuT/ASTOpMeshGeometryOperation.cpp b/Engine/Plugins/Experimental/Mutable/Source/MutableTools/Private/MuT/ASTOpMeshGeometryOperation.cpp
--- a/Engine/Plugins/Experimental/Mutable/Source/MutableTools/Private/MuT/ASTOpMeshGeometryOperation.cpp
+++ b/Engine/Plugins/Experimental/Mutable/Source/MutableTools/Private/MuT/ASTOpMeshGeometryOperation.cpp
@@ -51,7 +51,7 @@ namespace mu
 	}
 
 
-	mu::Ptr<ASTOp> ASTOpMeshGeometryOperation::Clone(MapChildFunc& mapChild) const
+	mu::Ptr<ASTOp> ASTOpMeshGeometryOperation::Clone(MapChildFuncRef mapChild) const
 	{
 		Ptr<ASTOpMeshGeometryOperation> n = new ASTOpMeshGeometryOperation();
 		n->meshA = mapChild(meshA.child());
@@ -62,7 +62,7 @@ namespace mu
 	}
 
 
-	void ASTOpMeshGeometryOperation::ForEachChild(const std::function<void(ASTChild&)>& f)
+	void ASTOpMeshGeometryOperation::ForEachChild(const TFunctionRef<void(ASTChild&)> f)
 	{
 		f(meshA);
 		f(meshB);
@@ -84,9 +84,9 @@ namespace mu
 			if (scalarA) args.scalarA = scalarA->linkedAddress;
 			if (scalarB) args.scalarB = scalarB->linkedAddress;
 
-			linkedAddress = (OP::ADDRESS)program.m_opAddress.size();
+			linkedAddress = (OP::ADDRESS)program.m_opAddress.Num();
 			//program.m_code.push_back(op);
-			program.m_opAddress.push_back((uint32_t)program.m_byteCode.size());
+			program.m_opAddress.Add((uint32_t)program.m_byteCode.Num());
 			AppendCode(program.m_byteCode, OP_TYPE::ME_GEOMETRYOPERATION);
 			AppendCode(program.m_byteCode, args);
 		}
diff --git a/Engine/Plugins/Experimental/Mutable/Source/MutableTools/Private/MuT/ASTOpMeshMaskClipMesh.cpp b/Engine/Plugins/Experimental/Mutable/Source/MutableTools/Private/MuT/ASTOpMeshMaskClipMesh.cpp
--- a/Engine/Plugins/Experimental/Mutable/Source/MutableTools/Private/MuT/ASTOpMeshMaskClipMesh.cpp
+++ b/Engine/Plugins/Experimental/Mutable/Source/MutableTools/Private/MuT/ASTOpMeshMaskClipMesh.cpp
@@ -12,233 +12,345 @@
 #include "MuT/ASTOpSwitch.h"
 #include "MuT/StreamsPrivate.h"
 
-#include <memory>
-#include <utility>
 
-
-using namespace mu;
-
-
-//-------------------------------------------------------------------------------------------------
-//-------------------------------------------------------------------------------------------------
-//-------------------------------------------------------------------------------------------------
-ASTOpMeshMaskClipMesh::ASTOpMeshMaskClipMesh()
-    : source(this)
-    , clip(this)
+namespace mu
 {
-}
 
 
-//-------------------------------------------------------------------------------------------------
-ASTOpMeshMaskClipMesh::~ASTOpMeshMaskClipMesh()
-{
-    // Explicit call needed to avoid recursive destruction
-    ASTOp::RemoveChildren();
-}
-
-
-//-------------------------------------------------------------------------------------------------
-bool ASTOpMeshMaskClipMesh::IsEqual(const ASTOp& otherUntyped) const
-{
-    if ( auto other = dynamic_cast<const ASTOpMeshMaskClipMesh*>(&otherUntyped) )
-    {
-        return source==other->source && clip==other->clip;
-    }
-    return false;
-}
-
-
-//-------------------------------------------------------------------------------------------------
-uint64 ASTOpMeshMaskClipMesh::Hash() const
-{
-	uint64 res = std::hash<void*>()(source.child().get());
-    hash_combine( res, clip.child() );
-    return res;
-}
-
-
-//-------------------------------------------------------------------------------------------------
-mu::Ptr<ASTOp> ASTOpMeshMaskClipMesh::Clone(MapChildFunc& mapChild) const
-{
-    Ptr<ASTOpMeshMaskClipMesh> n = new ASTOpMeshMaskClipMesh();
-    n->source = mapChild(source.child());
-    n->clip = mapChild(clip.child());
-    return n;
-}
-
-
-//-------------------------------------------------------------------------------------------------
-void ASTOpMeshMaskClipMesh::ForEachChild(const std::function<void(ASTChild&)>& f )
-{
-    f( source );
-    f( clip );
-}
-
-
-//-------------------------------------------------------------------------------------------------
-void ASTOpMeshMaskClipMesh::Link( PROGRAM& program, const FLinkerOptions* )
-{
-    // Already linked?
-    if (!linkedAddress)
-    {
-        OP::MeshMaskClipMeshArgs args;
-        memset( &args,0, sizeof(args) );
-
-        if (source) args.source = source->linkedAddress;
-        if (clip) args.clip = clip->linkedAddress;
-
-        linkedAddress = (OP::ADDRESS)program.m_opAddress.size();
-        //program.m_code.push_back(op);
-        program.m_opAddress.push_back((uint32_t)program.m_byteCode.size());
-        AppendCode(program.m_byteCode,OP_TYPE::ME_MASKCLIPMESH);
-        AppendCode(program.m_byteCode,args);
-    }
-}
-
-
-//-------------------------------------------------------------------------------------------------
-//-------------------------------------------------------------------------------------------------
-//-------------------------------------------------------------------------------------------------
-namespace
-{
-    class Sink_MeshMaskClipMeshAST
-    {
-    public:
-
-        // \TODO This is recursive and may cause stack overflows in big models.
-
-        mu::Ptr<ASTOp> Apply( const ASTOpMeshMaskClipMesh* root )
-        {
-            m_root = root;
-            m_oldToNew.clear();
-
-            m_initialSource = root->source.child();
-			mu::Ptr<ASTOp> newSource = Visit( m_initialSource );
-
-            // If there is any change, it is the new root.
-            if (newSource!=m_initialSource)
-            {
-                return newSource;
-            }
-
-            return nullptr;
-        }
-
-    protected:
-
-        const ASTOpMeshMaskClipMesh* m_root;
-		mu::Ptr<ASTOp> m_initialSource;
-        std::unordered_map<mu::Ptr<ASTOp>, mu::Ptr<ASTOp>> m_oldToNew;
-        vector<mu::Ptr<ASTOp>> m_newOps;
-
-		mu::Ptr<ASTOp> Visit( const mu::Ptr<ASTOp>& at )
-        {
-            if (!at) return nullptr;
-
-            // Newly created?
-            if (std::find(m_newOps.begin(), m_newOps.end(), at )!=m_newOps.end())
-            {
-                return at;
-            }
-
-            // Already visited?
-            auto cacheIt = m_oldToNew.find(at);
-            if (cacheIt!=m_oldToNew.end())
-            {
-                return cacheIt->second;
-            }
-
-			mu::Ptr<ASTOp> newAt = at;
-            switch ( at->GetOpType() )
-            {
-
-			// This cannot be sunk since the result is different. Since the clipping is now correctly
-			// generated at the end of the chain when really necessary, this wrong optimisation is no 
-			// longer needed.
-			//case OP_TYPE::ME_MORPH2:
-   //         {
-   //             // Remove this op: Let's clip the unmorphed mesh only
-   //             auto typedAt = dynamic_cast<const ASTOpFixed*>(at.get());
-   //             newAt = Visit( typedAt->children[typedAt->op.args.MeshMorph2.base].child() );
-   //             break;
-   //         }
-
-            case OP_TYPE::ME_REMOVEMASK:
-            {
-                // Remove this op:
-                // This may lead to the mask being bigger than needed since it will include
-                // faces removed by the ignored removemask, but it is ok
-
-                // TODO: Swap instead of ignore, and implement removemask on a mask?
-                auto typedAt = dynamic_cast<const ASTOpMeshRemoveMask*>(at.get());
-                newAt = Visit( typedAt->source.child() );
-                break;
-            }
-
-            case OP_TYPE::ME_CONDITIONAL:
-            {
-                // We move the mask creation down the two paths
-                // It always needs to be a clone because otherwise we could be modifying an
-                // instruction that shouldn't if the parent was a ME_REMOVEMASK above and we
-                // skipped the cloning for the parent.
-                auto newOp = mu::Clone<ASTOpConditional>(at);
-                newOp->yes = Visit(newOp->yes.child());
-                newOp->no = Visit(newOp->no.child());
-                newAt = newOp;
-                break;
-            }
-
-            case OP_TYPE::ME_SWITCH:
-            {
-                // We move the mask creation down all the paths
-                auto newOp = mu::Clone<ASTOpSwitch>(at);
-                newOp->def = Visit(newOp->def.child());
-                for( auto& c:newOp->cases )
-                {
-                    c.branch = Visit(c.branch.child());
-                }
-                newAt = newOp;
-                break;
-            }
-
-			// This cannot be sunk since the result is different. Since the clipping is now correctly
-			// generated at the end of the chain when really necessary, this wrong optimisation is no 
-			// longer needed.
-			//case OP_TYPE::ME_CLIPMORPHPLANE:
-   //         {
-   //             // We move the mask creation down the source
-   //             auto typedAt = dynamic_cast<const ASTOpMeshClipMorphPlane*>(at.get());
-   //             newAt = Visit(typedAt->source.child());
-   //             break;
-   //         }
-
-            default:
-            {
-                //
-                if (at!=m_initialSource)
-                {
-                    auto newOp = mu::Clone<ASTOpMeshMaskClipMesh>(m_root);
-                    newOp->source = at;
-                    newAt = newOp;
-                }
-                break;
-            }
-
-            }
-
-            m_oldToNew[at] = newAt;
-
-            return newAt;
-        }
-    };
-}
-
-
-//-------------------------------------------------------------------------------------------------
-mu::Ptr<ASTOp> ASTOpMeshMaskClipMesh::OptimiseSink(const MODEL_OPTIMIZATION_OPTIONS&, OPTIMIZE_SINK_CONTEXT&) const
-{
-    Sink_MeshMaskClipMeshAST sinker;
-    auto at = sinker.Apply(this);
+	//-------------------------------------------------------------------------------------------------
+	//-------------------------------------------------------------------------------------------------
+	//-------------------------------------------------------------------------------------------------
+	ASTOpMeshMaskClipMesh::ASTOpMeshMaskClipMesh()
+		: source(this)
+		, clip(this)
+	{
+	}
+
+
+	//-------------------------------------------------------------------------------------------------
+	ASTOpMeshMaskClipMesh::~ASTOpMeshMaskClipMesh()
+	{
+		// Explicit call needed to avoid recursive destruction
+		ASTOp::RemoveChildren();
+	}
+
+
+	//-------------------------------------------------------------------------------------------------
+	bool ASTOpMeshMaskClipMesh::IsEqual(const ASTOp& otherUntyped) const
+	{
+		if (const ASTOpMeshMaskClipMesh* other = dynamic_cast<const ASTOpMeshMaskClipMesh*>(&otherUntyped))
+		{
+			return source == other->source && clip == other->clip;
+		}
+		return false;
+	}
+
+
+	//-------------------------------------------------------------------------------------------------
+	uint64 ASTOpMeshMaskClipMesh::Hash() const
+	{
+		uint64 res = std::hash<void*>()(source.child().get());
+		hash_combine(res, clip.child());
+		return res;
+	}
+
+
+	//-------------------------------------------------------------------------------------------------
+	mu::Ptr<ASTOp> ASTOpMeshMaskClipMesh::Clone(MapChildFuncRef mapChild) const
+	{
+		Ptr<ASTOpMeshMaskClipMesh> n = new ASTOpMeshMaskClipMesh();
+		n->source = mapChild(source.child());
+		n->clip = mapChild(clip.child());
+		return n;
+	}
+
+
+	//-------------------------------------------------------------------------------------------------
+	void ASTOpMeshMaskClipMesh::ForEachChild(const TFunctionRef<void(ASTChild&)> f)
+	{
+		f(source);
+		f(clip);
+	}
+
+
+	//-------------------------------------------------------------------------------------------------
+	void ASTOpMeshMaskClipMesh::Link(PROGRAM& program, const FLinkerOptions*)
+	{
+		// Already linked?
+		if (!linkedAddress)
+		{
+			OP::MeshMaskClipMeshArgs args;
+			memset(&args, 0, sizeof(args));
+
+			if (source) args.source = source->linkedAddress;
+			if (clip) args.clip = clip->linkedAddress;
+
+			linkedAddress = (OP::ADDRESS)program.m_opAddress.Num();
+			//program.m_code.push_back(op);
+			program.m_opAddress.Add((uint32_t)program.m_byteCode.Num());
+			AppendCode(program.m_byteCode, OP_TYPE::ME_MASKCLIPMESH);
+			AppendCode(program.m_byteCode, args);
+		}
+	}
+
+
+	//-------------------------------------------------------------------------------------------------
+	//-------------------------------------------------------------------------------------------------
+	//-------------------------------------------------------------------------------------------------
+	namespace
+	{
+		class Sink_MeshMaskClipMeshSource
+		{
+		public:
+
+			// \TODO This is recursive and may cause stack overflows in big models.
+
+			mu::Ptr<ASTOp> Apply(const ASTOpMeshMaskClipMesh* root)
+			{
+				m_root = root;
+				m_oldToNew.Empty();
+
+				m_initialSource = root->source.child();
+				mu::Ptr<ASTOp> NewSource = Visit(m_initialSource);
+
+				// If there is any change, it is the new root.
+				if (NewSource != m_initialSource)
+				{
+					return NewSource;
+				}
+
+				return nullptr;
+			}
+
+		protected:
+
+			const ASTOpMeshMaskClipMesh* m_root;
+			mu::Ptr<ASTOp> m_initialSource;
+			TMap<mu::Ptr<ASTOp>, mu::Ptr<ASTOp>> m_oldToNew;
+			TArray<mu::Ptr<ASTOp>> m_newOps;
+
+			mu::Ptr<ASTOp> Visit(const mu::Ptr<ASTOp>& at)
+			{
+				if (!at) return nullptr;
+
+				// Newly created?
+				if (m_newOps.Find(at) != INDEX_NONE)
+				{
+					return at;
+				}
+
+				// Already visited?
+				mu::Ptr<ASTOp>* cacheIt = m_oldToNew.Find(at);
+				if (cacheIt)
+				{
+					return *cacheIt;
+				}
+
+				mu::Ptr<ASTOp> newAt = at;
+				switch (at->GetOpType())
+				{
+
+					// This cannot be sunk since the result is different. Since the clipping is now correctly
+					// generated at the end of the chain when really necessary, this wrong optimisation is no 
+					// longer needed.
+					//case OP_TYPE::ME_MORPH2:
+		   //         {
+		   //             // Remove this op: Let's clip the unmorphed mesh only
+		   //             auto typedAt = dynamic_cast<const ASTOpFixed*>(at.get());
+		   //             newAt = Visit( typedAt->children[typedAt->op.args.MeshMorph2.base].child() );
+		   //             break;
+		   //         }
+
+				case OP_TYPE::ME_REMOVEMASK:
+				{
+					// Remove this op:
+					// This may lead to the mask being bigger than needed since it will include
+					// faces removed by the ignored removemask, but it is ok
+
+					// TODO: Swap instead of ignore, and implement removemask on a mask?
+					const ASTOpMeshRemoveMask* typedAt = dynamic_cast<const ASTOpMeshRemoveMask*>(at.get());
+					newAt = Visit(typedAt->source.child());
+					break;
+				}
+
+				case OP_TYPE::ME_CONDITIONAL:
+				{
+					// We move the mask creation down the two paths
+					// It always needs to be a clone because otherwise we could be modifying an
+					// instruction that shouldn't if the parent was a ME_REMOVEMASK above and we
+					// skipped the cloning for the parent.
+					Ptr<ASTOpConditional> newOp = mu::Clone<ASTOpConditional>(at);
+					newOp->yes = Visit(newOp->yes.child());
+					newOp->no = Visit(newOp->no.child());
+					newAt = newOp;
+					break;
+				}
+
+				case OP_TYPE::ME_SWITCH:
+				{
+					// We move the mask creation down all the paths
+					Ptr<ASTOpSwitch> newOp = mu::Clone<ASTOpSwitch>(at);
+					newOp->def = Visit(newOp->def.child());
+					for (ASTOpSwitch::CASE& c : newOp->cases)
+					{
+						c.branch = Visit(c.branch.child());
+					}
+					newAt = newOp;
+					break;
+				}
+
+				// This cannot be sunk since the result is different. Since the clipping is now correctly
+				// generated at the end of the chain when really necessary, this wrong optimisation is no 
+				// longer needed.
+				//case OP_TYPE::ME_CLIPMORPHPLANE:
+	   //         {
+	   //             // We move the mask creation down the source
+	   //             auto typedAt = dynamic_cast<const ASTOpMeshClipMorphPlane*>(at.get());
+	   //             newAt = Visit(typedAt->source.child());
+	   //             break;
+	   //         }
+
+				default:
+				{
+					//
+					if (at != m_initialSource)
+					{
+						Ptr<ASTOpMeshMaskClipMesh> newOp = mu::Clone<ASTOpMeshMaskClipMesh>(m_root);
+						newOp->source = at;
+						newAt = newOp;
+					}
+					break;
+				}
+
+				}
+
+				m_oldToNew.Add(at, newAt);
+
+				return newAt;
+			}
+		};
+
+
+		class Sink_MeshMaskClipMeshClip
+		{
+		public:
+
+			// \TODO This is recursive and may cause stack overflows in big models.
+
+			mu::Ptr<ASTOp> Apply(const ASTOpMeshMaskClipMesh* root)
+			{
+				m_root = root;
+				m_oldToNew.Empty();
+
+				m_initialClip = root->clip.child();
+				mu::Ptr<ASTOp> NewClip = Visit(m_initialClip);
+
+				// If there is any change, it is the new root.
+				if (NewClip != m_initialClip)
+				{
+					return NewClip;
+				}
+
+				return nullptr;
+			}
+
+		protected:
+
+			const ASTOpMeshMaskClipMesh* m_root;
+			mu::Ptr<ASTOp> m_initialClip;
+			TMap<mu::Ptr<ASTOp>, mu::Ptr<ASTOp>> m_oldToNew;
+			TArray<mu::Ptr<ASTOp>> m_newOps;
+
+			mu::Ptr<ASTOp> Visit(const mu::Ptr<ASTOp>& at)
+			{
+				if (!at) return nullptr;
+
+				// Newly created?
+				if (m_newOps.Find(at) != INDEX_NONE)
+				{
+					return at;
+				}
+
+				// Already visited?
+				mu::Ptr<ASTOp>* cacheIt = m_oldToNew.Find(at);
+				if (cacheIt)
+				{
+					return *cacheIt;
+				}
+
+				mu::Ptr<ASTOp> newAt = at;
+				switch (at->GetOpType())
+				{
+
+				case OP_TYPE::ME_CONDITIONAL:
+				{
+					// We move the mask creation down the two paths
+					// It always needs to be a clone because otherwise we could be modifying an
+					// instruction that shouldn't if the parent was a ME_REMOVEMASK above and we
+					// skipped the cloning for the parent.
+					Ptr<ASTOpConditional> newOp = mu::Clone<ASTOpConditional>(at);
+					newOp->yes = Visit(newOp->yes.child());
+					newOp->no = Visit(newOp->no.child());
+					newAt = newOp;
+					break;
+				}
+
+				case OP_TYPE::ME_SWITCH:
+				{
+					// We move the mask creation down all the paths
+					Ptr<ASTOpSwitch> newOp = mu::Clone<ASTOpSwitch>(at);
+					newOp->def = Visit(newOp->def.child());
+					for (ASTOpSwitch::CASE& c : newOp->cases)
+					{
+						c.branch = Visit(c.branch.child());
+					}
+					newAt = newOp;
+					break;
+				}
+
+
+				default:
+				{
+					//
+					if (at != m_initialClip)
+					{
+						Ptr<ASTOpMeshMaskClipMesh> newOp = mu::Clone<ASTOpMeshMaskClipMesh>(m_root);
+						newOp->clip = at;
+						newAt = newOp;
+					}
+					break;
+				}
+
+				}
+
+				m_oldToNew.Add(at, newAt);
+
+				return newAt;
+			}
+		};
+
+	}
+
+
+	//-------------------------------------------------------------------------------------------------
+	mu::Ptr<ASTOp> ASTOpMeshMaskClipMesh::OptimiseSink(const MODEL_OPTIMIZATION_OPTIONS&, OPTIMIZE_SINK_CONTEXT&) const
+	{
+		// \TODO: Add logic state to the sinkers to avoid explosion with switches in both branches and similar cases.
+
+		// This will sink the operation down the source
+		Sink_MeshMaskClipMeshSource SinkerSource;
+		mu::Ptr<ASTOp> at = SinkerSource.Apply(this);
+
+		// If we didn't sink it.
+		if (!at || at == this)
+		{
+			// This will sink the operation down the clip child
+			Sink_MeshMaskClipMeshClip SinkerClip;
+			at = SinkerClip.Apply(this);
+		}
+
+		return at;
+	}
 
-    return at;
 }
+ 
\ No newline at end of file
diff --git a/Engine/Plugins/Experimental/Mutable/Source/MutableTools/Private/MuT/ASTOpMeshMorphReshape.cpp b/Engine/Plugins/Experimental/Mutable/Source/MutableTools/Private/MuT/ASTOpMeshMorphReshape.cpp
--- a/Engine/Plugins/Experimental/Mutable/Source/MutableTools/Private/MuT/ASTOpMeshMorphReshape.cpp
+++ b/Engine/Plugins/Experimental/Mutable/Source/MutableTools/Private/MuT/ASTOpMeshMorphReshape.cpp
@@ -49,7 +49,7 @@ uint64 ASTOpMeshMorphReshape::Hash() const
 }
 
 
-mu::Ptr<ASTOp> ASTOpMeshMorphReshape::Clone(MapChildFunc& mapChild) const
+mu::Ptr<ASTOp> ASTOpMeshMorphReshape::Clone(MapChildFuncRef mapChild) const
 {
     Ptr<ASTOpMeshMorphReshape> n = new ASTOpMeshMorphReshape();
 	n->Morph = mapChild(Morph.child());
@@ -58,7 +58,7 @@ mu::Ptr<ASTOp> ASTOpMeshMorphReshape::Clone(MapChildFunc& mapChild) const
 }
 
 
-void ASTOpMeshMorphReshape::ForEachChild(const std::function<void(ASTChild&)>& f )
+void ASTOpMeshMorphReshape::ForEachChild(const TFunctionRef<void(ASTChild&)> f )
 {
 	f(Morph);
 	f(Reshape);
@@ -83,8 +83,8 @@ void ASTOpMeshMorphReshape::Link( PROGRAM& program, const FLinkerOptions*)
 			Args.Reshape = Reshape->linkedAddress;
 		}
 
-        linkedAddress = (OP::ADDRESS)program.m_opAddress.size();
-        program.m_opAddress.push_back((uint32_t)program.m_byteCode.size());
+        linkedAddress = (OP::ADDRESS)program.m_opAddress.Num();
+        program.m_opAddress.Add((uint32_t)program.m_byteCode.Num());
         AppendCode(program.m_byteCode, OP_TYPE::ME_MORPHRESHAPE);
         AppendCode(program.m_byteCode, Args);
     }
diff --git a/Engine/Plugins/Experimental/Mutable/Source/MutableTools/Private/MuT/ASTOpMeshRemapIndices.cpp b/Engine/Plugins/Experimental/Mutable/Source/MutableTools/Private/MuT/ASTOpMeshRemapIndices.cpp
--- a/Engine/Plugins/Experimental/Mutable/Source/MutableTools/Private/MuT/ASTOpMeshRemapIndices.cpp
+++ b/Engine/Plugins/Experimental/Mutable/Source/MutableTools/Private/MuT/ASTOpMeshRemapIndices.cpp
@@ -54,7 +54,7 @@ uint64 ASTOpMeshRemapIndices::Hash() const
 
 
 //-------------------------------------------------------------------------------------------------
-mu::Ptr<ASTOp> ASTOpMeshRemapIndices::Clone(MapChildFunc& mapChild) const
+mu::Ptr<ASTOp> ASTOpMeshRemapIndices::Clone(MapChildFuncRef mapChild) const
 {
     Ptr<ASTOpMeshRemapIndices> n = new ASTOpMeshRemapIndices();
     n->source = mapChild(source.child());
@@ -64,7 +64,7 @@ mu::Ptr<ASTOp> ASTOpMeshRemapIndices::Clone(MapChildFunc& mapChild) const
 
 
 //-------------------------------------------------------------------------------------------------
-void ASTOpMeshRemapIndices::ForEachChild(const std::function<void(ASTChild&)>& f )
+void ASTOpMeshRemapIndices::ForEachChild(const TFunctionRef<void(ASTChild&)> f )
 {
     f( source );
     f( reference );
@@ -83,8 +83,8 @@ void ASTOpMeshRemapIndices::Link( PROGRAM& program, const FLinkerOptions* )
         if (source) args.source = source->linkedAddress;
         if (reference) args.reference = reference->linkedAddress;
 
-        linkedAddress = (OP::ADDRESS)program.m_opAddress.size();
-        program.m_opAddress.push_back((uint32_t)program.m_byteCode.size());
+        linkedAddress = (OP::ADDRESS)program.m_opAddress.Num();
+        program.m_opAddress.Add((uint32_t)program.m_byteCode.Num());
         AppendCode(program.m_byteCode,OP_TYPE::ME_REMAPINDICES);
         AppendCode(program.m_byteCode,args);
     }
diff --git a/Engine/Plugins/Experimental/Mutable/Source/MutableTools/Private/MuT/ASTOpMeshRemoveMask.cpp b/Engine/Plugins/Experimental/Mutable/Source/MutableTools/Private/MuT/ASTOpMeshRemoveMask.cpp
--- a/Engine/Plugins/Experimental/Mutable/Source/MutableTools/Private/MuT/ASTOpMeshRemoveMask.cpp
+++ b/Engine/Plugins/Experimental/Mutable/Source/MutableTools/Private/MuT/ASTOpMeshRemoveMask.cpp
@@ -34,8 +34,7 @@ ASTOpMeshRemoveMask::~ASTOpMeshRemoveMask()
 //---------------------------------------------------------------------------------------------
 void ASTOpMeshRemoveMask::AddRemove( const Ptr<ASTOp>& condition, const Ptr<ASTOp>& mask  )
 {
-    removes.push_back( std::make_pair<>( ASTChild(this,condition),
-                                         ASTChild(this,mask) ) );
+	removes.Add({ ASTChild(this,condition), ASTChild(this,mask) });
 }
 
 
@@ -51,15 +50,13 @@ bool ASTOpMeshRemoveMask::IsEqual(const ASTOp& otherUntyped) const
 
 
 //---------------------------------------------------------------------------------------------
-mu::Ptr<ASTOp> ASTOpMeshRemoveMask::Clone(MapChildFunc& mapChild) const
+mu::Ptr<ASTOp> ASTOpMeshRemoveMask::Clone(MapChildFuncRef mapChild) const
 {
     Ptr<ASTOpMeshRemoveMask> n = new ASTOpMeshRemoveMask();
     n->source = mapChild(source.child());
     for(const auto& r:removes)
     {
-        n->removes.push_back( std::make_pair<>( ASTChild(n,mapChild(r.first.child())),
-                                                ASTChild(n,mapChild(r.second.child()))
-                                                ) );
+		n->removes.Add({ ASTChild(n,mapChild(r.Key.child())), ASTChild(n,mapChild(r.Value.child())) });
     }
     return n;
 }
@@ -73,13 +70,13 @@ void ASTOpMeshRemoveMask::Assert()
 
 
 //---------------------------------------------------------------------------------------------
-void ASTOpMeshRemoveMask::ForEachChild(const std::function<void(ASTChild&)>& f )
+void ASTOpMeshRemoveMask::ForEachChild(const TFunctionRef<void(ASTChild&)> f )
 {
     f(source);
     for(auto& r:removes)
     {
-        f(r.first);
-        f(r.second);
+        f(r.Key);
+        f(r.Value);
     }
 }
 
@@ -90,8 +87,8 @@ uint64 ASTOpMeshRemoveMask::Hash() const
 	uint64 res = std::hash<ASTOp*>()( source.child().get() );
     for(const auto& r:removes)
     {
-        hash_combine( res, r.first.child().get() );
-        hash_combine( res, r.second.child().get() );
+        hash_combine( res, r.Key.child().get() );
+        hash_combine( res, r.Value.child().get() );
     }
     return res;
 }
@@ -103,19 +100,19 @@ void ASTOpMeshRemoveMask::Link( PROGRAM& program, const FLinkerOptions*)
     // Already linked?
     if (!linkedAddress)
     {
-        linkedAddress = (OP::ADDRESS)program.m_opAddress.size();
+        linkedAddress = (OP::ADDRESS)program.m_opAddress.Num();
 
-        program.m_opAddress.push_back((uint32_t)program.m_byteCode.size());
+        program.m_opAddress.Add((uint32_t)program.m_byteCode.Num());
         AppendCode(program.m_byteCode, OP_TYPE::ME_REMOVEMASK);
         OP::ADDRESS sourceAt = source ? source->linkedAddress : 0;
         AppendCode(program.m_byteCode, sourceAt );
-        AppendCode(program.m_byteCode, (uint16_t)removes.size() );
+        AppendCode(program.m_byteCode, (uint16)removes.Num() );
         for ( const auto& b: removes )
         {
-            OP::ADDRESS condition = b.first ? b.first->linkedAddress : 0;
+            OP::ADDRESS condition = b.Key ? b.Key->linkedAddress : 0;
             AppendCode(program.m_byteCode, condition );
 
-            OP::ADDRESS remove = b.second ? b.second->linkedAddress : 0;
+            OP::ADDRESS remove = b.Value ? b.Value->linkedAddress : 0;
             AppendCode(program.m_byteCode, remove );
         }
     }
@@ -136,7 +133,7 @@ namespace
 		mu::Ptr<ASTOp> Apply( const ASTOpMeshRemoveMask* root )
         {
             m_root = root;
-            m_oldToNew.clear();
+            m_oldToNew.Empty();
 
             m_initialSource = root->source.child();
 			mu::Ptr<ASTOp> newSource = Visit( m_initialSource );
@@ -154,24 +151,24 @@ namespace
 
         const ASTOpMeshRemoveMask* m_root;
 		mu::Ptr<ASTOp> m_initialSource;
-        std::unordered_map<mu::Ptr<ASTOp>, mu::Ptr<ASTOp>> m_oldToNew;
-        vector<mu::Ptr<ASTOp>> m_newOps;
+        TMap<mu::Ptr<ASTOp>, mu::Ptr<ASTOp>> m_oldToNew;
+        TArray<mu::Ptr<ASTOp>> m_newOps;
 
 		mu::Ptr<ASTOp> Visit( const mu::Ptr<ASTOp>& at )
         {
             if (!at) return nullptr;
 
             // Newly created?
-            if (std::find(m_newOps.begin(), m_newOps.end(), at )!=m_newOps.end())
+            if (m_newOps.Contains( at ))
             {
                 return at;
             }
 
             // Already visited?
-            auto cacheIt = m_oldToNew.find(at);
-            if (cacheIt!=m_oldToNew.end())
+            auto cacheIt = m_oldToNew.Find(at);
+            if (cacheIt)
             {
-                return cacheIt->second;
+                return *cacheIt;
             }
 
 			mu::Ptr<ASTOp> newAt = at;
@@ -222,7 +219,7 @@ namespace
 
             }
 
-            m_oldToNew[at] = newAt;
+            m_oldToNew.Add(at, newAt);
 
             return newAt;
         }
diff --git a/Engine/Plugins/Experimental/Mutable/Source/MutableTools/Private/MuT/ASTOpMeshRemoveMask.h b/Engine/Plugins/Experimental/Mutable/Source/MutableTools/Private/MuT/ASTOpMeshRemoveMask.h
--- a/Engine/Plugins/Experimental/Mutable/Source/MutableTools/Private/MuT/ASTOpMeshRemoveMask.h
+++ b/Engine/Plugins/Experimental/Mutable/Source/MutableTools/Private/MuT/ASTOpMeshRemoveMask.h
@@ -8,9 +8,6 @@
 #include "MuR/Ptr.h"
 #include "MuT/AST.h"
 
-#include <functional>
-#include <utility>
-
 
 namespace mu
 {
@@ -28,7 +25,7 @@ struct PROGRAM;
 		ASTChild source;
 
 		//! Pairs of remove candidates: condition + mesh to remove
-		vector< std::pair<ASTChild, ASTChild> > removes;
+		TArray< TPair<ASTChild, ASTChild> > removes;
 
 	public:
 
@@ -39,9 +36,9 @@ struct PROGRAM;
 		// ASTOp interface
 		OP_TYPE GetOpType() const override { return OP_TYPE::ME_REMOVEMASK; }
 		uint64 Hash() const override;
-		void ForEachChild(const std::function<void(ASTChild&)>&) override;
+		void ForEachChild(const TFunctionRef<void(ASTChild&)>) override;
 		bool IsEqual(const ASTOp& otherUntyped) const override;
-		Ptr<ASTOp> Clone(MapChildFunc& mapChild) const override;
+		Ptr<ASTOp> Clone(MapChildFuncRef mapChild) const override;
 		void Assert() override;
 		void Link(PROGRAM& program, const FLinkerOptions*) override;
 		Ptr<ASTOp> OptimiseSink(const MODEL_OPTIMIZATION_OPTIONS&, OPTIMIZE_SINK_CONTEXT&) const override;
diff --git a/Engine/Plugins/Experimental/Mutable/Source/MutableTools/Private/MuT/ASTOpMeshTransform.cpp b/Engine/Plugins/Experimental/Mutable/Source/MutableTools/Private/MuT/ASTOpMeshTransform.cpp
--- a/Engine/Plugins/Experimental/Mutable/Source/MutableTools/Private/MuT/ASTOpMeshTransform.cpp
+++ b/Engine/Plugins/Experimental/Mutable/Source/MutableTools/Private/MuT/ASTOpMeshTransform.cpp
@@ -51,7 +51,7 @@ uint64 ASTOpMeshTransform::Hash() const
 }
 
 
-mu::Ptr<ASTOp> ASTOpMeshTransform::Clone(MapChildFunc& mapChild) const
+mu::Ptr<ASTOp> ASTOpMeshTransform::Clone(MapChildFuncRef mapChild) const
 {
     Ptr<ASTOpMeshTransform> n = new ASTOpMeshTransform();
     n->matrix = matrix;
@@ -60,7 +60,7 @@ mu::Ptr<ASTOp> ASTOpMeshTransform::Clone(MapChildFunc& mapChild) const
 }
 
 
-void ASTOpMeshTransform::ForEachChild(const std::function<void(ASTChild&)>& f )
+void ASTOpMeshTransform::ForEachChild(const TFunctionRef<void(ASTChild&)> f )
 {
     f( source );
 }
@@ -78,9 +78,9 @@ void ASTOpMeshTransform::Link( PROGRAM& program, const FLinkerOptions*)
 
         args.matrix = program.AddConstant(matrix);
 
-        linkedAddress = (OP::ADDRESS)program.m_opAddress.size();
+        linkedAddress = (OP::ADDRESS)program.m_opAddress.Num();
         //program.m_code.push_back(op);
-        program.m_opAddress.push_back((uint32_t)program.m_byteCode.size());
+        program.m_opAddress.Add((uint32_t)program.m_byteCode.Num());
         AppendCode(program.m_byteCode,OP_TYPE::ME_TRANSFORM);
         AppendCode(program.m_byteCode,args);
     }
diff --git a/Engine/Plugins/Experimental/Mutable/Source/MutableTools/Private/MuT/ASTOpParameter.cpp b/Engine/Plugins/Experimental/Mutable/Source/MutableTools/Private/MuT/ASTOpParameter.cpp
--- a/Engine/Plugins/Experimental/Mutable/Source/MutableTools/Private/MuT/ASTOpParameter.cpp
+++ b/Engine/Plugins/Experimental/Mutable/Source/MutableTools/Private/MuT/ASTOpParameter.cpp
@@ -27,7 +27,7 @@ ASTOpParameter::~ASTOpParameter()
 }
 
 
-void ASTOpParameter::ForEachChild(const std::function<void(ASTChild&)>& f)
+void ASTOpParameter::ForEachChild(const TFunctionRef<void(ASTChild&)> f)
 {
     for (auto &c: ranges)
     {
@@ -53,18 +53,18 @@ bool ASTOpParameter::IsEqual(const ASTOp& otherUntyped) const
 }
 
 
-mu::Ptr<ASTOp> ASTOpParameter::Clone(MapChildFunc& mapChild) const
+mu::Ptr<ASTOp> ASTOpParameter::Clone(MapChildFuncRef mapChild) const
 {
     Ptr<ASTOpParameter> n = new ASTOpParameter();
     n->type = type;
     n->parameter = parameter;
     for(const auto& c:ranges)
     {
-        n->ranges.emplace_back( n.get(), mapChild(c.rangeSize.child()), c.rangeName, c.rangeUID );
+        n->ranges.Emplace( n.get(), mapChild(c.rangeSize.child()), c.rangeName, c.rangeUID );
     }
     for(const auto& c:additionalImages)
     {
-        n->additionalImages.emplace_back(n,mapChild(c.child()));
+        n->additionalImages.Emplace(n,mapChild(c.child()));
     }
     return n;
 }
@@ -93,27 +93,27 @@ void ASTOpParameter::Link( PROGRAM& program, const FLinkerOptions* )
         OP::ParameterArgs args;
         memset( &args,0, sizeof(args) );
 
-        args.variable = (OP::ADDRESS)program.m_parameters.size();
-        program.m_parameters.push_back(parameter);
+        args.variable = (OP::ADDRESS)program.m_parameters.Num();
+        program.m_parameters.Add(parameter);
 
         for( const auto& d: ranges)
         {
             OP::ADDRESS sizeAt = 0;
-            uint16_t rangeId = 0;
+            uint16 rangeId = 0;
             LinkRange( program, d, sizeAt, rangeId );
-            program.m_parameters.back().m_ranges.push_back(rangeId);
+            program.m_parameters.Last().m_ranges.Add(rangeId);
         }
 
         for( const auto& d: additionalImages)
         {
             OP::ADDRESS descAt = d ? d->linkedAddress : 0;
-            program.m_parameters.back().m_descImages.push_back(descAt);
+            program.m_parameters.Last().m_descImages.Add(descAt);
         }
 
-        linkedAddress = (OP::ADDRESS)program.m_opAddress.size();
+        linkedAddress = (OP::ADDRESS)program.m_opAddress.Num();
         //program.m_code.push_back(op);
 
-        program.m_opAddress.push_back((uint32_t)program.m_byteCode.size());
+        program.m_opAddress.Add((uint32_t)program.m_byteCode.Num());
         AppendCode(program.m_byteCode,type);
         AppendCode(program.m_byteCode,args);
     }
diff --git a/Engine/Plugins/Experimental/Mutable/Source/MutableTools/Private/MuT/ASTOpParameter.h b/Engine/Plugins/Experimental/Mutable/Source/MutableTools/Private/MuT/ASTOpParameter.h
--- a/Engine/Plugins/Experimental/Mutable/Source/MutableTools/Private/MuT/ASTOpParameter.h
+++ b/Engine/Plugins/Experimental/Mutable/Source/MutableTools/Private/MuT/ASTOpParameter.h
@@ -10,8 +10,6 @@
 #include "MuR/Ptr.h"
 #include "MuT/AST.h"
 
-#include <functional>
-
 
 namespace mu
 {
@@ -32,10 +30,10 @@ struct PROGRAM;
 		PARAMETER_DESC parameter;
 
 		//! Ranges adding dimensions to this parameter
-		vector<RANGE_DATA> ranges;
+		TArray<RANGE_DATA> ranges;
 
 		//! Additional images attached to the parameter
-		vector<ASTChild> additionalImages;
+		TArray<ASTChild> additionalImages;
 
 	public:
 
@@ -43,9 +41,9 @@ struct PROGRAM;
 
 		OP_TYPE GetOpType() const override { return type; }
 		uint64 Hash() const override;
-		void ForEachChild(const std::function<void(ASTChild&)>&) override;
+		void ForEachChild(const TFunctionRef<void(ASTChild&)>) override;
 		bool IsEqual(const ASTOp& otherUntyped) const override;
-		Ptr<ASTOp> Clone(MapChildFunc& mapChild) const override;
+		Ptr<ASTOp> Clone(MapChildFuncRef mapChild) const override;
 		void Assert() override;
 		void Link(PROGRAM& program, const FLinkerOptions*) override;
 		int EvaluateInt(ASTOpList& facts, bool& unknown) const override;
diff --git a/Engine/Plugins/Experimental/Mutable/Source/MutableTools/Private/MuT/ASTOpScalarCurve.cpp b/Engine/Plugins/Experimental/Mutable/Source/MutableTools/Private/MuT/ASTOpScalarCurve.cpp
--- a/Engine/Plugins/Experimental/Mutable/Source/MutableTools/Private/MuT/ASTOpScalarCurve.cpp
+++ b/Engine/Plugins/Experimental/Mutable/Source/MutableTools/Private/MuT/ASTOpScalarCurve.cpp
@@ -32,7 +32,7 @@ ASTOpScalarCurve::~ASTOpScalarCurve()
 }
 
 
-void ASTOpScalarCurve::ForEachChild(const std::function<void(ASTChild&)>& f)
+void ASTOpScalarCurve::ForEachChild(const TFunctionRef<void(ASTChild&)> f)
 {
     f( time );
 }
@@ -41,7 +41,7 @@ void ASTOpScalarCurve::ForEachChild(const std::function<void(ASTChild&)>& f)
 uint64 ASTOpScalarCurve::Hash() const
 {
 	uint64 res = std::hash<uint64>()(size_t(OP_TYPE::SC_CURVE));
-    hash_combine( res, curve.keyFrames.size() );
+    hash_combine( res, curve.keyFrames.Num() );
     return res;
 }
 
@@ -56,7 +56,7 @@ bool ASTOpScalarCurve::IsEqual(const ASTOp& otherUntyped) const
 }
 
 
-mu::Ptr<ASTOp> ASTOpScalarCurve::Clone(MapChildFunc& mapChild) const
+mu::Ptr<ASTOp> ASTOpScalarCurve::Clone(MapChildFuncRef mapChild) const
 {
     Ptr<ASTOpScalarCurve> n = new ASTOpScalarCurve();
     n->curve = curve;
@@ -74,9 +74,9 @@ void ASTOpScalarCurve::Link( PROGRAM& program, const FLinkerOptions* )
         args.time = time->linkedAddress;
         args.curve = program.AddConstant(curve);
 
-        linkedAddress = (OP::ADDRESS)program.m_opAddress.size();
+        linkedAddress = (OP::ADDRESS)program.m_opAddress.Num();
         //program.m_code.push_back(op);
-        program.m_opAddress.push_back((uint32_t)program.m_byteCode.size());
+        program.m_opAddress.Add((uint32_t)program.m_byteCode.Num());
         AppendCode(program.m_byteCode,OP_TYPE::SC_CURVE);
         AppendCode(program.m_byteCode,args);
     }
diff --git a/Engine/Plugins/Experimental/Mutable/Source/MutableTools/Private/MuT/ASTOpSwitch.cpp b/Engine/Plugins/Experimental/Mutable/Source/MutableTools/Private/MuT/ASTOpSwitch.cpp
--- a/Engine/Plugins/Experimental/Mutable/Source/MutableTools/Private/MuT/ASTOpSwitch.cpp
+++ b/Engine/Plugins/Experimental/Mutable/Source/MutableTools/Private/MuT/ASTOpSwitch.cpp
@@ -49,7 +49,7 @@ bool ASTOpSwitch::IsEqual(const ASTOp& otherUntyped) const
 
 
 //-------------------------------------------------------------------------------------------------
-mu::Ptr<ASTOp> ASTOpSwitch::Clone(MapChildFunc& mapChild) const
+mu::Ptr<ASTOp> ASTOpSwitch::Clone(MapChildFuncRef mapChild) const
 {
     Ptr<ASTOpSwitch> n = new ASTOpSwitch();
     n->type = type;
@@ -57,7 +57,7 @@ mu::Ptr<ASTOp> ASTOpSwitch::Clone(MapChildFunc& mapChild) const
     n->def = mapChild(def.child());
     for(const auto& c:cases)
     {
-        n->cases.emplace_back(c.condition, n, mapChild(c.branch.child()));
+        n->cases.Emplace(c.condition, n, mapChild(c.branch.child()));
     }
     return n;
 }
@@ -102,7 +102,7 @@ uint64 ASTOpSwitch::Hash() const
 //-------------------------------------------------------------------------------------------------
 mu::Ptr<ASTOp> ASTOpSwitch::GetFirstValidValue()
 {
-    for ( size_t i=0; i<cases.size(); ++i )
+    for ( int32 i=0; i<cases.Num(); ++i )
     {
         if ( cases[i].branch )
         {
@@ -118,7 +118,7 @@ bool ASTOpSwitch::IsCompatibleWith( const ASTOpSwitch* other ) const
 {
     if (!other) return false;
     if (variable.child()!=other->variable.child()) return false;
-    if (cases.size()!=other->cases.size()) return false;
+    if (cases.Num()!=other->cases.Num()) return false;
     for (const auto& c: cases)
     {
         bool found = false;
@@ -156,7 +156,7 @@ mu::Ptr<ASTOp> ASTOpSwitch::FindBranch( int32_t condition ) const
 
 
 //-------------------------------------------------------------------------------------------------
-void ASTOpSwitch::ForEachChild(const std::function<void(ASTChild&)>& f )
+void ASTOpSwitch::ForEachChild(const TFunctionRef<void(ASTChild&)> f )
 {
     f(variable);
     f(def);
@@ -173,8 +173,8 @@ void ASTOpSwitch::Link( PROGRAM& program, const FLinkerOptions* )
     // Already linked?
     if (!linkedAddress)
     {        
-        linkedAddress = (OP::ADDRESS)program.m_opAddress.size();
-        program.m_opAddress.push_back((uint32_t)program.m_byteCode.size());
+        linkedAddress = (OP::ADDRESS)program.m_opAddress.Num();
+        program.m_opAddress.Add((uint32_t)program.m_byteCode.Num());
  
         OP::ADDRESS VarAddress = variable ? variable->linkedAddress : 0;
         OP::ADDRESS DefAddress = def ? def->linkedAddress : 0;
@@ -182,7 +182,7 @@ void ASTOpSwitch::Link( PROGRAM& program, const FLinkerOptions* )
         AppendCode( program.m_byteCode, type );
         AppendCode( program.m_byteCode, VarAddress );
         AppendCode( program.m_byteCode, DefAddress );
-        AppendCode( program.m_byteCode, (uint32_t)cases.size() );
+        AppendCode( program.m_byteCode, (uint32_t)cases.Num() );
 
         for ( const CASE& Case : cases )
         {
@@ -233,7 +233,7 @@ FImageDesc ASTOpSwitch::GetImageDesc( bool returnBestOption, class GetImageDescC
         first = false;
     }
 
-    for ( int i=0; i<(int)cases.size(); ++i )
+    for ( int i=0; i<cases.Num(); ++i )
     {
         if ( cases[i].branch )
         {
@@ -419,7 +419,7 @@ mu::Ptr<ASTOp> ASTOpSwitch::OptimiseSemantic( const MODEL_OPTIMIZATION_OPTIONS&
         Ptr<ASTOp> branch = def.child();
 
         auto typedCondition = dynamic_cast<const ASTOpFixed*>(variable.child().get());
-        for ( size_t o=0; o<cases.size(); ++o )
+        for ( int32 o=0; o<cases.Num(); ++o )
         {
             if ( cases[o].branch &&
                 typedCondition->op.args.IntConstant.value
diff --git a/Engine/Plugins/Experimental/Mutable/Source/MutableTools/Private/MuT/ASTOpSwitch.h b/Engine/Plugins/Experimental/Mutable/Source/MutableTools/Private/MuT/ASTOpSwitch.h
--- a/Engine/Plugins/Experimental/Mutable/Source/MutableTools/Private/MuT/ASTOpSwitch.h
+++ b/Engine/Plugins/Experimental/Mutable/Source/MutableTools/Private/MuT/ASTOpSwitch.h
@@ -9,8 +9,6 @@
 #include "MuR/Ptr.h"
 #include "MuT/AST.h"
 
-#include <functional>
-
 
 namespace mu
 {
@@ -53,7 +51,7 @@ template <class SCALAR> class vec4;
 			}
 		};
 
-		vector<CASE> cases;
+		TArray<CASE> cases;
 
 	public:
 
@@ -64,9 +62,9 @@ template <class SCALAR> class vec4;
 		// ASTOp interface
 		OP_TYPE GetOpType() const override { return type; }
 		uint64 Hash() const override;
-		void ForEachChild(const std::function<void(ASTChild&)>& f) override;
+		void ForEachChild(const TFunctionRef<void(ASTChild&)> f) override;
 		bool IsEqual(const ASTOp& otherUntyped) const override;
-		Ptr<ASTOp> Clone(MapChildFunc& mapChild) const override;
+		Ptr<ASTOp> Clone(MapChildFuncRef mapChild) const override;
 		void Assert() override;
 		void Link(PROGRAM& program, const FLinkerOptions* Options) override;
 		FImageDesc GetImageDesc(bool returnBestOption,
diff --git a/Engine/Plugins/Experimental/Mutable/Source/MutableTools/Private/MuT/CodeGenerator.cpp b/Engine/Plugins/Experimental/Mutable/Source/MutableTools/Private/MuT/CodeGenerator.cpp
--- a/Engine/Plugins/Experimental/Mutable/Source/MutableTools/Private/MuT/CodeGenerator.cpp
+++ b/Engine/Plugins/Experimental/Mutable/Source/MutableTools/Private/MuT/CodeGenerator.cpp
@@ -92,7 +92,7 @@ namespace mu
         //UE_LOG(LogMutableCore, Verbose, " Begin partial optimisation.");
 
         ASTOpList roots;
-        roots.push_back( lastCompOp );
+        roots.Add( lastCompOp );
 
         //size_t initialOps = ASTOp::CountNodes(roots);
         //(void)initialOps;
@@ -219,7 +219,7 @@ namespace mu
         m_pErrorLog = new ErrorLog;
 
         // Add the parent at the top of the hierarchy
-        m_currentParents.push_back( PARENT_KEY() );
+        m_currentParents.Add( PARENT_KEY() );
     }
 
     //---------------------------------------------------------------------------------------------
@@ -256,7 +256,7 @@ namespace mu
                 MUTABLE_CPUPROFILER_SCOPE(MainPassState);
 
                 Ptr<ASTOp> stateRoot = Generate(pNode);
-                m_states.emplace_back( s.first, stateRoot );
+                m_states.Emplace( s.Key, stateRoot );
 
                 ++m_currentStateIndex;
             }
@@ -266,7 +266,7 @@ namespace mu
 
         // Free caches
         m_compiled.Reset();
-        m_constantMeshes.clear();
+        m_constantMeshes.Empty();
         m_addedLayouts.Empty();
         m_nodeVariables.clear();
         m_generatedMeshes.Reset();
@@ -275,10 +275,10 @@ namespace mu
         m_generatedTables.clear();
         m_firstPass = FirstPassGenerator();
         m_currentBottomUpState = BOTTOM_UP_STATE();
-        m_overrideLayoutsStack.clear();
-        m_imageState.clear();
-        m_currentParents.clear();
-        m_currentObject.clear();
+        m_overrideLayoutsStack.Empty();
+        m_imageState.Empty();
+        m_currentParents.Empty();
+        m_currentObject.Empty();
         m_additionalComponents.clear();
 
 
@@ -352,7 +352,7 @@ namespace mu
         {
             // This happens only if we generate a node graph that has a NodeSurfaceNew at the root.
             SURFACE_GENERATION_RESULT surfResult;
-            const vector<FirstPassGenerator::SURFACE::EDIT> edits;
+            const TArray<FirstPassGenerator::SURFACE::EDIT> edits;
             GenerateSurface( surfResult, surfNode, edits );
             return surfResult.surfaceOp;
         }
@@ -463,12 +463,12 @@ namespace mu
             // See if there is a string column. If there is one, we will use it as names for the
             // options. Only the first string column will be used.
             int nameCol = -1;
-            std::size_t cols = pTable->GetPrivate()->m_columns.size();
-            for ( std::size_t c=0; c<cols && nameCol<0; ++c )
+            int32 cols = pTable->GetPrivate()->m_columns.Num();
+            for ( int32 c=0; c<cols && nameCol<0; ++c )
             {
                 if ( pTable->GetPrivate()->m_columns[c].m_type==TCT_STRING )
                 {
-                    nameCol = (int)c;
+                    nameCol = c;
                 }
             }
 
@@ -477,12 +477,12 @@ namespace mu
 				PARAMETER_DESC::INT_VALUE_DESC nullValue;
 				nullValue.m_value = -1;
 				nullValue.m_name = "None";
-				param.m_possibleValues.push_back(nullValue);
+				param.m_possibleValues.Add(nullValue);
 				param.m_defaultValue.m_int = nullValue.m_value;
 			}
 
             // Add every row
-            std::size_t rows = pTable->GetPrivate()->m_rows.size();
+            int32 rows = pTable->GetPrivate()->m_rows.Num();
 			for (size_t i = 0; i<rows; ++i)
             {
                 PARAMETER_DESC::INT_VALUE_DESC value;
@@ -493,7 +493,7 @@ namespace mu
                     value.m_name = pTable->GetPrivate()->m_rows[i].m_values[nameCol].m_string;
                 }
 
-                param.m_possibleValues.push_back( value );
+                param.m_possibleValues.Add( value );
 
                 // The first row is the default one
                 if (i==0)
@@ -580,8 +580,8 @@ namespace mu
              blend = GenerateImageSize
                  (
                      blend,
-                     FImageSize( (uint16_t)m_imageState.back().m_imageRect.size[0],
-                                 (uint16_t)m_imageState.back().m_imageRect.size[1] )
+                     FImageSize( (uint16)m_imageState.Last().m_imageRect.size[0],
+                                 (uint16)m_imageState.Last().m_imageRect.size[1] )
                  );
              op->SetChild(op->op.args.ImageLayer.blended, blend );
 
@@ -601,8 +601,8 @@ namespace mu
              mask = GenerateImageSize
                  (
                      mask,
-                     FImageSize( (uint16_t)m_imageState.back().m_imageRect.size[0],
-                                 (uint16_t)m_imageState.back().m_imageRect.size[1] )
+                     FImageSize( (uint16)m_imageState.Last().m_imageRect.size[0],
+                                 (uint16)m_imageState.Last().m_imageRect.size[1] )
                  );
              op->SetChild(op->op.args.ImageLayer.mask, mask );
 
@@ -658,7 +658,7 @@ namespace mu
 		}
         else if ( auto pTypedSV = dynamic_cast<const NodeSurfaceVariation*>(pNode) )
 		{
-            if (pTypedSV->GetPrivate()->m_defaultSurfaces.size())
+            if (pTypedSV->GetPrivate()->m_defaultSurfaces.Num())
 			{
                 pResult = FindSourceMesh(pTypedSV->GetPrivate()->m_defaultSurfaces[0].get());
 			}
@@ -722,9 +722,9 @@ namespace mu
     {
         // Get the parent component layout
         const NodeObjectNew::Private* pParent = nullptr;
-        if (m_currentParents.size()>2)
+        if (m_currentParents.Num()>2)
         {
-            pParent = (m_currentParents.end()-2)->m_pObject;
+            pParent = m_currentParents[m_currentParents.Num()-2].m_pObject;
         }
 
         const NodeLayout* pNodeLayout = nullptr;
@@ -732,10 +732,10 @@ namespace mu
         {
             pNodeLayout = pParent->GetLayout
                 (
-                    m_currentParents.back().m_lod,
-                    m_currentParents.back().m_component,
-                    m_currentParents.back().m_surface,
-                    m_currentParents.back().m_texture
+                    m_currentParents.Last().m_lod,
+                    m_currentParents.Last().m_component,
+                    m_currentParents.Last().m_surface,
+                    m_currentParents.Last().m_texture
                 ).get();
         }
 
@@ -746,7 +746,7 @@ namespace mu
                 (
                     buf, 256,
                     "In object [%s] NodePatchImage couldn't find the layout in parent.",
-                    m_currentParents.rbegin()->m_pObject->m_name.c_str()
+                    m_currentParents.Last().m_pObject->m_name.c_str()
                 );
 
             m_pErrorLog->GetPrivate()->Add( buf, ELMT_ERROR, node.m_errorContext );
@@ -764,11 +764,11 @@ namespace mu
 
         // Create the expression for each component in this object
         // TODO: More components per operation
-        for ( std::size_t t=0; t<node.m_components.size(); ++t )
+        for ( int32 t=0; t<node.m_components.Num(); ++t )
         {
             if ( const NodeComponent* pComponentNode = node.m_components[t].get() )
             {
-                m_currentParents.back().m_component = (int)t;
+                m_currentParents.Last().m_component = (int)t;
 
                 Ptr<ASTOp> componentOp = Generate( pComponentNode );
 
@@ -788,8 +788,8 @@ namespace mu
 
         // Add components from child objects
         ADDITIONAL_COMPONENT_KEY thisKey;
-        thisKey.m_lod = m_currentParents.back().m_lod;
-        thisKey.m_pObject = m_currentParents.back().m_pObject;
+        thisKey.m_lod = m_currentParents.Last().m_lod;
+        thisKey.m_pObject = m_currentParents.Last().m_pObject;
         auto addIt = m_additionalComponents.find( thisKey );
         if (addIt!=m_additionalComponents.end())
         {
@@ -817,13 +817,13 @@ namespace mu
 
         // Store for possible parent objects if necessary
         // 2 is because there must be a parent and there is always a null element as well.
-        if (lastCompOp && m_currentParents.size()>2)
+        if (lastCompOp && m_currentParents.Num()>2)
         {
-            const auto& parentObjectKey = m_currentParents[ m_currentParents.size()-2 ];
+            const auto& parentObjectKey = m_currentParents[ m_currentParents.Num()-2 ];
             ADDITIONAL_COMPONENT_KEY parentKey;
-            parentKey.m_lod = m_currentParents.back().m_lod;
+            parentKey.m_lod = m_currentParents.Last().m_lod;
             parentKey.m_pObject = parentObjectKey.m_pObject;
-            m_additionalComponents[parentKey].push_back( lastCompOp );
+            m_additionalComponents[parentKey].Add( lastCompOp );
         }
 
         return lastCompOp;
@@ -834,7 +834,7 @@ namespace mu
     //---------------------------------------------------------------------------------------------
     void CodeGenerator::GenerateSurface( SURFACE_GENERATION_RESULT& result,
                                          NodeSurfaceNewPtrConst surfaceNode,
-                                         const vector<FirstPassGenerator::SURFACE::EDIT>& edits )
+                                         const TArray<FirstPassGenerator::SURFACE::EDIT>& edits )
     {
         //MUTABLE_CPUPROFILER_SCOPE(GenerateSurface);
 
@@ -853,12 +853,12 @@ namespace mu
         // Generate the mesh
         //------------------------------------------------------------------------
         NodeMeshPtr pMesh;
-        if ( node.m_meshes.size()>0 )
+        if ( node.m_meshes.Num()>0 )
         {
             pMesh = node.m_meshes[0].m_pMesh;
 
             // TODO: Support more meshes
-            check( node.m_meshes.size()==1 );
+            check( node.m_meshes.Num()==1 );
         }
 
 
@@ -873,30 +873,28 @@ namespace mu
 		// Do we need to generate the mesh? Or was it already generated for state conditions 
 		// accepting the current state?
 		FirstPassGenerator::SURFACE* targetSurface = nullptr;
-		for (vector<FirstPassGenerator::SURFACE>::iterator its = m_firstPass.surfaces.begin()
-			; its != m_firstPass.surfaces.end()
-			; ++its)
+		for (FirstPassGenerator::SURFACE& its : m_firstPass.surfaces)
 		{
-            if (its->node->GetPrivate() == &node)
+            if (its.node->GetPrivate() == &node)
             {
                 // Check state conditions
                 bool surfaceValidForThisState =
-                    m_currentStateIndex >= int(its->stateCondition.size())
+                    m_currentStateIndex >= its.stateCondition.Num()
                     ||
-                    its->stateCondition[m_currentStateIndex];
+                    its.stateCondition[m_currentStateIndex];
 
                 if (surfaceValidForThisState)
                 {
-                    if (its->resultSurfaceOp)
+                    if (its.resultSurfaceOp)
                     {
                         // Reuse the entire surface
-                        result.surfaceOp = its->resultSurfaceOp;
+                        result.surfaceOp = its.resultSurfaceOp;
                         return;
                     }
                     else
                     {
                         // Not already generated, we will generate this
-                        targetSurface = &(*its);
+                        targetSurface = &its;
                     }
                 }
             }
@@ -914,14 +912,14 @@ namespace mu
             Ptr<ASTOp> lastMeshOp;
 
 			// Store the data necessary to apply modifiers for the pre-normal operations stage.
-			m_activeTags.push_back(node.m_tags);
+			m_activeTags.Add(node.m_tags);
 
             // Generate the mesh
             GenerateMesh( meshResults, pMesh );
             lastMeshOp = meshResults.meshOp;
-            meshResults.extraMeshLayouts.resize( edits.size() );
+            meshResults.extraMeshLayouts.SetNum( edits.Num() );
 
-			m_activeTags.pop_back();
+			m_activeTags.Pop();
 
             // Let's remember the base mesh of each added mesh, so that we can use them for the
             // "remove mesh" operations that apply to them instead of the base.
@@ -929,7 +927,7 @@ namespace mu
 
 
             // Apply mesh merges from child objects "edit surface" nodes
-            for ( size_t editIndex=0; editIndex<edits.size(); ++editIndex )
+            for ( int32 editIndex=0; editIndex<edits.Num(); ++editIndex )
             {
                 const auto& e = edits[editIndex];
 
@@ -938,12 +936,12 @@ namespace mu
                     if ( NodeMeshPtr pAdd = e.node->m_pMesh->GetAdd() )
                     {
 						// Store the data necessary to apply modifiers for the pre-normal operations stage.
-						m_activeTags.push_back(e.node->m_tags);
+						m_activeTags.Add(e.node->m_tags);
 						
 						MESH_GENERATION_RESULT addResults;
                         GenerateMesh( addResults, pAdd );
 
-						m_activeTags.pop_back();
+						m_activeTags.Pop();
 
                         baseMeshesForEachAddedMesh[e.node] = addResults.baseMeshOp;
 
@@ -1129,7 +1127,7 @@ namespace mu
 				bModifiersForBeforeOperations, node.m_errorContext);
 
             // Layouts
-            for ( int l=0; l<(int)meshResults.layouts.size(); ++l )
+            for ( int32 l=0; l<meshResults.layouts.Num(); ++l )
             {
                 Ptr<ASTOp> layoutOp = nullptr;
 
@@ -1156,11 +1154,11 @@ namespace mu
                         if ( !data.meshFragment )
                         {
                             // No mesh to add, we assume there are no layouts to add either.
-                            check(data.layouts.empty());
+                            check(data.layouts.IsEmpty());
                             continue;
                         }
 
-                        if ( data.layouts.size() != meshResults.layouts.size() )
+                        if ( data.layouts.Num() != meshResults.layouts.Num() )
                         {
                             m_pErrorLog->GetPrivate()->Add
                                 ( "Merged layout has been ignored because the number of layouts"
@@ -1235,14 +1233,14 @@ namespace mu
                             op->op.type = OP_TYPE::ME_APPLYLAYOUT;
                             op->SetChild(op->op.args.MeshApplyLayout.mesh, lastMeshOp );
                             op->SetChild(op->op.args.MeshApplyLayout.layout, layoutOp );
-                            op->op.args.MeshApplyLayout.channel = (uint16_t)l;
+                            op->op.args.MeshApplyLayout.channel = (uint16)l;
                             lastMeshOp = op;
                         }
                     }
 
                 }
 
-                meshResults.layoutOps.push_back( layoutOp );
+                meshResults.layoutOps.Add( layoutOp );
             }
 
             // Store in the surface for later use.
@@ -1252,7 +1250,7 @@ namespace mu
 
         // Create the expression for each texture
         //------------------------------------------------------------------------
-        for ( std::size_t t=0; t<node.m_images.size(); ++t )
+        for ( int32 t=0; t<node.m_images.Num(); ++t )
         {
             //MUTABLE_CPUPROFILER_SCOPE(SurfaceTexture);
 
@@ -1300,7 +1298,7 @@ namespace mu
 
 				if (desc.m_size[0] == 0 || desc.m_size[1] == 0)
 				{
-					int currentLOD = m_currentParents.back().m_lod;
+					int currentLOD = m_currentParents.Last().m_lod;
 					char buf[256];
 					mutable_snprintf(buf, 256, "An image block for [%s] [%s] [%s] at lod [%d] has zero size and will not be generated. ", 
 						node.m_images[t].m_name.c_str(),
@@ -1333,15 +1331,15 @@ namespace mu
                     newState.m_imageRect = box<vec2<int>>::FromMinSize( vec2<int>(0,0), desc.m_size );
                     newState.m_layoutBlock = -1;
                     newState.m_pLayout = nullptr;
-                    m_imageState.push_back( newState );
+                    m_imageState.Add( newState );
 
                     Ptr<ASTOp> imageAd = Generate( pImageNode );
 
                     // Look for patches to this block
-                    for( size_t editIndex=0; editIndex<edits.size(); ++editIndex )
+                    for( int32 editIndex=0; editIndex<edits.Num(); ++editIndex )
                     {
                         const auto& e = edits[editIndex];
-                        if ( t<e.node->m_textures.size() )
+                        if ( t<e.node->m_textures.Num() )
                         {
                             if ( const NodePatchImage* pPatch = e.node->m_textures[t].m_pPatch.get() )
                             {
@@ -1353,7 +1351,7 @@ namespace mu
 
                     check(imageAd);
 
-                    m_imageState.pop_back();
+                    m_imageState.Pop();
 
                     if (mipmapNode)
                     {
@@ -1408,8 +1406,8 @@ namespace mu
                           ==
                           CompilerOptions::TextureLayoutStrategy::Pack )
                 {
-                    if( layout >= (int)meshResults.layouts.size() ||
-                        layout >= (int)meshResults.layoutOps.size() )
+                    if( layout >= meshResults.layouts.Num() ||
+                        layout >= meshResults.layoutOps.Num() )
                     {
                         m_pErrorLog->GetPrivate()->Add("Missing layout in object, or its parent.",
                                                        ELMT_ERROR, node.m_errorContext);
@@ -1443,7 +1441,7 @@ namespace mu
 							//UE_LOG( LogMutableCore, Warning, TEXT("A texture [%s] has been resized from [%d x %d] to [%d x %d] because it didn't fit the layout [%d x %d]. "),
 							//	ANSI_TO_TCHAR(node.m_images[t].m_name.c_str()), oldSize[0], oldSize[1], desc.m_size[0], desc.m_size[1], grid[0], grid[1] );
 
-							int currentLOD = m_currentParents.back().m_lod;
+							int currentLOD = m_currentParents.Last().m_lod;
 							char buf[256];
 							mutable_snprintf(buf, 256, "A texture [%s] for material [%s] parameter [%s] in LOD [%d] has been resized from [%d x %d] to [%d x %d] because it didn't fit the layout [%d x %d]. ",
 								node.m_images[t].m_name.c_str(), 
@@ -1454,8 +1452,8 @@ namespace mu
 							m_pErrorLog->GetPrivate()->Add(buf, ELMT_INFO, node.m_errorContext);
 						}
 
-                        int blockSizeX = std::max( 1, desc.m_size[0] / grid[0]);
-                        int blockSizeY = std::max( 1, desc.m_size[1] / grid[1]);
+                        int blockSizeX = FMath::Max( 1, desc.m_size[0] / grid[0]);
+                        int blockSizeY = FMath::Max( 1, desc.m_size[1] / grid[1]);
 
 						bool bBlocksHaveMips = desc.m_lods > 1;
 
@@ -1465,8 +1463,8 @@ namespace mu
                             Ptr<ASTOpFixed> bop = new ASTOpFixed();
                             bop->op.type = OP_TYPE::IM_BLANKLAYOUT;
                             bop->SetChild(bop->op.args.ImageBlankLayout.layout, meshResults.layoutOps[ layout ] );
-                            bop->op.args.ImageBlankLayout.blockSize[0] = uint16_t(blockSizeX);
-                            bop->op.args.ImageBlankLayout.blockSize[1] = uint16_t(blockSizeY);
+                            bop->op.args.ImageBlankLayout.blockSize[0] = uint16(blockSizeX);
+                            bop->op.args.ImageBlankLayout.blockSize[1] = uint16(blockSizeY);
                             // We support block compression directly here, but not non-block compression
                             if ( GetImageFormatData(desc.m_format).m_pixelsPerBlockX == 0 )
                             {
@@ -1507,23 +1505,20 @@ namespace mu
                             newState.m_imageRect = rect;
                             newState.m_layoutBlock = pLayout->m_blocks[b].m_id;
                             newState.m_pLayout = pLayout;
-                            m_imageState.push_back( newState );
+                            m_imageState.Add( newState );
 
                             Ptr<ASTOp> blockAd = Generate( pImageNode );
 
                             // Look for patches to this block
-                            for( size_t editIndex=0; editIndex<edits.size(); ++editIndex )
+                            for( int32 editIndex=0; editIndex<edits.Num(); ++editIndex )
                             {
                                 const auto& e = edits[editIndex];
-                                if ( t<e.node->m_textures.size() )
+                                if ( t<e.node->m_textures.Num() )
                                 {
                                     if ( const NodePatchImage* pPatch = e.node->m_textures[t].m_pPatch.get() )
                                     {
                                         // Is the current block to be patched?
-                                        if ( std::find( pPatch->GetPrivate()->m_blocks.begin(),
-                                                        pPatch->GetPrivate()->m_blocks.end(),
-                                                        b )
-                                             != pPatch->GetPrivate()->m_blocks.end() )
+                                        if (pPatch->GetPrivate()->m_blocks.Contains(b))
                                         {
                                             blockAd = GenerateImageBlockPatch(blockAd,pPatch,e.condition);
                                         }
@@ -1546,20 +1541,20 @@ namespace mu
 
                             imageAd = composeOp;
 
-                            m_imageState.pop_back();
+                            m_imageState.Pop();
                         }
                         check(imageAd);
 
                         // Apply composition of blocks coming from child objects
-                        for( size_t editIndex=0; editIndex<edits.size(); ++editIndex )
+                        for( int32 editIndex=0; editIndex<edits.Num(); ++editIndex )
                         {
                             const auto& e = edits[editIndex];
-                            if ( t<e.node->m_textures.size() )
+                            if ( t<e.node->m_textures.Num() )
                             {
                                 auto pExtend = e.node->m_textures[t].m_pExtend;
                                 if (pExtend)
                                 {
-                                    if ( layout>=int(meshResults.extraMeshLayouts[editIndex].layouts.size())
+                                    if ( layout>=meshResults.extraMeshLayouts[editIndex].layouts.Num()
                                          ||
                                          !meshResults.extraMeshLayouts[editIndex].layouts[layout] )
                                     {
@@ -1568,7 +1563,7 @@ namespace mu
                                             (
                                                 buf, 256,
                                                 "Trying to extend a layout that doesn't exist in object [%s].",
-                                                m_currentParents.rbegin()->m_pObject->m_name.c_str()
+                                                m_currentParents.Last().m_pObject->m_name.c_str()
                                             );
 
                                         m_pErrorLog->GetPrivate()->Add( buf, ELMT_ERROR, node.m_errorContext );
@@ -1612,17 +1607,17 @@ namespace mu
                                             newState.m_imageRect = rect;
                                             newState.m_layoutBlock = pExtendLayout->m_blocks[b].m_id;
                                             newState.m_pLayout = pExtendLayout;
-                                            m_imageState.push_back( newState );
+                                            m_imageState.Add( newState );
 
                                             Ptr<ASTOp> fragmentAd = Generate( pExtend );
 
-                                            m_imageState.pop_back();
+                                            m_imageState.Pop();
 
                                             // Adjust the format and size of the block to be added
                                             Ptr<ASTOp> formatted = GenerateImageFormat( fragmentAd, GetUncompressedFormat(desc.m_format) );
                                             FImageSize expectedSize;
-                                            expectedSize[0] = (uint16_t)(blockSizeX * blockRect.size[0]);
-                                            expectedSize[1] = (uint16_t)(blockSizeY * blockRect.size[1]);
+                                            expectedSize[0] = (uint16)(blockSizeX * blockRect.size[0]);
+                                            expectedSize[1] = (uint16)(blockSizeY * blockRect.size[1]);
                                             formatted = GenerateImageSize( formatted, expectedSize );
                                             // Direct compression support
                                             // formatted = GenerateImageUncompressed( formatted );
@@ -1685,7 +1680,7 @@ namespace mu
 
                             int mipsX = (int)ceilf( logf( (float)blockSizeX/finfo.m_pixelsPerBlockX )/logf(2.0f) );
                             int mipsY = (int)ceilf( logf( (float)blockSizeY/finfo.m_pixelsPerBlockY )/logf(2.0f) );
-                            mop->BlockLevels = (uint8_t)std::max( mipsX, mipsY );
+                            mop->BlockLevels = (uint8_t)FMath::Max( mipsX, mipsY );
 
 							mop->AddressMode	   = mipmapNode->GetPrivate()->m_settings.m_addressMode;
 							mop->FilterType		   = mipmapNode->GetPrivate()->m_settings.m_filterType;
@@ -1712,7 +1707,7 @@ namespace mu
 
 							int mipsX = (int)ceilf(logf((float)blockSizeX / finfo.m_pixelsPerBlockX) / logf(2.0f));
 							int mipsY = (int)ceilf(logf((float)blockSizeY / finfo.m_pixelsPerBlockY) / logf(2.0f));
-							mop->BlockLevels = (uint8_t)std::max(mipsX, mipsY);
+							mop->BlockLevels = (uint8_t)FMath::Max(mipsX, mipsY);
 
 							// Not important for the end of the mip tail?
 							mop->AddressMode = EAddressMode::AM_CLAMP;
@@ -1768,7 +1763,7 @@ namespace mu
 
         // Create the expression for each vector
         //------------------------------------------------------------------------
-        for (std::size_t t = 0; t<node.m_vectors.size(); ++t)
+        for (int32 t = 0; t<node.m_vectors.Num(); ++t)
         {
             //MUTABLE_CPUPROFILER_SCOPE(SurfaceVector);
 
@@ -1791,7 +1786,7 @@ namespace mu
 
         // Create the expression for each scalar
         //------------------------------------------------------------------------
-        for ( std::size_t t = 0; t < node.m_scalars.size(); ++t )
+        for ( int32 t = 0; t < node.m_scalars.Num(); ++t )
         {
             // MUTABLE_CPUPROFILER_SCOPE(SurfaceScalar);
 
@@ -1814,7 +1809,7 @@ namespace mu
 
         // Create the expression for each string
         //------------------------------------------------------------------------
-        for ( std::size_t t = 0; t < node.m_strings.size(); ++t )
+        for ( int32 t = 0; t < node.m_strings.Num(); ++t )
         {
             if ( NodeStringPtr pStringNode = node.m_strings[t].m_pString )
             {
@@ -1872,21 +1867,20 @@ namespace mu
         int surfaceID=1;
 
         // Look for all surfaces that belong to this component
-		for (vector<FirstPassGenerator::SURFACE>::const_iterator its = m_firstPass.surfaces.begin()
-			; its != m_firstPass.surfaces.end()
-            ; ++its, ++surfaceID)
+		for (int32 i = 0; i<m_firstPass.surfaces.Num(); ++i, ++surfaceID)
 		{
-			if (its->component==&node)
+			const auto& its = m_firstPass.surfaces[i];
+			if (its.component==&node)
 			{
                 // Apply state conditions: only generate it if it enabled in this state
                 {
                     bool enabledInThisState = true;
-                    if (its->stateCondition.size() && m_currentStateIndex >= 0)
+                    if (its.stateCondition.Num() && m_currentStateIndex >= 0)
                     {
                         enabledInThisState =
-                                ( m_currentStateIndex < int(its->stateCondition.size()) )
+                                ( m_currentStateIndex < its.stateCondition.Num() )
                                 &&
-                                ( its->stateCondition[m_currentStateIndex] );
+                                ( its.stateCondition[m_currentStateIndex] );
                     }
                     if (!enabledInThisState)
                     {
@@ -1896,22 +1890,22 @@ namespace mu
 
                 Ptr<ASTOpInstanceAdd> sop = new ASTOpInstanceAdd();
                 sop->type = OP_TYPE::IN_ADDSURFACE;
-                sop->name = its->node->GetPrivate()->m_name;
+                sop->name = its.node->GetPrivate()->m_name;
                 sop->instance = lastCompOp;
 
 				SURFACE_GENERATION_RESULT surfaceGenerationResult;
-                GenerateSurface( surfaceGenerationResult, its->node, its->edits );
+                GenerateSurface( surfaceGenerationResult, its.node, its.edits );
                 sop->value = surfaceGenerationResult.surfaceOp;
 
                 sop->id = surfaceID;
-                sop->externalId = its->node->GetPrivate()->m_customID;
+                sop->externalId = its.node->GetPrivate()->m_customID;
                 Ptr<ASTOp> surfaceAt = sop;
 
 				// TODO: This could be done earlier?
                 Ptr<ASTOpFixed> surfCondOp = new ASTOpFixed();
                 surfCondOp->op.type = OP_TYPE::BO_AND;
-                surfCondOp->SetChild(surfCondOp->op.args.BoolBinary.a, its->objectCondition );
-                surfCondOp->SetChild(surfCondOp->op.args.BoolBinary.b, its->surfaceCondition );
+                surfCondOp->SetChild(surfCondOp->op.args.BoolBinary.a, its.objectCondition );
+                surfCondOp->SetChild(surfCondOp->op.args.BoolBinary.b, its.surfaceCondition );
                 Ptr<ASTOp> surfaceCondition = surfCondOp;
 
                 if (surfaceCondition)
@@ -1930,16 +1924,16 @@ namespace mu
 
                 // Add the mesh with its condition
                 Ptr<ASTOp> mergeAd;
-                if (!lastMeshOp && its->node->GetPrivate()->m_meshes.size())
+                if (!lastMeshOp && its.node->GetPrivate()->m_meshes.Num())
                 {
                     // \todo: mesh index in case of multiple meshes.
                     const int meshIndex = 0;
-                    lastMeshName = its->node->GetPrivate()->m_meshes[ meshIndex ].m_name;
+                    lastMeshName = its.node->GetPrivate()->m_meshes[ meshIndex ].m_name;
                 }
 
                 // We add the merge op even for the first mesh, so that we set the surface id.
                 {
-                    Ptr<ASTOp> added = its->resultMeshOp;
+                    Ptr<ASTOp> added = its.resultMeshOp;
                     // if (added && m_compilerOptions->m_enablePartialOptimise)
                     // {
                     //     if (m_pTaskManager->HasRunner())
@@ -2039,11 +2033,11 @@ namespace mu
     {
         MUTABLE_CPUPROFILER_SCOPE(NodeObjectNew);
 
-        m_currentParents.push_back( PARENT_KEY() );
-        m_currentParents.back().m_pObject = &node;
+        m_currentParents.Add( PARENT_KEY() );
+        m_currentParents.Last().m_pObject = &node;
 
         // Parse the child objects first, which will accumulate operations in the patching lists
-        for ( std::size_t t=0; t<node.m_children.size(); ++t )
+        for ( int32 t=0; t<node.m_children.Num(); ++t )
         {
             if ( const NodeObject* pChildNode = node.m_children[t].get() )
             {
@@ -2051,9 +2045,9 @@ namespace mu
 
                 // If there are parent objects, the condition of this object depends on the
                 // condition of the parent object
-                if ( m_currentObject.size() )
+                if ( m_currentObject.Num() )
                 {
-                    paramOp = m_currentObject.back().m_condition;
+                    paramOp = m_currentObject.Last().m_condition;
                 }
                 else
                 {
@@ -2066,32 +2060,32 @@ namespace mu
 
                 OBJECT_GENERATION_DATA data;
                 data.m_condition = paramOp;
-                m_currentObject.push_back( data );
+                m_currentObject.Add( data );
 
                 // This op is ignored: everything is stored as patches to apply to the parent when
                 // it is compiled.
                 Generate( pChildNode );
 
-                m_currentObject.pop_back();
+                m_currentObject.Pop();
             }
         }
 
         // Create the expression for each lod
         Ptr<ASTOpAddLOD> lodsOp = new ASTOpAddLOD();
-        for ( std::size_t t=0; t<node.m_lods.size(); ++t )
+        for ( int32 t=0; t<node.m_lods.Num(); ++t )
         {
             if ( const NodeLOD* pLODNode = node.m_lods[t].get() )
             {
-                m_currentParents.back().m_lod = (int)t;
+                m_currentParents.Last().m_lod = (int)t;
 
                 Ptr<ASTOp> lodOp = Generate( pLODNode );
 
-                lodsOp->lods.emplace_back( lodsOp, lodOp );
+                lodsOp->lods.Emplace( lodsOp, lodOp );
             }
         }
         Ptr<ASTOp> rootOp = lodsOp;
 
-        m_currentParents.pop_back();
+        m_currentParents.Pop();
 
         return rootOp;
     }
@@ -2103,19 +2097,20 @@ namespace mu
 		vector<string> usedNames;
 
         // Parse the child objects first, which will accumulate operations in the patching lists
-        for ( std::size_t t=0; t<node.m_children.size(); ++t )
+        for ( int32 t=0; t<node.m_children.Num(); ++t )
         {
             if ( const NodeObject* pChildNode = node.m_children[t].get() )
             {
 				// Look for the child condition in the first pass
                 Ptr<ASTOp> conditionOp;
 				bool found = false;
-                for( auto it = m_firstPass.objects.begin(); !found && it != m_firstPass.objects.end(); it++ )
+                for( int32 i = 0; !found && i != m_firstPass.objects.Num(); i++ )
 				{
-					if (it->node == pChildNode->GetBasePrivate())
+					auto& it = m_firstPass.objects[i];
+					if (it.node == pChildNode->GetBasePrivate())
 					{
 						found = true;
-						conditionOp = it->condition;
+						conditionOp = it.condition;
 					}
 				}
 
@@ -2126,13 +2121,13 @@ namespace mu
 
                 OBJECT_GENERATION_DATA data;
                 data.m_condition = conditionOp;
-                m_currentObject.push_back( data );
+                m_currentObject.Add( data );
 
                 // This op is ignored: everything is stored as patches to apply to the parent when
                 // it is compiled.
                 Generate( pChildNode );
 
-                m_currentObject.pop_back();
+                m_currentObject.Pop();
 
 				// Check for duplicated child names
 				const char* strChildName = pChildNode->GetName();
@@ -2186,13 +2181,13 @@ namespace mu
 
 	//---------------------------------------------------------------------------------------------
 	void CodeGenerator::GetModifiersFor(
-		const vector<string>& tags,
+		const TArray<string>& tags,
 		int LOD, bool bModifiersForBeforeOperations,
-		vector<FirstPassGenerator::MODIFIER>& modifiers)
+		TArray<FirstPassGenerator::MODIFIER>& modifiers)
 	{
         MUTABLE_CPUPROFILER_SCOPE(GetModifiersFor);
 
-		if (tags.size())
+		if (tags.Num())
 		{
 			for (const auto& m: m_firstPass.modifiers)
 			{
@@ -2209,9 +2204,10 @@ namespace mu
 				}
 
 				// Already there?
-				bool alreadyAdded = std::find_if(modifiers.begin(), modifiers.end(),
-					[&m](const FirstPassGenerator::MODIFIER& c) {return c.node == m.node; })
-					!= modifiers.end();
+				bool alreadyAdded = 
+					modifiers.FindByPredicate( [&m](const FirstPassGenerator::MODIFIER& c) {return c.node == m.node; })
+					!= 
+					nullptr;
 
 				if (alreadyAdded)
 				{
@@ -2222,12 +2218,12 @@ namespace mu
 				bool found = false;
 				for (auto it = m.node->m_tags.begin(); !found && it != m.node->m_tags.end(); ++it)
 				{
-					found = std::find(tags.begin(), tags.end(), *it) != tags.end();
+					found = tags.Contains(*it);
 				}
 
 				if (found)
 				{
-					modifiers.push_back(m);
+					modifiers.Add(m);
 				}
 			}
 		}
@@ -2236,21 +2232,21 @@ namespace mu
 	//---------------------------------------------------------------------------------------------
 	Ptr<ASTOp> CodeGenerator::ApplyMeshModifiers(
 		const Ptr<ASTOp>& sourceOp,
-		const vector<string>& tags,
+		const TArray<string>& tags,
 		bool bModifiersForBeforeOperations,
 		const void* errorContext )
 	{
 		Ptr<ASTOp> lastMeshOp = sourceOp;
 
 		// Apply mesh modifiers
-		vector<FirstPassGenerator::MODIFIER> modifiers;
+		TArray<FirstPassGenerator::MODIFIER> modifiers;
 
-		int currentLOD = m_currentParents.back().m_lod;
+		int currentLOD = m_currentParents.Last().m_lod;
 		GetModifiersFor(tags, currentLOD, bModifiersForBeforeOperations, modifiers);
 
 		Ptr<ASTOp> preModifiersMesh = lastMeshOp;
 
-		m_activeTags.push_back({});
+		m_activeTags.Add({});
 
 		// Process clip-with-mesh modifiers
 		Ptr<ASTOpMeshRemoveMask> removeOp;
@@ -2423,7 +2419,7 @@ namespace mu
 			}
 		}
 			
-		m_activeTags.pop_back();
+		m_activeTags.Pop();
 
 		return lastMeshOp;
 	}
