[mutable] Optimize mesh clip versus switching of clipping mesh.

diff --git a/Engine/Plugins/Experimental/Mutable/Source/MutableTools/Private/MuT/CodeGenerator.h b/Engine/Plugins/Experimental/Mutable/Source/MutableTools/Private/MuT/CodeGenerator.h
--- a/Engine/Plugins/Experimental/Mutable/Source/MutableTools/Private/MuT/CodeGenerator.h
+++ b/Engine/Plugins/Experimental/Mutable/Source/MutableTools/Private/MuT/CodeGenerator.h
@@ -42,8 +42,6 @@
 #include "MuT/Visitor.h"
 #include "Templates/TypeHash.h"
 
-#include <functional>
-#include <memory>
 #include <shared_mutex>
 #include <utility>
 
@@ -176,7 +174,7 @@ namespace mu
 				KeyHash = HashCombine(KeyHash, ::GetTypeHash(InKey.imageRect.size[0]));
 				KeyHash = HashCombine(KeyHash, ::GetTypeHash(InKey.imageRect.size[1]));
 				KeyHash = HashCombine(KeyHash, ::GetTypeHash(InKey.state));
-				KeyHash = HashCombine(KeyHash, ::GetTypeHash((uint64)InKey.activeTags.size()));
+				KeyHash = HashCombine(KeyHash, ::GetTypeHash((uint64)InKey.activeTags.Num()));
 				return KeyHash;
 			}
 
@@ -198,8 +196,8 @@ namespace mu
             vec2<int> imageSize;
             box< vec2<int> > imageRect;
             int state = -1;
-			vector<mu::string> activeTags;
-			vector<LayoutPtrConst> overrideLayouts;
+			TArray<mu::string> activeTags;
+			TArray<LayoutPtrConst> overrideLayouts;
         };
 
         //! This struct contains additional state propagated from bottom to top of the object node graph.
@@ -222,13 +220,13 @@ namespace mu
         int m_currentStateIndex = -1;
 
         //! After the entire code generation this contains the information about all the states
-        typedef vector< std::pair<OBJECT_STATE, Ptr<ASTOp>> > StateList;
+        typedef TArray< std::pair<OBJECT_STATE, Ptr<ASTOp>> > StateList;
         StateList m_states;
 
     private:
 
         //! List of meshes generated to be able to reuse them
-        vector<MeshPtr> m_constantMeshes;
+		TArray<MeshPtr> m_constantMeshes;
 
         //! List of image resources for every image formata that have been generated so far as
         //! palceholders for missing images.
@@ -237,7 +235,7 @@ namespace mu
         //! If this has something, while generating meshes, the layouts will be ignored, because
         //! they are supposed to match some other set of layouts. If the vector is empty, layouts
         //! are generated normally.
-        vector< vector<Ptr<const Layout>> > m_overrideLayoutsStack;
+		TArray< TArray<Ptr<const Layout>> > m_overrideLayoutsStack;
 
         //! Map of layouts found in the code already generated. The map is from the source layout
         //! pointer to the cloned layout. The cloned layout will have absolute block ids assigned.
@@ -258,10 +256,10 @@ namespace mu
             int32_t m_layoutBlock;
             LayoutPtrConst m_pLayout;
         };
-        vector<IMAGE_STATE> m_imageState;
+		TArray<IMAGE_STATE> m_imageState;
 
 		// (top-down) Tags that are active when generating nodes.
-		vector< vector<mu::string> > m_activeTags;
+		TArray< TArray<mu::string> > m_activeTags;
 
         struct PARENT_KEY
         {
@@ -285,7 +283,7 @@ namespace mu
             int m_block;
         };
 
-        vector< PARENT_KEY > m_currentParents;
+		TArray< PARENT_KEY > m_currentParents;
 
         // List of additional components to add to an object that come from child objects.
         // The index is the object and lod that should receive the components.
@@ -307,7 +305,7 @@ namespace mu
                 return m_lod < o.m_lod;
             }
         };
-        std::map< ADDITIONAL_COMPONENT_KEY, vector<Ptr<ASTOp>> > m_additionalComponents;
+        std::map< ADDITIONAL_COMPONENT_KEY, TArray<Ptr<ASTOp>> > m_additionalComponents;
 
 
         struct OBJECT_GENERATION_DATA
@@ -315,7 +313,7 @@ namespace mu
             // Condition that enables a specific object
             Ptr<ASTOp> m_condition;
         };
-        vector< OBJECT_GENERATION_DATA > m_currentObject;
+		TArray< OBJECT_GENERATION_DATA > m_currentObject;
 
         map< std::pair<TablePtr,string>, std::pair<TablePtr,Ptr<ASTOp>> > m_generatedTables;
 
@@ -325,11 +323,11 @@ namespace mu
 		//-----------------------------------------------------------------------------------------
 
 		// Get the modifiers that have to be applied to elements with a specific tag.
-		void GetModifiersFor(const vector<string>& tags, int LOD,
-			bool bModifiersForBeforeOperations, vector<FirstPassGenerator::MODIFIER>& modifiers);
+		void GetModifiersFor(const TArray<string>& tags, int LOD,
+			bool bModifiersForBeforeOperations, TArray<FirstPassGenerator::MODIFIER>& modifiers);
 
 		// Apply the required mesh modifiers to the given operation.
-		Ptr<ASTOp> ApplyMeshModifiers( const Ptr<ASTOp>& sourceOp, const vector<string>& tags,
+		Ptr<ASTOp> ApplyMeshModifiers( const Ptr<ASTOp>& sourceOp, const TArray<string>& tags,
 			bool bModifiersForBeforeOperations, const void* errorContext);
 
 		// Get the modifiers that have to be applied to elements with a specific tag.
@@ -369,17 +367,17 @@ namespace mu
 			VISITED_MAP_KEY key;
 			key.pNode = InNode;
 			key.state = m_currentStateIndex;
-			if (!m_imageState.empty())
+			if (!m_imageState.IsEmpty())
 			{
-				key.imageSize = m_imageState.back().m_imageSize;
-				key.imageRect = m_imageState.back().m_imageRect;
+				key.imageSize = m_imageState.Last().m_imageSize;
+				key.imageRect = m_imageState.Last().m_imageRect;
 			}
-			if (!m_activeTags.empty())
+			if (!m_activeTags.IsEmpty())
 			{
-				key.activeTags = m_activeTags.back();
-				if (!m_overrideLayoutsStack.empty())
+				key.activeTags = m_activeTags.Last();
+				if (!m_overrideLayoutsStack.IsEmpty())
 				{
-					key.overrideLayouts = m_overrideLayoutsStack.back();
+					key.overrideLayouts = m_overrideLayoutsStack.Last();
 				}
 			}
 			return key;
@@ -591,7 +589,7 @@ namespace mu
 
         void GenerateSurface( SURFACE_GENERATION_RESULT& result,
                               NodeSurfaceNewPtrConst node,
-                              const vector<FirstPassGenerator::SURFACE::EDIT>& edits );
+                              const TArray<FirstPassGenerator::SURFACE::EDIT>& edits );
 
         TaskManager* m_pTaskManager = nullptr;
     };
@@ -657,7 +655,7 @@ namespace mu
 
         // Create the switch to cover all the options
         Ptr<ASTOp> lastSwitch;
-        std::size_t rows = pTable->GetPrivate()->m_rows.size();
+        std::size_t rows = pTable->GetPrivate()->m_rows.Num();
 
         Ptr<ASTOpSwitch> SwitchOp = new ASTOpSwitch();
 		SwitchOp->type = OPTYPE;
@@ -667,9 +665,9 @@ namespace mu
 		for (size_t i = 0; i < rows; ++i)
         {
             check( pTable->GetPrivate()->m_rows[i].m_id <= 0xFFFF);
-            auto condition = (uint16_t)pTable->GetPrivate()->m_rows[i].m_id;
+            auto condition = (uint16)pTable->GetPrivate()->m_rows[i].m_id;
             Ptr<ASTOp> Branch = GenerateOption( node, colIndex, (int)i, m_pErrorLog.get() );
-			SwitchOp->cases.push_back(ASTOpSwitch::CASE(condition, SwitchOp, Branch ));
+			SwitchOp->cases.Add(ASTOpSwitch::CASE(condition, SwitchOp, Branch ));
         }
 
         return SwitchOp;
diff --git a/Engine/Plugins/Experimental/Mutable/Source/MutableTools/Private/MuT/CodeGenerator_Bool.cpp b/Engine/Plugins/Experimental/Mutable/Source/MutableTools/Private/MuT/CodeGenerator_Bool.cpp
--- a/Engine/Plugins/Experimental/Mutable/Source/MutableTools/Private/MuT/CodeGenerator_Bool.cpp
+++ b/Engine/Plugins/Experimental/Mutable/Source/MutableTools/Private/MuT/CodeGenerator_Bool.cpp
@@ -107,12 +107,11 @@ void CodeGenerator::GenerateBool_Parameter( BOOL_GENERATION_RESULT& result, cons
 		op->parameter = param;
 
 		// Generate the code for the ranges
-		for (std::size_t a = 0; a < node.m_ranges.size(); ++a)
+		for (int32 a = 0; a < node.m_ranges.Num(); ++a)
 		{
 			RANGE_GENERATION_RESULT rangeResult;
 			GenerateRange(rangeResult, node.m_ranges[a]);
-			op->ranges.emplace_back
-			(op.get(), rangeResult.sizeOp, rangeResult.rangeName, rangeResult.rangeUID);
+			op->ranges.Emplace(op.get(), rangeResult.sizeOp, rangeResult.rangeName, rangeResult.rangeUID);
 		}
 
 		m_nodeVariables[node.m_pNode] = op;
diff --git a/Engine/Plugins/Experimental/Mutable/Source/MutableTools/Private/MuT/CodeGenerator_Color.cpp b/Engine/Plugins/Experimental/Mutable/Source/MutableTools/Private/MuT/CodeGenerator_Color.cpp
--- a/Engine/Plugins/Experimental/Mutable/Source/MutableTools/Private/MuT/CodeGenerator_Color.cpp
+++ b/Engine/Plugins/Experimental/Mutable/Source/MutableTools/Private/MuT/CodeGenerator_Color.cpp
@@ -156,12 +156,11 @@ void CodeGenerator::GenerateColor_Parameter( COLOR_GENERATION_RESULT& result, co
 		op->parameter = param;
 
 		// Generate the code for the ranges
-		for (std::size_t a = 0; a < node.m_ranges.size(); ++a)
+		for (int32 a = 0; a < node.m_ranges.Num(); ++a)
 		{
 			RANGE_GENERATION_RESULT rangeResult;
 			GenerateRange(rangeResult, node.m_ranges[a]);
-			op->ranges.emplace_back
-			(op.get(), rangeResult.sizeOp, rangeResult.rangeName, rangeResult.rangeUID);
+			op->ranges.Emplace(op.get(), rangeResult.sizeOp, rangeResult.rangeName, rangeResult.rangeUID);
 		}
 
 		m_nodeVariables[node.m_pNode] = op;
@@ -182,7 +181,7 @@ void CodeGenerator::GenerateColor_Switch(COLOR_GENERATION_RESULT& result, const
 
 	MUTABLE_CPUPROFILER_SCOPE(NodeColourSwitch);
 
-	if (node.m_options.size() == 0)
+	if (node.m_options.Num() == 0)
 	{
 		// No options in the switch!
 		Ptr<ASTOp> missingOp = GenerateMissingColourCode("Switch option",
@@ -206,7 +205,7 @@ void CodeGenerator::GenerateColor_Switch(COLOR_GENERATION_RESULT& result, const
 	}
 
 	// Options
-	for (std::size_t t = 0; t < node.m_options.size(); ++t)
+	for (std::size_t t = 0; t < node.m_options.Num(); ++t)
 	{
 		Ptr<ASTOp> branch;
 		if (node.m_options[t])
@@ -219,7 +218,7 @@ void CodeGenerator::GenerateColor_Switch(COLOR_GENERATION_RESULT& result, const
 			branch = GenerateMissingColourCode("Switch option",
 				node.m_errorContext);
 		}
-		op->cases.push_back(ASTOpSwitch::CASE((int16_t)t, op, branch));
+		op->cases.Emplace((int16)t, op, branch);
 	}
 
 	result.op = op;
@@ -246,11 +245,11 @@ void CodeGenerator::GenerateColor_Variation(COLOR_GENERATION_RESULT& result, con
 	}
 
 	// Process variations in reverse order, since conditionals are built bottom-up.
-	for (int t = int(node.m_variations.size()) - 1; t >= 0; --t)
+	for (int t = node.m_variations.Num() - 1; t >= 0; --t)
 	{
 		int tagIndex = -1;
 		const string& tag = node.m_variations[t].m_tag;
-		for (int i = 0; i < int(m_firstPass.m_tags.size()); ++i)
+		for (int i = 0; i < m_firstPass.m_tags.Num(); ++i)
 		{
 			if (m_firstPass.m_tags[i].tag == tag)
 			{
@@ -314,13 +313,13 @@ void CodeGenerator::GenerateColor_SampleImage(COLOR_GENERATION_RESULT& result, c
 		newState.m_imageRect.min[0] = 0;
 		newState.m_imageRect.min[1] = 0;
 		newState.m_imageRect.size = desc.m_size;
-		m_imageState.push_back(newState);
+		m_imageState.Add(newState);
 
 		// Generate
 		base = Generate(pSource);
 
 		// Restore rect
-		m_imageState.pop_back();
+		m_imageState.Pop();
 	}
 	else
 	{
diff --git a/Engine/Plugins/Experimental/Mutable/Source/MutableTools/Private/MuT/CodeGenerator_FirstPass.cpp b/Engine/Plugins/Experimental/Mutable/Source/MutableTools/Private/MuT/CodeGenerator_FirstPass.cpp
--- a/Engine/Plugins/Experimental/Mutable/Source/MutableTools/Private/MuT/CodeGenerator_FirstPass.cpp
+++ b/Engine/Plugins/Experimental/Mutable/Source/MutableTools/Private/MuT/CodeGenerator_FirstPass.cpp
@@ -40,8 +40,8 @@ namespace mu
 	{
 		// Default conditions when there is no restriction accumulated.
 		CONDITION_CONTEXT noCondition;
-        m_currentCondition.push_back(noCondition);
-        m_currentStateCondition.push_back(StateCondition());
+        m_currentCondition.Add(noCondition);
+        m_currentStateCondition.Add(StateCondition());
 	}
 
 
@@ -78,14 +78,14 @@ namespace mu
 		}
 
 		// Step 2: Collect all tags and a list of the surfaces that activate them
-		for (std::size_t s=0; s<surfaces.size(); ++s)
+		for (int32 s=0; s<surfaces.Num(); ++s)
 		{
             // \todo: edit surfaces should also be able to activate tags. T1245
-            for (std::size_t t=0; t<surfaces[s].node->GetPrivate()->m_tags.size(); ++t)
+            for (int32 t=0; t<surfaces[s].node->GetPrivate()->m_tags.Num(); ++t)
 			{
 				int tag = -1;
                 auto tagStr = surfaces[s].node->GetPrivate()->m_tags[t];
-                for (std::size_t i = 0; i<m_tags.size() && tag<0; ++i)
+                for (std::size_t i = 0; i<m_tags.Num() && tag<0; ++i)
 				{
                     if (m_tags[i].tag == tagStr)
 					{
@@ -96,30 +96,28 @@ namespace mu
 				// New tag?
 				if (tag < 0)
 				{
-                    tag = (int)m_tags.size();
+                    tag = (int)m_tags.Num();
 					TAG newTag;
                     newTag.tag = tagStr;
-                    m_tags.push_back(newTag);
+                    m_tags.Add(newTag);
 				}
 
-                if (std::find(m_tags[tag].surfaces.begin(), m_tags[tag].surfaces.end(),s)
-					== 
-                    m_tags[tag].surfaces.end())
+                if (m_tags[tag].surfaces.Find(s)==INDEX_NONE)
 				{
-                    m_tags[tag].surfaces.push_back((int)s);
+                    m_tags[tag].surfaces.Add((int)s);
 				}
 			}
 
             // Collect the tags in edit surfaces
-            for (std::size_t e=0; e<surfaces[s].edits.size(); ++e)
+            for (std::size_t e=0; e<surfaces[s].edits.Num(); ++e)
             {
                 const auto& edit = surfaces[s].edits[e];
-                for (std::size_t t=0; t<edit.node->m_tags.size(); ++t)
+                for (int32 t=0; t<edit.node->m_tags.Num(); ++t)
                 {
                     int tag = -1;
                     auto tagStr = edit.node->m_tags[t];
 
-                    for (std::size_t i = 0; i<m_tags.size() && tag<0; ++i)
+                    for (int32 i = 0; i<m_tags.Num() && tag<0; ++i)
                     {
                         if (m_tags[i].tag == tagStr)
                         {
@@ -130,18 +128,15 @@ namespace mu
                     // New tag?
                     if (tag < 0)
                     {
-                        tag = (int)m_tags.size();
+                        tag = (int)m_tags.Num();
                         TAG newTag;
                         newTag.tag = tagStr;
-                        m_tags.push_back(newTag);
+                        m_tags.Add(newTag);
                     }
 
-                    auto key = std::make_pair(int(s),int(e));
-                    if (std::find(m_tags[tag].edits.begin(), m_tags[tag].edits.end(),key)
-                        ==
-                        m_tags[tag].edits.end())
+                    if (m_tags[tag].edits.Find({int(s),int(e)}) == INDEX_NONE)
                     {
-                        m_tags[tag].edits.push_back(key);
+                        m_tags[tag].edits.Add({ int(s),int(e) });
                     }
                 }
             }
@@ -152,14 +147,14 @@ namespace mu
         // Step 3: Create default state if necessary
         if ( ignoreStates )
         {
-            m_states.clear();
+            m_states.Empty();
         }
 
-        if ( m_states.empty() )
+        if ( m_states.IsEmpty() )
         {
             OBJECT_STATE data;
             data.m_name = "Default";
-            m_states.emplace_back( data, root );
+            m_states.Emplace( data, root );
         }
 	}
 
@@ -170,12 +165,12 @@ namespace mu
         // Add the data about this modifier
         MODIFIER thisData;
         thisData.node = &node;
-        thisData.objectCondition = m_currentCondition.back().objectCondition;
-        thisData.stateCondition = m_currentStateCondition.back();
+        thisData.objectCondition = m_currentCondition.Last().objectCondition;
+        thisData.stateCondition = m_currentStateCondition.Last();
         thisData.lod = m_currentLOD;
         thisData.positiveTags = m_currentPositiveTags;
         thisData.negativeTags = m_currentNegativeTags;
-        modifiers.push_back(thisData);
+        modifiers.Add(thisData);
 
         return nullptr;
     }
@@ -187,12 +182,12 @@ namespace mu
         // Add the data about this modifier
 		MODIFIER thisData;
 		thisData.node = &node;
-        thisData.objectCondition = m_currentCondition.back().objectCondition;
-        thisData.stateCondition = m_currentStateCondition.back();
+        thisData.objectCondition = m_currentCondition.Last().objectCondition;
+        thisData.stateCondition = m_currentStateCondition.Last();
         thisData.lod = m_currentLOD;
         thisData.positiveTags = m_currentPositiveTags;
         thisData.negativeTags = m_currentNegativeTags;
-        modifiers.push_back(thisData);
+        modifiers.Add(thisData);
 
         return nullptr;
 	}
@@ -204,12 +199,12 @@ namespace mu
         // Add the data about this modifier
 		MODIFIER thisData;
 		thisData.node = &node;
-        thisData.objectCondition = m_currentCondition.back().objectCondition;
-        thisData.stateCondition = m_currentStateCondition.back();
+        thisData.objectCondition = m_currentCondition.Last().objectCondition;
+        thisData.stateCondition = m_currentStateCondition.Last();
         thisData.lod = m_currentLOD;
         thisData.positiveTags = m_currentPositiveTags;
         thisData.negativeTags = m_currentNegativeTags;
-        modifiers.push_back(thisData);
+        modifiers.Add(thisData);
 
         return nullptr;
 	}
@@ -221,11 +216,11 @@ namespace mu
 		SURFACE thisData;
 		thisData.node = dynamic_cast<const NodeSurfaceNew*>(node.m_pNode);
 		thisData.component = m_currentComponent;
-        thisData.objectCondition = m_currentCondition.back().objectCondition;
-        thisData.stateCondition = m_currentStateCondition.back();
+        thisData.objectCondition = m_currentCondition.Last().objectCondition;
+        thisData.stateCondition = m_currentStateCondition.Last();
         thisData.positiveTags = m_currentPositiveTags;
         thisData.negativeTags = m_currentNegativeTags;
-        surfaces.push_back(thisData);
+        surfaces.Add(thisData);
 
         return nullptr;
 	}
@@ -236,8 +231,7 @@ namespace mu
 	{
 		// Store a reference to this node in the surface data for the surface that this node is
 		// editing.
-		auto its = std::find_if(surfaces.begin(), surfaces.end(), 
-            [&node](const FirstPassGenerator::SURFACE& s)
+		auto its = surfaces.FindByPredicate([&node](const FirstPassGenerator::SURFACE& s)
         {
             // Are we editing the main surface node of this surface?
             if (s.node.get() == node.m_pParent.get()) return true;
@@ -256,14 +250,14 @@ namespace mu
 		
 		// The surface could be missing if the parent is not in the hierarchy. This could happen
 		// with wrong input or in case of partial models for preview.
-		if (its != surfaces.end())
+		if (its)
 		{
 			SURFACE& surface = *its;
 
             SURFACE::EDIT edit;
             edit.node = &node;
-            edit.condition = m_currentCondition.back().objectCondition;
-            surface.edits.push_back(edit);
+            edit.condition = m_currentCondition.Last().objectCondition;
+            surface.edits.Add(edit);
 		}
 		else
 		{
@@ -285,9 +279,9 @@ namespace mu
         {
             // Any of the tags in the variations would prevent the default surface
             auto oldNegativeTags = m_currentNegativeTags;
-            for (size_t v=0; v<node.m_variations.size(); ++v)
+            for (int32 v=0; v<node.m_variations.Num(); ++v)
             {
-                m_currentNegativeTags.push_back(node.m_variations[v].m_tag);
+                m_currentNegativeTags.Add(node.m_variations[v].m_tag);
             }
 
             for(const auto& n:node.m_defaultSurfaces)
@@ -301,9 +295,9 @@ namespace mu
 
             m_currentNegativeTags = oldNegativeTags;
 
-            for (size_t v=0; v<node.m_variations.size(); ++v)
+            for (int32 v=0; v<node.m_variations.Num(); ++v)
             {
-                m_currentPositiveTags.push_back(node.m_variations[v].m_tag);
+                m_currentPositiveTags.Add(node.m_variations[v].m_tag);
                 for (const auto& s : node.m_variations[v].m_surfaces)
                 {
                     s->GetBasePrivate()->Accept(*this);
@@ -314,11 +308,11 @@ namespace mu
                     s->GetBasePrivate()->Accept(*this);
                 }
 
-                m_currentPositiveTags.pop_back();
+                m_currentPositiveTags.Pop();
 
                 // Tags have an order in a variation node: the current tag should prevent any following
                 // variation surface
-                m_currentNegativeTags.push_back(node.m_variations[v].m_tag);
+                m_currentNegativeTags.Add(node.m_variations[v].m_tag);
             }
 
             m_currentNegativeTags = oldNegativeTags;
@@ -329,22 +323,24 @@ namespace mu
 
         case NodeSurfaceVariation::VariationType::State:
         {
-            size_t stateCount = m_states.size();
+            size_t stateCount = m_states.Num();
 
             // Default
             {
                 // Store the states for the default branch here
                 StateCondition defaultStates;
                 {
-                    defaultStates = m_currentStateCondition.back().empty()
-                            ? StateCondition( stateCount, true )
-                            : m_currentStateCondition.back();
+					StateCondition AllTrue;
+					AllTrue.Init(true,stateCount);
+                    defaultStates = m_currentStateCondition.Last().IsEmpty()
+                            ? AllTrue
+                            : m_currentStateCondition.Last();
 
                     for (const auto& v:node.m_variations)
                     {
                         for( size_t s=0; s<stateCount; ++s )
                         {
-                            if (m_states[s].first.m_name==v.m_tag)
+                            if (m_states[s].Key.m_name==v.m_tag)
                             {
                                 // Remove this state from the default options, since it has its own variation
                                 defaultStates[s] = false;
@@ -353,7 +349,7 @@ namespace mu
                     }
                 }
 
-                m_currentStateCondition.push_back(defaultStates);
+                m_currentStateCondition.Add(defaultStates);
 
                 for (const auto& n : node.m_defaultSurfaces)
                 {
@@ -364,7 +360,7 @@ namespace mu
                     n->GetBasePrivate()->Accept(*this);
                 }
 
-                m_currentStateCondition.pop_back();
+                m_currentStateCondition.Pop();
             }
 
 
@@ -372,17 +368,18 @@ namespace mu
             for (const auto& v:node.m_variations)
             {
                 // Store the states for this variation here
-                StateCondition variationStates( stateCount, false );
+				StateCondition variationStates;
+				variationStates.Init(false,stateCount);
 
                 for( size_t s=0; s<stateCount; ++s )
                 {
-                    if (m_states[s].first.m_name==v.m_tag)
+                    if (m_states[s].Key.m_name==v.m_tag)
                     {
                         variationStates[s] = true;
                     }
                 }
 
-                m_currentStateCondition.push_back(variationStates);
+                m_currentStateCondition.Add(variationStates);
 
                 for (const auto& n : v.m_surfaces)
                 {
@@ -393,7 +390,7 @@ namespace mu
                     n->GetBasePrivate()->Accept(*this);
                 }
 
-                m_currentStateCondition.pop_back();
+                m_currentStateCondition.Pop();
             }
 
             break;
@@ -476,16 +473,16 @@ namespace mu
 		// Add the data about this object
 		OBJECT thisData;
 		thisData.node = &node;
-        thisData.condition = m_currentCondition.back().objectCondition;
-		objects.push_back(thisData);
+        thisData.condition = m_currentCondition.Last().objectCondition;
+		objects.Add(thisData);
 
         // Accumulate the model states
         for ( const auto& s: node.m_states )
         {
-            m_states.emplace_back( s, &node );
+            m_states.Emplace( s, &node );
 
             // \todo move this 32 constant to a macro or constexpr with a meaningful name and location
-            if ( s.m_runtimeParams.size() > 32 )
+            if ( s.m_runtimeParams.Num() > 32 )
             {
                 char temp[256];
                 mutable_snprintf( temp, 256,
@@ -542,7 +539,7 @@ namespace mu
                 PARAMETER_DESC::INT_VALUE_DESC nullValue;
                 nullValue.m_value = -1;
                 nullValue.m_name = "None";
-                op->parameter.m_possibleValues.push_back( nullValue );
+                op->parameter.m_possibleValues.Add( nullValue );
                 op->parameter.m_defaultValue.m_int = nullValue.m_value;
             }
 
@@ -551,7 +548,7 @@ namespace mu
 
 
         // Parse the child objects
-		for ( std::size_t t=0; t<node.m_children.size(); ++t )
+		for ( int32 t=0; t<node.m_children.Num(); ++t )
         {
             if ( const NodeObject* pChildNode = node.m_children[t].get() )
             {
@@ -591,7 +588,7 @@ namespace mu
                         PARAMETER_DESC::INT_VALUE_DESC value;
                         value.m_value = (int16_t)t;
                         value.m_name = pChildNode->GetName();
-                        enumOp->parameter.m_possibleValues.push_back( value );
+                        enumOp->parameter.m_possibleValues.Add( value );
 
                         // Set as default if none has been set.
                         if ( enumOp->parameter.m_defaultValue.m_int == -1 )
@@ -616,22 +613,22 @@ namespace mu
                 }
 
                 // Combine the new condition with previous conditions coming from parent objects
-                if (m_currentCondition.back().objectCondition)
+                if (m_currentCondition.Last().objectCondition)
                 {
                     Ptr<ASTOpFixed> op = new ASTOpFixed();
                     op->op.type = OP_TYPE::BO_AND;
-                    op->SetChild( op->op.args.BoolBinary.a,m_currentCondition.back().objectCondition);
+                    op->SetChild( op->op.args.BoolBinary.a,m_currentCondition.Last().objectCondition);
                     op->SetChild( op->op.args.BoolBinary.b,paramOp);
                     paramOp = op;
                 }
 
 				CONDITION_CONTEXT data;
                 data.objectCondition = paramOp;
-                m_currentCondition.push_back( data );
+                m_currentCondition.Add( data );
 
 				pChildNode->GetBasePrivate()->Accept(*this);
 
-                m_currentCondition.pop_back();
+                m_currentCondition.Pop();
             }
         }
 
@@ -653,7 +650,7 @@ namespace mu
             if (node.m_pRoot)
             {
                 // Remember the new state with the new root.
-                m_states.emplace_back( node.m_state, node.m_pRoot->GetBasePrivate() );
+                m_states.Emplace( node.m_state, node.m_pRoot->GetBasePrivate() );
             }
         }
 
diff --git a/Engine/Plugins/Experimental/Mutable/Source/MutableTools/Private/MuT/CodeGenerator_FirstPass.h b/Engine/Plugins/Experimental/Mutable/Source/MutableTools/Private/MuT/CodeGenerator_FirstPass.h
--- a/Engine/Plugins/Experimental/Mutable/Source/MutableTools/Private/MuT/CodeGenerator_FirstPass.h
+++ b/Engine/Plugins/Experimental/Mutable/Source/MutableTools/Private/MuT/CodeGenerator_FirstPass.h
@@ -31,9 +31,8 @@
 
 namespace mu
 {
-class Layout;
-struct OBJECT_STATE;
-
+	class Layout;
+	struct OBJECT_STATE;
 
     //! Store the results of the code generation of a mesh.
     struct MESH_GENERATION_RESULT
@@ -45,19 +44,19 @@ struct OBJECT_STATE;
         Ptr<ASTOp> baseMeshOp;
 
 		/** Source node layouts to use with these extra mesh. They don't have block ids. */
-		vector<Ptr<const Layout>> layouts;
+		TArray<Ptr<const Layout>> layouts;
 
-        vector<Ptr<ASTOp>> layoutOps;
+		TArray<Ptr<ASTOp>> layoutOps;
 
 
         struct EXTRA_LAYOUTS
         {
 			/** Source node layouts to use with these extra mesh. They don't have block ids. */
-            vector<Ptr<const Layout>> layouts;
+			TArray<Ptr<const Layout>> layouts;
             Ptr<ASTOp> condition;
             Ptr<ASTOp> meshFragment;
         };
-        vector< EXTRA_LAYOUTS > extraMeshLayouts;
+		TArray< EXTRA_LAYOUTS > extraMeshLayouts;
 
     };
 
@@ -105,13 +104,13 @@ struct OBJECT_STATE;
 			const NodeObjectNew::Private* node;
             Ptr<ASTOp> condition;
 		};
-		vector<OBJECT> objects;
+		TArray<OBJECT> objects;
 
         //! Type used to represent the activation conditions regarding states
         //! This is the state mask for the states in which this surface must be added. If it
         //! is empty it means the surface is valid for all states. Otherwise it is only valid
         //! for the states whose index is true.
-        using StateCondition = std::vector<uint8_t>;
+        using StateCondition = TArray<uint8_t>;
 
 		//! Store information about every surface including
 		//! - the component it may be added to
@@ -128,10 +127,10 @@ struct OBJECT_STATE;
             const NodeComponentNew::Private* component = nullptr;
 
             // List of tags that are required for the presence of this surface
-            vector<string> positiveTags;
+			TArray<string> positiveTags;
 
             // List of tags that block the presence of this surface
-            vector<string> negativeTags;
+			TArray<string> negativeTags;
 
 			// This conditions is the condition of the object defining this surface which may not
 			// be the parent object where this surface will be added.
@@ -153,13 +152,13 @@ struct OBJECT_STATE;
                 //! Weak reference to the edit node, used during compilation.
                 const NodeSurfaceEdit::Private* node = nullptr;
             };
-            vector<EDIT> edits;
+			TArray<EDIT> edits;
 
             // This is filled in the final code generation pass
             Ptr<ASTOp> resultSurfaceOp;
             Ptr<ASTOp> resultMeshOp;
         };
-		vector<SURFACE> surfaces;
+		TArray<SURFACE> surfaces;
 
 		//! Store the conditions that enable every modifier.
 		struct MODIFIER
@@ -167,10 +166,10 @@ struct OBJECT_STATE;
             const NodeModifier::Private* node = nullptr;
 
             // List of tags that are required for the presence of this surface
-            vector<string> positiveTags;
+			TArray<string> positiveTags;
 
             // List of tags that block the presence of this surface
-            vector<string> negativeTags;
+			TArray<string> negativeTags;
 
             // This conditions is the condition of the object defining this modifier which may not
             // be the parent object where this surface will be added.
@@ -187,7 +186,7 @@ struct OBJECT_STATE;
             //
             int lod = 0;
         };
-		vector<MODIFIER> modifiers;
+		TArray<MODIFIER> modifiers;
 
 		//! Info about all found tags.
 		struct TAG
@@ -196,22 +195,22 @@ struct OBJECT_STATE;
 
             // Surfaces that activate the tag. These are indices to the FirstPassGenerator::surfaces
             // vector.
-            vector<int> surfaces;
+			TArray<int> surfaces;
 
             // Edit Surfaces that activate the tag. These first element of the pair are indices to
             // the FirstPassGenerator::surfaces vector. The second element are indices to the
             // "edits" in the specific surface.
-            vector<std::pair<int,int>> edits;
+			TArray<TPair<int,int>> edits;
 
             // This conditions is the condition for this tag to be enabled considering no other
             // condition. This is filled in CodeGenerator_SecondPass.
             Ptr<ASTOp> genericCondition;
         };
-        vector<TAG> m_tags;
+        TArray<TAG> m_tags;
 
         //! Accumulate the model states found while generating code, with their generated root
         //! nodes.
-        typedef vector< std::pair<OBJECT_STATE, const Node::Private*> > StateList;
+        typedef TArray< TPair<OBJECT_STATE, const Node::Private*> > StateList;
         StateList m_states;
 
 	public:
@@ -236,17 +235,17 @@ struct OBJECT_STATE;
         {
             Ptr<ASTOp> objectCondition;
         };
-        vector< CONDITION_CONTEXT > m_currentCondition;
+		TArray< CONDITION_CONTEXT > m_currentCondition;
 
         //!
-        vector< StateCondition > m_currentStateCondition;
+		TArray< StateCondition > m_currentStateCondition;
 
 		//! When processing surfaces, this is the parent component the surfaces may be added to
         const NodeComponentNew::Private* m_currentComponent = nullptr;
 
         //! Current relevant tags so far. Used during traversal.
-        vector<string> m_currentPositiveTags;
-        vector<string> m_currentNegativeTags;
+		TArray<string> m_currentPositiveTags;
+		TArray<string> m_currentNegativeTags;
 
 		//! Index of the LOD we are processing
         int m_currentLOD = -1;
diff --git a/Engine/Plugins/Experimental/Mutable/Source/MutableTools/Private/MuT/CodeGenerator_Image.cpp b/Engine/Plugins/Experimental/Mutable/Source/MutableTools/Private/MuT/CodeGenerator_Image.cpp
--- a/Engine/Plugins/Experimental/Mutable/Source/MutableTools/Private/MuT/CodeGenerator_Image.cpp
+++ b/Engine/Plugins/Experimental/Mutable/Source/MutableTools/Private/MuT/CodeGenerator_Image.cpp
@@ -125,7 +125,7 @@ namespace mu
 
 		// Generate the block size in case we are generating an expression whose root is an image
 		bool addedImageState = false;
-		if ( m_imageState.empty() )
+		if ( m_imageState.IsEmpty() )
 		{
 
 			FImageDesc desc = CalculateImageDesc(*Untyped->GetBasePrivate());
@@ -136,7 +136,7 @@ namespace mu
 			newState.m_imageRect.min[0] = 0;
 			newState.m_imageRect.min[1] = 0;
 			newState.m_layoutBlock = -1;
-			m_imageState.push_back(newState);
+			m_imageState.Add(newState);
 			addedImageState = true;
 		}
 
@@ -191,7 +191,7 @@ namespace mu
 		// Restore the modified image state
 		if (addedImageState)
 		{
-			m_imageState.pop_back();
+			m_imageState.Pop();
 		}
 
 	}
@@ -232,14 +232,14 @@ namespace mu
 
         // Order of the operations is important: multiply first to avoid losing precision.
         // It will not overflow since image sizes are limited to 16 bit
-        cropRect.min = ( m_imageState.back().m_imageRect.min * imageSize )
-                / m_imageState.back().m_imageSize;
-        cropRect.size = ( m_imageState.back().m_imageRect.size * imageSize )
-                / m_imageState.back().m_imageSize;
+        cropRect.min = ( m_imageState.Last().m_imageRect.min * imageSize )
+                / m_imageState.Last().m_imageSize;
+        cropRect.size = ( m_imageState.Last().m_imageRect.size * imageSize )
+                / m_imageState.Last().m_imageSize;
 
         //check( cropRect.size[0]>0 && cropRect.size[1]>0 );
-        cropRect.size[0] = std::max( cropRect.size[0], 1 );
-        cropRect.size[1] = std::max( cropRect.size[1], 1 );
+        cropRect.size[0] = FMath::Max( cropRect.size[0], 1 );
+        cropRect.size[1] = FMath::Max( cropRect.size[1], 1 );
 
         //
         if ( pImage->GetSizeX()!=cropRect.size[0] || pImage->GetSizeY()!=cropRect.size[1] )
@@ -315,8 +315,8 @@ namespace mu
 		EImageFormat baseFormat = base->GetImageDesc().m_format;
         //base = GenerateImageFormat( base, EImageFormat::IF_RGB_UBYTE );
         base = GenerateImageSize
-                ( base, FImageSize((uint16_t)m_imageState.back().m_imageRect.size[0],
-                                   (uint16_t)m_imageState.back().m_imageRect.size[1]) );
+                ( base, FImageSize((uint16)m_imageState.Last().m_imageRect.size[0],
+                                   (uint16)m_imageState.Last().m_imageRect.size[1]) );
         op->SetChild( op->op.args.ImageLayer.base, base);
 
         // Mask of the effect
@@ -326,8 +326,8 @@ namespace mu
             mask = Generate( pMask );
             mask = GenerateImageFormat( mask, EImageFormat::IF_L_UBYTE );
             mask = GenerateImageSize
-                    ( mask, FImageSize((uint16_t)m_imageState.back().m_imageRect.size[0],
-                                       (uint16_t)m_imageState.back().m_imageRect.size[1]) );
+                    ( mask, FImageSize((uint16)m_imageState.Last().m_imageRect.size[0],
+                                       (uint16)m_imageState.Last().m_imageRect.size[1]) );
         }
         op->SetChild( op->op.args.ImageLayer.mask, mask);
 
@@ -345,8 +345,8 @@ namespace mu
         //blended = GenerateImageFormat( blended, EImageFormat::IF_RGB_UBYTE );
         blended = GenerateImageFormat( blended, baseFormat );
         blended = GenerateImageSize
-                ( blended, FImageSize((uint16_t)m_imageState.back().m_imageRect.size[0],
-                                      (uint16_t)m_imageState.back().m_imageRect.size[1]) );
+                ( blended, FImageSize((uint16)m_imageState.Last().m_imageRect.size[0],
+                                      (uint16)m_imageState.Last().m_imageRect.size[1]) );
         op->SetChild( op->op.args.ImageLayer.blended, blended);
 
         result.op = op;
@@ -377,7 +377,7 @@ namespace mu
         }
         base = GenerateImageFormat( base, EImageFormat::IF_RGB_UBYTE );
         base = GenerateImageSize
-                ( base, FImageSize((uint16_t)m_imageState.back().m_imageRect.size[0],(uint16_t)m_imageState.back().m_imageRect.size[1]) );
+                ( base, FImageSize((uint16)m_imageState.Last().m_imageRect.size[0],(uint16)m_imageState.Last().m_imageRect.size[1]) );
         op->SetChild( op->op.args.ImageLayerColour.base, base);
 
         // Mask of the effect
@@ -387,7 +387,7 @@ namespace mu
             mask = Generate( pMask );
             mask = GenerateImageFormat( mask, EImageFormat::IF_L_UBYTE );
             mask = GenerateImageSize
-                    ( mask, FImageSize((uint16_t)m_imageState.back().m_imageRect.size[0],(uint16_t)m_imageState.back().m_imageRect.size[1]) );
+                    ( mask, FImageSize((uint16)m_imageState.Last().m_imageRect.size[0],(uint16)m_imageState.Last().m_imageRect.size[1]) );
         }
         op->SetChild( op->op.args.ImageLayerColour.mask, mask);
 
@@ -434,8 +434,8 @@ namespace mu
 
 		EImageFormat baseFormat = base->GetImageDesc().m_format;
         base = GenerateImageSize
-                ( base, FImageSize((uint16_t)m_imageState.back().m_imageRect.size[0],
-                                   (uint16_t)m_imageState.back().m_imageRect.size[1]) );
+                ( base, FImageSize((uint16)m_imageState.Last().m_imageRect.size[0],
+                                   (uint16)m_imageState.Last().m_imageRect.size[1]) );
         op->base = base;
 
         // Mask of the effect
@@ -445,8 +445,8 @@ namespace mu
             mask = Generate( pMask );
             mask = GenerateImageFormat( mask, EImageFormat::IF_L_UBYTE );
             mask = GenerateImageSize
-                    ( mask, FImageSize((uint16_t)m_imageState.back().m_imageRect.size[0],
-                                       (uint16_t)m_imageState.back().m_imageRect.size[1]) );
+                    ( mask, FImageSize((uint16)m_imageState.Last().m_imageRect.size[0],
+                                       (uint16)m_imageState.Last().m_imageRect.size[1]) );
         }
         op->mask = mask;
 
@@ -464,8 +464,8 @@ namespace mu
         //blended = GenerateImageFormat( blended, EImageFormat::IF_RGB_UBYTE );
         blended = GenerateImageFormat( blended, baseFormat );
         blended = GenerateImageSize
-                ( blended, FImageSize((uint16_t)m_imageState.back().m_imageRect.size[0],
-                                      (uint16_t)m_imageState.back().m_imageRect.size[1]) );
+                ( blended, FImageSize((uint16)m_imageState.Last().m_imageRect.size[0],
+                                      (uint16)m_imageState.Last().m_imageRect.size[1]) );
         op->blend = blended;
 
         // Range of iteration
@@ -511,8 +511,8 @@ namespace mu
 
 		EImageFormat baseFormat = base->GetImageDesc().m_format;
         base = GenerateImageSize
-                ( base, FImageSize((uint16_t)m_imageState.back().m_imageRect.size[0],
-                                   (uint16_t)m_imageState.back().m_imageRect.size[1]) );
+                ( base, FImageSize((uint16)m_imageState.Last().m_imageRect.size[0],
+                                   (uint16)m_imageState.Last().m_imageRect.size[1]) );
         op->Base = base;
 
         Ptr<ASTOp> normal;
@@ -521,8 +521,8 @@ namespace mu
             normal = Generate( pNormal );
             normal = GenerateImageFormat( normal, EImageFormat::IF_RGB_UBYTE );
             normal = GenerateImageSize
-                    ( normal, FImageSize((uint16_t)m_imageState.back().m_imageRect.size[0],
-                                         (uint16_t)m_imageState.back().m_imageRect.size[1]) );
+                    ( normal, FImageSize((uint16)m_imageState.Last().m_imageRect.size[0],
+                                         (uint16)m_imageState.Last().m_imageRect.size[1]) );
         }
 		else
 		{
@@ -596,8 +596,8 @@ namespace mu
 
 		EImageFormat baseFormat = base->GetImageDesc().m_format;
         base = GenerateImageSize
-                ( base, FImageSize((uint16_t)m_imageState.back().m_imageRect.size[0],
-									  (uint16_t)m_imageState.back().m_imageRect.size[1]) );
+                ( base, FImageSize((uint16)m_imageState.Last().m_imageRect.size[0],
+									  (uint16)m_imageState.Last().m_imageRect.size[1]) );
         op->base = base;
 
         result.op = op; 
@@ -629,7 +629,7 @@ namespace mu
         int numTargets = 0;
 
         for ( std::size_t t=0
-            ; t< node.m_targets.size() && numTargets<MUTABLE_OP_MAX_INTERPOLATE_COUNT
+            ; t< node.m_targets.Num() && numTargets<MUTABLE_OP_MAX_INTERPOLATE_COUNT
             ; ++t )
         {
             if ( Node* pA = node.m_targets[t].get() )
@@ -639,7 +639,7 @@ namespace mu
                 // TODO: Support other formats
                 target = GenerateImageFormat( target, EImageFormat::IF_RGB_UBYTE );
                 target = GenerateImageSize
-                        ( target, FImageSize((uint16_t)m_imageState.back().m_imageRect.size[0],(uint16_t)m_imageState.back().m_imageRect.size[1]) );
+                        ( target, FImageSize((uint16)m_imageState.Last().m_imageRect.size[0],(uint16)m_imageState.Last().m_imageRect.size[1]) );
 
                 op->SetChild( op->op.args.ImageInterpolate.targets[numTargets], target);
                 numTargets++;
@@ -652,7 +652,7 @@ namespace mu
             Ptr<ASTOp> target = GenerateMissingImageCode( "First interpolation image", EImageFormat::IF_RGB_UBYTE,
                                                        node.m_errorContext );
             target = GenerateImageSize
-                    ( target, FImageSize((uint16_t)m_imageState.back().m_imageRect.size[0],(uint16_t)m_imageState.back().m_imageRect.size[1]) );
+                    ( target, FImageSize((uint16)m_imageState.Last().m_imageRect.size[0],(uint16)m_imageState.Last().m_imageRect.size[1]) );
             op->SetChild( op->op.args.ImageInterpolate.targets[0], target);
         }
 
@@ -707,7 +707,7 @@ namespace mu
         }
         target = GenerateImageFormat( target, EImageFormat::IF_RGB_UBYTE );
         target = GenerateImageSize
-                ( target, FImageSize((uint16_t)m_imageState.back().m_imageRect.size[0],(uint16_t)m_imageState.back().m_imageRect.size[1]) );
+                ( target, FImageSize((uint16)m_imageState.Last().m_imageRect.size[0],(uint16)m_imageState.Last().m_imageRect.size[1]) );
         op->SetChild( op->op.args.ImageInterpolate3.target0, target);
 
         // Target 1
@@ -723,8 +723,8 @@ namespace mu
         }
         target = GenerateImageFormat( target, EImageFormat::IF_RGB_UBYTE );
         target = GenerateImageSize
-                ( target, FImageSize( (uint16_t)m_imageState.back().m_imageRect.size[0],
-                                      (uint16_t)m_imageState.back().m_imageRect.size[1]) );
+                ( target, FImageSize( (uint16)m_imageState.Last().m_imageRect.size[0],
+                                      (uint16)m_imageState.Last().m_imageRect.size[1]) );
         op->SetChild( op->op.args.ImageInterpolate3.target1, target);
 
         // Target 2
@@ -740,8 +740,8 @@ namespace mu
         }
         target = GenerateImageFormat( target, EImageFormat::IF_RGB_UBYTE );
         target = GenerateImageSize
-                ( target, FImageSize( (uint16_t)m_imageState.back().m_imageRect.size[0],
-                                      (uint16_t)m_imageState.back().m_imageRect.size[1] ) );
+                ( target, FImageSize( (uint16)m_imageState.Last().m_imageRect.size[0],
+                                      (uint16)m_imageState.Last().m_imageRect.size[1] ) );
         op->SetChild( op->op.args.ImageInterpolate3.target2, target);
 
 
@@ -802,11 +802,11 @@ namespace mu
 		check(node.m_format != EImageFormat::IF_NONE);
 
 		// Source images and channels
-		check(node.m_sources.size() == node.m_sourceChannels.size());
+		check(node.m_sources.Num() == node.m_sourceChannels.Num());
 
 		// First source, for reference in the size
         Ptr<ASTOp> first = 0;
-		for (std::size_t t = 0; t< node.m_sources.size(); ++t)
+		for (std::size_t t = 0; t< node.m_sources.Num(); ++t)
 		{
 			if (Node* pChannel = node.m_sources[t].get())
 			{
@@ -911,7 +911,7 @@ namespace mu
         }
         base = GenerateImageFormat( base, EImageFormat::IF_RGB_UBYTE );
         base = GenerateImageSize
-                ( base, FImageSize((uint16_t)m_imageState.back().m_imageRect.size[0],(uint16_t)m_imageState.back().m_imageRect.size[1]) );
+                ( base, FImageSize((uint16)m_imageState.Last().m_imageRect.size[0],(uint16)m_imageState.Last().m_imageRect.size[1]) );
         op->SetChild( op->op.args.ImageSaturate.base, base);
 
 
@@ -980,7 +980,7 @@ namespace mu
                                              node.m_errorContext );
         }
         base = GenerateImageSize
-                ( base, FImageSize((uint16_t)m_imageState.back().m_imageRect.size[0],(uint16_t)m_imageState.back().m_imageRect.size[1]) );
+                ( base, FImageSize((uint16)m_imageState.Last().m_imageRect.size[0],(uint16)m_imageState.Last().m_imageRect.size[1]) );
         op->SetChild( op->op.args.ImageSelectColour.base, base);
 
 
@@ -1021,8 +1021,8 @@ namespace mu
         }
         a = GenerateImageFormat( a, EImageFormat::IF_RGB_UBYTE );
         a = GenerateImageSize( a,
-                               FImageSize( (uint16_t)m_imageState.back().m_imageRect.size[0],
-                                           (uint16_t)m_imageState.back().m_imageRect.size[1]) );
+                               FImageSize( (uint16)m_imageState.Last().m_imageRect.size[0],
+                                           (uint16)m_imageState.Last().m_imageRect.size[1]) );
         op->SetChild( op->op.args.ImageDifference.a, a);
 
         // B image
@@ -1038,7 +1038,7 @@ namespace mu
         }
         b = GenerateImageFormat( b, EImageFormat::IF_RGB_UBYTE );
         b = GenerateImageSize
-                ( b, FImageSize((uint16_t)m_imageState.back().m_imageRect.size[0],(uint16_t)m_imageState.back().m_imageRect.size[1]) );
+                ( b, FImageSize((uint16)m_imageState.Last().m_imageRect.size[0],(uint16)m_imageState.Last().m_imageRect.size[1]) );
         op->SetChild( op->op.args.ImageDifference.b, b);
 
         result.op = op;
@@ -1066,8 +1066,8 @@ namespace mu
                                              node.m_errorContext );
         }
         base = GenerateImageSize
-                ( base, FImageSize( (uint16_t)m_imageState.back().m_imageRect.size[0],
-                                    (uint16_t)m_imageState.back().m_imageRect.size[1]) );
+                ( base, FImageSize( (uint16)m_imageState.Last().m_imageRect.size[0],
+                                    (uint16)m_imageState.Last().m_imageRect.size[1]) );
         op->SetChild( op->op.args.ImageColourMap.base, base);
 
         // Mask of the effect
@@ -1084,8 +1084,8 @@ namespace mu
         }
         mask = GenerateImageFormat( mask, EImageFormat::IF_L_UBYTE );
         mask = GenerateImageSize
-                ( mask, FImageSize( (uint16_t)m_imageState.back().m_imageRect.size[0],
-                                    (uint16_t)m_imageState.back().m_imageRect.size[1]) );
+                ( mask, FImageSize( (uint16)m_imageState.Last().m_imageRect.size[0],
+                                    (uint16)m_imageState.Last().m_imageRect.size[1]) );
         op->SetChild( op->op.args.ImageColourMap.mask, mask);
 
         // Map image
@@ -1101,8 +1101,8 @@ namespace mu
         }
         mapImage = GenerateImageFormat( mapImage, EImageFormat::IF_RGB_UBYTE );
         mapImage = GenerateImageSize
-                ( mapImage, FImageSize( (uint16_t)m_imageState.back().m_imageRect.size[0],
-                                        (uint16_t)m_imageState.back().m_imageRect.size[1]) );
+                ( mapImage, FImageSize( (uint16)m_imageState.Last().m_imageRect.size[0],
+                                        (uint16)m_imageState.Last().m_imageRect.size[1]) );
         op->SetChild( op->op.args.ImageColourMap.map, mapImage);
 
         result.op = op;
@@ -1143,8 +1143,8 @@ namespace mu
         }
         op->SetChild( op->op.args.ImageGradient.colour1, colour1);
 
-        op->op.args.ImageGradient.size[0] = (uint16_t)std::max( 2, std::min( 1024, node.m_size[0] ) );
-        op->op.args.ImageGradient.size[1] = (uint16_t)std::max( 1, std::min( 1024, node.m_size[1] ) );
+        op->op.args.ImageGradient.size[0] = (uint16)FMath::Max( 2, FMath::Min( 1024, node.m_size[0] ) );
+        op->op.args.ImageGradient.size[1] = (uint16)FMath::Max( 1, FMath::Min( 1024, node.m_size[1] ) );
 
         result.op = op;
     }
@@ -1172,8 +1172,8 @@ namespace mu
         }
         a = GenerateImageFormat( a, EImageFormat::IF_RGB_UBYTE );
         a = GenerateImageSize
-                ( a, FImageSize( (uint16_t)m_imageState.back().m_imageRect.size[0],
-                                 (uint16_t)m_imageState.back().m_imageRect.size[1]) );
+                ( a, FImageSize( (uint16)m_imageState.Last().m_imageRect.size[0],
+                                 (uint16)m_imageState.Last().m_imageRect.size[1]) );
         op->SetChild( op->op.args.ImageBinarise.base, a );
 
         // Threshold
@@ -1206,7 +1206,7 @@ namespace mu
         Ptr<ASTOp> base = 0;
         if ( Node* pSource = node.m_pBase.get() )
         {
-            IMAGE_STATE newState = m_imageState.back();
+            IMAGE_STATE newState = m_imageState.Last();
             if ( node.m_relative )
             {
                 newState.m_imageSize[0]=(int)std::roundf(newState.m_imageSize[0]/node.m_sizeX);
@@ -1217,11 +1217,11 @@ namespace mu
                 newState.m_imageRect.size[1]=(int)std::roundf(newState.m_imageRect.size[1]/node.m_sizeY);
             }
 
-            m_imageState.push_back(newState);
+            m_imageState.Add(newState);
 
             base = Generate( pSource );
 
-            m_imageState.pop_back();
+            m_imageState.Pop();
         }
         else
         {
@@ -1243,8 +1243,8 @@ namespace mu
         {
             Ptr<ASTOpFixed> op = new ASTOpFixed();
             op->op.type = OP_TYPE::IM_RESIZE;
-            op->op.args.ImageResize.size[0] = (uint16_t)node.m_sizeX;
-            op->op.args.ImageResize.size[1] = (uint16_t)node.m_sizeY;
+            op->op.args.ImageResize.size[0] = (uint16)node.m_sizeX;
+            op->op.args.ImageResize.size[1] = (uint16)node.m_sizeY;
             op->SetChild( op->op.args.ImageResize.source, base);
             at = op;
         }
@@ -1274,20 +1274,20 @@ namespace mu
         op->op.type = OP_TYPE::IM_PLAINCOLOUR;
         op->SetChild( op->op.args.ImagePlainColour.colour, base);
         op->op.args.ImagePlainColour.format = EImageFormat::IF_RGB_UBYTE;
-        op->op.args.ImagePlainColour.size[0] = uint16_t(node.m_sizeX);
-        op->op.args.ImagePlainColour.size[1] = uint16_t(node.m_sizeY);
+        op->op.args.ImagePlainColour.size[0] = uint16(node.m_sizeX);
+        op->op.args.ImagePlainColour.size[1] = uint16(node.m_sizeY);
 
         Ptr<ASTOpFixed> opSize = new ASTOpFixed();
         opSize->op.type = OP_TYPE::IM_RESIZE;
-        if (m_imageState.size())
+        if (m_imageState.Num())
         {
-            opSize->op.args.ImageResize.size[0] = (uint16_t)m_imageState.back().m_imageRect.size[0];
-            opSize->op.args.ImageResize.size[1] = (uint16_t)m_imageState.back().m_imageRect.size[1];
+            opSize->op.args.ImageResize.size[0] = (uint16)m_imageState.Last().m_imageRect.size[0];
+            opSize->op.args.ImageResize.size[1] = (uint16)m_imageState.Last().m_imageRect.size[1];
         }
         else
         {
-            opSize->op.args.ImageResize.size[0] = (uint16_t)node.m_sizeX;
-            opSize->op.args.ImageResize.size[1] = (uint16_t)node.m_sizeY;
+            opSize->op.args.ImageResize.size[0] = (uint16)node.m_sizeX;
+            opSize->op.args.ImageResize.size[1] = (uint16)node.m_sizeY;
         }
         opSize->SetChild( opSize->op.args.ImageResize.source, op);
 
@@ -1302,7 +1302,7 @@ namespace mu
 
 		MUTABLE_CPUPROFILER_SCOPE(NodeImageSwitch);
 
-        if (node.m_options.size() == 0)
+        if (node.m_options.Num() == 0)
 		{
 			// No options in the switch!
             Ptr<ASTOp> missingOp = GenerateMissingImageCode("Switch option",
@@ -1327,7 +1327,7 @@ namespace mu
 		}
 
 		// Options
-        for ( std::size_t t=0; t< node.m_options.size(); ++t )
+        for ( std::size_t t=0; t< node.m_options.Num(); ++t )
         {
             Ptr<ASTOp> branch;
 
@@ -1343,7 +1343,7 @@ namespace mu
                                                     node.m_errorContext);
             }
 
-            op->cases.emplace_back((int16_t)t,op,branch);
+            op->cases.Emplace((int16_t)t,op,branch);
         }
 
         Ptr<ASTOp> switchAt = op;
@@ -1441,7 +1441,7 @@ namespace mu
         if ( node.m_pMesh )
         {
             MESH_GENERATION_RESULT MeshResult;
-			check(m_overrideLayoutsStack.empty());
+			check(m_overrideLayoutsStack.IsEmpty());
             GenerateMesh(MeshResult, node.m_pMesh );
 
             pop->SetChild( pop->op.args.MeshProject.mesh, MeshResult.meshOp );
@@ -1467,12 +1467,12 @@ namespace mu
             else
             {
                 // Extract the mesh layout block
-                if ( m_imageState.size() && m_imageState.back().m_layoutBlock>=0 )
+                if ( m_imageState.Num() && m_imageState.Last().m_layoutBlock>=0 )
                 {
                     Ptr<ASTOpMeshExtractLayoutBlocks> eop = new ASTOpMeshExtractLayoutBlocks();
                     eop->source = pop->children[pop->op.args.MeshProject.mesh].child();
                     eop->layout = node.m_layout;
-                    eop->blocks.push_back( m_imageState.back().m_layoutBlock );
+                    eop->blocks.Add( m_imageState.Last().m_layoutBlock );
 
                     pop->SetChild( pop->op.args.MeshProject.mesh, eop );
                 }
@@ -1526,13 +1526,13 @@ namespace mu
             newState.m_imageRect.min[1] = 0;
             newState.m_imageRect.size = desc.m_size;
             newState.m_layoutBlock = -1;
-            m_imageState.push_back( newState );
+            m_imageState.Add( newState );
 
             // Generate
             op->SetChild( op->op.args.ImageRasterMesh.image, Generate( node.m_pImage.get() ) );
 
             // Restore rect
-            m_imageState.pop_back();
+            m_imageState.Pop();
 
         }
         else
@@ -1543,9 +1543,9 @@ namespace mu
         }
 
         // Image size, from the current block being generated
-        op->op.args.ImageRasterMesh.sizeX = (uint16_t)m_imageState.back().m_imageRect.size[0];
-        op->op.args.ImageRasterMesh.sizeY = (uint16_t)m_imageState.back().m_imageRect.size[1];
-        op->op.args.ImageRasterMesh.blockIndex = m_imageState.back().m_layoutBlock;
+        op->op.args.ImageRasterMesh.sizeX = (uint16)m_imageState.Last().m_imageRect.size[0];
+        op->op.args.ImageRasterMesh.sizeY = (uint16)m_imageState.Last().m_imageRect.size[1];
+        op->op.args.ImageRasterMesh.blockIndex = m_imageState.Last().m_layoutBlock;
 
         // Fading properties are optional, and stored in a colour
         if (node.m_pAngleFadeStart||node.m_pAngleFadeEnd)
@@ -1571,8 +1571,8 @@ namespace mu
             mask = GenerateImageFormat( mask, EImageFormat::IF_L_UBYTE );
             op->SetChild( op->op.args.ImageRasterMesh.mask, GenerateImageSize
                     ( mask,
-                      FImageSize( (uint16_t)m_imageState.back().m_imageRect.size[0],
-                                  (uint16_t)m_imageState.back().m_imageRect.size[1]) )
+                      FImageSize( (uint16)m_imageState.Last().m_imageRect.size[0],
+                                  (uint16)m_imageState.Last().m_imageRect.size[1]) )
                               );
         }
 
@@ -1648,7 +1648,7 @@ namespace mu
         {
             int mipsX = (int)ceilf( logf( (float)blockX )/logf(2.0f) );
             int mipsY = (int)ceilf( logf( (float)blockY )/logf(2.0f) );
-            op->BlockLevels = (uint8_t)std::max( mipsX, mipsY );
+            op->BlockLevels = (uint8_t)FMath::Max( mipsX, mipsY );
         }
         else
         {
@@ -1689,8 +1689,8 @@ namespace mu
 		}
 		a = GenerateImageFormat(a, EImageFormat::IF_RGB_UBYTE);
 		a = GenerateImageSize
-		(a, FImageSize((uint16_t)m_imageState.back().m_imageRect.size[0],
-			(uint16_t)m_imageState.back().m_imageRect.size[1]));
+		(a, FImageSize((uint16)m_imageState.Last().m_imageRect.size[0],
+			(uint16)m_imageState.Last().m_imageRect.size[1]));
 		op->SetChild(op->op.args.ImageInvert.base, a);
 
 		result.op = op;
@@ -1718,11 +1718,11 @@ namespace mu
         }
 
         // Process variations in reverse order, since conditionals are built bottom-up.
-        for ( int t = int( node.m_variations.size() ) - 1; t >= 0; --t )
+        for ( int t = int( node.m_variations.Num() ) - 1; t >= 0; --t )
         {
             int tagIndex = -1;
             const string& tag = node.m_variations[t].m_tag;
-            for ( int i = 0; i < int( m_firstPass.m_tags.size() ); ++i )
+            for ( int i = 0; i < int( m_firstPass.m_tags.Num() ); ++i )
             {
                 if ( m_firstPass.m_tags[i].tag == tag )
                 {
@@ -1946,9 +1946,9 @@ namespace mu
         uint8_t* pData = pImage->GetData();
         for ( int p=0; p<size*size; ++p )
         {
-            pData[0] = (uint8_t)std::min( 255.0f, std::max( 0.0f, 255*colour[0] ) );
-            pData[1] = (uint8_t)std::min( 255.0f, std::max( 0.0f, 255*colour[1] ) );
-            pData[2] = (uint8_t)std::min( 255.0f, std::max( 0.0f, 255*colour[2] ) );
+            pData[0] = (uint8_t)FMath::Min( 255.0f, FMath::Max( 0.0f, 255*colour[0] ) );
+            pData[1] = (uint8_t)FMath::Min( 255.0f, FMath::Max( 0.0f, 255*colour[1] ) );
+            pData[2] = (uint8_t)FMath::Min( 255.0f, FMath::Max( 0.0f, 255*colour[2] ) );
             pData += 3;
         }
 
diff --git a/Engine/Plugins/Experimental/Mutable/Source/MutableTools/Private/MuT/CodeGenerator_Mesh.cpp b/Engine/Plugins/Experimental/Mutable/Source/MutableTools/Private/MuT/CodeGenerator_Mesh.cpp
--- a/Engine/Plugins/Experimental/Mutable/Source/MutableTools/Private/MuT/CodeGenerator_Mesh.cpp
+++ b/Engine/Plugins/Experimental/Mutable/Source/MutableTools/Private/MuT/CodeGenerator_Mesh.cpp
@@ -109,7 +109,7 @@ class Node;
         check( channel>=0 );
 
         // Create the layout block vertex buffer
-        uint16_t* pLayoutData = 0;
+        uint16* pLayoutData = 0;
         {
             int layoutBuf = currentLayoutMesh->GetVertexBuffers().GetBufferCount();
             currentLayoutMesh->GetVertexBuffers().SetBufferCount( layoutBuf+1 );
@@ -124,13 +124,13 @@ class Node;
             currentLayoutMesh->GetVertexBuffers().SetBuffer
                     (
                         layoutBuf,
-                        sizeof(uint16_t),
+                        sizeof(uint16),
                         1,
                         &semantic, &semanticIndex,
                         &format, &components,
                         &offset
                     );
-            pLayoutData = (uint16_t*)currentLayoutMesh->GetVertexBuffers().GetBufferData( layoutBuf );
+            pLayoutData = (uint16*)currentLayoutMesh->GetVertexBuffers().GetBufferData( layoutBuf );
         }
 
         // Get the information about the texture coordinates channel
@@ -187,7 +187,7 @@ class Node;
 
                         // Set the value to the unique block id
                         check( pLayout->m_blocks[b].m_id < 65535 );
-                        pLayoutData[v] =(uint16_t) pLayout->m_blocks[b].m_id;
+                        pLayoutData[v] =(uint16) pLayout->m_blocks[b].m_id;
 
                         inside++;
                     }
@@ -213,7 +213,7 @@ class Node;
 
                         // Set the value to the unique block id
                         check( pLayout->m_blocks[b].m_id < 65535 );
-                        pLayoutData[v] = (uint16_t)pLayout->m_blocks[b].m_id;
+                        pLayoutData[v] = (uint16)pLayout->m_blocks[b].m_id;
 
                         inside++;
 
@@ -232,13 +232,13 @@ class Node;
                 (
                     buf, 256,
                     "Source mesh has %d vertices not assigned to any layout block in LOD %d",
-					outside, m_currentParents.back().m_lod
+					outside, m_currentParents.Last().m_lod
                 );
 			int blockCount = pLayout->GetBlockCount();
         
             //m_pErrorLog->GetPrivate()->Add( buf, blockCount==1?ELMT_INFO:ELMT_WARNING, errorContext );
-            vector< float > unassignedUVs;
-            unassignedUVs.reserve(64);           
+            TArray< float > unassignedUVs;
+            unassignedUVs.Reserve(64);           
             
             const uint8_t* pVertices = pData;
             for (int i=0; i<currentLayoutMesh->GetVertexBuffers().GetElementCount(); ++i)
@@ -258,14 +258,14 @@ class Node;
 
                 if (pLayoutData[i]==65535)
                 {
-                    unassignedUVs.push_back(UV[0]);
-                    unassignedUVs.push_back(UV[1]);
+                    unassignedUVs.Add(UV[0]);
+                    unassignedUVs.Add(UV[1]);
                 }
             }
 
             ErrorLogMessageAttachedDataView attachedDataView;
-            attachedDataView.m_unassignedUVs = unassignedUVs.data();
-            attachedDataView.m_unassignedUVsSize = unassignedUVs.size();
+            attachedDataView.m_unassignedUVs = unassignedUVs.GetData();
+            attachedDataView.m_unassignedUVsSize = (size_t)unassignedUVs.Num();
 
             m_pErrorLog->GetPrivate()->Add( buf, attachedDataView, blockCount==1?ELMT_INFO:ELMT_WARNING, errorContext );
         }
@@ -367,22 +367,22 @@ class Node;
         }
 
         // TODO: Support more than MUTABLE_OP_MAX_MORPH2_TARGETS targets
-        if ( node.m_morphs.size()>MUTABLE_OP_MAX_MORPH2_TARGETS )
+        if ( node.m_morphs.Num()>MUTABLE_OP_MAX_MORPH2_TARGETS )
         {
             char temp[256];
             mutable_snprintf( temp, 256,
                               "A morph node has more targets [%d] than currently supported [%d].",
-                              node.m_morphs.size(),
+                              node.m_morphs.Num(),
                               MUTABLE_OP_MAX_MORPH2_TARGETS );
             m_pErrorLog->GetPrivate()->Add( temp, ELMT_WARNING, node.m_errorContext );
         }
 
-        m_overrideLayoutsStack.push_back( baseResult.layouts );
-		m_activeTags.push_back({});
+        m_overrideLayoutsStack.Add( baseResult.layouts );
+		m_activeTags.Add({});
 
         int count = 0;
         for ( std::size_t t=0
-            ; t<node.m_morphs.size() && t<MUTABLE_OP_MAX_MORPH2_TARGETS
+            ; t<node.m_morphs.Num() && t<MUTABLE_OP_MAX_MORPH2_TARGETS
             ; ++t )
         {
             if ( auto pA = node.m_morphs[t].get() )
@@ -441,8 +441,8 @@ class Node;
             OpMorphReshape->Reshape = OpApply;
         }
 
-        m_overrideLayoutsStack.pop_back();
-		m_activeTags.pop_back();
+        m_overrideLayoutsStack.Pop();
+		m_activeTags.Pop();
 
 		if (OpMorphReshape)
 		{
@@ -486,8 +486,8 @@ class Node;
         }
 
         // Target
-        m_overrideLayoutsStack.push_back( baseResult.layouts );
-		m_activeTags.push_back({});
+        m_overrideLayoutsStack.Add( baseResult.layouts );
+		m_activeTags.Add({});
 		if ( node.m_pTarget )
         {
             MESH_GENERATION_RESULT targetResult;
@@ -501,8 +501,8 @@ class Node;
             m_pErrorLog->GetPrivate()->Add( "Mesh make morph target node is not set.",
                                             ELMT_ERROR, node.m_errorContext );
         }
-        m_overrideLayoutsStack.pop_back();
-		m_activeTags.pop_back();
+        m_overrideLayoutsStack.Pop();
+		m_activeTags.Pop();
 
         result.meshOp = op;
         result.baseMeshOp = baseResult.baseMeshOp;
@@ -527,18 +527,18 @@ class Node;
 
                 op->source = baseResult.meshOp;
 
-                if ( baseResult.layouts.size()>(size_t)node.m_layoutOrGroup )
+                if ( baseResult.layouts.Num()>node.m_layoutOrGroup )
                 {
                     LayoutPtrConst pSourceLayout = baseResult.layouts[ node.m_layoutOrGroup ];
                     const Layout* pLayout = m_addedLayouts[ pSourceLayout.get() ].get();
-                    op->layout = (uint16_t)node.m_layoutOrGroup;
+                    op->layout = (uint16)node.m_layoutOrGroup;
 
-                    for ( size_t i=0; i<node.m_blocks.size(); ++i )
+                    for ( int32 i=0; i<node.m_blocks.Num(); ++i )
                     {
                         if (node.m_blocks[i]>=0 && node.m_blocks[i]<pLayout->m_blocks.Num() )
                         {
                             int bid = pLayout->m_blocks[ node.m_blocks[i] ].m_id;
-                            op->blocks.push_back(bid);
+                            op->blocks.Add(bid);
                         }
                         else
                         {
@@ -605,15 +605,15 @@ class Node;
         //
         Ptr<ASTOp> base = 0;
         int count = 0;
-        for ( std::size_t t=0
-            ; t<node.m_targets.size() && t<MUTABLE_OP_MAX_INTERPOLATE_COUNT-1
+        for ( int32 t=0
+            ; t<node.m_targets.Num() && t<MUTABLE_OP_MAX_INTERPOLATE_COUNT-1
             ; ++t )
         {
             if ( NodeMesh* pA = node.m_targets[t].get() )
             {
                 if ( count>0 )
                 {
-                    m_overrideLayoutsStack.push_back( result.layouts );
+                    m_overrideLayoutsStack.Add( result.layouts );
                 }
 
                 MESH_GENERATION_RESULT targetResult;
@@ -621,7 +621,7 @@ class Node;
 
                 if ( count>0 )
                 {
-                    m_overrideLayoutsStack.pop_back();
+                    m_overrideLayoutsStack.Pop();
                 }
 
                 // The first target is the base
@@ -643,18 +643,18 @@ class Node;
                     // \todo Texcoords are broken?
                     dop->op.args.MeshDifference.ignoreTextureCoords = 1;
 
-                    if ( node.m_channels.size()>MUTABLE_OP_MAX_MORPH_CHANNELS )
+                    if ( node.m_channels.Num()>MUTABLE_OP_MAX_MORPH_CHANNELS )
                     {
                         char temp[256];
                         mutable_snprintf( temp, 256,
                                           "Morph uses too many channels [%d]. The maximum is [%d].",
-                                          node.m_channels.size(),
+                                          node.m_channels.Num(),
                                           MUTABLE_OP_MAX_MORPH_CHANNELS );
                         m_pErrorLog->GetPrivate()->Add( temp, ELMT_ERROR, node.m_errorContext );
                     }
 
                     for ( size_t c=0;
-                          c<node.m_channels.size() && c<MUTABLE_OP_MAX_MORPH_CHANNELS;
+                          c<node.m_channels.Num() && c<MUTABLE_OP_MAX_MORPH_CHANNELS;
                           ++c)
                     {
                         check( node.m_channels[c].semantic<256 );
@@ -689,7 +689,7 @@ class Node;
     {
         NodeMeshSwitch::Private& node = *sw->GetPrivate();
 
-        if (node.m_options.size() == 0)
+        if (node.m_options.Num() == 0)
         {
             // No options in the switch!
             // TODO
@@ -712,11 +712,11 @@ class Node;
         }
 
         // Options
-        for ( std::size_t t=0; t< node.m_options.size(); ++t )
+        for ( std::size_t t=0; t< node.m_options.Num(); ++t )
         {
             if (t!=0)
             {
-                m_overrideLayoutsStack.push_back( result.layouts );
+                m_overrideLayoutsStack.Add( result.layouts );
             }
 
             if ( node.m_options[t] )
@@ -725,7 +725,7 @@ class Node;
                 GenerateMesh( branchResults, node.m_options[t] );
 
                 auto branch = branchResults.meshOp;
-                op->cases.emplace_back((int16_t)t,op,branch);
+                op->cases.Emplace((int16)t,op,branch);
 
                 if (t==0)
                 {
@@ -735,7 +735,7 @@ class Node;
 
             if (t!=0)
             {
-                m_overrideLayoutsStack.pop_back();
+                m_overrideLayoutsStack.Pop();
             }
         }
 
@@ -767,7 +767,7 @@ class Node;
 				pCell->SetValue(pMesh);
 
 				// TODO Take into account layout strategy
-				int numLayouts = (int)node.m_layouts.size();
+				int numLayouts = node.m_layouts.Num();
 				pCell->SetLayoutCount(numLayouts);
 				for (int i = 0; i < numLayouts; ++i)
 				{
@@ -776,7 +776,7 @@ class Node;
 
 				if (t != 0)
 				{
-					m_overrideLayoutsStack.push_back(NewResult.layouts);
+					m_overrideLayoutsStack.Add(NewResult.layouts);
 				}
 
 				MESH_GENERATION_RESULT branchResults;
@@ -788,7 +788,7 @@ class Node;
 				}
 				else
 				{
-					m_overrideLayoutsStack.pop_back();
+					m_overrideLayoutsStack.Pop();
 				}
 
 				++t;
@@ -823,11 +823,11 @@ class Node;
         }
 
         // Process variations in reverse order, since conditionals are built bottom-up.
-        for ( int t = int(node.m_variations.size())-1; t >= 0; --t )
+        for ( int32 t = node.m_variations.Num()-1; t >= 0; --t )
         {
             int tagIndex = -1;
             const string& tag = node.m_variations[t].m_tag;
-            for ( int i = 0; i < int( m_firstPass.m_tags.size() ); ++i )
+            for ( int i = 0; i < m_firstPass.m_tags.Num(); ++i )
             {
                 if ( m_firstPass.m_tags[i].tag==tag)
                 {
@@ -850,7 +850,7 @@ class Node;
             {
                 if (firstOptionProcessed)
                 {
-                    m_overrideLayoutsStack.push_back( currentResult.layouts );
+                    m_overrideLayoutsStack.Add( currentResult.layouts );
                 }
          
                 MESH_GENERATION_RESULT branchResults;
@@ -860,7 +860,7 @@ class Node;
 
 				if (firstOptionProcessed )
 				{
-					m_overrideLayoutsStack.pop_back();
+					m_overrideLayoutsStack.Pop();
 				}
 
                 if ( !firstOptionProcessed )
@@ -896,7 +896,7 @@ class Node;
 		result.baseMeshOp = op;
 		result.meshOp = op;
 
-		bool bIsOverridingLayouts = !m_overrideLayoutsStack.empty();
+		bool bIsOverridingLayouts = !m_overrideLayoutsStack.IsEmpty();
 
         MeshPtr pMesh = node.m_pValue.get();
         if (pMesh)
@@ -909,10 +909,10 @@ class Node;
                 // This means that we are processing a base mesh
 
                 // We will redefine the layouts
-				result.layouts.clear();
+				result.layouts.Empty();
 
                 // Apply whatever transform is necessary for every layout
-                for ( std::size_t LayoutIndex=0; LayoutIndex<node.m_layouts.size(); ++LayoutIndex)
+                for ( std::size_t LayoutIndex=0; LayoutIndex<node.m_layouts.Num(); ++LayoutIndex)
                 {
 					NodeLayoutPtr pLayoutNode = node.m_layouts[LayoutIndex];
 					// TODO: In a cleanup of the design of the layouts, we should remove this cast.
@@ -927,7 +927,7 @@ class Node;
 						}
 
 						LayoutPtr pSourceLayout = TypedNode->GetPrivate()->m_pLayout;
-						result.layouts.push_back(pSourceLayout);
+						result.layouts.Add(pSourceLayout);
 
 						PrepareForLayout(pSourceLayout,
 							pCloned, LayoutIndex,
@@ -939,15 +939,15 @@ class Node;
             else
             {
                 // We need to apply the transform of the source layouts
-                for ( std::size_t LayoutIndex=0; LayoutIndex<m_overrideLayoutsStack.back().size(); ++LayoutIndex)
+                for ( int32 LayoutIndex=0; LayoutIndex<m_overrideLayoutsStack.Last().Num(); ++LayoutIndex)
                 {
-                    PrepareForLayout( m_overrideLayoutsStack.back()[LayoutIndex],
+                    PrepareForLayout( m_overrideLayoutsStack.Last()[LayoutIndex],
                                       pCloned, LayoutIndex,
                                       node.m_errorContext  );
 
                 }
 
-				result.layouts = m_overrideLayoutsStack.back();
+				result.layouts = m_overrideLayoutsStack.Last();
 
                 // Disabled: lots of misleading messages (when not using layouts?).
                 // Warn about unnecessary layouts
@@ -963,7 +963,7 @@ class Node;
 			// For now we cannot handle the case were we are overriding layouts and reusing meshes at the same time.
             bool bIsDuplicated = false;
             MeshPtrConst Candidate;
-            for ( size_t i=0; !bIsOverridingLayouts && !bIsDuplicated && i<m_constantMeshes.size(); ++i)
+            for ( int32 i=0; !bIsOverridingLayouts && !bIsDuplicated && i<m_constantMeshes.Num(); ++i)
             {
 				Candidate = m_constantMeshes[i];
 
@@ -992,14 +992,14 @@ class Node;
 						const Layout* GeneratedMeshLayout = Candidate->GetLayout(l);
 						if (bIsOverridingLayouts)
 						{
-							const Layout* OverridingLayout = m_overrideLayoutsStack.back()[l].get();
+							const Layout* OverridingLayout = m_overrideLayoutsStack.Last()[l].get();
 							check(OverridingLayout->m_blocks.IsEmpty() || OverridingLayout->m_blocks[0].m_id == -1);
 							GeneratedMeshLayout = m_addedLayouts[OverridingLayout].get();
 						}
 
 						check(DestLayoutKey->m_blocks.IsEmpty() || DestLayoutKey->m_blocks[0].m_id == -1);
 						check(GeneratedMeshLayout->m_blocks.IsEmpty() || GeneratedMeshLayout->m_blocks[0].m_id != -1);
-						m_addedLayouts[DestLayoutKey] = GeneratedMeshLayout;
+						m_addedLayouts.Add(DestLayoutKey, GeneratedMeshLayout);
 					}
                 }
             }
@@ -1042,18 +1042,18 @@ class Node;
                 }
 
                 // Add the constant data
-                m_constantMeshes.push_back(pCloned);
+                m_constantMeshes.Add(pCloned);
                 op->SetValue( pCloned.get(), m_compilerOptions->m_optimisationOptions.m_useDiskCache );
             }
         }
         else
         {
-			result.layouts.clear();
+			result.layouts.Empty();
 
             // This data is required
             MeshPtr pTempMesh = new Mesh();
             op->SetValue( pTempMesh, m_compilerOptions->m_optimisationOptions.m_useDiskCache );
-            m_constantMeshes.push_back(pTempMesh);
+            m_constantMeshes.Add(pTempMesh);
 
             // Log an error message
             m_pErrorLog->GetPrivate()->Add( "Constant mesh not set.",
@@ -1062,14 +1062,14 @@ class Node;
 
 		// Apply the modifier for the pre-normal operations stage.
 		BOTTOM_UP_STATE temp = m_currentBottomUpState;
-		if (!m_activeTags.empty())
+		if (!m_activeTags.IsEmpty())
 		{
 			// Clearing layout stack to avoid unwanted information
-			vector< vector<Ptr<const Layout>> > m_overrideLayoutsStackCopy = m_overrideLayoutsStack;
-			m_overrideLayoutsStack.clear();
+			TArray< TArray<Ptr<const Layout>> > m_overrideLayoutsStackCopy = m_overrideLayoutsStack;
+			m_overrideLayoutsStack.Empty();
 
 			bool bModifiersForBeforeOperations = true;
-			result.meshOp = ApplyMeshModifiers(op, m_activeTags.back(),
+			result.meshOp = ApplyMeshModifiers(op, m_activeTags.Last(),
 				bModifiersForBeforeOperations, node.m_errorContext);
 
 			// Retrieving stack information
@@ -1097,11 +1097,6 @@ class Node;
             op->Source = baseResult.meshOp;
             op->Buffers = 0;
 
-            if ( node.m_rebuildTangents )
-            {
-                op->Buffers |= OP::MeshFormatArgs::BT_REBUILD_TANGENTS;
-            }
-
             MeshPtr pFormatMesh = new Mesh();
 
             if (node.m_VertexBuffers.GetBufferCount())
@@ -1127,7 +1122,7 @@ class Node;
             cop->SetValue( pFormatMesh, m_compilerOptions->m_optimisationOptions.m_useDiskCache );
             op->Format = cop;
 
-            m_constantMeshes.push_back(pFormatMesh);
+            m_constantMeshes.Add(pFormatMesh);
 
             result.meshOp = op;
             result.baseMeshOp = baseResult.baseMeshOp;
@@ -1273,13 +1268,13 @@ class Node;
         // Clipping mesh
         if (node.m_pClipMesh)
         {
-			m_activeTags.push_back({});
+			m_activeTags.Add({});
 
             MESH_GENERATION_RESULT clipResult;
             GenerateMesh( clipResult, node.m_pClipMesh);
             op->SetChild( op->op.args.MeshClipWithMesh.clipMesh, clipResult.meshOp );
 
-			m_activeTags.pop_back();
+			m_activeTags.Pop();
 		}
         else
         {
@@ -1352,15 +1347,15 @@ class Node;
         if (node.m_pPose)
         {
             // We don't need layouts for the pose mesh
-            m_overrideLayoutsStack.push_back( vector<LayoutPtrConst>() );
-			m_activeTags.push_back({});
+            m_overrideLayoutsStack.Add( TArray<LayoutPtrConst>() );
+			m_activeTags.Add({});
 
             MESH_GENERATION_RESULT poseResult;
             GenerateMesh( poseResult, node.m_pPose );
             op->pose = poseResult.meshOp;
 
-            m_overrideLayoutsStack.pop_back();
-			m_activeTags.pop_back();
+            m_overrideLayoutsStack.Pop();
+			m_activeTags.Pop();
 		}
         else
         {
@@ -1443,8 +1438,8 @@ class Node;
 		}
 
 		// Base and target shapes shouldn't have layouts or modifiers.
-		m_overrideLayoutsStack.push_back({});
-		m_activeTags.push_back({});
+		m_overrideLayoutsStack.Add({});
+		m_activeTags.Add({});
 
 		// Base Shape
 		if (node.m_pBaseShape)
@@ -1464,8 +1459,8 @@ class Node;
 			opApply->Shape = targetResult.meshOp;
 		}
 
-		m_overrideLayoutsStack.pop_back();
-		m_activeTags.pop_back();
+		m_overrideLayoutsStack.Pop();
+		m_activeTags.Pop();
 
 		result.meshOp = opApply;
 	}
diff --git a/Engine/Plugins/Experimental/Mutable/Source/MutableTools/Private/MuT/CodeGenerator_Projector.cpp b/Engine/Plugins/Experimental/Mutable/Source/MutableTools/Private/MuT/CodeGenerator_Projector.cpp
--- a/Engine/Plugins/Experimental/Mutable/Source/MutableTools/Private/MuT/CodeGenerator_Projector.cpp
+++ b/Engine/Plugins/Experimental/Mutable/Source/MutableTools/Private/MuT/CodeGenerator_Projector.cpp
@@ -127,12 +127,11 @@ void CodeGenerator::GenerateProjector_Parameter( PROJECTOR_GENERATION_RESULT& re
         op->parameter = param;
 
         // Generate the code for the ranges
-        for ( std::size_t a=0; a<node.m_ranges.size(); ++a )
+        for ( int32 a=0; a<node.m_ranges.Num(); ++a )
         {
             RANGE_GENERATION_RESULT rangeResult;
             GenerateRange( rangeResult, node.m_ranges[a] );
-            op->ranges.emplace_back
-                    ( op.get(), rangeResult.sizeOp, rangeResult.rangeName, rangeResult.rangeUID );
+            op->ranges.Emplace( op.get(), rangeResult.sizeOp, rangeResult.rangeName, rangeResult.rangeUID );
         }
 
         m_nodeVariables[node.m_pNode] = op;
diff --git a/Engine/Plugins/Experimental/Mutable/Source/MutableTools/Private/MuT/CodeGenerator_Scalar.cpp b/Engine/Plugins/Experimental/Mutable/Source/MutableTools/Private/MuT/CodeGenerator_Scalar.cpp
--- a/Engine/Plugins/Experimental/Mutable/Source/MutableTools/Private/MuT/CodeGenerator_Scalar.cpp
+++ b/Engine/Plugins/Experimental/Mutable/Source/MutableTools/Private/MuT/CodeGenerator_Scalar.cpp
@@ -152,16 +152,15 @@ void CodeGenerator::GenerateScalar_Parameter( SCALAR_GENERATION_RESULT& result,
 		op->parameter = param;
 
 		// Generate the code for the ranges
-		for (std::size_t a = 0; a < node.m_ranges.size(); ++a)
+		for (int32 a = 0; a < node.m_ranges.Num(); ++a)
 		{
 			RANGE_GENERATION_RESULT rangeResult;
 			GenerateRange(rangeResult, node.m_ranges[a]);
-			op->ranges.emplace_back
-			(op.get(), rangeResult.sizeOp, rangeResult.rangeName, rangeResult.rangeUID);
+			op->ranges.Emplace(op.get(), rangeResult.sizeOp, rangeResult.rangeName, rangeResult.rangeUID);
 		}
 
 		// Generate the code for the additional images in the parameter
-		for (std::size_t a = 0; a < node.m_additionalImages.size(); ++a)
+		for (int32 a = 0; a < node.m_additionalImages.Num(); ++a)
 		{
 			Ptr<ASTOp> descAd;
 			if (node.m_additionalImages[a])
@@ -174,16 +173,16 @@ void CodeGenerator::GenerateScalar_Parameter( SCALAR_GENERATION_RESULT& result,
 				newState.m_imageRect.min[1] = 0;
 				newState.m_imageRect.size = desc.m_size;
 				newState.m_layoutBlock = -1;
-				m_imageState.push_back(newState);
+				m_imageState.Add(newState);
 
 				// Generate
 				descAd = Generate(node.m_additionalImages[a]);
 				check(descAd);
 
 				// Restore rect
-				m_imageState.pop_back();
+				m_imageState.Pop();
 			}
-			op->additionalImages.emplace_back(op, descAd);
+			op->additionalImages.Emplace(op, descAd);
 		}
 
 		m_nodeVariables[node.m_pNode] = op;
@@ -214,8 +213,8 @@ void CodeGenerator::GenerateScalar_EnumParameter(SCALAR_GENERATION_RESULT& resul
 		param.m_defaultValue.m_int = node.m_defaultValue;
 		param.m_detailedType = node.m_detailedType;
 
-		param.m_possibleValues.resize(node.m_options.size());
-		for (size_t i = 0; i < node.m_options.size(); ++i)
+		param.m_possibleValues.SetNum(node.m_options.Num());
+		for (int32 i = 0; i < node.m_options.Num(); ++i)
 		{
 			param.m_possibleValues[i].m_value = (int16_t)node.m_options[i].value;
 			param.m_possibleValues[i].m_name = node.m_options[i].name;
@@ -226,12 +225,11 @@ void CodeGenerator::GenerateScalar_EnumParameter(SCALAR_GENERATION_RESULT& resul
 		op->parameter = param;
 
 		// Generate the code for the ranges
-		for (std::size_t a = 0; a < node.m_ranges.size(); ++a)
+		for (int32 a = 0; a < node.m_ranges.Num(); ++a)
 		{
 			RANGE_GENERATION_RESULT rangeResult;
 			GenerateRange(rangeResult, node.m_ranges[a]);
-			op->ranges.emplace_back
-			(op.get(), rangeResult.sizeOp, rangeResult.rangeName, rangeResult.rangeUID);
+			op->ranges.Emplace(op.get(), rangeResult.sizeOp, rangeResult.rangeName, rangeResult.rangeUID);
 		}
 
 		m_nodeVariables[node.m_pNode] = op;
@@ -251,7 +249,7 @@ void CodeGenerator::GenerateScalar_Switch(SCALAR_GENERATION_RESULT& result, cons
 {
 	const NodeScalarSwitch::Private& node = *Typed->GetPrivate();
 
-	if (node.m_options.size() == 0)
+	if (node.m_options.Num() == 0)
 	{
 		// No options in the switch!
 		Ptr<ASTOp> missingOp = GenerateMissingScalarCode("Switch option",
@@ -278,7 +276,7 @@ void CodeGenerator::GenerateScalar_Switch(SCALAR_GENERATION_RESULT& result, cons
 	}
 
 	// Options
-	for (std::size_t t = 0; t < node.m_options.size(); ++t)
+	for (int32 t = 0; t < node.m_options.Num(); ++t)
 	{
 		Ptr<ASTOp> branch;
 		if (node.m_options[t])
@@ -292,7 +290,7 @@ void CodeGenerator::GenerateScalar_Switch(SCALAR_GENERATION_RESULT& result, cons
 			// This argument is required
 			branch = GenerateMissingScalarCode("Switch option", 1.0f, node.m_errorContext);
 		}
-		op->cases.push_back(ASTOpSwitch::CASE((int16_t)t, op, branch));
+		op->cases.Emplace((int16_t)t, op, branch);
 	}
 
 	result.op = op;
@@ -322,11 +320,11 @@ void CodeGenerator::GenerateScalar_Variation(SCALAR_GENERATION_RESULT& result, c
 	}
 
 	// Process variations in reverse order, since conditionals are built bottom-up.
-	for (int t = int(node.m_variations.size()) - 1; t >= 0; --t)
+	for (int t = node.m_variations.Num() - 1; t >= 0; --t)
 	{
 		int tagIndex = -1;
 		const string& tag = node.m_variations[t].m_tag;
-		for (int i = 0; i < int(m_firstPass.m_tags.size()); ++i)
+		for (int i = 0; i < m_firstPass.m_tags.Num(); ++i)
 		{
 			if (m_firstPass.m_tags[i].tag == tag)
 			{
diff --git a/Engine/Plugins/Experimental/Mutable/Source/MutableTools/Private/MuT/CodeGenerator_SecondPass.cpp b/Engine/Plugins/Experimental/Mutable/Source/MutableTools/Private/MuT/CodeGenerator_SecondPass.cpp
--- a/Engine/Plugins/Experimental/Mutable/Source/MutableTools/Private/MuT/CodeGenerator_SecondPass.cpp
+++ b/Engine/Plugins/Experimental/Mutable/Source/MutableTools/Private/MuT/CodeGenerator_SecondPass.cpp
@@ -145,7 +145,7 @@ mu::Ptr<ASTOp> SecondPassGenerator::GenerateTagCondition( size_t tagIndex,
     // Condition expression for all the edit-surfaces that activate the tag
     for ( auto editKey: t.edits )
     {
-        if ( negSurf.find(editKey.first) != negSurf.end() )
+        if ( negSurf.find(editKey.Key) != negSurf.end() )
         {
             // The surface in the edit is a negative requirement higher up in the condition so
             // this branch never be true.
@@ -153,7 +153,7 @@ mu::Ptr<ASTOp> SecondPassGenerator::GenerateTagCondition( size_t tagIndex,
         }
 
         Ptr<ASTOp> surfCondition;
-        if ( posSurf.find(editKey.first) != posSurf.end() )
+        if ( posSurf.find(editKey.Key) != posSurf.end() )
         {
             // This surface in the edit is already a positive requirement higher up in the condition
             // so we don't need that part of the condition
@@ -163,15 +163,15 @@ mu::Ptr<ASTOp> SecondPassGenerator::GenerateTagCondition( size_t tagIndex,
             auto positiveTags = posTag;
             positiveTags.insert( tagIndex );
 
-            surfCondition = GenerateSurfaceCondition( editKey.first,
+            surfCondition = GenerateSurfaceCondition( editKey.Key,
                                                       posSurf,
                                                       negSurf,
                                                       positiveTags,
                                                       negTag );
         }
 
-        const auto& surface = m_pFirstPass->surfaces[editKey.first];
-        const auto& edit = surface.edits[editKey.second];
+        const auto& surface = m_pFirstPass->surfaces[editKey.Key];
+        const auto& edit = surface.edits[editKey.Value];
 
         // Combine object and surface conditions
         Ptr<ASTOp> fullCondition;
@@ -248,16 +248,15 @@ mu::Ptr<ASTOp> SecondPassGenerator::GenerateSurfaceCondition( size_t surfIndex,
 
     for (const auto& t: surf.positiveTags)
     {
-        auto it = std::find_if( m_pFirstPass->m_tags.begin(), m_pFirstPass->m_tags.end(),
-                                [&](const FirstPassGenerator::TAG& e) { return e.tag==t; } );
-        if (it==m_pFirstPass->m_tags.end())
+        auto it = m_pFirstPass->m_tags.FindByPredicate([&](const FirstPassGenerator::TAG& e) { return e.tag == t; } );
+        if (!it)
         {
             // This could happen if a tag is in a variation but noone defines it.
             // This surface depends on a tag that will never be active, so it will never be used.
             return m_opPool.Add( new ASTOpConstantBool(false) );
         }
 
-        size_t tagIndex = it - m_pFirstPass->m_tags.begin();
+        size_t tagIndex = it - &m_pFirstPass->m_tags[0];
 
         set<size_t> positiveSurfacesVisited = posSurf;
         positiveSurfacesVisited.insert( surfIndex );
@@ -313,15 +312,14 @@ mu::Ptr<ASTOp> SecondPassGenerator::GenerateSurfaceCondition( size_t surfIndex,
 
     for (const auto& t: surf.negativeTags)
     {
-        auto it = std::find_if( m_pFirstPass->m_tags.begin(), m_pFirstPass->m_tags.end(),
-                                [&](const FirstPassGenerator::TAG& e) { return e.tag==t; } );
-        if (it==m_pFirstPass->m_tags.end())
+        auto it = m_pFirstPass->m_tags.FindByPredicate([&](const FirstPassGenerator::TAG& e) { return e.tag == t; } );
+        if (!it)
         {
             // This could happen if a tag is in a variation but noone defines it.
             continue;
         }
 
-        size_t tagIndex = it - m_pFirstPass->m_tags.begin();
+        size_t tagIndex = it - &m_pFirstPass->m_tags[0];
 
         set<size_t> positiveSurfacesVisited = negSurf;
         set<size_t> negativeSurfacesVisited = posSurf;
@@ -389,9 +387,8 @@ mu::Ptr<ASTOp> SecondPassGenerator::GenerateModifierCondition( size_t modIndex )
     bool done = false;
     for (const auto& t: mod.positiveTags)
     {
-        auto it = std::find_if( m_pFirstPass->m_tags.begin(), m_pFirstPass->m_tags.end(),
-                                [&](const FirstPassGenerator::TAG& e) { return e.tag==t; } );
-        if (it==m_pFirstPass->m_tags.end())
+        auto it = m_pFirstPass->m_tags.FindByPredicate([&](const FirstPassGenerator::TAG& e) { return e.tag == t; } );
+        if (!it)
         {
             // This could happen if a tag is in a variation but noone defines it.
 
@@ -403,7 +400,7 @@ mu::Ptr<ASTOp> SecondPassGenerator::GenerateModifierCondition( size_t modIndex )
             continue;
         }
 
-        size_t tagIndex = it - m_pFirstPass->m_tags.begin();
+        size_t tagIndex = it - &m_pFirstPass->m_tags[0];
 
         set<size_t> empty;
         Ptr<ASTOp> tagCondition = GenerateTagCondition( tagIndex, empty, empty, empty, empty );
@@ -454,15 +451,14 @@ mu::Ptr<ASTOp> SecondPassGenerator::GenerateModifierCondition( size_t modIndex )
     {
         for (const auto& t: mod.negativeTags)
         {
-            auto it = std::find_if( m_pFirstPass->m_tags.begin(), m_pFirstPass->m_tags.end(),
-                                    [&](const FirstPassGenerator::TAG& e) { return e.tag==t; } );
-            if (it==m_pFirstPass->m_tags.end())
+            auto it = m_pFirstPass->m_tags.FindByPredicate( [&](const FirstPassGenerator::TAG& e) { return e.tag == t; } );
+            if (!it)
             {
                 // This could happen if a tag is in a variation but noone defines it.
                 continue;
             }
 
-            size_t tagIndex = it - m_pFirstPass->m_tags.begin();
+            size_t tagIndex = it - &m_pFirstPass->m_tags[0];
 
             set<size_t> empty;
             Ptr<ASTOp> tagCondition = GenerateTagCondition( tagIndex, empty, empty, empty, empty );
@@ -529,10 +525,10 @@ bool SecondPassGenerator::Generate(
 
     // Find the list of surfaces every tag depends on
 	m_surfacesPerTag.clear();
-    m_surfacesPerTag.resize(m_pFirstPass->m_tags.size());
+    m_surfacesPerTag.resize(m_pFirstPass->m_tags.Num());
 	m_tagsPerTag.clear();
-    m_tagsPerTag.resize(m_pFirstPass->m_tags.size());
-    for( size_t t=0; t<m_pFirstPass->m_tags.size(); ++t )
+    m_tagsPerTag.resize(m_pFirstPass->m_tags.Num());
+    for( size_t t=0; t<m_pFirstPass->m_tags.Num(); ++t )
     {
         set<size_t> pendingSurfs;
         for( size_t s: m_pFirstPass->m_tags[t].surfaces )
@@ -559,15 +555,14 @@ bool SecondPassGenerator::Generate(
             auto& csurf = m_pFirstPass->surfaces[cs];
             for ( auto sct: csurf.positiveTags )
             {
-                auto it = std::find_if( m_pFirstPass->m_tags.begin(), m_pFirstPass->m_tags.end(),
-                                        [&](const FirstPassGenerator::TAG& e) { return e.tag==sct; } );
-                if (it==m_pFirstPass->m_tags.end())
+                auto it = m_pFirstPass->m_tags.FindByPredicate( [&](const FirstPassGenerator::TAG& e) { return e.tag == sct; } );
+                if (!it)
                 {
                     // This could happen if a tag is in a variation but noone defines it.
                     continue;
                 }
 
-                size_t ct = it - m_pFirstPass->m_tags.begin();
+                size_t ct = it - &m_pFirstPass->m_tags[0];
 
 				m_tagsPerTag[t].insert(ct);
 
@@ -581,15 +576,14 @@ bool SecondPassGenerator::Generate(
             }
             for ( auto sct: csurf.negativeTags )
             {
-                auto it = std::find_if( m_pFirstPass->m_tags.begin(), m_pFirstPass->m_tags.end(),
-                                        [&](const FirstPassGenerator::TAG& e) { return e.tag==sct; } );
-                if (it==m_pFirstPass->m_tags.end())
+                auto it = m_pFirstPass->m_tags.FindByPredicate([&](const FirstPassGenerator::TAG& e) { return e.tag == sct; } );
+                if (!it)
                 {
                     // This could happen if a tag is in a variation but noone defines it.
                     continue;
                 }
 
-                size_t ct = it - m_pFirstPass->m_tags.begin();
+                size_t ct = it - &m_pFirstPass->m_tags[0];
 
 				m_tagsPerTag[t].insert(ct);
 
@@ -607,20 +601,20 @@ bool SecondPassGenerator::Generate(
     // Create the conditions for every surface, modifier and individual tag.
     m_tagConditionGenerationCache.clear();
 
-    for ( size_t s=0; s<m_pFirstPass->surfaces.size(); ++s )
+    for ( int32 s=0; s<m_pFirstPass->surfaces.Num(); ++s )
     {
         set<size_t> empty;
         Ptr<ASTOp> c = GenerateSurfaceCondition( s, empty, empty, empty, empty );
         m_pFirstPass->surfaces[s].surfaceCondition = c;
     }
 
-    for ( size_t m=0; m<m_pFirstPass->modifiers.size(); ++m )
+    for ( int32 m=0; m<m_pFirstPass->modifiers.Num(); ++m )
     {
         Ptr<ASTOp> c = GenerateModifierCondition( m );
         m_pFirstPass->modifiers[m].surfaceCondition = c;
     }
 
-    for ( size_t s = 0; s < m_pFirstPass->m_tags.size(); ++s )
+    for ( int32  s = 0; s < m_pFirstPass->m_tags.Num(); ++s )
     {
         set<size_t> empty;
         Ptr<ASTOp> c = GenerateTagCondition( s, empty, empty, empty, empty );
diff --git a/Engine/Plugins/Experimental/Mutable/Source/MutableTools/Private/MuT/CodeGenerator_SecondPass.h b/Engine/Plugins/Experimental/Mutable/Source/MutableTools/Private/MuT/CodeGenerator_SecondPass.h
--- a/Engine/Plugins/Experimental/Mutable/Source/MutableTools/Private/MuT/CodeGenerator_SecondPass.h
+++ b/Engine/Plugins/Experimental/Mutable/Source/MutableTools/Private/MuT/CodeGenerator_SecondPass.h
@@ -14,7 +14,6 @@
 #include "map"
 #include "set"
 
-#include <functional>
 #include <memory>
 #include <utility>
 
diff --git a/Engine/Plugins/Experimental/Mutable/Source/MutableTools/Private/MuT/CodeGenerator_String.cpp b/Engine/Plugins/Experimental/Mutable/Source/MutableTools/Private/MuT/CodeGenerator_String.cpp
--- a/Engine/Plugins/Experimental/Mutable/Source/MutableTools/Private/MuT/CodeGenerator_String.cpp
+++ b/Engine/Plugins/Experimental/Mutable/Source/MutableTools/Private/MuT/CodeGenerator_String.cpp
@@ -94,7 +94,7 @@ void CodeGenerator::GenerateString_Parameter( STRING_GENERATION_RESULT& result,
 		param.m_type = PARAMETER_TYPE::T_FLOAT;
 		param.m_detailedType = node.m_detailedType;
 		// param.m_defaultValue.m_text = node.m_defaultValue;
-		int len = std::min(MUTABLE_MAX_STRING_PARAM_LENGTH, int(node.m_defaultValue.size()));
+		int len = FMath::Min(MUTABLE_MAX_STRING_PARAM_LENGTH, int(node.m_defaultValue.size()));
 		FMemory::Memcpy(param.m_defaultValue.m_text, node.m_defaultValue.c_str(), len);
 		param.m_defaultValue.m_text[len] = 0;
 
@@ -104,12 +104,11 @@ void CodeGenerator::GenerateString_Parameter( STRING_GENERATION_RESULT& result,
 		op->parameter = param;
 
 		// Generate the code for the ranges
-		for (std::size_t a = 0; a < node.m_ranges.size(); ++a)
+		for (int32 a = 0; a < node.m_ranges.Num(); ++a)
 		{
 			RANGE_GENERATION_RESULT rangeResult;
 			GenerateRange(rangeResult, node.m_ranges[a]);
-			op->ranges.emplace_back(
-				op.get(), rangeResult.sizeOp, rangeResult.rangeName, rangeResult.rangeUID);
+			op->ranges.Emplace(op.get(), rangeResult.sizeOp, rangeResult.rangeName, rangeResult.rangeUID);
 		}
 
 		m_nodeVariables[node.m_pNode] = op;
