[mutable] Optimize mesh clip versus switching of clipping mesh.

diff --git a/Engine/Plugins/Experimental/Mutable/Source/MutableTools/Private/MuT/CodeOptimiser.cpp b/Engine/Plugins/Experimental/Mutable/Source/MutableTools/Private/MuT/CodeOptimiser.cpp
--- a/Engine/Plugins/Experimental/Mutable/Source/MutableTools/Private/MuT/CodeOptimiser.cpp
+++ b/Engine/Plugins/Experimental/Mutable/Source/MutableTools/Private/MuT/CodeOptimiser.cpp
@@ -336,10 +336,10 @@ bool FindOpTypeVisitor::Apply( PROGRAM& program, OP::ADDRESS rootAt )
     bool found = false;
 
     // If the program added new operations, we haven't visited them.
-    m_visited.resize( program.m_opAddress.size(), 0 );
+    m_visited.resize( program.m_opAddress.Num(), 0 );
 
     m_pending.clear();
-    m_pending.reserve( program.m_opAddress.size()/4 );
+    m_pending.reserve( program.m_opAddress.Num()/4 );
     m_pending.push_back( std::make_pair(false,rootAt) );
 
     // Don't early out to be able to complete parent op cached flags
@@ -473,7 +473,7 @@ public:
 
         PROGRAM::STATE state;
         state.m_root = at;
-        model->GetPrivate()->m_program.m_states.push_back(state);
+        model->GetPrivate()->m_program.m_states.Add(state);
 
         ParametersPtr localParams = model->NewParameters();
         pSystem->GetPrivate()->BeginBuild( model );
@@ -669,8 +669,8 @@ bool mu::ConstantGeneratorAST( const CompilerOptions::Private* options, Ptr<ASTO
 
     std::shared_timed_mutex codeAccessMutex;
 
-    vector< Ptr<ASTOp> > roots;
-    roots.push_back(root);
+    TArray< Ptr<ASTOp> > roots;
+    roots.Add(root);
 
     // Generate constant operations
     ASTOp::Traverse_TopDown_Unique_Imprecise( roots, [&](Ptr<ASTOp>& n)
@@ -813,7 +813,6 @@ CodeOptimiser::CodeOptimiser( CompilerOptionsPtr options, vector<STATE_COMPILATI
     : m_states( states )
 {
     m_options = options;
-    m_pModelReport = new ModelReport;
 }
 
 
@@ -854,7 +853,7 @@ void CodeOptimiser::FullOptimiseAST( ASTOpList& roots,
         // TODO: Conservative logic optimiser
         if (!modified)
         {
-//                AXE_INT_VALUE("Mutable", Verbose, "program size", (int64_t)program.m_opAddress.size());
+//                AXE_INT_VALUE("Mutable", Verbose, "program size", (int64_t)program.m_opAddress.Num());
 //                UE_LOG(LogMutableCore, Verbose, " - logic optimiser");
 //                LocalLogicOptimiser log;
 //                modified |= log.Apply( program, (int)s );
@@ -1039,11 +1038,11 @@ private:
 struct ADDMESH_SKELETON
 {
 	mu::Ptr<ASTOp> m_pAddMeshOp;
-    vector<mu::Ptr<ASTOpConstantResource>> m_contributingMeshes;
+	TArray<mu::Ptr<ASTOpConstantResource>> m_contributingMeshes;
 	mu::Ptr<Skeleton> m_pFinalSkeleton;
 
     ADDMESH_SKELETON( const mu::Ptr<ASTOp>& pAddMeshOp,
-                      vector<mu::Ptr<ASTOpConstantResource>>& contributingMeshes,
+                      TArray<mu::Ptr<ASTOpConstantResource>>& contributingMeshes,
                       const mu::Ptr<Skeleton>& pFinalSkeleton )
     {
         m_pAddMeshOp = pAddMeshOp;
@@ -1054,13 +1053,13 @@ struct ADDMESH_SKELETON
 
 
 //---------------------------------------------------------------------------------------------
-void SkeletonCleanerAST( vector<mu::Ptr<ASTOp>>& roots, const MODEL_OPTIMIZATION_OPTIONS& options )
+void SkeletonCleanerAST( TArray<mu::Ptr<ASTOp>>& roots, const MODEL_OPTIMIZATION_OPTIONS& options )
 {
     // This collects all the meshes that require a skeleton because they are used in operations
     // that require it.
     CollectAllMeshesForSkeletonVisitorAST requireSkeletonCollector( roots );
 
-    vector<ADDMESH_SKELETON> replacementsFound;
+    TArray<ADDMESH_SKELETON> replacementsFound;
 
     ASTOp::Traverse_TopDown_Unique_Imprecise( roots, [&](mu::Ptr<ASTOp>& at )
     {
@@ -1075,9 +1074,9 @@ void SkeletonCleanerAST( vector<mu::Ptr<ASTOp>>& roots, const MODEL_OPTIMIZATION
             if (meshRoot)
             {
                 // Gather constant meshes contributing to the final mesh
-                vector<mu::Ptr<ASTOpConstantResource>> subtreeMeshes;
-                vector<mu::Ptr<ASTOp>> tempRoots;
-                tempRoots.push_back(meshRoot);
+                TArray<mu::Ptr<ASTOpConstantResource>> subtreeMeshes;
+				TArray<mu::Ptr<ASTOp>> tempRoots;
+                tempRoots.Add(meshRoot);
                 ASTOp::Traverse_TopDown_Unique_Imprecise( tempRoots, [&](mu::Ptr<ASTOp>& lat )
                 {
                     // \todo: refine to avoid instruction branches with irrelevant skeletons.
@@ -1086,11 +1085,11 @@ void SkeletonCleanerAST( vector<mu::Ptr<ASTOp>>& roots, const MODEL_OPTIMIZATION
 						mu::Ptr<ASTOpConstantResource> typedOp = dynamic_cast<ASTOpConstantResource*>(lat.get());
                         check(typedOp);
 
-                        if ( std::find(subtreeMeshes.begin(), subtreeMeshes.end(), typedOp)
+                        if ( subtreeMeshes.Find(typedOp)
                              ==
-                             subtreeMeshes.end() )
+                             INDEX_NONE )
                         {
-                            subtreeMeshes.push_back(typedOp);
+                            subtreeMeshes.Add(typedOp);
                         }
                     }
                     return true;
@@ -1108,7 +1107,7 @@ void SkeletonCleanerAST( vector<mu::Ptr<ASTOp>>& roots, const MODEL_OPTIMIZATION
                     }
                 }
 
-                replacementsFound.emplace_back( at, subtreeMeshes, pFinalSkeleton );
+                replacementsFound.Emplace( at, subtreeMeshes, pFinalSkeleton );
             }
         }
 
@@ -1130,11 +1129,7 @@ void SkeletonCleanerAST( vector<mu::Ptr<ASTOp>>& roots, const MODEL_OPTIMIZATION
 
             for( auto& rep: replacementsFound )
             {
-                if ( std::find( rep.m_contributingMeshes.begin(),
-                                rep.m_contributingMeshes.end(),
-                                at)
-                     !=
-                     rep.m_contributingMeshes.end() )
+                if ( rep.m_contributingMeshes.Contains(at) )
                 {
 					mu::Ptr<const Mesh> pMesh =  static_cast<const Mesh*>(typedOp->GetValue().get());
                     pMesh->CheckIntegrity();
@@ -1191,10 +1186,10 @@ void CodeOptimiser::OptimiseAST( TaskManager* pTaskManager )
     MUTABLE_CPUPROFILER_SCOPE(OptimiseAST);
 
     // Gather all the roots
-    vector<Ptr<ASTOp>> roots;
+    TArray<Ptr<ASTOp>> roots;
     for(const auto& s:m_states)
     {
-        roots.push_back(s.root);
+        roots.Add(s.root);
     }
 
 	UE_LOG(LogMutableCore, Verbose, TEXT("(int) %s : %ld"), TEXT("ast size"), int64(ASTOp::CountNodes(roots)));
@@ -1266,7 +1261,7 @@ void CodeOptimiser::OptimiseAST( TaskManager* pTaskManager )
 //                // TODO: We have disable this since the new operation "layout remove block" is used.
 //                // Proper support needs to be implemented for this case.
 //    //                MUTABLE_CPUPROFILER_SCOPE(ImageComposeConditionGenerator);
-//    //                AXE_INT_VALUE("Mutable", Verbose, "program size", (int64_t)program.m_opAddress.size());
+//    //                AXE_INT_VALUE("Mutable", Verbose, "program size", (int64_t)program.m_opAddress.Num());
 //    //                UE_LOG(LogMutableCore, Verbose, " - image compose conditions generator");
 //    //                ImageComposeConditionGenerator composeGen( program );
 //            }
@@ -1298,8 +1293,7 @@ void CodeOptimiser::OptimiseAST( TaskManager* pTaskManager )
         // Make sure we didn't lose track of pointers
         for ( size_t s=0;  s<m_states.size(); ++s )
         {
-            check( std::find(roots.begin(), roots.end(), m_states[s].root )
-                            != roots.end() );
+            check( roots.Contains( m_states[s].root ) );
         }
 
         ASTOp::LogHistogram(roots);
@@ -1350,8 +1344,7 @@ void CodeOptimiser::OptimiseAST( TaskManager* pTaskManager )
         // Make sure we didn't lose track of pointers
         for ( size_t s=0;  s<m_states.size(); ++s )
         {
-            check( std::find(roots.begin(), roots.end(), m_states[s].root )
-                            != roots.end() );
+            check( roots.Contains( m_states[s].root ) );
         }
     }
 
diff --git a/Engine/Plugins/Experimental/Mutable/Source/MutableTools/Private/MuT/CodeOptimiser.h b/Engine/Plugins/Experimental/Mutable/Source/MutableTools/Private/MuT/CodeOptimiser.h
--- a/Engine/Plugins/Experimental/Mutable/Source/MutableTools/Private/MuT/CodeOptimiser.h
+++ b/Engine/Plugins/Experimental/Mutable/Source/MutableTools/Private/MuT/CodeOptimiser.h
@@ -26,9 +26,6 @@ namespace mu
         //! Optimise the virtual machine code, using several transforms.
         void OptimiseAST( TaskManager* );
 
-        //! Report about states external gpu usage
-        ModelReportPtr m_pModelReport;
-
     private:
 
         CompilerOptionsPtr m_options;
diff --git a/Engine/Plugins/Experimental/Mutable/Source/MutableTools/Private/MuT/CodeOptimiser_Logic.cpp b/Engine/Plugins/Experimental/Mutable/Source/MutableTools/Private/MuT/CodeOptimiser_Logic.cpp
--- a/Engine/Plugins/Experimental/Mutable/Source/MutableTools/Private/MuT/CodeOptimiser_Logic.cpp
+++ b/Engine/Plugins/Experimental/Mutable/Source/MutableTools/Private/MuT/CodeOptimiser_Logic.cpp
@@ -12,7 +12,6 @@
 #include "MuT/ASTOpSwitch.h"
 #include "MuT/CodeOptimiser.h"
 
-#include <functional>
 #include <memory>
 #include <utility>
 
@@ -113,9 +112,9 @@ namespace mu
 
 //        if (cache)
 //        {
-//            if (cache->m_values.size()<program.m_opAddress.size())
+//            if (cache->m_values.size()<program.m_opAddress.Num())
 //            {
-//                cache->m_values.resize(program.m_opAddress.size());
+//                cache->m_values.resize(program.m_opAddress.Num());
 //            }
 //            cache->m_values[at]=pRes;
 //        }
@@ -403,7 +402,7 @@ namespace mu
 //    //---------------------------------------------------------------------------------------------
 //    bool LogicOptimiser::Apply( PROGRAM& program, int state )
 //    {
-//        m_initialCodeSize = program.m_opAddress.size();
+//        m_initialCodeSize = program.m_opAddress.Num();
 //        m_modified = false;
 //        Traverse( program.m_states[state].m_root, program );
 
@@ -416,7 +415,7 @@ namespace mu
 //    {
 //        // Sanity stop
 //        const int sanityFactor = 4;
-//        if ( program.m_opAddress.size()>m_initialCodeSize*sanityFactor)
+//        if ( program.m_opAddress.Num()>m_initialCodeSize*sanityFactor)
 //        {
 //            return at;
 //        }
@@ -1042,7 +1041,7 @@ namespace mu
                                 bool conditionaAreExclusive = false;
 
                                 ASTOpList facts;
-                                facts.push_back(topConditional->condition.child());
+                                facts.Add(topConditional->condition.child());
 
                                 // Check if the child condition has a value with the current facts
                                 Ptr<ASTOp> pChildCond = bottomConditional->condition.child();
@@ -1101,7 +1100,7 @@ namespace mu
                                 bool conditionaAreExclusive = false;
 
                                 ASTOpList facts;
-                                facts.push_back(topConditional->condition.child());
+                                facts.Add(topConditional->condition.child());
 
                                 // Check if the child condition has a value with the current facts
 								Ptr<ASTOp> pChildCond = bottomConditional->condition.child();
@@ -1241,7 +1240,7 @@ namespace mu
                                     check(compare);
                                     if ( compareValue == switchOp->variable.child() )
                                     {
-                                        switchOp->cases.emplace_back( compare->op.args.BoolEqualScalarConst.constant,
+                                        switchOp->cases.Emplace( compare->op.args.BoolEqualScalarConst.constant,
                                                                       switchOp,
                                                                       conditional->yes.child() );
 
@@ -1259,7 +1258,7 @@ namespace mu
                         }
 
                         const int MIN_CONDITIONS_TO_CREATE_SWITCH = 3;
-                        if (switchOp->cases.size()>=MIN_CONDITIONS_TO_CREATE_SWITCH)
+                        if (switchOp->cases.Num()>=MIN_CONDITIONS_TO_CREATE_SWITCH)
                         {
                             ASTOp::Replace(n,switchOp);
                             n = switchOp;
diff --git a/Engine/Plugins/Experimental/Mutable/Source/MutableTools/Private/MuT/CodeOptimiser_Parameters.cpp b/Engine/Plugins/Experimental/Mutable/Source/MutableTools/Private/MuT/CodeOptimiser_Parameters.cpp
--- a/Engine/Plugins/Experimental/Mutable/Source/MutableTools/Private/MuT/CodeOptimiser_Parameters.cpp
+++ b/Engine/Plugins/Experimental/Mutable/Source/MutableTools/Private/MuT/CodeOptimiser_Parameters.cpp
@@ -26,10 +26,10 @@
 #include "MuT/Compiler.h"
 #include "MuT/CompilerPrivate.h"
 #include "MuT/DataPacker.h"
+#include "MuT/Platform.h"
 #include "Trace/Detail/Channel.h"
 
 #include <array>
-#include <functional>
 #include <memory>
 #include <utility>
 
@@ -50,7 +50,7 @@ class TaskManager;
 
     bool RuntimeParameterVisitorAST::HasAny( const Ptr<ASTOp>& root )
     {
-        if (!m_pState->nodeState.m_runtimeParams.size())
+        if (!m_pState->nodeState.m_runtimeParams.Num())
         {
             return false;
         }
@@ -142,9 +142,9 @@ class TaskManager;
                     {
                         auto typed = dynamic_cast<const ASTOpParameter*>(at.get());
                         const auto& params = m_pState->nodeState.m_runtimeParams;
-                        if ( std::find(params.cbegin(), params.cend(), typed->parameter.m_name)
+                        if ( params.Find( typed->parameter.m_name)
                              !=
-                             params.cend() )
+                             INDEX_NONE )
                         {
                             found = true;
                             m_visited[at] = OP_STATE::VISITED_HASRUNTIME;
@@ -633,7 +633,7 @@ class TaskManager;
 
 //                bool branchHasAny = false;
 //                OP_TYPE branchType = (OP_TYPE)program.m_code[program.m_code[at].args.Switch.values[0]].type;
-//                for ( map<uint16_t,OP::ADDRESS>::const_iterator it=chain.cases.begin();
+//                for ( map<uint16,OP::ADDRESS>::const_iterator it=chain.cases.begin();
 //                      it != chain.cases.end();
 //                      ++it )
 //                {
@@ -667,7 +667,7 @@ class TaskManager;
 //                    {
 //                        // Move the switch down the base
 //                        OP::ADDRESS baseAt = 0;
-//                        for ( map<uint16_t,OP::ADDRESS>::const_iterator it=chain.cases.begin();
+//                        for ( map<uint16,OP::ADDRESS>::const_iterator it=chain.cases.begin();
 //                              it != chain.cases.end();
 //                              )
 //                        {
@@ -691,7 +691,7 @@ class TaskManager;
 
 //                        // Move the switch down the mask
 //                        OP::ADDRESS maskAt = 0;
-//                        for ( map<uint16_t,OP::ADDRESS>::const_iterator it=chain.cases.begin();
+//                        for ( map<uint16,OP::ADDRESS>::const_iterator it=chain.cases.begin();
 //                              it != chain.cases.end();
 //                              )
 //                        {
@@ -715,7 +715,7 @@ class TaskManager;
 
 //                        // Move the switch down the blended
 //                        OP::ADDRESS blendedAt = 0;
-//                        for ( map<uint16_t,OP::ADDRESS>::const_iterator it=chain.cases.begin();
+//                        for ( map<uint16,OP::ADDRESS>::const_iterator it=chain.cases.begin();
 //                              it != chain.cases.end();
 //                              )
 //                        {
@@ -753,7 +753,7 @@ class TaskManager;
 //                    {
 //                        // Move the switch down the base
 //                        OP::ADDRESS baseAt = 0;
-//                        for ( map<uint16_t,OP::ADDRESS>::const_iterator it=chain.cases.begin();
+//                        for ( map<uint16,OP::ADDRESS>::const_iterator it=chain.cases.begin();
 //                              it != chain.cases.end();
 //                              )
 //                        {
@@ -777,7 +777,7 @@ class TaskManager;
 
 //                        // Move the switch down the mask
 //                        OP::ADDRESS maskAt = 0;
-//                        for ( map<uint16_t,OP::ADDRESS>::const_iterator it=chain.cases.begin();
+//                        for ( map<uint16,OP::ADDRESS>::const_iterator it=chain.cases.begin();
 //                              it != chain.cases.end();
 //                              )
 //                        {
@@ -801,7 +801,7 @@ class TaskManager;
 
 //                        // Move the switch down the colour
 //                        OP::ADDRESS colourAt = 0;
-//                        for ( map<uint16_t,OP::ADDRESS>::const_iterator it=chain.cases.begin();
+//                        for ( map<uint16,OP::ADDRESS>::const_iterator it=chain.cases.begin();
 //                              it != chain.cases.end();
 //                              )
 //                        {
@@ -1521,19 +1521,19 @@ class TaskManager;
             // Remove unsupported formats
             if (GetOpDataType( at->GetOpType() )==DT_IMAGE)
             {
-                auto it = m_supportedFormats.find(at);
-                if (it==m_supportedFormats.end())
+                auto it = m_supportedFormats.Find(at);
+                if (!it)
                 {
                     // Default to all supported
-                    m_supportedFormats[at] = m_allSupported;
-                    it = m_supportedFormats.find(at);
+                    m_supportedFormats.Add(at, m_allSupported);
+                    it = m_supportedFormats.Find(at);
                 }
 
                 for ( unsigned f=0; f< (unsigned)EImageFormat::IF_COUNT; ++f )
                 {
                     if ( !currentFormats[f] )
                     {
-                        it->second[f] = 0;
+                        (*it)[f] = 0;
                     }
                 }
             }
@@ -1606,13 +1606,13 @@ class TaskManager;
 
         bool IsSupportedFormat( Ptr<ASTOp> at, EImageFormat format ) const
         {
-            auto it = m_supportedFormats.find(at);
-            if (it==m_supportedFormats.end())
+            auto it = m_supportedFormats.Find(at);
+            if (!it)
             {
                 return false;
             }
 
-            return it->second[(size_t)format]!=0;
+            return (*it)[(size_t)format]!=0;
         }
 
 
@@ -1620,7 +1620,7 @@ class TaskManager;
 
         //! Formats known to be supported for every instruction.
         //! IF_COUNT*code.size() entries
-        std::unordered_map< Ptr<ASTOp>, std::array<uint8_t, (size_t)EImageFormat::IF_COUNT> > m_supportedFormats;
+        TMap< Ptr<ASTOp>, std::array<uint8_t, (size_t)EImageFormat::IF_COUNT> > m_supportedFormats;
 
         //! Constant convenience initial value
         std::array<uint8_t, (size_t)EImageFormat::IF_COUNT> m_initialState;
@@ -1650,7 +1650,7 @@ class TaskManager;
 
             // The state is the onlyLayoutRelevant flag
             ASTOp::Traverse_TopDown_Unique_Imprecise_WithState<bool>( root, false,
-                [&]( Ptr<ASTOp>& at, bool& state, vector<pair<Ptr<ASTOp>,bool>>& pending )
+                [&]( Ptr<ASTOp>& at, bool& state, TArray<TPair<Ptr<ASTOp>,bool>>& pending )
             {
                 (void)state;
 
@@ -1679,13 +1679,13 @@ class TaskManager;
 
                     if ( auto& source = pTyped->children[pTyped->op.args.LayoutRemoveBlocks.source] )
                     {
-                        pending.push_back( std::make_pair<>(source.m_child, state) );
+						pending.Add({ source.m_child, state });
                     }
 
                     // For that mesh we only want to know about the layouts
                     if ( auto& mesh = pTyped->children[pTyped->op.args.LayoutRemoveBlocks.mesh] )
                     {
-                        pending.push_back( std::make_pair<>(mesh.m_child, true) );
+						pending.Add({ mesh.m_child, true });
                     }
 
                     return false;
@@ -1698,7 +1698,7 @@ class TaskManager;
 
                     if ( auto& base = pTyped->children[pTyped->op.args.MeshMorph2.base] )
                     {
-                        pending.push_back( std::make_pair<>(base.m_child, state) );
+						pending.Add({ base.m_child, state });
                     }
 
                     // Mesh morphs don't modify the layouts, so we can ignore the factor and morphs
@@ -1706,14 +1706,14 @@ class TaskManager;
                     {
                         if ( auto& factor = pTyped->children[pTyped->op.args.MeshMorph2.factor] )
                         {
-                            pending.push_back( std::make_pair<>(factor.m_child, state) );
+							pending.Add({ factor.m_child, state });
                         }
 
                         for (int t=0;t<MUTABLE_OP_MAX_MORPH2_TARGETS;++t)
                         {
                             if ( auto& target = pTyped->children[pTyped->op.args.MeshMorph2.targets[t]] )
                             {
-                                pending.push_back( std::make_pair<>(target.m_child, state) );
+								pending.Add({ target.m_child, state });
                             }
                         }
                     }
@@ -1741,7 +1741,7 @@ class TaskManager;
     //!   .first IsResourceRoot
     //!   .second ParentIsRuntime
     //---------------------------------------------------------------------------------------------
-    class StateCacheDetectorAST : public Visitor_TopDown_Unique_Const< std::pair<bool,bool> >
+    class StateCacheDetectorAST : public Visitor_TopDown_Unique_Const< TPair<bool,bool> >
     {
     public:
 
@@ -1749,33 +1749,36 @@ class TaskManager;
             : m_hasRuntimeParamVisitor( pState )
         {
             ASTOpList roots;
-            roots.push_back(pState->root);
-            Traverse( roots, std::make_pair<>(false,false) );
+            roots.Add(pState->root);
+			Traverse(roots, { false,false });
 
-            pState->m_updateCache.clear();
-            pState->m_dynamicResources.clear();
+            pState->m_updateCache.Empty();
+            pState->m_dynamicResources.Empty();
 
-            for( auto i=m_cache.cbegin();
-                 i != m_cache.cend();
-                 ++i )
+            for( const auto& i : m_cache )
             {
-                if ( i->second )
+                if ( i.Value )
                 {
-                    pState->m_updateCache.push_back( i->first );
+                    pState->m_updateCache.Add( i.Key );
                 }
             }
 
-            for( auto i=m_dynamicResourceRoot.cbegin();
-                 i != m_dynamicResourceRoot.cend();
-                 ++i )
+            for(const auto& i : m_dynamicResourceRoot )
             {
-                if ( i->second )
+                if ( i.Value )
                 {
                     // Generate the list of relevant parameters
                     SubtreeRelevantParametersVisitorAST subtreeParams;
-                    subtreeParams.Run( i->first );
+                    subtreeParams.Run( i.Key );
 
-                    pState->m_dynamicResources.emplace_back( i->first, std::move(subtreeParams.m_params) );
+					// Temp copy
+					TArray<mu::string> ParamCopy;
+					for ( const auto& e: subtreeParams.m_params )
+					{
+						ParamCopy.Add(e);
+					}
+
+                    pState->m_dynamicResources.Emplace( i.Key, MoveTemp(ParamCopy) );
                 }
             }
         }
@@ -1785,8 +1788,10 @@ class TaskManager;
         {
             bool thisIsRuntime = m_hasRuntimeParamVisitor.HasAny( at );
 
-            bool resourceRoot = GetCurrentState().first;
-            bool parentIsRuntime = GetCurrentState().second;
+            bool resourceRoot = GetCurrentState().Key;
+            bool parentIsRuntime = GetCurrentState().Value;
+
+			m_cache.FindOrAdd(at, false);
 
             OP_TYPE type = at->GetOpType();
             if ( GetOpDesc( type ).cached )
@@ -1800,13 +1805,13 @@ class TaskManager;
                 {
                     // We want to cache this result to update the instances.
                     // Mark this as update cache
-                    m_cache[at] = true;
+                    m_cache.Add(at, true);
                 }
             }
 
             if ( !m_cache[at] && resourceRoot && thisIsRuntime )
             {
-                m_dynamicResourceRoot[at] = true;
+                m_dynamicResourceRoot.Add(at, true);
             }
 
             if ( !m_cache[at] && thisIsRuntime )
@@ -1821,15 +1826,15 @@ class TaskManager;
                 {
                     auto typedAt = dynamic_cast<const ASTOpInstanceAdd*>(at.get());
 
-                    std::pair<bool,bool> newState;
-                    newState.first = false; //resource root
-                    newState.second = thisIsRuntime;
+                    TPair<bool,bool> newState;
+                    newState.Key = false; //resource root
+                    newState.Value = thisIsRuntime;
                     RecurseWithState( typedAt->instance.child(), newState );
 
                     if ( typedAt->value )
                     {
-                        newState.first = true; //resource root
-                        newState.second = thisIsRuntime;
+                        newState.Key = true; //resource root
+                        newState.Value = thisIsRuntime;
                         RecurseWithState( typedAt->value.child(), newState );
                     }
                     return false;
@@ -1837,9 +1842,9 @@ class TaskManager;
 
                 default:
                 {
-                    std::pair<bool,bool> newState;
-                    newState.first = false; //resource root
-                    newState.second = thisIsRuntime;
+                    TPair<bool,bool> newState;
+                    newState.Key = false; //resource root
+                    newState.Value = thisIsRuntime;
                     SetCurrentState(newState);
                     return true;
                 }
@@ -1853,8 +1858,8 @@ class TaskManager;
     private:
 
         //!
-        std::map<Ptr<ASTOp>,bool> m_cache;
-        std::map<Ptr<ASTOp>,bool> m_dynamicResourceRoot;
+        TMap<Ptr<ASTOp>,bool> m_cache;
+		TMap<Ptr<ASTOp>,bool> m_dynamicResourceRoot;
 
         RuntimeParameterVisitorAST m_hasRuntimeParamVisitor;
     };
@@ -1883,8 +1888,7 @@ class TaskManager;
         {
             processChildren = true;
 
-            auto it = std::find(m_state.m_updateCache.cbegin(), m_state.m_updateCache.cend(), at);
-            bool isUpdateCache = it!=m_state.m_updateCache.cend();
+            bool isUpdateCache = m_state.m_updateCache.Contains(at);
 
             if ( isUpdateCache )
             {
@@ -1998,12 +2002,12 @@ class TaskManager;
         {
             auto* typedAt = dynamic_cast<ASTOpAddLOD*>(at.get());
 
-            if (typedAt->lods.size()>(size_t)m_lodCount)
+            if (typedAt->lods.Num()>(size_t)m_lodCount)
             {
                 Ptr<ASTOpAddLOD> newAt = mu::Clone<ASTOpAddLOD>(at);
-                while( newAt->lods.size()>(size_t)m_lodCount )
+                while( newAt->lods.Num()>(size_t)m_lodCount )
                 {
-                    newAt->lods.pop_back();
+                    newAt->lods.Pop();
                 }
                 at = newAt;
             }
@@ -2031,7 +2035,7 @@ class TaskManager;
             }
 
             // If a state has no runtime parameters, skip its optimisation alltogether
-            if (m_states[s].nodeState.m_runtimeParams.size())
+            if (m_states[s].nodeState.m_runtimeParams.Num())
             {
                 // Remove unnecessary image compression
                 if (m_states[s].optimisationFlags.m_avoidRuntimeCompression)
@@ -2059,8 +2063,8 @@ class TaskManager;
                                                  m_options->GetPrivate()->m_optimisationOptions );
                     modified = param.Apply();
 
-                    vector<Ptr<ASTOp>> roots;
-                    roots.push_back(m_states[s].root);
+                    TArray<Ptr<ASTOp>> roots;
+                    roots.Add(m_states[s].root);
 
                     UE_LOG(LogMutableCore, Verbose, TEXT(" - after parameter optimiser"));
 					UE_LOG(LogMutableCore, Verbose, TEXT("(int) %s : %ld"), TEXT("ast size"), int64(ASTOp::CountNodes(roots)));
@@ -2087,8 +2091,8 @@ class TaskManager;
                     //modified |= log.Apply( program, s );
                 }
 
-                vector<Ptr<ASTOp>> roots;
-                roots.push_back(m_states[s].root);
+                TArray<Ptr<ASTOp>> roots;
+                roots.Add(m_states[s].root);
 
                 UE_LOG(LogMutableCore, Verbose, TEXT(" - duplicated data remover"));
                 DuplicatedDataRemoverAST( roots );
@@ -2098,17 +2102,17 @@ class TaskManager;
                 DuplicatedCodeRemoverAST( roots );
 				//UE_LOG(LogMutableCore, Verbose, TEXT("(int) %s : %ld"), TEXT("ast size"), int64(ASTOp::CountNodes(roots)));
 
-                m_states[s].root = roots.front();
+                m_states[s].root = roots[0];
             }
         }
 
         // Mark the instructions that don't depend on runtime parameters to be cached. This is
         // necessary at this stage before GPU optimisation.
         {
-            vector<Ptr<ASTOp>> roots;
+            TArray<Ptr<ASTOp>> roots;
             for(const auto& s:m_states)
             {
-                roots.push_back(s.root);
+                roots.Add(s.root);
             }
 
             AccumulateAllImageFormatsOpAST opFormats;
@@ -2157,10 +2161,10 @@ class TaskManager;
             int numIterations = 0;
             while (modified && (!m_optimizeIterationsMax || m_optimizeIterationsLeft>0 || !numIterations ))
             {
-                vector<Ptr<ASTOp>> roots;
+                TArray<Ptr<ASTOp>> roots;
                 for(const auto& s:m_states)
                 {
-                    roots.push_back(s.root);
+                    roots.Add(s.root);
                 }
 
                 ++numIterations;
@@ -2187,10 +2191,10 @@ class TaskManager;
                 ConstantGeneratorAST( m_options->GetPrivate(), s.root, pTaskManager );
             }
 
-            vector<Ptr<ASTOp>> roots;
+            TArray<Ptr<ASTOp>> roots;
             for(const auto& s:m_states)
             {
-                roots.push_back(s.root);
+                roots.Add(s.root);
             }
 			UE_LOG(LogMutableCore, Verbose, TEXT("(int) %s : %ld"), TEXT("ast size"), int64(ASTOp::CountNodes(roots)));
 
@@ -2204,10 +2208,10 @@ class TaskManager;
 		}
 
         // Gather all the current roots
-        vector<Ptr<ASTOp>> roots;
+        TArray<Ptr<ASTOp>> roots;
         for(const auto& s:m_states)
         {
-            roots.push_back(s.root);
+            roots.Add(s.root);
         }
 
         // Optimise the data formats
diff --git a/Engine/Plugins/Experimental/Mutable/Source/MutableTools/Private/MuT/CodeOptimiser_Semantic.cpp b/Engine/Plugins/Experimental/Mutable/Source/MutableTools/Private/MuT/CodeOptimiser_Semantic.cpp
--- a/Engine/Plugins/Experimental/Mutable/Source/MutableTools/Private/MuT/CodeOptimiser_Semantic.cpp
+++ b/Engine/Plugins/Experimental/Mutable/Source/MutableTools/Private/MuT/CodeOptimiser_Semantic.cpp
@@ -23,7 +23,6 @@
 #include "MuT/CodeOptimiser.h"
 #include "MuT/Table.h"
 
-#include <functional>
 #include <memory>
 #include <utility>
 
@@ -411,7 +410,7 @@ namespace mu
                 FImageRect maskUsage;
                 maskUsage.min[0] = (sourceMaskUsage.min[0]/blockSize)*blockSize;
                 maskUsage.min[1] = (sourceMaskUsage.min[1]/blockSize)*blockSize;
-                vec2<uint16_t> minOffset = sourceMaskUsage.min - maskUsage.min;
+                vec2<uint16> minOffset = sourceMaskUsage.min - maskUsage.min;
                 maskUsage.size[0] = ((sourceMaskUsage.size[0]+minOffset[0]+blockSize-1)/blockSize)*blockSize;
                 maskUsage.size[1] = ((sourceMaskUsage.size[1]+minOffset[1]+blockSize-1)/blockSize)*blockSize;
 
@@ -547,7 +546,7 @@ namespace mu
                 }
 
                 // We need to copy the options because we change them
-                for ( size_t v=0; v<nop->cases.size(); ++v )
+                for ( int32 v=0; v<nop->cases.Num(); ++v )
                 {
                     if ( nop->cases[v].branch )
                     {
@@ -1030,8 +1029,8 @@ namespace mu
     Ptr<ASTOp> Sink_ImageCropAST::Apply( const ASTOp* root )
     {
         m_root = root;
-        m_oldToNew.clear();
-		m_newOps.clear();
+        m_oldToNew.Empty();
+		m_newOps.Empty();
 
         check(root->GetOpType()==OP_TYPE::IM_CROP);
 
@@ -1055,17 +1054,20 @@ namespace mu
         if (!at) return nullptr;
 
         // Newly created?
-        if (std::find(m_newOps.begin(), m_newOps.end(), at )!=m_newOps.end())
+        if (m_newOps.Find( at )!=INDEX_NONE)
         {
             return at;
         }
 
         // Already visited?
-        auto cacheIt = m_oldToNew[currentCropOp].find(at);
-        if (cacheIt!=m_oldToNew[currentCropOp].end())
-        {
-            return cacheIt->second;
-        }
+		{
+			auto& CurrentSet = m_oldToNew.FindOrAdd(currentCropOp);
+			auto cacheIt = CurrentSet.find(at);
+			if (cacheIt != CurrentSet.end())
+			{
+				return cacheIt->second;
+			}
+		}
 
         bool skipSinking=false;
         Ptr<ASTOp> newAt = at;
@@ -1358,7 +1360,7 @@ namespace mu
                     }
 
                     // We need to copy the options because we change them
-                    for ( size_t v=0; v<nop->cases.size(); ++v )
+                    for ( int32 v=0; v<nop->cases.Num(); ++v )
                     {
                         if ( nop->cases[v].branch )
                         {
@@ -1530,7 +1532,7 @@ namespace mu
                 }
 
                 // We need to copy the options because we change them
-                for ( size_t o=0; o<nop->cases.size(); ++o )
+                for ( size_t o=0; o<nop->cases.Num(); ++o )
                 {
                     if ( nop->cases[o].branch )
                     {
@@ -1588,7 +1590,7 @@ namespace mu
 					}
 
 					// We need to copy the options because we change them
-					for (size_t o = 0; o < nop->cases.size(); ++o)
+					for (size_t o = 0; o < nop->cases.Num(); ++o)
 					{
 						if (nop->cases[o].branch)
 						{
@@ -1682,7 +1684,7 @@ namespace mu
 //                    nop->def = defOp;
 //                }
 
-//                for ( size_t v=0; v<nop->cases.size(); ++v )
+//                for ( size_t v=0; v<nop->cases.Num(); ++v )
 //                {
 //                    if ( nop->cases[v].branch )
 //                    {
@@ -1804,7 +1806,7 @@ namespace mu
                 }
 
                 // We need to copy the options because we change them
-                for ( size_t v=0; v<nop->cases.size(); ++v )
+                for ( size_t v=0; v<nop->cases.Num(); ++v )
                 {
                     if ( nop->cases[v].branch )
                     {
@@ -1887,7 +1889,7 @@ namespace mu
                         }
 
                         // We need to copy the options because we change them
-                        for ( size_t v=0; v<nop->cases.size(); ++v )
+                        for ( size_t v=0; v<nop->cases.Num(); ++v )
                         {
                             if ( nop->cases[v].branch )
                             {
@@ -1916,7 +1918,7 @@ namespace mu
 					}
 
 					// We need to copy the options because we change them
-					for (size_t o = 0; o < nop->cases.size(); ++o)
+					for (size_t o = 0; o < nop->cases.Num(); ++o)
 					{
 						if (nop->cases[o].branch)
 						{
@@ -2128,8 +2130,8 @@ namespace mu
                 newOp->patch = Visit( newOp->patch.child(), currentSinkingOp);
 
                 // todo: review if this is always correct, or we need some "divisible" check
-                newOp->location[0] = uint16_t( newOp->location[0] *  scaleX );
-                newOp->location[1] = uint16_t( newOp->location[1] *  scaleY );
+                newOp->location[0] = uint16( newOp->location[0] *  scaleX );
+                newOp->location[1] = uint16( newOp->location[1] *  scaleY );
 
                 newAt = newOp;
 
@@ -2263,8 +2265,8 @@ namespace mu
 				//auto imageOp = nop->children[nop->op.args.ImageRasterMesh.image].child();
 				//nop->SetChild(nop->op.args.ImageRasterMesh.image, Visit(imageOp, currentSinkingOp));
 
-                nop->op.args.ImageRasterMesh.sizeX = uint16_t( nop->op.args.ImageRasterMesh.sizeX * scaleX + 0.5f );
-                nop->op.args.ImageRasterMesh.sizeY = uint16_t( nop->op.args.ImageRasterMesh.sizeY * scaleY + 0.5f );
+                nop->op.args.ImageRasterMesh.sizeX = uint16( nop->op.args.ImageRasterMesh.sizeX * scaleX + 0.5f );
+                nop->op.args.ImageRasterMesh.sizeY = uint16( nop->op.args.ImageRasterMesh.sizeY * scaleY + 0.5f );
                 newAt = nop;
                 break;
             }
@@ -2579,11 +2581,11 @@ namespace mu
                     auto newOp = mu::Clone<ASTOpFixed>(sourceAt);
 
                     newOp->op.args.ImageBlankLayout.blockSize[0] =
-                            uint16_t( newOp->op.args.ImageBlankLayout.blockSize[0]
+                            uint16( newOp->op.args.ImageBlankLayout.blockSize[0]
                                     * op.args.ImageResizeRel.factor[0]
                                     + 0.5f );
                     newOp->op.args.ImageBlankLayout.blockSize[1] =
-                            uint16_t( newOp->op.args.ImageBlankLayout.blockSize[1]
+                            uint16( newOp->op.args.ImageBlankLayout.blockSize[1]
                                     * op.args.ImageResizeRel.factor[1]
                                     + 0.5f );
                     at = newOp;
diff --git a/Engine/Plugins/Experimental/Mutable/Source/MutableTools/Private/MuT/Compiler.cpp b/Engine/Plugins/Experimental/Mutable/Source/MutableTools/Private/MuT/Compiler.cpp
--- a/Engine/Plugins/Experimental/Mutable/Source/MutableTools/Private/MuT/Compiler.cpp
+++ b/Engine/Plugins/Experimental/Mutable/Source/MutableTools/Private/MuT/Compiler.cpp
@@ -200,7 +200,7 @@ namespace mu
 
             gen.GenerateRoot( pNode, pTaskManager );
 
-            check( !gen.m_states.empty() );
+            check( !gen.m_states.IsEmpty() );
 
             for ( const auto& s: gen.m_states )
             {
@@ -253,10 +253,10 @@ namespace mu
         // Set the runtime parameter indices.
         for( auto& s: states )
         {
-            for ( size_t p=0; p<s.nodeState.m_runtimeParams.size(); ++p )
+            for ( int32 p=0; p<s.nodeState.m_runtimeParams.Num(); ++p )
             {
                 int paramIndex = -1;
-                for ( size_t i=0; paramIndex<0 && i<program.m_parameters.size(); ++i )
+                for ( size_t i=0; paramIndex<0 && i<program.m_parameters.Num(); ++i )
                 {
                     if ( program.m_parameters[i].m_name
                          ==
@@ -268,7 +268,7 @@ namespace mu
 
                 if (paramIndex>=0)
                 {
-                    s.state.m_runtimeParameters.push_back( paramIndex );
+                    s.state.m_runtimeParameters.Add( paramIndex );
                 }
                 else
                 {
@@ -288,21 +288,21 @@ namespace mu
             // Generate the mask of update cache ops
             for ( const auto& a: s.m_updateCache )
             {
-                s.state.m_updateCache.emplace_back(a->linkedAddress);
+                s.state.m_updateCache.Emplace(a->linkedAddress);
             }
 
             // Sort the update cache addresses for performance and determinism
-            std::sort( s.state.m_updateCache.begin(), s.state.m_updateCache.end() );
+            s.state.m_updateCache.Sort();
 
             // Generate the mask of dynamic resources
             for ( const auto& a: s.m_dynamicResources )
             {
                 uint64_t relevantMask = 0;
-                for ( const auto& b: a.second )
+                for ( const auto& b: a.Value )
                 {
                     // Find the index in the model parameter list
                     int paramIndex = -1;
-                    for ( size_t i=0; paramIndex<0 && i<program.m_parameters.size(); ++i )
+                    for ( int32 i=0; paramIndex<0 && i<program.m_parameters.Num(); ++i )
                     {
                         if ( program.m_parameters[i].m_name == b )
                         {
@@ -312,14 +312,11 @@ namespace mu
                     check(paramIndex>=0);
 
                     // Find the position in the state data vector.
-                    auto it = std::find( s.state.m_runtimeParameters.cbegin(),
-                                         s.state.m_runtimeParameters.cend(),
-                                         paramIndex );
+                    int32 it = s.state.m_runtimeParameters.Find( paramIndex );
 
-                    if ( it!=s.state.m_runtimeParameters.end() )
+                    if ( it!=INDEX_NONE )
                     {
-                        size_t pos = it - s.state.m_runtimeParameters.begin();
-                        relevantMask |= uint64_t(1) << pos;
+                        relevantMask |= uint64_t(1) << it;
                     }
                 }
 
@@ -329,27 +326,17 @@ namespace mu
                 //check(relevantMask!=0);
                 if (relevantMask!=0)
                 {
-                    s.state.m_dynamicResources.emplace_back( a.first->linkedAddress, relevantMask );
+                    s.state.m_dynamicResources.Emplace( a.Key->linkedAddress, relevantMask );
                 }
             }            
 
             // Sort for performance and determinism
-            std::sort( s.state.m_dynamicResources.begin(), s.state.m_dynamicResources.end() );
+            s.state.m_dynamicResources.Sort();
 
-            program.m_states.push_back(s.state);
+            program.m_states.Add(s.state);
         }
 
-        //
-        if ( m_pD->m_options->GetPrivate()->m_log )
-        {
-			UE_LOG(LogMutableCore, Log, TEXT("----------------------------------------------------------------------\n") );
-			UE_LOG(LogMutableCore, Log, TEXT("non-optimised size : %d\n"), program.m_opAddress.size() );
-			UE_LOG(LogMutableCore, Log, TEXT("----------------------------------------------------------------------\n") );
-            Log( pResult );
-			UE_LOG(LogMutableCore, Log, TEXT("----------------------------------------------------------------------\n") );
-        }
-
-		UE_LOG(LogMutableCore, Verbose, TEXT("(int) %s : %ld"), TEXT("program size"), int64(program.m_opAddress.size()));
+		UE_LOG(LogMutableCore, Verbose, TEXT("(int) %s : %ld"), TEXT("program size"), int64(program.m_opAddress.Num()));
 
         // Optimize the generated code
         {        
@@ -369,19 +356,10 @@ namespace mu
 		int32 MinimumBytesPerRom = 1024; // \TODO: compilation parameter
 		m_pD->GenerateRoms(pResult.get(),MinimumBytesPerRom);
 
-		UE_LOG(LogMutableCore, Verbose, TEXT("(int) %s : %ld"), TEXT("program size"), int64(program.m_opAddress.size()));
+		UE_LOG(LogMutableCore, Verbose, TEXT("(int) %s : %ld"), TEXT("program size"), int64(program.m_opAddress.Num()));
 
         delete pTaskManager;
 
-        // Finish logging
-        if ( m_pD->m_options->GetPrivate()->m_log )
-        {
-            Log( pResult );
-			UE_LOG(LogMutableCore, Log, TEXT("----------------------------------------------------------------------") );
-            m_pD->m_pErrorLog->Log();
-			UE_LOG(LogMutableCore, Log, TEXT("----------------------------------------------------------------------") );
-        }
-
         return pResult;
     }
 
@@ -393,13 +371,6 @@ namespace mu
     }
 
 
-    //---------------------------------------------------------------------------------------------
-    ModelReportPtrConst Compiler::GetModelReport() const
-    {
-        return m_pD->m_pModelReport;
-    }
-
-
 	//---------------------------------------------------------------------------------------------
 	namespace
 	{
diff --git a/Engine/Plugins/Experimental/Mutable/Source/MutableTools/Private/MuT/CompilerPrivate.h b/Engine/Plugins/Experimental/Mutable/Source/MutableTools/Private/MuT/CompilerPrivate.h
--- a/Engine/Plugins/Experimental/Mutable/Source/MutableTools/Private/MuT/CompilerPrivate.h
+++ b/Engine/Plugins/Experimental/Mutable/Source/MutableTools/Private/MuT/CompilerPrivate.h
@@ -6,7 +6,6 @@
 
 #include "MuT/AST.h"
 #include "MuT/ErrorLogPrivate.h"
-#include "MuT/ModelReport.h"
 
 #include "MuR/Operations.h"
 #include "MuR/MemoryPrivate.h"
@@ -95,7 +94,7 @@ namespace mu
         STATE_OPTIMIZATION_OPTIONS m_optimisation;
 
         //! List of names of the runtime parameters in this state
-        vector<string> m_runtimeParams;
+        TArray<string> m_runtimeParams;
 
         void Serialise( OutputArchive& arch ) const
         {
@@ -130,11 +129,11 @@ namespace mu
         STATE_OPTIMIZATION_OPTIONS optimisationFlags;
 
         //! List of instructions that need to be cached to efficiently update this state
-        vector<Ptr<ASTOp>> m_updateCache;
+        TArray<Ptr<ASTOp>> m_updateCache;
 
         //! List of root instructions for the dynamic resources that depend on the runtime
         //! parameters of this state.
-        vector< std::pair<Ptr<ASTOp>, std::unordered_set<string> > > m_dynamicResources;
+		TArray< TPair<Ptr<ASTOp>, TArray<string> > > m_dynamicResources;
     };
 
 
@@ -149,7 +148,6 @@ namespace mu
         }
 
         ErrorLogPtr m_pErrorLog;
-        ModelReportPtr m_pModelReport;
 
         //! Detailed options
         CompilerOptionsPtr m_options;
diff --git a/Engine/Plugins/Experimental/Mutable/Source/MutableTools/Private/MuT/DataPacker.cpp b/Engine/Plugins/Experimental/Mutable/Source/MutableTools/Private/MuT/DataPacker.cpp
--- a/Engine/Plugins/Experimental/Mutable/Source/MutableTools/Private/MuT/DataPacker.cpp
+++ b/Engine/Plugins/Experimental/Mutable/Source/MutableTools/Private/MuT/DataPacker.cpp
@@ -24,7 +24,6 @@
 #include "MuT/ASTOpMeshRemoveMask.h"
 
 #include <cstdint>
-#include <functional>
 #include <memory>
 #include <utility>
 
@@ -42,10 +41,10 @@ namespace mu
           m_constant( constant ),
           m_opType( optype )
     {
-        m_visited.resize( program.m_opAddress.size() );
-        if ( m_visited.size() )
+        m_visited.SetNum( program.m_opAddress.Num() );
+        if ( m_visited.Num() )
         {
-            memset( &m_visited[0], 0, m_visited.size() );
+            memset( &m_visited[0], 0, m_visited.Num() );
         }
     }
 
@@ -59,64 +58,63 @@ namespace mu
 
         bool found = false;
 
-        vector< std::pair<bool,int> > pending;
-        pending.reserve(m_program.m_opAddress.size()/4);
+        TArray< TPair<bool,int> > pending;
+        pending.Reserve(m_program.m_opAddress.Num()/4);
 
-        pending.push_back( std::make_pair(false,root) );
+		pending.Add({ false,root });
 
-        while( pending.size() )
+        while( pending.Num() )
         {
-            std::pair<bool,int> item = pending.back();
-            pending.pop_back();
+			TPair<bool,int> item = pending.Pop();
 
-            if (item.first)
+            if (item.Key)
             {
                 // Item indicating we finished with all the children of a parent.
 
                 // Propagate the "found" state upwards
-                m_visited[item.second]=1;
+                m_visited[item.Value]=1;
 
-                //OP cop = m_program.m_code[item.second];
-                ForEachReference( m_program, (OP::ADDRESS)item.second, [&](OP::ADDRESS ref)
+                //OP cop = m_program.m_code[item.Value];
+                ForEachReference( m_program, (OP::ADDRESS)item.Value, [&](OP::ADDRESS ref)
                 {
                     if (ref && m_visited[ref]==2)
                     {
-                        m_visited[item.second]=2;
+                        m_visited[item.Value]=2;
                     }
                 });
             }
             else
             {
-                if ( !m_visited[item.second] )
+                if ( !m_visited[item.Value] )
                 {
-                    OP_TYPE thisOpType = m_program.GetOpType(item.second);
+                    OP_TYPE thisOpType = m_program.GetOpType(item.Value);
                     if ( m_opType == thisOpType )
                     {
-                        auto args = m_program.GetOpArgs<OP::ResourceConstantArgs>(item.second);
+                        auto args = m_program.GetOpArgs<OP::ResourceConstantArgs>(item.Value);
                         if ( m_constant == args.value )
                         {
                             found = true;
-                            m_visited[item.second] = 2;
+                            m_visited[item.Value] = 2;
                         }
                     }
 
                     if (!found)
                     {
-                        pending.push_back( std::make_pair(true,item.second) );
+						pending.Add({ true,item.Value });
 
-                        //OP cop = m_program.m_code[item.second];
-                        ForEachReference( m_program, item.second, [&](OP::ADDRESS ref)
+                        //OP cop = m_program.m_code[item.Value];
+                        ForEachReference( m_program, item.Value, [&](OP::ADDRESS ref)
                         {
                             if (ref && !m_visited[ref])
                             {
-                                pending.push_back( std::make_pair(false,ref) );
+								pending.Add({ false,ref });
                             }
                         });
                     }
                 }
                 else
                 {
-                    found = (m_visited[item.second]==2);
+                    found = (m_visited[item.Value]==2);
                 }
             }
         }
@@ -130,7 +128,7 @@ namespace mu
     //! Get all the parameters that affect the constant.
     //! "Affect" means that the constant may be used or not depending on the parameter.
     //---------------------------------------------------------------------------------------------
-    class GatherParametersVisitor : public UniqueConstCodeVisitorIterative< vector<int> >
+    class GatherParametersVisitor : public UniqueConstCodeVisitorIterative< TArray<int> >
     {
     public:
 
@@ -139,22 +137,22 @@ namespace mu
         {
             MUTABLE_CPUPROFILER_SCOPE(GatherParametersVisitor);
 
-            vector<int> currentParams;
-            currentParams.resize( program.m_parameters.size(), 0 );
+			TArray<int> currentParams;
+            currentParams.SetNumZeroed( program.m_parameters.Num() );
             SetDefaultState(currentParams);
 
             // TODO: we can optimize by precalculating what ops have the required optype below
             m_opType = opType;
             m_constant = constant;
 
-            m_allParams.resize( program.m_parameters.size(), 0 );
+            m_allParams.SetNumZeroed( program.m_parameters.Num() );
             FullTraverse( program );
 
-            for (size_t i=0; i<m_allParams.size(); ++i )
+            for (size_t i=0; i<m_allParams.Num(); ++i )
             {
                 if ( m_allParams[i] )
                 {
-                    m_sortedParams.push_back( int(i) );
+                    m_sortedParams.Add( int(i) );
                 }
             }
         }
@@ -186,8 +184,8 @@ namespace mu
                     {
                         m_conditionVisitor.Run( args.condition, program );
 
-                        vector<int> currentParams = GetCurrentState();
-                        for (size_t p=0; p<m_conditionVisitor.m_params.size(); ++p )
+                        TArray<int> currentParams = GetCurrentState();
+                        for (size_t p=0; p<m_conditionVisitor.m_params.Num(); ++p )
                         {
                             currentParams[ m_conditionVisitor.m_params[p] ]++;
                         }
@@ -213,8 +211,8 @@ namespace mu
 
                     m_conditionVisitor.Run( VarAddress, program );
 
-                    vector<int> currentParams = GetCurrentState();
-                    for (size_t p=0; p<m_conditionVisitor.m_params.size(); ++p )
+					TArray<int> currentParams = GetCurrentState();
+                    for (size_t p=0; p<m_conditionVisitor.m_params.Num(); ++p )
                     {
                         currentParams[ m_conditionVisitor.m_params[p] ]++;
                     }
@@ -236,8 +234,8 @@ namespace mu
                         if ( args.value == m_constant )
                         {
                             // Accumulate the currently relevant parameters
-                            const vector<int>& currentParams = GetCurrentState();
-                            for (size_t i=0; i<currentParams.size(); ++i )
+                            const TArray<int>& currentParams = GetCurrentState();
+                            for (size_t i=0; i<currentParams.Num(); ++i )
                             {
                                 m_allParams[i] += currentParams[i];
                             }
@@ -249,8 +247,8 @@ namespace mu
                         if ( args.value == m_constant )
                         {
                             // Accumulate the currently relevant parameters
-                            const vector<int>& currentParams = GetCurrentState();
-                            for (size_t i=0; i<currentParams.size(); ++i )
+                            const TArray<int>& currentParams = GetCurrentState();
+                            for (size_t i=0; i<currentParams.Num(); ++i )
                             {
                                 m_allParams[i] += currentParams[i];
                             }
@@ -270,14 +268,14 @@ namespace mu
     public:
 
         //! Output result
-        vector<int> m_sortedParams;
+		TArray<int> m_sortedParams;
 
     private:
 
         OP::ADDRESS m_constant;
         OP_TYPE m_opType;
 
-        vector<int> m_allParams;
+		TArray<int> m_allParams;
 
         SubtreeParametersVisitor m_conditionVisitor;
         SubtreeSearchConstantVisitor m_constSearch;
@@ -288,7 +286,7 @@ namespace mu
     //---------------------------------------------------------------------------------------------
     //---------------------------------------------------------------------------------------------
     //---------------------------------------------------------------------------------------------
-    class AccumulateImageFormatsAST : public Visitor_TopDown_Unique_Const< vector<bool> >
+    class AccumulateImageFormatsAST : public Visitor_TopDown_Unique_Const< TArray<bool> >
     {
     public:
 
@@ -300,7 +298,8 @@ namespace mu
             //vector<bool> initial( IF_COUNT, true );
             //std::fill( m_supportedFormats.begin(), m_supportedFormats.end(), initial );
 
-            vector<bool> defaultState(size_t(EImageFormat::IF_COUNT), false);
+			TArray<bool> defaultState;
+			defaultState.SetNumZeroed(size_t(EImageFormat::IF_COUNT));
 
             Traverse( roots, defaultState );
         }
@@ -310,8 +309,11 @@ namespace mu
         {
             bool recurse = true;
 
-            const vector<bool>& currentFormats = GetCurrentState();
-            bool allFalse = currentFormats == vector<bool>(size_t(EImageFormat::IF_COUNT), false );
+            const TArray<bool>& currentFormats = GetCurrentState();
+
+			TArray<bool> defaultState;
+			defaultState.SetNumZeroed(size_t(EImageFormat::IF_COUNT));
+			bool allFalse = currentFormats == defaultState;
 
             // Can we use the cache?
             if (allFalse)
@@ -332,7 +334,8 @@ namespace mu
                 auto op = dynamic_cast<const ASTOpConstantResource*>(node.get());
                 if (!m_supportedFormats.count(op))
                 {
-                    vector<bool> initial(size_t(EImageFormat::IF_COUNT), true );
+					TArray<bool> initial;
+					initial.Init( true, size_t(EImageFormat::IF_COUNT) );
                     m_supportedFormats.insert( std::make_pair(op, std::move(initial)) );
                 }
 
@@ -357,7 +360,8 @@ namespace mu
                 recurse = false;
 				const ASTOpImageCompose* op = dynamic_cast<const ASTOpImageCompose*>(node.get());
 
-                vector<bool> newState = vector<bool>(size_t(EImageFormat::IF_COUNT), false);
+				TArray<bool> newState;
+				newState.Init(false, size_t(EImageFormat::IF_COUNT));
                 RecurseWithState( op->Layout.child(), newState );
                 RecurseWithState( op->Base.child(), newState );
                 RecurseWithState( op->BlockImage.child(), newState );
@@ -374,8 +378,9 @@ namespace mu
             {
                 recurse = false;
 				const ASTOpFixed* op = dynamic_cast<const ASTOpFixed*>(node.get());
-                vector<bool> newState = vector<bool>(size_t(EImageFormat::IF_COUNT), false);
-                RecurseWithState( op->children[op->op.args.ImageLayerColour.base].child(), newState );
+				TArray<bool> newState;
+				newState.Init(false, size_t(EImageFormat::IF_COUNT));
+				RecurseWithState( op->children[op->op.args.ImageLayerColour.base].child(), newState );
                 RecurseWithState( op->children[op->op.args.ImageLayerColour.colour].child(), newState );
 
                 if ( op->children[op->op.args.ImageLayerColour.mask] )
@@ -392,8 +397,9 @@ namespace mu
             {
                 recurse = false;
 				const ASTOpFixed* op = dynamic_cast<const ASTOpFixed*>(node.get());
-                vector<bool> newState = vector<bool>(size_t(EImageFormat::IF_COUNT), false);
-                RecurseWithState( op->children[op->op.args.ImageLayer.base].child(), newState );
+				TArray<bool> newState;
+				newState.Init(false, size_t(EImageFormat::IF_COUNT));
+				RecurseWithState( op->children[op->op.args.ImageLayer.base].child(), newState );
                 RecurseWithState( op->children[op->op.args.ImageLayer.blended].child(), newState );
 
                 if (op->children[op->op.args.ImageLayer.mask])
@@ -410,8 +416,9 @@ namespace mu
             {
                 recurse = false;
 				const ASTOpImageMultiLayer* op = dynamic_cast<const ASTOpImageMultiLayer*>(node.get());
-                vector<bool> newState = vector<bool>(size_t(EImageFormat::IF_COUNT), false);
-                RecurseWithState( op->base.child(), newState );
+				TArray<bool> newState;
+				newState.Init(false, size_t(EImageFormat::IF_COUNT));
+				RecurseWithState( op->base.child(), newState );
                 RecurseWithState( op->blend.child(), newState );
 
                 if (op->mask)
@@ -428,8 +435,9 @@ namespace mu
             {
                 recurse = false;
 				const ASTOpFixed* op = dynamic_cast<const ASTOpFixed*>(node.get());
-                vector<bool> newState = vector<bool>(size_t(EImageFormat::IF_COUNT), false);
-                RecurseWithState( op->children[op->op.args.ImageDisplace.source].child(), newState );
+				TArray<bool> newState;
+				newState.Init(false, size_t(EImageFormat::IF_COUNT));
+				RecurseWithState( op->children[op->op.args.ImageDisplace.source].child(), newState );
 
                 newState[(size_t)EImageFormat::IF_L_UBYTE ] = true;
                 newState[(size_t)EImageFormat::IF_L_UBYTE_RLE ] = true;
@@ -440,12 +448,13 @@ namespace mu
 
             default:
             {
-                //m_currentFormats.push_back(vector<bool>(IF_COUNT, false));
+                //m_currentFormats.Add(vector<bool>(IF_COUNT, false));
                 //Recurse(at, program);
                 //m_currentFormats.pop_back();
 
-                vector<bool> newState = vector<bool>(size_t(EImageFormat::IF_COUNT), false);
-                if (currentFormats != newState)
+				TArray<bool> newState;
+				newState.Init(false, size_t(EImageFormat::IF_COUNT));
+				if (currentFormats != newState)
                 {
                     RecurseWithState(node, newState);
                     recurse = false;
@@ -466,7 +475,7 @@ namespace mu
 
         //! Result of this visitor:
         //! Formats known to be supported by every constant image.
-        std::unordered_map< Ptr<const ASTOpConstantResource>, vector<bool> > m_supportedFormats;
+        std::unordered_map< Ptr<const ASTOpConstantResource>, TArray<bool> > m_supportedFormats;
 
     private:
 
@@ -559,7 +568,7 @@ namespace mu
                 RecurseWithState( op->source.child(), newState );
                 for( auto& r: op->removes )
                 {
-                    RecurseWithState( r.second.child(), newState );
+                    RecurseWithState( r.Value.child(), newState );
                 }
                 break;
              }
diff --git a/Engine/Plugins/Experimental/Mutable/Source/MutableTools/Private/MuT/DataPacker.h b/Engine/Plugins/Experimental/Mutable/Source/MutableTools/Private/MuT/DataPacker.h
--- a/Engine/Plugins/Experimental/Mutable/Source/MutableTools/Private/MuT/DataPacker.h
+++ b/Engine/Plugins/Experimental/Mutable/Source/MutableTools/Private/MuT/DataPacker.h
@@ -44,7 +44,7 @@ namespace mu
         //! 0: not visited
         //! 1: visited and not found
         //! 2: visited and found
-        vector<uint8_t> m_visited;
+        TArray<uint8> m_visited;
     };
 
 }
diff --git a/Engine/Plugins/Experimental/Mutable/Source/MutableTools/Private/MuT/ErrorLog.cpp b/Engine/Plugins/Experimental/Mutable/Source/MutableTools/Private/MuT/ErrorLog.cpp
--- a/Engine/Plugins/Experimental/Mutable/Source/MutableTools/Private/MuT/ErrorLog.cpp
+++ b/Engine/Plugins/Experimental/Mutable/Source/MutableTools/Private/MuT/ErrorLog.cpp
@@ -25,7 +25,6 @@
 #include "Templates/Tuple.h"
 #include "Trace/Detail/Channel.h"
 
-#include <functional>
 #include <memory>
 #include <utility>
 
@@ -60,7 +59,7 @@ namespace mu
     int ErrorLog::GetMessageCount( ErrorLogMessageType ) const
 	{
         // \todo: by type argument
-		return (int)m_pD->m_messages.size();
+		return (int)m_pD->m_messages.Num();
 	}
 
 
@@ -69,7 +68,7 @@ namespace mu
 	{
 		const char* strResult = "";
 
-		if ( index >=0 && index<(int)m_pD->m_messages.size() )
+		if ( index >=0 && index<(int)m_pD->m_messages.Num() )
 		{
 			strResult = m_pD->m_messages[index].m_text.c_str();
 		}
@@ -83,7 +82,7 @@ namespace mu
 	{
 		const void* result = 0;
 
-		if ( index >=0 && index<(int)m_pD->m_messages.size() )
+		if ( index >=0 && index<(int)m_pD->m_messages.Num() )
 		{
 			result = m_pD->m_messages[index].m_context;
 		}
@@ -97,7 +96,7 @@ namespace mu
 	{
 		ErrorLogMessageType result = ELMT_NONE;
 
-		if ( index >=0 && index<(int)m_pD->m_messages.size() )
+		if ( index >=0 && index<(int)m_pD->m_messages.Num() )
 		{
 			result = m_pD->m_messages[index].m_type;
 		}
@@ -110,14 +109,14 @@ namespace mu
 	{
         ErrorLogMessageAttachedDataView result;
 
-		if ( index >=0 && index<(int)m_pD->m_messages.size() )
+		if ( index >=0 && index<(int)m_pD->m_messages.Num() )
 		{
 			const Private::MSG& message = m_pD->m_messages[index];
             
             if ( message.m_data ) 
             {
-                result.m_unassignedUVs = message.m_data->m_unassignedUVs.data();
-			    result.m_unassignedUVsSize = message.m_data->m_unassignedUVs.size();
+                result.m_unassignedUVs = message.m_data->m_unassignedUVs.GetData();
+			    result.m_unassignedUVsSize = message.m_data->m_unassignedUVs.Num();
             }
 		}
 
@@ -128,10 +127,10 @@ namespace mu
 	void ErrorLog::Private::Add( const char* strMessage, ErrorLogMessageType type,
 								 const void* context )
 	{
-		m_messages.push_back( MSG() );
-		m_messages.back().m_type = type;
-		m_messages.back().m_text = strMessage;
-		m_messages.back().m_context = context;
+		m_messages.Add( MSG() );
+		m_messages.Last().m_type = type;
+		m_messages.Last().m_text = strMessage;
+		m_messages.Last().m_context = context;
 	}
 
 	//---------------------------------------------------------------------------------------------
@@ -139,29 +138,26 @@ namespace mu
                                  const ErrorLogMessageAttachedDataView& dataView,
                                  ErrorLogMessageType type, const void* context )
 	{
-		m_messages.push_back( MSG() );
-		m_messages.back().m_type = type;
-		m_messages.back().m_text = strMessage;
-		m_messages.back().m_context = context;
-		m_messages.back().m_data = std::make_shared< DATA >();
+		m_messages.Add( MSG() );
+		m_messages.Last().m_type = type;
+		m_messages.Last().m_text = strMessage;
+		m_messages.Last().m_context = context;
+		m_messages.Last().m_data = std::make_shared< DATA >();
 
         if ( dataView.m_unassignedUVs && dataView.m_unassignedUVsSize > 0 )
         {
-            m_messages.back().m_data->m_unassignedUVs = vector<float>(
-                dataView.m_unassignedUVs, dataView.m_unassignedUVs + dataView.m_unassignedUVsSize);
+			// \TODO: Review
+			m_messages.Last().m_data->m_unassignedUVs.Append(dataView.m_unassignedUVs, dataView.m_unassignedUVsSize);
         }
 	}
+	
 	//---------------------------------------------------------------------------------------------
 	void ErrorLog::Log() const
 	{
 		UE_LOG(LogMutableCore, Log, TEXT(" Error Log :\n"));
 
-		for ( vector<Private::MSG>::const_iterator it = m_pD->m_messages.begin()
-			; it != m_pD->m_messages.end()
-			; ++it )
+		for ( const Private::MSG& msg : m_pD->m_messages )
 		{
-			const Private::MSG& msg = *it;
-
 			switch ( msg.m_type )
 			{
 			case ELMT_ERROR: 	UE_LOG(LogMutableCore, Log, TEXT("  ERR  %s\n"), msg.m_text.c_str()); break;
@@ -176,12 +172,7 @@ namespace mu
 	//---------------------------------------------------------------------------------------------
 	void ErrorLog::Merge( const ErrorLog* pOther )
 	{
-		m_pD->m_messages.insert
-			(
-				m_pD->m_messages.end(),
-				pOther->GetPrivate()->m_messages.begin(),
-				pOther->GetPrivate()->m_messages.end()
-			);
+		m_pD->m_messages.Append(pOther->GetPrivate()->m_messages);
 	}
 
 
@@ -581,17 +572,6 @@ namespace mu
 			break;
 		}
 
-
-		//-----------------------------------------------------------------------------------------
-//        case OP_TYPE::ME_REMOVEMASK:
-//        {
-//            auto args = program.GetOpArgs<OP::MeshRemoveMaskArgs>(at);
-//            mutable_snprintf( temp, 1024,
-//                              "source : %3d    mask : %3d",
-//                    (int)args.source,
-//                    (int)args.mask );
-//            break;
-//        }
         case OP_TYPE::ME_MERGE:
         {
             auto args = program.GetOpArgs<OP::MeshMergeArgs>(at);
@@ -626,223 +606,4 @@ namespace mu
 		return temp;
 	}
 
-
-	//---------------------------------------------------------------------------------------------
-    void LogOp( const PROGRAM& prog, int at )
-	{
-        OP_TYPE type = prog.GetOpType(at);
-		UE_LOG(LogMutableCore, Log, TEXT("  @%3d  %s %s"), (int)at, s_opNames[(int)type], *GetOpDesc(prog, at));
-	}
-
-
-	//---------------------------------------------------------------------------------------------
-	void Log( ModelPtrConst pModel )
-	{
-		const PROGRAM& prog = pModel->GetPrivate()->m_program;
-
-		for ( vector<PROGRAM::STATE>::const_iterator it = prog.m_states.begin()
-			; it != prog.m_states.end()
-			; ++it )
-		{
-			UE_LOG(LogMutableCore, Log, TEXT( "state name : % s\n      root : %d" ),
-					  *FString(it->m_name.c_str()),
-					  (int)it->m_root );
-		}
-
-		UE_LOG(LogMutableCore, Log, TEXT("code : ") );
-
-//		int at = 0;
-//		for ( vector<OP>::const_iterator it = prog.m_code.begin()
-//			; it != prog.m_code.end()
-//			; ++it )
-//		{
-//            LogOp( prog, at, *it );
-//			++at;
-//		}
-
-		UE_LOG(LogMutableCore, Log, TEXT("parameters : ") );
-		for ( vector<PARAMETER_DESC>::const_iterator it = prog.m_parameters.begin()
-			; it != prog.m_parameters.end()
-			; ++it )
-		{
-			UE_LOG(LogMutableCore, Log, TEXT("    [%s] : %2d"), *FString(it->m_name.c_str()), (int)it->m_type  );
-		}
-
-
-//		LogDebug( "constant images : \n" );
-//		for ( vector<ImagePtr>::const_iterator it = prog.m_constantImages.begin()
-//			; it != prog.m_constantImages.end()
-//			; ++it )
-//		{
-//			LogDebug( "    image : %2d", (int)(*it)->GetFormat() );
-//			LogDebug( "\t %4d x %4d\n", (int)(*it)->GetSizeX(), (int)(*it)->GetSizeY() );
-//		}
-
-		UE_LOG(LogMutableCore, Log, TEXT("constant meshes : ") );
-        for ( const TPair<int32, mu::MeshPtrConst>& at: prog.m_constantMeshes )
-        {
-            const Mesh* m = at.Value.get();
-			UE_LOG(LogMutableCore, Log, TEXT("    mesh : %6d vertices \t %6d indices "),
-                    (int)m->GetVertexCount(), (int)m->GetIndexCount());
-        }
-
-		UE_LOG(LogMutableCore, Log, TEXT("constant layouts : ") );
-        for ( vector<LayoutPtrConst>::const_iterator it = prog.m_constantLayouts.begin()
-			; it != prog.m_constantLayouts.end()
-			; ++it )
-		{
-            auto pLayout = *it;
-			UE_LOG(LogMutableCore, Log, TEXT("    layout : %3d x %3d"),
-					pLayout->m_size[0],
-					pLayout->m_size[1]
-					);
-			for ( int b=0; b<pLayout->GetBlockCount(); ++b )
-			{
-				UE_LOG(LogMutableCore, Log, TEXT("        block : %3d x %3d"),
-						pLayout->m_blocks[b].m_size[0],
-						pLayout->m_blocks[b].m_size[1]
-						);
-			}
-		}
-
-		UE_LOG(LogMutableCore, Log, TEXT("constant skeletons : ") );
-        for ( auto it = prog.m_constantSkeletons.begin()
-            ; it != prog.m_constantSkeletons.end()
-            ; ++it )
-        {
-            auto pSkeleton = *it;
-			UE_LOG(LogMutableCore, Log, TEXT("    skeleton : %3d bones"), pSkeleton->GetBoneCount() );
-            for ( int b=0; b<pSkeleton->GetBoneCount(); ++b )
-            {
-				UE_LOG(LogMutableCore, Log, TEXT("        bone : %s"), *FString(pSkeleton->GetBoneName(b)) );
-            }
-        }
-
-		UE_LOG(LogMutableCore, Log, TEXT("constant strings : ") );
-		for ( vector<string>::const_iterator it = prog.m_constantStrings.begin()
-			; it != prog.m_constantStrings.end()
-			; ++it )
-		{
-			UE_LOG(LogMutableCore, Log, TEXT("    string : [%s]"), *FString(it->c_str()) );
-		}
-
-	}
-
-
-    //---------------------------------------------------------------------------------------------
-    void RecurseLogOp( int indent, int at, PROGRAM& prog, vector<uint8_t>& done )
-    {
-        OP_TYPE type = prog.GetOpType(at);
-
-        if (done[at])
-        {
-			UE_LOG(LogMutableCore, Log, TEXT("%*s@%3d  %s [repeated]"), indent, " ", (int)at, s_opNames[(int)type] );
-        }
-        else
-        {
-            done[at] = 1;
-
-            string desc = "";//GetOpDesc(prog,op);
-			UE_LOG(LogMutableCore, Log, TEXT("%*s@%3d  %s %s"), indent, " ", (int)at, s_opNames[(int)type], desc.c_str() );
-
-            ForEachReference( prog, at, [&]( OP::ADDRESS ref )
-            {
-                if (ref)
-                {
-                    RecurseLogOp( indent+2, ref, prog, done );
-                }
-            });
-        }
-    }
-
-
-    //---------------------------------------------------------------------------------------------
-    void LogSubtree( PROGRAM& prog, OP::ADDRESS root )
-    {
-        for ( vector<PROGRAM::STATE>::const_iterator it = prog.m_states.begin()
-            ; it != prog.m_states.end()
-            ; ++it )
-        {
-			UE_LOG(LogMutableCore, Log, TEXT("state name : % s\n      root : %d"),
-                      it->m_name.c_str(),
-                      (int)it->m_root );
-        }
-
-		UE_LOG(LogMutableCore, Log, TEXT("code : ") );
-
-        vector<uint8_t> done( prog.m_opAddress.size(), 0 );
-        RecurseLogOp( 0, root, prog, done );
-
-		UE_LOG(LogMutableCore, Log, TEXT("parameters : ") );
-        for ( vector<PARAMETER_DESC>::const_iterator it = prog.m_parameters.begin()
-            ; it != prog.m_parameters.end()
-            ; ++it )
-        {
-			UE_LOG(LogMutableCore, Log, TEXT("    [%s] : %2d"), it->m_name.c_str(), (int)it->m_type  );
-        }
-
-
-//		LogDebug( "constant images : \n" );
-//		for ( vector<ImagePtr>::const_iterator it = prog.m_constantImages.begin()
-//			; it != prog.m_constantImages.end()
-//			; ++it )
-//		{
-//			LogDebug( "    image : %2d", (int)(*it)->GetFormat() );
-//			LogDebug( "\t %4d x %4d\n", (int)(*it)->GetSizeX(), (int)(*it)->GetSizeY() );
-//		}
-
-		UE_LOG(LogMutableCore, Log, TEXT("constant meshes : ") );
-        int meshIndex = 0;
-        for ( const auto & at: prog.m_constantMeshes )
-        {
-            const Mesh* m = at.Value.get();
-			UE_LOG(LogMutableCore, Log, TEXT("    %4d mesh : %6d vertices \t %6d indices "),
-                      ++meshIndex,
-                      (int)m->GetVertexCount(), (int)m->GetIndexCount() );
-        }
-
-		UE_LOG(LogMutableCore, Log, TEXT("constant layouts : ") );
-        for ( vector<LayoutPtrConst>::const_iterator it = prog.m_constantLayouts.begin()
-            ; it != prog.m_constantLayouts.end()
-            ; ++it )
-        {
-            auto pLayout = *it;
-			UE_LOG(LogMutableCore, Log, TEXT("    layout : %3d x %3d"),
-                    pLayout->m_size[0],
-                    pLayout->m_size[1]
-                    );
-            for ( int b=0; b<pLayout->GetBlockCount(); ++b )
-            {
-				UE_LOG(LogMutableCore, Log, TEXT("        block : %3d x %3d"),
-                        pLayout->m_blocks[b].m_size[0],
-                        pLayout->m_blocks[b].m_size[1]
-                        );
-            }
-        }
-
-		UE_LOG(LogMutableCore, Log, TEXT("constant skeletons : ") );
-        for ( auto it = prog.m_constantSkeletons.begin()
-            ; it != prog.m_constantSkeletons.end()
-            ; ++it )
-        {
-            auto pSkeleton = *it;
-			UE_LOG(LogMutableCore, Log, TEXT("    skeleton : %3d bones"), pSkeleton->GetBoneCount() );
-            for ( int b=0; b<pSkeleton->GetBoneCount(); ++b )
-            {
-				UE_LOG(LogMutableCore, Log, TEXT("        bone : %s"), pSkeleton->GetBoneName(b) );
-            }
-        }
-
-		UE_LOG(LogMutableCore, Log, TEXT("constant strings : ") );
-        for ( vector<string>::const_iterator it = prog.m_constantStrings.begin()
-            ; it != prog.m_constantStrings.end()
-            ; ++it )
-        {
-			UE_LOG(LogMutableCore, Log, TEXT("    string : [%s]"), it->c_str() );
-        }
-
-    }
-
-
-
 }
diff --git a/Engine/Plugins/Experimental/Mutable/Source/MutableTools/Private/MuT/ErrorLogPrivate.h b/Engine/Plugins/Experimental/Mutable/Source/MutableTools/Private/MuT/ErrorLogPrivate.h
--- a/Engine/Plugins/Experimental/Mutable/Source/MutableTools/Private/MuT/ErrorLogPrivate.h
+++ b/Engine/Plugins/Experimental/Mutable/Source/MutableTools/Private/MuT/ErrorLogPrivate.h
@@ -25,7 +25,7 @@ namespace mu
 
         struct DATA
         {
-            vector< float > m_unassignedUVs;
+            TArray< float > m_unassignedUVs;
         };
 
 		struct MSG
@@ -42,7 +42,7 @@ namespace mu
 			const void* m_context;
 		};
 
-		vector< MSG > m_messages;
+		TArray< MSG > m_messages;
 
 
 		//-----------------------------------------------------------------------------------------
@@ -61,13 +61,8 @@ namespace mu
 	typedef Ptr<Model> ModelPtr;
 	typedef Ptr<const Model> ModelPtrConst;
 
-	extern void Log( ModelPtrConst pModel );
     extern const TCHAR* GetOpName( OP_TYPE type );
     extern FString GetOpDesc( const PROGRAM& program, OP::ADDRESS at );
-    extern void LogSubtree( PROGRAM& program, OP::ADDRESS root );
-
-    // Log in a graphviz/dot format
-	MUTABLETOOLS_API extern void LogSubtreeGraph( PROGRAM& program, OP::ADDRESS root );
 
 }
 
diff --git a/Engine/Plugins/Experimental/Mutable/Source/MutableTools/Private/MuT/ImageDescGenerator.cpp b/Engine/Plugins/Experimental/Mutable/Source/MutableTools/Private/MuT/ImageDescGenerator.cpp
--- a/Engine/Plugins/Experimental/Mutable/Source/MutableTools/Private/MuT/ImageDescGenerator.cpp
+++ b/Engine/Plugins/Experimental/Mutable/Source/MutableTools/Private/MuT/ImageDescGenerator.cpp
@@ -53,16 +53,16 @@ namespace mu
     //---------------------------------------------------------------------------------------------
     void ImageDescGenerator::Generate( const Node::Private& node )
     {
-        VisitedMap::const_iterator it = m_compiled.find(&node);
-        if ( it != m_compiled.end() )
+        auto it = m_compiled.Find(&node);
+        if ( it )
         {
-            m_desc = it->second;
+            m_desc = *it;
         }
         else
         {
             // The result will be stored in m_size
             node.Accept( *this );
-            m_compiled[ &node ] = m_desc;
+            m_compiled.Add( &node, m_desc );
         }
     }
 
@@ -125,7 +125,7 @@ namespace mu
         Ptr<const Image> pImage;
 
         std::size_t i = 0;
-        while ( !pImage && i<node.m_pTable->GetPrivate()->m_rows.size() )
+        while ( !pImage && i<node.m_pTable->GetPrivate()->m_rows.Num() )
         {
             pImage = node.m_pTable->GetPrivate()->m_rows[i].m_values[ colIndex ].m_pProxyImage->Get();
             ++i;
@@ -155,13 +155,13 @@ namespace mu
 
             if ( node.m_relative )
             {
-                m_desc.m_size[0] = uint16_t( m_desc.m_size[0] * node.m_sizeX );
-                m_desc.m_size[1] = uint16_t( m_desc.m_size[1] * node.m_sizeY );
+                m_desc.m_size[0] = uint16( m_desc.m_size[0] * node.m_sizeX );
+                m_desc.m_size[1] = uint16( m_desc.m_size[1] * node.m_sizeY );
             }
             else
             {
-                m_desc.m_size[0] = uint16_t( node.m_sizeX );
-                m_desc.m_size[1] = uint16_t( node.m_sizeY );
+                m_desc.m_size[0] = uint16( node.m_sizeX );
+                m_desc.m_size[1] = uint16( node.m_sizeY );
             }
         }
 
@@ -192,12 +192,12 @@ namespace mu
         {
             Generate( *pSource->GetBasePrivate() );
 
-            int mipLevel = std::max(
+            int mipLevel = FMath::Max(
                         (int)ceilf( logf( (float)m_desc.m_size[0] )/logf(2.0f) ),
                             (int)ceilf( logf( (float)m_desc.m_size[1] )/logf(2.0f) )
                             );
 
-            m_desc.m_lods = std::max( m_desc.m_lods, (uint8_t)mipLevel );
+            m_desc.m_lods = FMath::Max( m_desc.m_lods, (uint8_t)mipLevel );
         }
 
         return 0;
@@ -226,7 +226,7 @@ namespace mu
         m_desc = MUTABLE_MISSING_IMAGE_DESC;
 
         bool found = false;
-        for ( std::size_t t=0; !found && t<node.m_targets.size(); ++t )
+        for ( int32 t=0; !found && t<node.m_targets.Num(); ++t )
         {
             if ( NodeImage* pB = node.m_targets[t].get() )
             {
@@ -272,7 +272,7 @@ namespace mu
 
 		// The base image size has higher priority
 		bool found = false;
-		for (std::size_t t = 0; !found && t<node.m_sources.size(); ++t)
+		for (int32 t = 0; !found && t<node.m_sources.Num(); ++t)
 		{
 			if (NodeImage* pB = node.m_sources[t].get())
 			{
@@ -388,8 +388,8 @@ namespace mu
     Ptr<ASTOp> ImageDescGenerator::Visit( const NodeImageGradient::Private& node )
     {
         m_desc.m_format = EImageFormat::IF_RGB_UBYTE;
-        m_desc.m_size[0] = (uint16_t)node.m_size[0];
-        m_desc.m_size[1] = (uint16_t)node.m_size[1];
+        m_desc.m_size[0] = (uint16)node.m_size[0];
+        m_desc.m_size[1] = (uint16)node.m_size[1];
         m_desc.m_lods = 1;
         return 0;
     }
@@ -471,7 +471,7 @@ namespace mu
     {
         m_desc = MUTABLE_MISSING_IMAGE_DESC;
 
-        if ( node.m_options.size() > 0 && node.m_options[0] )
+        if ( node.m_options.Num() > 0 && node.m_options[0] )
         {
             Generate( *node.m_options[0]->GetBasePrivate() );
         }
@@ -489,7 +489,7 @@ namespace mu
         {
             Generate( *node.m_defaultImage->GetBasePrivate() );
         }
-        else if ( node.m_variations.size() > 0 && node.m_variations[0].m_image )
+        else if ( node.m_variations.Num() > 0 && node.m_variations[0].m_image )
         {
             Generate( *node.m_variations[0].m_image->GetBasePrivate() );
         }
@@ -533,8 +533,8 @@ namespace mu
 		// Size override?
 		if (node.m_imageSize.X > 0 && node.m_imageSize.Y > 0)
 		{
-			m_desc.m_size[0] = uint16_t(node.m_imageSize[0]);
-			m_desc.m_size[1] = uint16_t(node.m_imageSize[1]);
+			m_desc.m_size[0] = uint16(node.m_imageSize[0]);
+			m_desc.m_size[1] = uint16(node.m_imageSize[1]);
 		}
 
         // Size from the mask
diff --git a/Engine/Plugins/Experimental/Mutable/Source/MutableTools/Private/MuT/ImageDescGenerator.h b/Engine/Plugins/Experimental/Mutable/Source/MutableTools/Private/MuT/ImageDescGenerator.h
--- a/Engine/Plugins/Experimental/Mutable/Source/MutableTools/Private/MuT/ImageDescGenerator.h
+++ b/Engine/Plugins/Experimental/Mutable/Source/MutableTools/Private/MuT/ImageDescGenerator.h
@@ -114,7 +114,7 @@ namespace mu
 
     public:
 
-        typedef map<const Node::Private*,FImageDesc> VisitedMap;
+        typedef TMap<const Node::Private*,FImageDesc> VisitedMap;
         VisitedMap m_compiled;
 
         //! Result
diff --git a/Engine/Plugins/Experimental/Mutable/Source/MutableTools/Private/MuT/Node.cpp b/Engine/Plugins/Experimental/Mutable/Source/MutableTools/Private/MuT/Node.cpp
--- a/Engine/Plugins/Experimental/Mutable/Source/MutableTools/Private/MuT/Node.cpp
+++ b/Engine/Plugins/Experimental/Mutable/Source/MutableTools/Private/MuT/Node.cpp
@@ -150,81 +150,6 @@ namespace mu
 	}
 
 
-	//---------------------------------------------------------------------------------------------
-	void Node::Private::CloneChildren( const Node* pThis, NodePtr pTarget, NodeMapPtr pMap )
-	{
-		// Clone the children
-		for ( int i=0; i<pThis->GetInputCount(); ++i )
-		{
-			NodePtr pSource = pThis->GetInputNode(i);
-			NodePtr pCloned;
-			if (pSource)
-			{
-				pCloned = pMap->Get( pSource.get() );
-				if (!pCloned)
-				{
-					pCloned = pSource->Clone( pMap );
-				}
-			}
-
-			pTarget->SetInputNode( i, pCloned.get() );
-		}
-	}
-
-
-	//---------------------------------------------------------------------------------------------
-	//---------------------------------------------------------------------------------------------
-	//---------------------------------------------------------------------------------------------
-	NodeMap::NodeMap()
-	{
-		m_pD = new Private();
-	}
-
-
-	//---------------------------------------------------------------------------------------------
-	NodeMap::~NodeMap()
-	{
-        check( m_pD );
-		delete m_pD;
-		m_pD = 0;
-	}
-
-
-	//---------------------------------------------------------------------------------------------
-	NodeMap::Private* NodeMap::GetPrivate() const
-	{
-		return m_pD;
-	}
-
-
-	//---------------------------------------------------------------------------------------------
-	int NodeMap::GetSize( ) const
-	{
-		return (int)m_pD->m_map.size();
-	}
-
-
-	//---------------------------------------------------------------------------------------------
-	void NodeMap::Add( const void* key, NodePtr value )
-	{
-		m_pD->m_map[ key ] = value;
-	}
-
-
-	//---------------------------------------------------------------------------------------------
-	NodePtr NodeMap::Get( const void* key ) const
-	{
-		NodePtr pResult;
-
-		map<const void*,NodePtr>::const_iterator it = m_pD->m_map.find(key);
-		if ( it != m_pD->m_map.end() )
-		{
-			pResult = it->second;
-		}
-
-		return pResult;
-	}
-
 }
 
 
diff --git a/Engine/Plugins/Experimental/Mutable/Source/MutableTools/Private/MuT/NodeBoolConstant.cpp b/Engine/Plugins/Experimental/Mutable/Source/MutableTools/Private/MuT/NodeBoolConstant.cpp
--- a/Engine/Plugins/Experimental/Mutable/Source/MutableTools/Private/MuT/NodeBoolConstant.cpp
+++ b/Engine/Plugins/Experimental/Mutable/Source/MutableTools/Private/MuT/NodeBoolConstant.cpp
@@ -8,8 +8,6 @@
 #include "MuT/NodePrivate.h"
 
 
-#define NODE_INPUT_COUNT 	0
-
 
 namespace mu
 {
@@ -33,7 +31,7 @@ namespace mu
 	//---------------------------------------------------------------------------------------------
 	int NodeBoolConstant::GetInputCount() const
 	{
-		return NODE_INPUT_COUNT;
+		return 0;
 	}
 
 
@@ -68,5 +66,3 @@ namespace mu
 
 
 }
-
-#undef NODE_INPUT_COUNT
diff --git a/Engine/Plugins/Experimental/Mutable/Source/MutableTools/Private/MuT/NodeBoolParameter.cpp b/Engine/Plugins/Experimental/Mutable/Source/MutableTools/Private/MuT/NodeBoolParameter.cpp
--- a/Engine/Plugins/Experimental/Mutable/Source/MutableTools/Private/MuT/NodeBoolParameter.cpp
+++ b/Engine/Plugins/Experimental/Mutable/Source/MutableTools/Private/MuT/NodeBoolParameter.cpp
@@ -37,7 +37,7 @@ namespace mu
 	//---------------------------------------------------------------------------------------------
 	int NodeBoolParameter::GetInputCount() const
 	{
-        return int( m_pD->m_ranges.size() );
+        return int( m_pD->m_ranges.Num() );
     }
 
 
@@ -105,15 +105,15 @@ namespace mu
     void NodeBoolParameter::SetRangeCount( int i )
     {
         check(i>=0);
-        m_pD->m_ranges.resize(i);
+        m_pD->m_ranges.SetNum(i);
     }
 
 
     //---------------------------------------------------------------------------------------------
     void NodeBoolParameter::SetRange( int i, NodeRangePtr pRange )
     {
-        check( i>=0 && i<int(m_pD->m_ranges.size()) );
-        if ( i>=0 && i<int(m_pD->m_ranges.size()) )
+        check( i>=0 && i<int(m_pD->m_ranges.Num()) );
+        if ( i>=0 && i<int(m_pD->m_ranges.Num()) )
         {
             m_pD->m_ranges[i] = pRange;
         }
diff --git a/Engine/Plugins/Experimental/Mutable/Source/MutableTools/Private/MuT/NodeBoolPrivate.h b/Engine/Plugins/Experimental/Mutable/Source/MutableTools/Private/MuT/NodeBoolPrivate.h
--- a/Engine/Plugins/Experimental/Mutable/Source/MutableTools/Private/MuT/NodeBoolPrivate.h
+++ b/Engine/Plugins/Experimental/Mutable/Source/MutableTools/Private/MuT/NodeBoolPrivate.h
@@ -66,7 +66,7 @@ namespace mu
 		string m_name;
 		string m_uid;
 
-        vector<Ptr<NodeRange>> m_ranges;
+        TArray<Ptr<NodeRange>> m_ranges;
 
 		//!
 		void Serialise( OutputArchive& arch ) const
diff --git a/Engine/Plugins/Experimental/Mutable/Source/MutableTools/Private/MuT/NodeColourArithmeticOperation.cpp b/Engine/Plugins/Experimental/Mutable/Source/MutableTools/Private/MuT/NodeColourArithmeticOperation.cpp
--- a/Engine/Plugins/Experimental/Mutable/Source/MutableTools/Private/MuT/NodeColourArithmeticOperation.cpp
+++ b/Engine/Plugins/Experimental/Mutable/Source/MutableTools/Private/MuT/NodeColourArithmeticOperation.cpp
@@ -8,9 +8,6 @@
 #include "MuT/NodePrivate.h"
 
 
-#define NODE_INPUT_COUNT 	2
-
-
 namespace mu
 {
 
@@ -42,14 +39,14 @@ namespace mu
 	//---------------------------------------------------------------------------------------------
 	int NodeColourArithmeticOperation::GetInputCount() const
 	{
-		return NODE_INPUT_COUNT;
+		return 2;
 	}
 
 
 	//---------------------------------------------------------------------------------------------
 	Node* NodeColourArithmeticOperation::GetInputNode( int i ) const
 	{
-		check( i>=0 && i<NODE_INPUT_COUNT );
+		check( i>=0 && i< GetInputCount());
 
 		Node* pResult = 0;
 
@@ -66,7 +63,7 @@ namespace mu
 	//---------------------------------------------------------------------------------------------
 	void NodeColourArithmeticOperation::SetInputNode( int i, NodePtr pNode )
 	{
-		check( i>=0 && i<NODE_INPUT_COUNT );
+		check( i>=0 && i< GetInputCount());
 
 		switch (i)
 		{
@@ -119,5 +116,3 @@ namespace mu
 	}
 
 }
-
-#undef NODE_INPUT_COUNT
diff --git a/Engine/Plugins/Experimental/Mutable/Source/MutableTools/Private/MuT/NodeColourConstant.cpp b/Engine/Plugins/Experimental/Mutable/Source/MutableTools/Private/MuT/NodeColourConstant.cpp
--- a/Engine/Plugins/Experimental/Mutable/Source/MutableTools/Private/MuT/NodeColourConstant.cpp
+++ b/Engine/Plugins/Experimental/Mutable/Source/MutableTools/Private/MuT/NodeColourConstant.cpp
@@ -9,9 +9,6 @@
 #include "MuT/NodePrivate.h"
 
 
-#define NODE_INPUT_COUNT 	0
-
-
 namespace mu
 {
 
@@ -34,7 +31,7 @@ namespace mu
 	//---------------------------------------------------------------------------------------------
 	int NodeColourConstant::GetInputCount() const
 	{
-		return NODE_INPUT_COUNT;
+		return 0;
 	}
 
 
@@ -84,4 +81,3 @@ namespace mu
 
 }
 
-#undef NODE_INPUT_COUNT
diff --git a/Engine/Plugins/Experimental/Mutable/Source/MutableTools/Private/MuT/NodeColourFromScalars.cpp b/Engine/Plugins/Experimental/Mutable/Source/MutableTools/Private/MuT/NodeColourFromScalars.cpp
--- a/Engine/Plugins/Experimental/Mutable/Source/MutableTools/Private/MuT/NodeColourFromScalars.cpp
+++ b/Engine/Plugins/Experimental/Mutable/Source/MutableTools/Private/MuT/NodeColourFromScalars.cpp
@@ -9,9 +9,6 @@
 #include "MuT/NodeScalar.h"
 
 
-#define NODE_INPUT_COUNT 	4
-
-
 namespace mu
 {
 
@@ -34,14 +31,14 @@ namespace mu
 	//---------------------------------------------------------------------------------------------
 	int NodeColourFromScalars::GetInputCount() const
 	{
-		return NODE_INPUT_COUNT;
+		return 4;
 	}
 
 
 	//---------------------------------------------------------------------------------------------
 	Node* NodeColourFromScalars::GetInputNode( int i ) const
 	{
-		check( i>=0 && i<NODE_INPUT_COUNT );
+		check( i>=0 && i< GetInputCount());
 
 		Node* pResult = 0;
 
@@ -60,7 +57,7 @@ namespace mu
 	//---------------------------------------------------------------------------------------------
 	void NodeColourFromScalars::SetInputNode( int i, NodePtr pNode )
 	{
-		check( i>=0 && i<NODE_INPUT_COUNT );
+		check( i>=0 && i< GetInputCount());
 
 		switch (i)
 		{
@@ -131,4 +128,3 @@ namespace mu
 
 }
 
-#undef NODE_INPUT_COUNT
diff --git a/Engine/Plugins/Experimental/Mutable/Source/MutableTools/Private/MuT/NodeColourParameter.cpp b/Engine/Plugins/Experimental/Mutable/Source/MutableTools/Private/MuT/NodeColourParameter.cpp
--- a/Engine/Plugins/Experimental/Mutable/Source/MutableTools/Private/MuT/NodeColourParameter.cpp
+++ b/Engine/Plugins/Experimental/Mutable/Source/MutableTools/Private/MuT/NodeColourParameter.cpp
@@ -36,7 +36,7 @@ namespace mu
 	//---------------------------------------------------------------------------------------------
 	int NodeColourParameter::GetInputCount() const
 	{
-        return int( m_pD->m_ranges.size() );
+        return  m_pD->m_ranges.Num();
     }
 
 
@@ -126,15 +126,15 @@ namespace mu
     void NodeColourParameter::SetRangeCount( int i )
     {
         check(i>=0);
-        m_pD->m_ranges.resize(i);
+        m_pD->m_ranges.SetNum(i);
     }
 
 
     //---------------------------------------------------------------------------------------------
     void NodeColourParameter::SetRange( int i, NodeRangePtr pRange )
     {
-        check( i>=0 && i<int(m_pD->m_ranges.size()) );
-        if ( i>=0 && i<int(m_pD->m_ranges.size()) )
+        check( i>=0 && i<int(m_pD->m_ranges.Num()) );
+        if ( i>=0 && i<int(m_pD->m_ranges.Num()) )
         {
             m_pD->m_ranges[i] = pRange;
         }
diff --git a/Engine/Plugins/Experimental/Mutable/Source/MutableTools/Private/MuT/NodeColourParameterPrivate.h b/Engine/Plugins/Experimental/Mutable/Source/MutableTools/Private/MuT/NodeColourParameterPrivate.h
--- a/Engine/Plugins/Experimental/Mutable/Source/MutableTools/Private/MuT/NodeColourParameterPrivate.h
+++ b/Engine/Plugins/Experimental/Mutable/Source/MutableTools/Private/MuT/NodeColourParameterPrivate.h
@@ -26,7 +26,7 @@ namespace mu
 		string m_name;
 		string m_uid;
 
-        vector<Ptr<NodeRange>> m_ranges;
+        TArray<Ptr<NodeRange>> m_ranges;
 
 		//!
 		void Serialise( OutputArchive& arch ) const
diff --git a/Engine/Plugins/Experimental/Mutable/Source/MutableTools/Private/MuT/NodeColourSampleImage.cpp b/Engine/Plugins/Experimental/Mutable/Source/MutableTools/Private/MuT/NodeColourSampleImage.cpp
--- a/Engine/Plugins/Experimental/Mutable/Source/MutableTools/Private/MuT/NodeColourSampleImage.cpp
+++ b/Engine/Plugins/Experimental/Mutable/Source/MutableTools/Private/MuT/NodeColourSampleImage.cpp
@@ -10,9 +10,6 @@
 #include "MuT/NodeScalar.h"
 
 
-#define NODE_INPUT_COUNT 	3
-
-
 namespace mu
 {
 
@@ -35,14 +32,14 @@ namespace mu
 	//---------------------------------------------------------------------------------------------
 	int NodeColourSampleImage::GetInputCount() const
 	{
-		return NODE_INPUT_COUNT;
+		return 3;
 	}
 
 
 	//---------------------------------------------------------------------------------------------
 	Node* NodeColourSampleImage::GetInputNode( int i ) const
 	{
-		check( i>=0 && i<NODE_INPUT_COUNT );
+		check( i>=0 && i< GetInputCount());
 
 		Node* pResult = 0;
 
@@ -60,7 +57,7 @@ namespace mu
 	//---------------------------------------------------------------------------------------------
 	void NodeColourSampleImage::SetInputNode( int i, NodePtr pNode )
 	{
-		check( i>=0 && i<NODE_INPUT_COUNT );
+		check( i>=0 && i< GetInputCount());
 
 		switch (i)
 		{
@@ -117,6 +114,3 @@ namespace mu
 
 
 }
-
-#undef NODE_INPUT_COUNT
-
diff --git a/Engine/Plugins/Experimental/Mutable/Source/MutableTools/Private/MuT/NodeColourSwitch.cpp b/Engine/Plugins/Experimental/Mutable/Source/MutableTools/Private/MuT/NodeColourSwitch.cpp
--- a/Engine/Plugins/Experimental/Mutable/Source/MutableTools/Private/MuT/NodeColourSwitch.cpp
+++ b/Engine/Plugins/Experimental/Mutable/Source/MutableTools/Private/MuT/NodeColourSwitch.cpp
@@ -34,7 +34,7 @@ namespace mu
 	//---------------------------------------------------------------------------------------------
     int NodeColourSwitch::GetInputCount() const
 	{
-		return 1 + (int)m_pD->m_options.size();
+		return 1 + m_pD->m_options.Num();
 	}
 
 
@@ -97,14 +97,14 @@ namespace mu
 	//---------------------------------------------------------------------------------------------
     void NodeColourSwitch::SetOptionCount( int t )
 	{
-		m_pD->m_options.resize(t);
+		m_pD->m_options.SetNum(t);
 	}
 
 
 	//---------------------------------------------------------------------------------------------
     NodeColourPtr NodeColourSwitch::GetOption( int t ) const
 	{
-		check( t>=0 && t<(int)m_pD->m_options.size() );
+		check( t>=0 && t<m_pD->m_options.Num() );
 		return m_pD->m_options[t].get();
 	}
 
@@ -112,7 +112,7 @@ namespace mu
 	//---------------------------------------------------------------------------------------------
     void NodeColourSwitch::SetOption( int t, NodeColourPtr pNode )
 	{
-		check( t>=0 && t<(int)m_pD->m_options.size() );
+		check( t>=0 && t<m_pD->m_options.Num() );
 		m_pD->m_options[t] = pNode;
 	}
 
diff --git a/Engine/Plugins/Experimental/Mutable/Source/MutableTools/Private/MuT/NodeColourSwitchPrivate.h b/Engine/Plugins/Experimental/Mutable/Source/MutableTools/Private/MuT/NodeColourSwitchPrivate.h
--- a/Engine/Plugins/Experimental/Mutable/Source/MutableTools/Private/MuT/NodeColourSwitchPrivate.h
+++ b/Engine/Plugins/Experimental/Mutable/Source/MutableTools/Private/MuT/NodeColourSwitchPrivate.h
@@ -23,7 +23,7 @@ namespace mu
 		static NODE_TYPE s_type;
 
 		NodeScalarPtr m_pParameter;
-        vector<NodeColourPtr> m_options;
+        TArray<NodeColourPtr> m_options;
 
 		//!
 		void Serialise( OutputArchive& arch ) const
diff --git a/Engine/Plugins/Experimental/Mutable/Source/MutableTools/Private/MuT/NodeColourVariation.cpp b/Engine/Plugins/Experimental/Mutable/Source/MutableTools/Private/MuT/NodeColourVariation.cpp
--- a/Engine/Plugins/Experimental/Mutable/Source/MutableTools/Private/MuT/NodeColourVariation.cpp
+++ b/Engine/Plugins/Experimental/Mutable/Source/MutableTools/Private/MuT/NodeColourVariation.cpp
@@ -30,7 +30,7 @@ namespace mu
 
 
     //---------------------------------------------------------------------------------------------
-    int NodeColourVariation::GetInputCount() const { return 1 + int( m_pD->m_variations.size() ); }
+    int NodeColourVariation::GetInputCount() const { return 1 + m_pD->m_variations.Num(); }
 
 
     //---------------------------------------------------------------------------------------------
@@ -44,11 +44,11 @@ namespace mu
         }
         i -= 1;
 
-        if ( i < int( m_pD->m_variations.size() ) )
+        if ( i < int( m_pD->m_variations.Num() ) )
         {
             return m_pD->m_variations[i].m_colour.get();
         }
-        i -= int( m_pD->m_variations.size() );
+        i -= int( m_pD->m_variations.Num() );
 
         return nullptr;
     }
@@ -66,13 +66,13 @@ namespace mu
         }
 
         i -= 1;
-        if ( i < int( m_pD->m_variations.size() ) )
+        if ( i < int( m_pD->m_variations.Num() ) )
         {
 
             m_pD->m_variations[i].m_colour = dynamic_cast<NodeColour*>( pNode.get() );
             return;
         }
-        i -= (int)m_pD->m_variations.size();
+        i -= (int)m_pD->m_variations.Num();
     }
 
 
@@ -83,20 +83,20 @@ namespace mu
 
 
     //---------------------------------------------------------------------------------------------
-    int NodeColourVariation::GetVariationCount() const { return int( m_pD->m_variations.size() ); }
+    int NodeColourVariation::GetVariationCount() const { return m_pD->m_variations.Num(); }
 
 
     //---------------------------------------------------------------------------------------------
     void NodeColourVariation::SetVariationCount( int num )
     {
         check( num >= 0 );
-        m_pD->m_variations.resize( num );
+        m_pD->m_variations.SetNum( num );
     }
 
     //---------------------------------------------------------------------------------------------
     void NodeColourVariation::SetVariationTag( int index, const char* strTag )
     {
-        check( index >= 0 && index < (int)m_pD->m_variations.size() );
+        check( index >= 0 && index < m_pD->m_variations.Num() );
         check( strTag );
 
         if ( strTag )
@@ -113,7 +113,7 @@ namespace mu
     //---------------------------------------------------------------------------------------------
     void NodeColourVariation::SetVariationColour( int index, NodeColour* pNode )
     {
-        check( index >= 0 && index < (int)m_pD->m_variations.size() );
+        check( index >= 0 && index < m_pD->m_variations.Num() );
 
         m_pD->m_variations[index].m_colour = pNode;
     }
diff --git a/Engine/Plugins/Experimental/Mutable/Source/MutableTools/Private/MuT/NodeColourVariationPrivate.h b/Engine/Plugins/Experimental/Mutable/Source/MutableTools/Private/MuT/NodeColourVariationPrivate.h
--- a/Engine/Plugins/Experimental/Mutable/Source/MutableTools/Private/MuT/NodeColourVariationPrivate.h
+++ b/Engine/Plugins/Experimental/Mutable/Source/MutableTools/Private/MuT/NodeColourVariationPrivate.h
@@ -52,7 +52,7 @@ namespace mu
             }
         };
 
-        vector<VARIATION> m_variations;
+        TArray<VARIATION> m_variations;
 
         //!
         void Serialise( OutputArchive& arch ) const
