[mutable] Optimize mesh clip versus switching of clipping mesh.

diff --git a/Engine/Plugins/Experimental/Mutable/Source/MutableTools/Private/MuT/NodeComponentEdit.cpp b/Engine/Plugins/Experimental/Mutable/Source/MutableTools/Private/MuT/NodeComponentEdit.cpp
--- a/Engine/Plugins/Experimental/Mutable/Source/MutableTools/Private/MuT/NodeComponentEdit.cpp
+++ b/Engine/Plugins/Experimental/Mutable/Source/MutableTools/Private/MuT/NodeComponentEdit.cpp
@@ -37,7 +37,7 @@ namespace mu
     //---------------------------------------------------------------------------------------------
     int NodeComponentEdit::GetInputCount() const
     {
-        return (int) m_pD->m_surfaces.size();
+        return m_pD->m_surfaces.Num();
     }
 
 
@@ -48,7 +48,7 @@ namespace mu
 
         NodePtr pResult;
 
-        if ( i<(int)m_pD->m_surfaces.size() )
+        if ( i<m_pD->m_surfaces.Num() )
         {
             pResult = m_pD->m_surfaces[i];
         }
@@ -62,7 +62,7 @@ namespace mu
     {
         check( i >=0 && i < GetInputCount() );
 
-        if ( i<(int)m_pD->m_surfaces.size() )
+        if ( i<m_pD->m_surfaces.Num() )
         {
             m_pD->m_surfaces[ i ] = dynamic_cast<NodeSurface*>(pNode.get());
         }
@@ -88,7 +88,7 @@ namespace mu
     //---------------------------------------------------------------------------------------------
     int NodeComponentEdit::GetSurfaceCount() const
     {
-        return (int)m_pD->m_surfaces.size();
+        return m_pD->m_surfaces.Num();
     }
 
 
@@ -96,14 +96,14 @@ namespace mu
     void NodeComponentEdit::SetSurfaceCount( int num )
     {
         check( num >=0 );
-        m_pD->m_surfaces.resize( num );
+        m_pD->m_surfaces.SetNum( num );
     }
 
 
     //---------------------------------------------------------------------------------------------
     NodeSurface* NodeComponentEdit::GetSurface( int index ) const
     {
-        check( index >=0 && index < (int)m_pD->m_surfaces.size() );
+        check( index >=0 && index < m_pD->m_surfaces.Num() );
 
         return m_pD->m_surfaces[ index ].get();
     }
@@ -112,7 +112,7 @@ namespace mu
     //---------------------------------------------------------------------------------------------
     void NodeComponentEdit::SetSurface( int index, NodeSurface* pNode )
     {
-        check( index >=0 && index < (int)m_pD->m_surfaces.size() );
+        check( index >=0 && index < m_pD->m_surfaces.Num() );
 
         m_pD->m_surfaces[ index ] = pNode;
     }
diff --git a/Engine/Plugins/Experimental/Mutable/Source/MutableTools/Private/MuT/NodeComponentEditPrivate.h b/Engine/Plugins/Experimental/Mutable/Source/MutableTools/Private/MuT/NodeComponentEditPrivate.h
--- a/Engine/Plugins/Experimental/Mutable/Source/MutableTools/Private/MuT/NodeComponentEditPrivate.h
+++ b/Engine/Plugins/Experimental/Mutable/Source/MutableTools/Private/MuT/NodeComponentEditPrivate.h
@@ -29,15 +29,11 @@ namespace mu
 
 	public:
 
-		Private()
-		{
-		}
-
 		static NODE_TYPE s_type;
 
         NodeComponentPtr m_pParent;
 
-        vector<NodeSurfacePtr> m_surfaces;
+        TArray<NodeSurfacePtr> m_surfaces;
 
 		//!
 		void Serialise( OutputArchive& arch ) const
diff --git a/Engine/Plugins/Experimental/Mutable/Source/MutableTools/Private/MuT/NodeComponentNew.cpp b/Engine/Plugins/Experimental/Mutable/Source/MutableTools/Private/MuT/NodeComponentNew.cpp
--- a/Engine/Plugins/Experimental/Mutable/Source/MutableTools/Private/MuT/NodeComponentNew.cpp
+++ b/Engine/Plugins/Experimental/Mutable/Source/MutableTools/Private/MuT/NodeComponentNew.cpp
@@ -36,7 +36,7 @@ namespace mu
 	//---------------------------------------------------------------------------------------------
 	int NodeComponentNew::GetInputCount() const
 	{
-        return (int) m_pD->m_surfaces.size();
+        return m_pD->m_surfaces.Num();
 	}
 
 
@@ -47,7 +47,7 @@ namespace mu
 
 		NodePtr pResult;
 
-        if ( i<(int)m_pD->m_surfaces.size() )
+        if ( i<m_pD->m_surfaces.Num() )
 		{
             pResult = m_pD->m_surfaces[i];
 		}
@@ -61,7 +61,7 @@ namespace mu
 	{
         check( i >=0 && i < GetInputCount() );
 
-        if ( i<(int)m_pD->m_surfaces.size() )
+        if ( i<m_pD->m_surfaces.Num() )
 		{
             m_pD->m_surfaces[ i ] = dynamic_cast<NodeSurface*>(pNode.get());
 		}
@@ -110,7 +110,7 @@ namespace mu
 	//---------------------------------------------------------------------------------------------
     int NodeComponentNew::GetSurfaceCount() const
 	{
-        return (int)m_pD->m_surfaces.size();
+        return m_pD->m_surfaces.Num();
 	}
 
 
@@ -118,14 +118,14 @@ namespace mu
     void NodeComponentNew::SetSurfaceCount( int num )
 	{
 		check( num >=0 );
-        m_pD->m_surfaces.resize( num );
+        m_pD->m_surfaces.SetNum( num );
 	}
 
 
 	//---------------------------------------------------------------------------------------------
     NodeSurface* NodeComponentNew::GetSurface( int index ) const
 	{
-        check( index >=0 && index < (int)m_pD->m_surfaces.size() );
+        check( index >=0 && index < m_pD->m_surfaces.Num() );
 
         return m_pD->m_surfaces[ index ].get();
 	}
@@ -134,7 +134,7 @@ namespace mu
 	//---------------------------------------------------------------------------------------------
     void NodeComponentNew::SetSurface( int index, NodeSurface* pNode )
 	{
-        check( index >=0 && index < (int)m_pD->m_surfaces.size() );
+        check( index >=0 && index < m_pD->m_surfaces.Num() );
 
         m_pD->m_surfaces[ index ] = pNode;
 	}
diff --git a/Engine/Plugins/Experimental/Mutable/Source/MutableTools/Private/MuT/NodeComponentNewPrivate.h b/Engine/Plugins/Experimental/Mutable/Source/MutableTools/Private/MuT/NodeComponentNewPrivate.h
--- a/Engine/Plugins/Experimental/Mutable/Source/MutableTools/Private/MuT/NodeComponentNewPrivate.h
+++ b/Engine/Plugins/Experimental/Mutable/Source/MutableTools/Private/MuT/NodeComponentNewPrivate.h
@@ -35,7 +35,7 @@ namespace mu
 
 		uint16 m_id = 0;
 
-        vector<NodeSurfacePtr> m_surfaces;
+        TArray<NodeSurfacePtr> m_surfaces;
 
 		//!
 		void Serialise( OutputArchive& arch ) const
diff --git a/Engine/Plugins/Experimental/Mutable/Source/MutableTools/Private/MuT/NodeImageBinarise.cpp b/Engine/Plugins/Experimental/Mutable/Source/MutableTools/Private/MuT/NodeImageBinarise.cpp
--- a/Engine/Plugins/Experimental/Mutable/Source/MutableTools/Private/MuT/NodeImageBinarise.cpp
+++ b/Engine/Plugins/Experimental/Mutable/Source/MutableTools/Private/MuT/NodeImageBinarise.cpp
@@ -9,9 +9,6 @@
 #include "MuT/NodeScalar.h"
 
 
-#define NODE_INPUT_COUNT 	2
-
-
 namespace mu
 {
 
@@ -31,7 +28,7 @@ namespace mu
 	//---------------------------------------------------------------------------------------------
 	void NodeImageBinarise::SetInputNode( int i, NodePtr pNode )
 	{
-		check( i>=0 && i<NODE_INPUT_COUNT );
+		check( i>=0 && i< GetInputCount());
 
 		switch (i)
 		{
@@ -44,14 +41,14 @@ namespace mu
 	//---------------------------------------------------------------------------------------------
 	int NodeImageBinarise::GetInputCount() const
 	{
-		return NODE_INPUT_COUNT;
+		return 2;
 	}
 
 
 	//---------------------------------------------------------------------------------------------
 	Node* NodeImageBinarise::GetInputNode( int i ) const
 	{
-		check( i>=0 && i<NODE_INPUT_COUNT );
+		check( i>=0 && i< GetInputCount());
 
 		Node* pResult = 0;
 
@@ -96,5 +93,3 @@ namespace mu
 
 }
 
-#undef NODE_INPUT_COUNT
-
diff --git a/Engine/Plugins/Experimental/Mutable/Source/MutableTools/Private/MuT/NodeImageDifference.cpp b/Engine/Plugins/Experimental/Mutable/Source/MutableTools/Private/MuT/NodeImageDifference.cpp
--- a/Engine/Plugins/Experimental/Mutable/Source/MutableTools/Private/MuT/NodeImageDifference.cpp
+++ b/Engine/Plugins/Experimental/Mutable/Source/MutableTools/Private/MuT/NodeImageDifference.cpp
@@ -8,9 +8,6 @@
 #include "MuT/NodePrivate.h"
 
 
-#define NODE_INPUT_COUNT 	2
-
-
 namespace mu
 {
 
@@ -30,7 +27,7 @@ namespace mu
 	//---------------------------------------------------------------------------------------------
 	void NodeImageDifference::SetInputNode( int i, NodePtr pNode )
 	{
-		check( i>=0 && i<NODE_INPUT_COUNT );
+		check( i>=0 && i< GetInputCount());
 
 		switch (i)
 		{
@@ -43,14 +40,14 @@ namespace mu
 	//---------------------------------------------------------------------------------------------
 	int NodeImageDifference::GetInputCount() const
 	{
-		return NODE_INPUT_COUNT;
+		return 2;
 	}
 
 
 	//---------------------------------------------------------------------------------------------
 	Node* NodeImageDifference::GetInputNode( int i ) const
 	{
-		check( i>=0 && i<NODE_INPUT_COUNT );
+		check( i>=0 && i< GetInputCount());
 
 		Node* pResult = 0;
 
@@ -95,5 +92,4 @@ namespace mu
 
 }
 
-#undef NODE_INPUT_COUNT
 
diff --git a/Engine/Plugins/Experimental/Mutable/Source/MutableTools/Private/MuT/NodeImageFormat.cpp b/Engine/Plugins/Experimental/Mutable/Source/MutableTools/Private/MuT/NodeImageFormat.cpp
--- a/Engine/Plugins/Experimental/Mutable/Source/MutableTools/Private/MuT/NodeImageFormat.cpp
+++ b/Engine/Plugins/Experimental/Mutable/Source/MutableTools/Private/MuT/NodeImageFormat.cpp
@@ -8,9 +8,6 @@
 #include "MuT/NodePrivate.h"
 
 
-#define NODE_INPUT_COUNT 	4
-
-
 namespace mu
 {
 
@@ -88,5 +85,3 @@ namespace mu
 	}
 
 }
-
-#undef NODE_INPUT_COUNT
diff --git a/Engine/Plugins/Experimental/Mutable/Source/MutableTools/Private/MuT/NodeImageGradient.cpp b/Engine/Plugins/Experimental/Mutable/Source/MutableTools/Private/MuT/NodeImageGradient.cpp
--- a/Engine/Plugins/Experimental/Mutable/Source/MutableTools/Private/MuT/NodeImageGradient.cpp
+++ b/Engine/Plugins/Experimental/Mutable/Source/MutableTools/Private/MuT/NodeImageGradient.cpp
@@ -10,9 +10,6 @@
 #include "MuT/NodePrivate.h"
 
 
-#define NODE_INPUT_COUNT 	2
-
-
 namespace mu
 {
 
@@ -35,14 +32,14 @@ namespace mu
 	//---------------------------------------------------------------------------------------------
 	int NodeImageGradient::GetInputCount() const
 	{
-		return NODE_INPUT_COUNT;
+		return 2;
 	}
 
 
 	//---------------------------------------------------------------------------------------------
 	Node* NodeImageGradient::GetInputNode( int i ) const
 	{
-		check( i>=0 && i<NODE_INPUT_COUNT );
+		check( i>=0 && i< GetInputCount());
 
 		Node* pResult = 0;
 
@@ -59,7 +56,7 @@ namespace mu
 	//---------------------------------------------------------------------------------------------
 	void NodeImageGradient::SetInputNode( int i, NodePtr pNode )
 	{
-		check( i>=0 && i<NODE_INPUT_COUNT );
+		check( i>=0 && i< GetInputCount());
 
 		switch (i)
 		{
@@ -121,4 +118,3 @@ namespace mu
 
 }
 
-#undef NODE_INPUT_COUNT
diff --git a/Engine/Plugins/Experimental/Mutable/Source/MutableTools/Private/MuT/NodeImageInterpolate.cpp b/Engine/Plugins/Experimental/Mutable/Source/MutableTools/Private/MuT/NodeImageInterpolate.cpp
--- a/Engine/Plugins/Experimental/Mutable/Source/MutableTools/Private/MuT/NodeImageInterpolate.cpp
+++ b/Engine/Plugins/Experimental/Mutable/Source/MutableTools/Private/MuT/NodeImageInterpolate.cpp
@@ -35,7 +35,7 @@ namespace mu
 	//---------------------------------------------------------------------------------------------
 	int NodeImageInterpolate::GetInputCount() const
 	{
-		return 1 + (int)m_pD->m_targets.size();
+		return 1 + (int)m_pD->m_targets.Num();
 	}
 
 
@@ -98,7 +98,7 @@ namespace mu
 	//---------------------------------------------------------------------------------------------
 	NodeImagePtr NodeImageInterpolate::GetTarget( int t ) const
 	{
-		check( t>=0 && t<(int)m_pD->m_targets.size() );
+		check( t>=0 && t<(int)m_pD->m_targets.Num() );
 		return m_pD->m_targets[t].get();
 	}
 
@@ -106,7 +106,7 @@ namespace mu
 	//---------------------------------------------------------------------------------------------
 	void NodeImageInterpolate::SetTarget( int t, NodeImagePtr pNode )
 	{
-		check( t>=0 && t<(int)m_pD->m_targets.size() );
+		check( t>=0 && t<(int)m_pD->m_targets.Num() );
 		m_pD->m_targets[t] = pNode;
 	}
 
@@ -114,14 +114,14 @@ namespace mu
 	//---------------------------------------------------------------------------------------------
 	void NodeImageInterpolate::SetTargetCount( int t )
 	{
-		m_pD->m_targets.resize(t);
+		m_pD->m_targets.SetNum(t);
 	}
 
 
 	//---------------------------------------------------------------------------------------------
 	int NodeImageInterpolate::GetTargetCount() const
 	{
-		return int(m_pD->m_targets.size());
+		return int(m_pD->m_targets.Num());
 	}
 
 
diff --git a/Engine/Plugins/Experimental/Mutable/Source/MutableTools/Private/MuT/NodeImageInterpolatePrivate.h b/Engine/Plugins/Experimental/Mutable/Source/MutableTools/Private/MuT/NodeImageInterpolatePrivate.h
--- a/Engine/Plugins/Experimental/Mutable/Source/MutableTools/Private/MuT/NodeImageInterpolatePrivate.h
+++ b/Engine/Plugins/Experimental/Mutable/Source/MutableTools/Private/MuT/NodeImageInterpolatePrivate.h
@@ -21,7 +21,7 @@ namespace mu
 		static NODE_TYPE s_type;
 
 		NodeScalarPtr m_pFactor;
-		vector<NodeImagePtr> m_targets;
+		TArray<NodeImagePtr> m_targets;
 
 		//!
 		void Serialise( OutputArchive& arch ) const
diff --git a/Engine/Plugins/Experimental/Mutable/Source/MutableTools/Private/MuT/NodeImageInvert.cpp b/Engine/Plugins/Experimental/Mutable/Source/MutableTools/Private/MuT/NodeImageInvert.cpp
--- a/Engine/Plugins/Experimental/Mutable/Source/MutableTools/Private/MuT/NodeImageInvert.cpp
+++ b/Engine/Plugins/Experimental/Mutable/Source/MutableTools/Private/MuT/NodeImageInvert.cpp
@@ -8,8 +8,6 @@
 #include "MuT/NodePrivate.h"
 
 
-#define NODE_INPUT_COUNT 1
-
 namespace mu
 {
 
@@ -28,7 +26,7 @@ namespace mu
 	//---------------------------------------------------------------------------------------------
 	void NodeImageInvert::SetInputNode(int i, NodePtr pNode)
 	{
-		check(i >= 0 && i < NODE_INPUT_COUNT);
+		check(i >= 0 && i < GetInputCount());
 
 		m_pD->m_pBase = dynamic_cast<NodeImage*>(pNode.get());
 	}
@@ -36,13 +34,13 @@ namespace mu
 	//---------------------------------------------------------------------------------------------
 	int NodeImageInvert::GetInputCount() const
 	{
-		return NODE_INPUT_COUNT;
+		return 1;
 	}
 
 	//---------------------------------------------------------------------------------------------
 	Node* NodeImageInvert::GetInputNode(int i)const
 	{
-		check(i >= 0 && i < NODE_INPUT_COUNT);
+		check(i >= 0 && i < GetInputCount());
 
 		Node* pResult = 0;
 		pResult = m_pD->m_pBase.get();
@@ -64,5 +62,3 @@ namespace mu
 	}
 	
 }
-
-#undef NODE_INPUT_COUNT
diff --git a/Engine/Plugins/Experimental/Mutable/Source/MutableTools/Private/MuT/NodeImageLayer.cpp b/Engine/Plugins/Experimental/Mutable/Source/MutableTools/Private/MuT/NodeImageLayer.cpp
--- a/Engine/Plugins/Experimental/Mutable/Source/MutableTools/Private/MuT/NodeImageLayer.cpp
+++ b/Engine/Plugins/Experimental/Mutable/Source/MutableTools/Private/MuT/NodeImageLayer.cpp
@@ -8,9 +8,6 @@
 #include "MuT/NodePrivate.h"
 
 
-#define NODE_INPUT_COUNT 	3
-
-
 namespace mu
 {
 
@@ -48,14 +45,14 @@ namespace mu
 	//---------------------------------------------------------------------------------------------
 	int NodeImageLayer::GetInputCount() const
 	{
-		return NODE_INPUT_COUNT;
+		return 3;
 	}
 
 
 	//---------------------------------------------------------------------------------------------
 	Node* NodeImageLayer::GetInputNode( int i ) const
 	{
-		check( i>=0 && i<NODE_INPUT_COUNT );
+		check( i>=0 && i< GetInputCount());
 
 		Node* pResult = 0;
 
@@ -73,7 +70,7 @@ namespace mu
 	//---------------------------------------------------------------------------------------------
 	void NodeImageLayer::SetInputNode( int i, NodePtr pNode )
 	{
-		check( i>=0 && i<NODE_INPUT_COUNT );
+		check( i>=0 && i< GetInputCount());
 
 		switch (i)
 		{
@@ -143,5 +140,3 @@ namespace mu
 
 
 }
-
-#undef NODE_INPUT_COUNT
diff --git a/Engine/Plugins/Experimental/Mutable/Source/MutableTools/Private/MuT/NodeImageLayerColour.cpp b/Engine/Plugins/Experimental/Mutable/Source/MutableTools/Private/MuT/NodeImageLayerColour.cpp
--- a/Engine/Plugins/Experimental/Mutable/Source/MutableTools/Private/MuT/NodeImageLayerColour.cpp
+++ b/Engine/Plugins/Experimental/Mutable/Source/MutableTools/Private/MuT/NodeImageLayerColour.cpp
@@ -9,9 +9,6 @@
 #include "MuT/NodePrivate.h"
 
 
-#define NODE_INPUT_COUNT 	3
-
-
 namespace mu
 {
 
@@ -34,14 +31,14 @@ namespace mu
 	//---------------------------------------------------------------------------------------------
 	int NodeImageLayerColour::GetInputCount() const
 	{
-		return NODE_INPUT_COUNT;
+		return 3;
 	}
 
 
 	//---------------------------------------------------------------------------------------------
 	Node* NodeImageLayerColour::GetInputNode( int i ) const
 	{
-		check( i>=0 && i<NODE_INPUT_COUNT );
+		check( i>=0 && i< GetInputCount());
 
 		Node* pResult = 0;
 
@@ -59,7 +56,7 @@ namespace mu
 	//---------------------------------------------------------------------------------------------
 	void NodeImageLayerColour::SetInputNode( int i, NodePtr pNode )
 	{
-		check( i>=0 && i<NODE_INPUT_COUNT );
+		check( i>=0 && i< GetInputCount());
 
 		switch (i)
 		{
@@ -124,5 +121,3 @@ namespace mu
 
 }
 
-#undef NODE_INPUT_COUNT
-
diff --git a/Engine/Plugins/Experimental/Mutable/Source/MutableTools/Private/MuT/NodeImageLuminance.cpp b/Engine/Plugins/Experimental/Mutable/Source/MutableTools/Private/MuT/NodeImageLuminance.cpp
--- a/Engine/Plugins/Experimental/Mutable/Source/MutableTools/Private/MuT/NodeImageLuminance.cpp
+++ b/Engine/Plugins/Experimental/Mutable/Source/MutableTools/Private/MuT/NodeImageLuminance.cpp
@@ -8,9 +8,6 @@
 #include "MuT/NodePrivate.h"
 
 
-#define NODE_INPUT_COUNT 	1
-
-
 namespace mu
 {
 
@@ -33,14 +30,14 @@ namespace mu
 	//---------------------------------------------------------------------------------------------
 	int NodeImageLuminance::GetInputCount() const
 	{
-		return NODE_INPUT_COUNT;
+		return 1;
 	}
 
 
 	//---------------------------------------------------------------------------------------------
 	Node* NodeImageLuminance::GetInputNode( int i ) const
 	{
-		check( i>=0 && i<NODE_INPUT_COUNT );
+		check( i>=0 && i< GetInputCount());
 
 		Node* pResult = 0;
 
@@ -56,7 +53,7 @@ namespace mu
 	//---------------------------------------------------------------------------------------------
 	void NodeImageLuminance::SetInputNode( int i, NodePtr pNode )
 	{
-		check( i>=0 && i<NODE_INPUT_COUNT );
+		check( i>=0 && i< GetInputCount());
 
 		switch (i)
 		{
@@ -83,5 +80,3 @@ namespace mu
 
 
 }
-
-#undef NODE_INPUT_COUNT
diff --git a/Engine/Plugins/Experimental/Mutable/Source/MutableTools/Private/MuT/NodeImageMipmap.cpp b/Engine/Plugins/Experimental/Mutable/Source/MutableTools/Private/MuT/NodeImageMipmap.cpp
--- a/Engine/Plugins/Experimental/Mutable/Source/MutableTools/Private/MuT/NodeImageMipmap.cpp
+++ b/Engine/Plugins/Experimental/Mutable/Source/MutableTools/Private/MuT/NodeImageMipmap.cpp
@@ -9,9 +9,6 @@
 #include "MuT/NodePrivate.h"
 
 
-#define NODE_INPUT_COUNT 	1
-
-
 namespace mu
 {
 
@@ -34,14 +31,14 @@ namespace mu
 	//---------------------------------------------------------------------------------------------
     int NodeImageMipmap::GetInputCount() const
 	{
-		return NODE_INPUT_COUNT;
+		return 1;
 	}
 
 
 	//---------------------------------------------------------------------------------------------
     Node* NodeImageMipmap::GetInputNode( int i ) const
 	{
-		check( i>=0 && i<NODE_INPUT_COUNT );
+		check( i>=0 && i< GetInputCount());
 
 		Node* pResult = 0;
 
@@ -57,7 +54,7 @@ namespace mu
 	//---------------------------------------------------------------------------------------------
     void NodeImageMipmap::SetInputNode( int i, NodePtr pNode )
 	{
-		check( i>=0 && i<NODE_INPUT_COUNT );
+		check( i>=0 && i< GetInputCount());
 
 		switch (i)
 		{
@@ -91,5 +88,3 @@ namespace mu
 	}
 
 }
-
-#undef NODE_INPUT_COUNT
diff --git a/Engine/Plugins/Experimental/Mutable/Source/MutableTools/Private/MuT/NodeImageMultiLayer.cpp b/Engine/Plugins/Experimental/Mutable/Source/MutableTools/Private/MuT/NodeImageMultiLayer.cpp
--- a/Engine/Plugins/Experimental/Mutable/Source/MutableTools/Private/MuT/NodeImageMultiLayer.cpp
+++ b/Engine/Plugins/Experimental/Mutable/Source/MutableTools/Private/MuT/NodeImageMultiLayer.cpp
@@ -8,9 +8,6 @@
 #include "MuT/NodePrivate.h"
 
 
-#define NODE_INPUT_COUNT 	4
-
-
 namespace mu
 {
 
@@ -33,14 +30,14 @@ namespace mu
 	//---------------------------------------------------------------------------------------------
     int NodeImageMultiLayer::GetInputCount() const
 	{
-		return NODE_INPUT_COUNT;
+		return 4;
 	}
 
 
 	//---------------------------------------------------------------------------------------------
     Node* NodeImageMultiLayer::GetInputNode( int i ) const
 	{
-		check( i>=0 && i<NODE_INPUT_COUNT );
+		check( i>=0 && i< GetInputCount());
 
 		Node* pResult = 0;
 
@@ -59,7 +56,7 @@ namespace mu
 	//---------------------------------------------------------------------------------------------
     void NodeImageMultiLayer::SetInputNode( int i, NodePtr pNode )
 	{
-		check( i>=0 && i<NODE_INPUT_COUNT );
+		check( i>=0 && i< GetInputCount());
 
 		switch (i)
 		{
@@ -144,5 +141,3 @@ namespace mu
 
 
 }
-
-#undef NODE_INPUT_COUNT
diff --git a/Engine/Plugins/Experimental/Mutable/Source/MutableTools/Private/MuT/NodeImagePlainColour.cpp b/Engine/Plugins/Experimental/Mutable/Source/MutableTools/Private/MuT/NodeImagePlainColour.cpp
--- a/Engine/Plugins/Experimental/Mutable/Source/MutableTools/Private/MuT/NodeImagePlainColour.cpp
+++ b/Engine/Plugins/Experimental/Mutable/Source/MutableTools/Private/MuT/NodeImagePlainColour.cpp
@@ -9,9 +9,6 @@
 #include "MuT/NodePrivate.h"
 
 
-#define NODE_INPUT_COUNT 	1
-
-
 namespace mu
 {
 
@@ -32,7 +29,7 @@ namespace mu
 	//---------------------------------------------------------------------------------------------
 	void NodeImagePlainColour::SetInputNode( int i, NodePtr pNode )
 	{
-		check( i>=0 && i<NODE_INPUT_COUNT );
+		check( i>=0 && i< GetInputCount());
 
 		switch (i)
 		{
@@ -44,14 +41,14 @@ namespace mu
 	//---------------------------------------------------------------------------------------------
 	int NodeImagePlainColour::GetInputCount() const
 	{
-		return NODE_INPUT_COUNT;
+		return 1;
 	}
 
 
 	//---------------------------------------------------------------------------------------------
 	Node* NodeImagePlainColour::GetInputNode( int i ) const
 	{
-		check( i>=0 && i<NODE_INPUT_COUNT );
+		check( i>=0 && i< GetInputCount());
 
 		Node* pResult = 0;
 
@@ -89,5 +86,3 @@ namespace mu
 
 }
 
-#undef NODE_INPUT_COUNT
-
diff --git a/Engine/Plugins/Experimental/Mutable/Source/MutableTools/Private/MuT/NodeImageResize.cpp b/Engine/Plugins/Experimental/Mutable/Source/MutableTools/Private/MuT/NodeImageResize.cpp
--- a/Engine/Plugins/Experimental/Mutable/Source/MutableTools/Private/MuT/NodeImageResize.cpp
+++ b/Engine/Plugins/Experimental/Mutable/Source/MutableTools/Private/MuT/NodeImageResize.cpp
@@ -8,9 +8,6 @@
 #include "MuT/NodePrivate.h"
 
 
-#define NODE_INPUT_COUNT 	1
-
-
 namespace mu
 {
 
@@ -30,7 +27,7 @@ namespace mu
 	//---------------------------------------------------------------------------------------------
 	void NodeImageResize::SetInputNode( int i, NodePtr pNode )
 	{
-		check( i>=0 && i<NODE_INPUT_COUNT );
+		check( i>=0 && i< GetInputCount());
 
 		switch (i)
 		{
@@ -42,14 +39,14 @@ namespace mu
 	//---------------------------------------------------------------------------------------------
 	int NodeImageResize::GetInputCount() const
 	{
-		return NODE_INPUT_COUNT;
+		return 1;
 	}
 
 
 	//---------------------------------------------------------------------------------------------
 	Node* NodeImageResize::GetInputNode( int i ) const
 	{
-		check( i>=0 && i<NODE_INPUT_COUNT );
+		check( i>=0 && i< GetInputCount());
 
 		Node* pResult = 0;
 
@@ -93,5 +90,3 @@ namespace mu
 	}
 
 }
-
-#undef NODE_INPUT_COUNT
diff --git a/Engine/Plugins/Experimental/Mutable/Source/MutableTools/Private/MuT/NodeImageSaturate.cpp b/Engine/Plugins/Experimental/Mutable/Source/MutableTools/Private/MuT/NodeImageSaturate.cpp
--- a/Engine/Plugins/Experimental/Mutable/Source/MutableTools/Private/MuT/NodeImageSaturate.cpp
+++ b/Engine/Plugins/Experimental/Mutable/Source/MutableTools/Private/MuT/NodeImageSaturate.cpp
@@ -9,9 +9,6 @@
 #include "MuT/NodeScalar.h"
 
 
-#define NODE_INPUT_COUNT 	2
-
-
 namespace mu
 {
 
@@ -34,14 +31,14 @@ namespace mu
 	//---------------------------------------------------------------------------------------------
 	int NodeImageSaturate::GetInputCount() const
 	{
-		return NODE_INPUT_COUNT;
+		return 2;
 	}
 
 
 	//---------------------------------------------------------------------------------------------
 	Node* NodeImageSaturate::GetInputNode( int i ) const
 	{
-		check( i>=0 && i<NODE_INPUT_COUNT );
+		check( i>=0 && i< GetInputCount());
 
 		Node* pResult = 0;
 
@@ -58,7 +55,7 @@ namespace mu
 	//---------------------------------------------------------------------------------------------
 	void NodeImageSaturate::SetInputNode( int i, NodePtr pNode )
 	{
-		check( i>=0 && i<NODE_INPUT_COUNT );
+		check( i>=0 && i< GetInputCount());
 
 		switch (i)
 		{
@@ -101,5 +98,3 @@ namespace mu
 
 }
 
-#undef NODE_INPUT_COUNT
-
diff --git a/Engine/Plugins/Experimental/Mutable/Source/MutableTools/Private/MuT/NodeImageSelectColour.cpp b/Engine/Plugins/Experimental/Mutable/Source/MutableTools/Private/MuT/NodeImageSelectColour.cpp
--- a/Engine/Plugins/Experimental/Mutable/Source/MutableTools/Private/MuT/NodeImageSelectColour.cpp
+++ b/Engine/Plugins/Experimental/Mutable/Source/MutableTools/Private/MuT/NodeImageSelectColour.cpp
@@ -9,9 +9,6 @@
 #include "MuT/NodePrivate.h"
 
 
-#define NODE_INPUT_COUNT 	2
-
-
 namespace mu
 {
 
@@ -33,14 +30,14 @@ namespace mu
 	//---------------------------------------------------------------------------------------------
 	int NodeImageSelectColour::GetInputCount() const
 	{
-		return NODE_INPUT_COUNT;
+		return 2;
 	}
 
 
 	//---------------------------------------------------------------------------------------------
 	Node* NodeImageSelectColour::GetInputNode( int i ) const
 	{
-		check( i>=0 && i<NODE_INPUT_COUNT );
+		check( i>=0 && i< GetInputCount());
 
 		Node* pResult = 0;
 
@@ -57,7 +54,7 @@ namespace mu
 	//---------------------------------------------------------------------------------------------
 	void NodeImageSelectColour::SetInputNode( int i, NodePtr pNode )
 	{
-		check( i>=0 && i<NODE_INPUT_COUNT );
+		check( i>=0 && i< GetInputCount());
 
 		switch (i)
 		{
@@ -99,5 +96,3 @@ namespace mu
 
 
 }
-
-#undef NODE_INPUT_COUNT
diff --git a/Engine/Plugins/Experimental/Mutable/Source/MutableTools/Private/MuT/NodeImageSwitch.cpp b/Engine/Plugins/Experimental/Mutable/Source/MutableTools/Private/MuT/NodeImageSwitch.cpp
--- a/Engine/Plugins/Experimental/Mutable/Source/MutableTools/Private/MuT/NodeImageSwitch.cpp
+++ b/Engine/Plugins/Experimental/Mutable/Source/MutableTools/Private/MuT/NodeImageSwitch.cpp
@@ -35,7 +35,7 @@ namespace mu
 	//---------------------------------------------------------------------------------------------
 	int NodeImageSwitch::GetInputCount() const
 	{
-		return 1 + (int)m_pD->m_options.size();
+		return 1 + (int)m_pD->m_options.Num();
 	}
 
 
@@ -98,21 +98,21 @@ namespace mu
 	//---------------------------------------------------------------------------------------------
 	void NodeImageSwitch::SetOptionCount( int t )
 	{
-		m_pD->m_options.resize(t);
+		m_pD->m_options.SetNum(t);
 	}
 
 
 	//---------------------------------------------------------------------------------------------
 	int NodeImageSwitch::GetOptionCount() const 
 	{
-		return int(m_pD->m_options.size());
+		return int(m_pD->m_options.Num());
 	}
 
 
 	//---------------------------------------------------------------------------------------------
 	NodeImagePtr NodeImageSwitch::GetOption( int t ) const
 	{
-		check( t>=0 && t<(int)m_pD->m_options.size() );
+		check( t>=0 && t<(int)m_pD->m_options.Num() );
 		return m_pD->m_options[t].get();
 	}
 
@@ -120,7 +120,7 @@ namespace mu
 	//---------------------------------------------------------------------------------------------
 	void NodeImageSwitch::SetOption( int t, NodeImagePtr pNode )
 	{
-		check( t>=0 && t<(int)m_pD->m_options.size() );
+		check( t>=0 && t<(int)m_pD->m_options.Num() );
 		m_pD->m_options[t] = pNode;
 	}
 
diff --git a/Engine/Plugins/Experimental/Mutable/Source/MutableTools/Private/MuT/NodeImageSwitchPrivate.h b/Engine/Plugins/Experimental/Mutable/Source/MutableTools/Private/MuT/NodeImageSwitchPrivate.h
--- a/Engine/Plugins/Experimental/Mutable/Source/MutableTools/Private/MuT/NodeImageSwitchPrivate.h
+++ b/Engine/Plugins/Experimental/Mutable/Source/MutableTools/Private/MuT/NodeImageSwitchPrivate.h
@@ -22,7 +22,7 @@ namespace mu
 		static NODE_TYPE s_type;
 
 		NodeScalarPtr m_pParameter;
-		vector<NodeImagePtr> m_options;
+		TArray<NodeImagePtr> m_options;
 
 		//!
 		void Serialise( OutputArchive& arch ) const
diff --git a/Engine/Plugins/Experimental/Mutable/Source/MutableTools/Private/MuT/NodeImageSwizzle.cpp b/Engine/Plugins/Experimental/Mutable/Source/MutableTools/Private/MuT/NodeImageSwizzle.cpp
--- a/Engine/Plugins/Experimental/Mutable/Source/MutableTools/Private/MuT/NodeImageSwizzle.cpp
+++ b/Engine/Plugins/Experimental/Mutable/Source/MutableTools/Private/MuT/NodeImageSwizzle.cpp
@@ -13,9 +13,6 @@
 #include <utility>
 
 
-#define NODE_INPUT_COUNT 	4
-
-
 namespace mu
 {
 
@@ -38,7 +35,7 @@ namespace mu
 	//---------------------------------------------------------------------------------------------
 	int NodeImageSwizzle::GetInputCount() const
 	{
-		return (int)m_pD->m_sources.size();
+		return (int)m_pD->m_sources.Num();
 	}
 
 
@@ -79,15 +76,15 @@ namespace mu
 		m_pD->m_format = format;
 
 		int channelCount = GetImageFormatData( format ).m_channels;
-		m_pD->m_sources.resize( channelCount );
-		m_pD->m_sourceChannels.resize( channelCount );
+		m_pD->m_sources.SetNum( channelCount );
+		m_pD->m_sourceChannels.SetNum( channelCount );
 	}
 
 
 	//---------------------------------------------------------------------------------------------
 	NodeImagePtr NodeImageSwizzle::GetSource( int t ) const
 	{
-		check( t>=0 && t<(int)m_pD->m_sources.size() );
+		check( t>=0 && t<(int)m_pD->m_sources.Num() );
 		return m_pD->m_sources[t].get();
 	}
 
@@ -95,7 +92,7 @@ namespace mu
 	//---------------------------------------------------------------------------------------------
 	void NodeImageSwizzle::SetSource( int t, NodeImagePtr pNode )
 	{
-		check( t>=0 && t<(int)m_pD->m_sources.size() );
+		check( t>=0 && t<(int)m_pD->m_sources.Num() );
 		m_pD->m_sources[t] = pNode;
 	}
 
@@ -103,7 +100,7 @@ namespace mu
 	//---------------------------------------------------------------------------------------------
 	int NodeImageSwizzle::GetSourceChannel( int t ) const
 	{
-		check( t>=0 && t<(int)m_pD->m_sourceChannels.size() );
+		check( t>=0 && t<(int)m_pD->m_sourceChannels.Num() );
 		return m_pD->m_sourceChannels[t];
 	}
 
@@ -111,12 +108,10 @@ namespace mu
 	//---------------------------------------------------------------------------------------------
 	void NodeImageSwizzle::SetSourceChannel( int t, int channel )
 	{
-		check( t>=0 && t<(int)m_pD->m_sourceChannels.size() );
+		check( t>=0 && t<(int)m_pD->m_sourceChannels.Num() );
 		m_pD->m_sourceChannels[t] = channel;
 	}
 
 
 
 }
-
-#undef NODE_INPUT_COUNT
diff --git a/Engine/Plugins/Experimental/Mutable/Source/MutableTools/Private/MuT/NodeImageSwizzlePrivate.h b/Engine/Plugins/Experimental/Mutable/Source/MutableTools/Private/MuT/NodeImageSwizzlePrivate.h
--- a/Engine/Plugins/Experimental/Mutable/Source/MutableTools/Private/MuT/NodeImageSwizzlePrivate.h
+++ b/Engine/Plugins/Experimental/Mutable/Source/MutableTools/Private/MuT/NodeImageSwizzlePrivate.h
@@ -22,8 +22,8 @@ namespace mu
 		static NODE_TYPE s_type;
 
 		EImageFormat m_format;
-		vector<NodeImagePtr> m_sources;
-		vector<int> m_sourceChannels;
+		TArray<NodeImagePtr> m_sources;
+		TArray<int> m_sourceChannels;
 
 		//!
 		void Serialise( OutputArchive& arch ) const
diff --git a/Engine/Plugins/Experimental/Mutable/Source/MutableTools/Private/MuT/NodeImageTransform.cpp b/Engine/Plugins/Experimental/Mutable/Source/MutableTools/Private/MuT/NodeImageTransform.cpp
--- a/Engine/Plugins/Experimental/Mutable/Source/MutableTools/Private/MuT/NodeImageTransform.cpp
+++ b/Engine/Plugins/Experimental/Mutable/Source/MutableTools/Private/MuT/NodeImageTransform.cpp
@@ -9,9 +9,6 @@
 #include "MuT/NodeScalar.h"
 
 
-#define NODE_INPUT_COUNT 	6
-
-
 namespace mu
 {
 
@@ -31,7 +28,7 @@ namespace mu
 	//---------------------------------------------------------------------------------------------
 	void NodeImageTransform::SetInputNode( int i, NodePtr pNode )
 	{
-		check( i>=0 && i<NODE_INPUT_COUNT );
+		check( i>=0 && i< GetInputCount());
 
 		switch (i)
 		{
@@ -48,14 +45,14 @@ namespace mu
 	//---------------------------------------------------------------------------------------------
 	int NodeImageTransform::GetInputCount() const
 	{
-		return NODE_INPUT_COUNT;
+		return 6;
 	}
 
 
 	//---------------------------------------------------------------------------------------------
 	Node* NodeImageTransform::GetInputNode( int i ) const
 	{
-		check( i>=0 && i<NODE_INPUT_COUNT );
+		check( i>=0 && i< GetInputCount());
 
 		Node* pResult = 0;
 
@@ -149,5 +146,3 @@ namespace mu
 
 }
 
-#undef NODE_INPUT_COUNT
-
diff --git a/Engine/Plugins/Experimental/Mutable/Source/MutableTools/Private/MuT/NodeImageVariation.cpp b/Engine/Plugins/Experimental/Mutable/Source/MutableTools/Private/MuT/NodeImageVariation.cpp
--- a/Engine/Plugins/Experimental/Mutable/Source/MutableTools/Private/MuT/NodeImageVariation.cpp
+++ b/Engine/Plugins/Experimental/Mutable/Source/MutableTools/Private/MuT/NodeImageVariation.cpp
@@ -30,7 +30,7 @@ namespace mu
 
 
     //---------------------------------------------------------------------------------------------
-    int NodeImageVariation::GetInputCount() const { return 1 + int( m_pD->m_variations.size() ); }
+    int NodeImageVariation::GetInputCount() const { return 1 + int( m_pD->m_variations.Num() ); }
 
 
     //---------------------------------------------------------------------------------------------
@@ -44,11 +44,11 @@ namespace mu
         }
         i -= 1;
 
-        if ( i < int( m_pD->m_variations.size() ) )
+        if ( i < int( m_pD->m_variations.Num() ) )
         {
             return m_pD->m_variations[i].m_image.get();
         }
-        i -= int( m_pD->m_variations.size() );
+        i -= int( m_pD->m_variations.Num() );
 
         return nullptr;
     }
@@ -66,13 +66,13 @@ namespace mu
         }
 
         i -= 1;
-        if ( i < int( m_pD->m_variations.size() ) )
+        if ( i < int( m_pD->m_variations.Num() ) )
         {
 
             m_pD->m_variations[i].m_image = dynamic_cast<NodeImage*>( pNode.get() );
             return;
         }
-        i -= (int)m_pD->m_variations.size();
+        i -= (int)m_pD->m_variations.Num();
     }
 
 
@@ -83,20 +83,20 @@ namespace mu
 
 
     //---------------------------------------------------------------------------------------------
-    int NodeImageVariation::GetVariationCount() const { return int( m_pD->m_variations.size() ); }
+    int NodeImageVariation::GetVariationCount() const { return int( m_pD->m_variations.Num() ); }
 
 
     //---------------------------------------------------------------------------------------------
     void NodeImageVariation::SetVariationCount( int num )
     {
         check( num >= 0 );
-        m_pD->m_variations.resize( num );
+        m_pD->m_variations.SetNum( num );
     }
 
     //---------------------------------------------------------------------------------------------
     void NodeImageVariation::SetVariationTag( int index, const char* strTag )
     {
-        check( index >= 0 && index < (int)m_pD->m_variations.size() );
+        check( index >= 0 && index < (int)m_pD->m_variations.Num() );
         check( strTag );
 
         if ( strTag )
@@ -113,7 +113,7 @@ namespace mu
     //---------------------------------------------------------------------------------------------
     void NodeImageVariation::SetVariationImage( int index, NodeImage* pNode )
     {
-        check( index >= 0 && index < (int)m_pD->m_variations.size() );
+        check( index >= 0 && index < (int)m_pD->m_variations.Num() );
 
         m_pD->m_variations[index].m_image = pNode;
     }
diff --git a/Engine/Plugins/Experimental/Mutable/Source/MutableTools/Private/MuT/NodeImageVariationPrivate.h b/Engine/Plugins/Experimental/Mutable/Source/MutableTools/Private/MuT/NodeImageVariationPrivate.h
--- a/Engine/Plugins/Experimental/Mutable/Source/MutableTools/Private/MuT/NodeImageVariationPrivate.h
+++ b/Engine/Plugins/Experimental/Mutable/Source/MutableTools/Private/MuT/NodeImageVariationPrivate.h
@@ -52,7 +52,7 @@ namespace mu
             }
         };
 
-        vector<VARIATION> m_variations;
+		TArray<VARIATION> m_variations;
 
         //!
         void Serialise( OutputArchive& arch ) const
diff --git a/Engine/Plugins/Experimental/Mutable/Source/MutableTools/Private/MuT/NodeLOD.cpp b/Engine/Plugins/Experimental/Mutable/Source/MutableTools/Private/MuT/NodeLOD.cpp
--- a/Engine/Plugins/Experimental/Mutable/Source/MutableTools/Private/MuT/NodeLOD.cpp
+++ b/Engine/Plugins/Experimental/Mutable/Source/MutableTools/Private/MuT/NodeLOD.cpp
@@ -38,7 +38,7 @@ namespace mu
 	//---------------------------------------------------------------------------------------------
 	int NodeLOD::GetInputCount() const
 	{
-		return int(m_pD->m_components.size()+m_pD->m_modifiers.size());
+		return m_pD->m_components.Num()+m_pD->m_modifiers.Num();
 	}
 
 
@@ -47,14 +47,14 @@ namespace mu
 	{
         check(i >= 0 && i < GetInputCount());
 
-        if (i < int(m_pD->m_components.size()))
+        if (i < int(m_pD->m_components.Num()))
         {
             return m_pD->m_components[i].get();
 		}
 
-        i -= int(m_pD->m_components.size());
+        i -= int(m_pD->m_components.Num());
 
-        if (i < int(m_pD->m_modifiers.size()))
+        if (i < int(m_pD->m_modifiers.Num()))
         {
             return m_pD->m_modifiers[i].get();
         }
@@ -68,14 +68,14 @@ namespace mu
 	{
         check(i >= 0 && i < GetInputCount());
 
-        if (i<int(m_pD->m_components.size()))
+        if (i<int(m_pD->m_components.Num()))
 		{
             m_pD->m_components[i] = dynamic_cast<NodeComponent*>(pNode.get());
             return;
         }
 
-        i -= int(m_pD->m_components.size());
-        if (i < int(m_pD->m_modifiers.size()))
+        i -= int(m_pD->m_components.Num());
+        if (i < int(m_pD->m_modifiers.Num()))
         {
             m_pD->m_modifiers[i] = dynamic_cast<NodeModifier*>(pNode.get());
         }
@@ -87,7 +87,7 @@ namespace mu
 	//---------------------------------------------------------------------------------------------
 	int NodeLOD::GetComponentCount() const
 	{
-		return (int)m_pD->m_components.size();
+		return (int)m_pD->m_components.Num();
 	}
 
 
@@ -95,14 +95,14 @@ namespace mu
 	void NodeLOD::SetComponentCount(int num)
 	{
 		check(num >= 0);
-		m_pD->m_components.resize(num);
+		m_pD->m_components.SetNum(num);
 	}
 
 
 	//---------------------------------------------------------------------------------------------
 	NodeComponentPtr NodeLOD::GetComponent(int index) const
 	{
-		check(index >= 0 && index < (int)m_pD->m_components.size());
+		check(index >= 0 && index < m_pD->m_components.Num());
 
 		return m_pD->m_components[index].get();
 	}
@@ -111,7 +111,7 @@ namespace mu
 	//---------------------------------------------------------------------------------------------
 	void NodeLOD::SetComponent(int index, NodeComponentPtr pComponent)
 	{
-		check(index >= 0 && index < (int)m_pD->m_components.size());
+		check(index >= 0 && index < m_pD->m_components.Num());
 
 		m_pD->m_components[index] = pComponent;
 	}
@@ -120,7 +120,7 @@ namespace mu
 	//---------------------------------------------------------------------------------------------
 	int NodeLOD::GetModifierCount() const
 	{
-		return (int)m_pD->m_modifiers.size();
+		return (int)m_pD->m_modifiers.Num();
 	}
 
 
@@ -128,14 +128,14 @@ namespace mu
 	void NodeLOD::SetModifierCount(int num)
 	{
 		check(num >= 0);
-		m_pD->m_modifiers.resize(num);
+		m_pD->m_modifiers.SetNum(num);
 	}
 
 
 	//---------------------------------------------------------------------------------------------
 	void NodeLOD::SetModifier(int index, NodeModifierPtr pModifier)
 	{
-		check(index >= 0 && index < (int)m_pD->m_modifiers.size());
+		check(index >= 0 && index < m_pD->m_modifiers.Num());
 
 		m_pD->m_modifiers[index] = pModifier;
 	}
diff --git a/Engine/Plugins/Experimental/Mutable/Source/MutableTools/Private/MuT/NodeLODPrivate.h b/Engine/Plugins/Experimental/Mutable/Source/MutableTools/Private/MuT/NodeLODPrivate.h
--- a/Engine/Plugins/Experimental/Mutable/Source/MutableTools/Private/MuT/NodeLODPrivate.h
+++ b/Engine/Plugins/Experimental/Mutable/Source/MutableTools/Private/MuT/NodeLODPrivate.h
@@ -25,8 +25,8 @@ namespace mu
 
 		static NODE_TYPE s_type;
 
-		vector<NodeComponentPtr> m_components;
-		vector<NodeModifierPtr> m_modifiers;
+		TArray<NodeComponentPtr> m_components;
+		TArray<NodeModifierPtr> m_modifiers;
 
 		//!
 		void Serialise( OutputArchive& arch ) const
diff --git a/Engine/Plugins/Experimental/Mutable/Source/MutableTools/Private/MuT/NodeMeshClipMorphPlane.cpp b/Engine/Plugins/Experimental/Mutable/Source/MutableTools/Private/MuT/NodeMeshClipMorphPlane.cpp
--- a/Engine/Plugins/Experimental/Mutable/Source/MutableTools/Private/MuT/NodeMeshClipMorphPlane.cpp
+++ b/Engine/Plugins/Experimental/Mutable/Source/MutableTools/Private/MuT/NodeMeshClipMorphPlane.cpp
@@ -15,8 +15,6 @@
 #include <utility>
 
 
-#define NODE_INPUT_COUNT 	1
-
 
 namespace mu
 {
@@ -40,14 +38,14 @@ namespace mu
 	//---------------------------------------------------------------------------------------------
     int NodeMeshClipMorphPlane::GetInputCount() const
 	{
-		return NODE_INPUT_COUNT;
+		return 1;
 	}
 
 
 	//---------------------------------------------------------------------------------------------
     Node* NodeMeshClipMorphPlane::GetInputNode( int i ) const
 	{
-		check( i>=0 && i<NODE_INPUT_COUNT );
+		check( i>=0 && i< GetInputCount());
         (void)i;
         return m_pD->m_pSource.get();
 	}
@@ -56,7 +54,7 @@ namespace mu
 	//---------------------------------------------------------------------------------------------
     void NodeMeshClipMorphPlane::SetInputNode( int i, NodePtr pNode )
 	{
-		check( i>=0 && i<NODE_INPUT_COUNT );
+		check( i>=0 && i< GetInputCount());
 		if (i==0)
 		{
 			m_pD->m_pSource = dynamic_cast<NodeMesh*>( pNode.get() );
@@ -67,7 +65,7 @@ namespace mu
 	//---------------------------------------------------------------------------------------------
     const char* NodeMeshClipMorphPlane::GetInputName( int i ) const
 	{
-		check( i>=0 && i<NODE_INPUT_COUNT );
+		check( i>=0 && i< GetInputCount());
         (void)i;
         return "Source";
 	}
@@ -76,7 +74,7 @@ namespace mu
 	//---------------------------------------------------------------------------------------------
     const NODE_TYPE* NodeMeshClipMorphPlane::GetInputType( int i ) const
 	{
-		check( i>=0 && i<NODE_INPUT_COUNT );
+		check( i>=0 && i< GetInputCount());
         (void)i;
         return NodeMesh::GetStaticType();
 	}
@@ -156,8 +154,6 @@ namespace mu
 	//---------------------------------------------------------------------------------------------
 	void NodeMeshClipMorphPlane::AddTag(const char* tagName)
 	{
-		m_pD->m_tags.push_back(tagName);
+		m_pD->m_tags.Add(tagName);
 	}
 }
-
-#undef NODE_INPUT_COUNT
diff --git a/Engine/Plugins/Experimental/Mutable/Source/MutableTools/Private/MuT/NodeMeshClipMorphPlanePrivate.h b/Engine/Plugins/Experimental/Mutable/Source/MutableTools/Private/MuT/NodeMeshClipMorphPlanePrivate.h
--- a/Engine/Plugins/Experimental/Mutable/Source/MutableTools/Private/MuT/NodeMeshClipMorphPlanePrivate.h
+++ b/Engine/Plugins/Experimental/Mutable/Source/MutableTools/Private/MuT/NodeMeshClipMorphPlanePrivate.h
@@ -67,7 +67,7 @@ namespace mu
 		vec3f m_selectionBoxRadius;
 		string m_vertexSelectionBone;
 
-		vector<mu::string> m_tags;
+		TArray<mu::string> m_tags;
 
 		// Max distance a vertex can have to the bone in order to be affected. A negative value
 		// means no limit.
diff --git a/Engine/Plugins/Experimental/Mutable/Source/MutableTools/Private/MuT/NodeMeshClipWithMesh.cpp b/Engine/Plugins/Experimental/Mutable/Source/MutableTools/Private/MuT/NodeMeshClipWithMesh.cpp
--- a/Engine/Plugins/Experimental/Mutable/Source/MutableTools/Private/MuT/NodeMeshClipWithMesh.cpp
+++ b/Engine/Plugins/Experimental/Mutable/Source/MutableTools/Private/MuT/NodeMeshClipWithMesh.cpp
@@ -14,9 +14,6 @@
 #include <utility>
 
 
-#define NODE_INPUT_COUNT 	1
-
-
 namespace mu
 {
 
@@ -39,14 +36,14 @@ namespace mu
 	//---------------------------------------------------------------------------------------------
     int NodeMeshClipWithMesh::GetInputCount() const
 	{
-		return NODE_INPUT_COUNT;
+		return 1;
 	}
 
 
 	//---------------------------------------------------------------------------------------------
     Node* NodeMeshClipWithMesh::GetInputNode( int i ) const
 	{
-		check( i>=0 && i<NODE_INPUT_COUNT );
+		check( i>=0 && i< GetInputCount());
         (void)i;
         return m_pD->m_pSource.get();
 	}
@@ -55,7 +52,7 @@ namespace mu
 	//---------------------------------------------------------------------------------------------
     void NodeMeshClipWithMesh::SetInputNode( int i, NodePtr pNode )
 	{
-		check( i>=0 && i<NODE_INPUT_COUNT );
+		check( i>=0 && i< GetInputCount());
 		if (i==0)
 		{
 			m_pD->m_pSource = dynamic_cast<NodeMesh*>( pNode.get() );
@@ -66,7 +63,7 @@ namespace mu
 	//---------------------------------------------------------------------------------------------
     const char* NodeMeshClipWithMesh::GetInputName( int i ) const
 	{
-		check( i>=0 && i<NODE_INPUT_COUNT );
+		check( i>=0 && i< GetInputCount());
         (void)i;
         return "Source";
 	}
@@ -75,7 +72,7 @@ namespace mu
 	//---------------------------------------------------------------------------------------------
     const NODE_TYPE* NodeMeshClipWithMesh::GetInputType( int i ) const
 	{
-		check( i>=0 && i<NODE_INPUT_COUNT );
+		check( i>=0 && i< GetInputCount());
         (void)i;
         return NodeMesh::GetStaticType();
 	}
@@ -123,8 +120,6 @@ namespace mu
 	//---------------------------------------------------------------------------------------------
 	void NodeMeshClipWithMesh::AddTag(const char* tagName)
 	{
-		m_pD->m_tags.push_back(tagName);
+		m_pD->m_tags.Add(tagName);
 	}
 }
-
-#undef NODE_INPUT_COUNT
diff --git a/Engine/Plugins/Experimental/Mutable/Source/MutableTools/Private/MuT/NodeMeshClipWithMeshPrivate.h b/Engine/Plugins/Experimental/Mutable/Source/MutableTools/Private/MuT/NodeMeshClipWithMeshPrivate.h
--- a/Engine/Plugins/Experimental/Mutable/Source/MutableTools/Private/MuT/NodeMeshClipWithMeshPrivate.h
+++ b/Engine/Plugins/Experimental/Mutable/Source/MutableTools/Private/MuT/NodeMeshClipWithMeshPrivate.h
@@ -22,21 +22,17 @@ namespace mu
 
 	public:
 
-		Private()
-		{
-		}
-
 		static NODE_TYPE s_type;
 
 		NodeMeshPtr m_pSource;
 		NodeMeshPtr m_pClipMesh;
 
-		vector<mu::string> m_tags;
+		TArray<mu::string> m_tags;
 
 		//!
 		void Serialise( OutputArchive& arch ) const
 		{
-            uint32_t ver = 1;
+            uint32 ver = 1;
 			arch << ver;
 
 			arch << m_pSource;
@@ -47,7 +43,7 @@ namespace mu
 		//!
 		void Unserialise( InputArchive& arch )
 		{
-            uint32_t ver;
+            uint32 ver;
 			arch >> ver;
 
 			arch >> m_pSource;
diff --git a/Engine/Plugins/Experimental/Mutable/Source/MutableTools/Private/MuT/NodeMeshConstant.cpp b/Engine/Plugins/Experimental/Mutable/Source/MutableTools/Private/MuT/NodeMeshConstant.cpp
--- a/Engine/Plugins/Experimental/Mutable/Source/MutableTools/Private/MuT/NodeMeshConstant.cpp
+++ b/Engine/Plugins/Experimental/Mutable/Source/MutableTools/Private/MuT/NodeMeshConstant.cpp
@@ -35,14 +35,14 @@ namespace mu
 	//---------------------------------------------------------------------------------------------
 	int NodeMeshConstant::GetInputCount() const
 	{
-		return (int)m_pD->m_layouts.size();
+		return m_pD->m_layouts.Num();
 	}
 
 
 	//---------------------------------------------------------------------------------------------
 	Node* NodeMeshConstant::GetInputNode( int i ) const
 	{
-		check( i>=0 && i<(int)m_pD->m_layouts.size() );
+		check( i>=0 && i<m_pD->m_layouts.Num() );
 		return m_pD->m_layouts[i].get();
 	}
 
@@ -50,7 +50,7 @@ namespace mu
 	//---------------------------------------------------------------------------------------------
 	void NodeMeshConstant::SetInputNode( int i, NodePtr pNode )
 	{
-		check( i>=0 && i<(int)m_pD->m_layouts.size() );
+		check( i>=0 && i<m_pD->m_layouts.Num() );
 
 		m_pD->m_layouts[i] = dynamic_cast<NodeLayout*>( pNode.get() );
 	}
@@ -81,7 +81,7 @@ namespace mu
 	//---------------------------------------------------------------------------------------------
 	int NodeMeshConstant::GetLayoutCount() const
 	{
-		return (int)m_pD->m_layouts.size();
+		return m_pD->m_layouts.Num();
 	}
 
 
@@ -89,14 +89,14 @@ namespace mu
 	void NodeMeshConstant::SetLayoutCount( int num )
 	{
 		check( num >=0 );
-		m_pD->m_layouts.resize( num );
+		m_pD->m_layouts.SetNum( num );
 	}
 
 
 	//---------------------------------------------------------------------------------------------
 	NodeLayoutPtr NodeMeshConstant::GetLayout( int index ) const
 	{
-		check( index >=0 && index < (int)m_pD->m_layouts.size() );
+		check( index >=0 && index < m_pD->m_layouts.Num() );
 
 		return m_pD->GetLayout( index );
 	}
@@ -107,7 +107,7 @@ namespace mu
 	{
 		NodeLayoutPtr pResult;
 
-		if ( index >=0 && index < (int)m_layouts.size() )
+		if ( index >=0 && index < m_layouts.Num() )
 		{
 			pResult = m_layouts[ index ].get();
 		}
@@ -119,7 +119,7 @@ namespace mu
 	//---------------------------------------------------------------------------------------------
 	void NodeMeshConstant::SetLayout( int index, NodeLayoutPtr pLayout )
 	{
-		check( index >=0 && index < (int)m_pD->m_layouts.size() );
+		check( index >=0 && index < m_pD->m_layouts.Num() );
 
 		m_pD->m_layouts[ index ] = pLayout;
 	}
diff --git a/Engine/Plugins/Experimental/Mutable/Source/MutableTools/Private/MuT/NodeMeshConstantPrivate.h b/Engine/Plugins/Experimental/Mutable/Source/MutableTools/Private/MuT/NodeMeshConstantPrivate.h
--- a/Engine/Plugins/Experimental/Mutable/Source/MutableTools/Private/MuT/NodeMeshConstantPrivate.h
+++ b/Engine/Plugins/Experimental/Mutable/Source/MutableTools/Private/MuT/NodeMeshConstantPrivate.h
@@ -24,7 +24,7 @@ namespace mu
 
 		MeshPtr m_pValue;
 
-		vector<NodeLayoutPtr> m_layouts;
+		TArray<NodeLayoutPtr> m_layouts;
 
 		//!
 		void Serialise( OutputArchive& arch ) const
diff --git a/Engine/Plugins/Experimental/Mutable/Source/MutableTools/Private/MuT/NodeMeshFormat.cpp b/Engine/Plugins/Experimental/Mutable/Source/MutableTools/Private/MuT/NodeMeshFormat.cpp
--- a/Engine/Plugins/Experimental/Mutable/Source/MutableTools/Private/MuT/NodeMeshFormat.cpp
+++ b/Engine/Plugins/Experimental/Mutable/Source/MutableTools/Private/MuT/NodeMeshFormat.cpp
@@ -92,13 +92,6 @@ class FMeshBufferSet;
 	}
 
 
-	//---------------------------------------------------------------------------------------------
-	void NodeMeshFormat::SetRebuildTangents( bool enabled )
-	{
-		m_pD->m_rebuildTangents = enabled;
-	}
-
-
 	//---------------------------------------------------------------------------------------------
 	NodeLayoutPtr NodeMeshFormat::Private::GetLayout( int index ) const
 	{
diff --git a/Engine/Plugins/Experimental/Mutable/Source/MutableTools/Private/MuT/NodeMeshFormatPrivate.h b/Engine/Plugins/Experimental/Mutable/Source/MutableTools/Private/MuT/NodeMeshFormatPrivate.h
--- a/Engine/Plugins/Experimental/Mutable/Source/MutableTools/Private/MuT/NodeMeshFormatPrivate.h
+++ b/Engine/Plugins/Experimental/Mutable/Source/MutableTools/Private/MuT/NodeMeshFormatPrivate.h
@@ -21,11 +21,6 @@ namespace mu
 
 	public:
 
-		Private()
-		{
-			m_rebuildTangents = false;
-		}
-
 
 		static NODE_TYPE s_type;
 
@@ -38,19 +33,16 @@ namespace mu
 		FMeshBufferSet m_IndexBuffers;
 		FMeshBufferSet m_FaceBuffers;
 
-		bool m_rebuildTangents;
-
 		//!
 		void Serialise( OutputArchive& arch ) const
 		{
-            uint32_t ver = 3;
+            uint32_t ver = 4;
 			arch << ver;
 
 			arch << m_pSource;
 			arch << m_VertexBuffers;
 			arch << m_IndexBuffers;
 			arch << m_FaceBuffers;
-			arch << m_rebuildTangents;
 		}
 
 		//!
@@ -58,13 +50,17 @@ namespace mu
 		{
             uint32_t ver;
 			arch >> ver;
-            check(ver==3);
+            check(ver>=3 && ver<=4);
 
 			arch >> m_pSource;
 			arch >> m_VertexBuffers;
 			arch >> m_IndexBuffers;
 			arch >> m_FaceBuffers;
-			arch >> m_rebuildTangents;
+			if (ver == 3)
+			{
+				bool bDummy;
+				arch >> bDummy;
+			}
 		}
 
 		// NodeMesh::Private interface
diff --git a/Engine/Plugins/Experimental/Mutable/Source/MutableTools/Private/MuT/NodeMeshFragment.cpp b/Engine/Plugins/Experimental/Mutable/Source/MutableTools/Private/MuT/NodeMeshFragment.cpp
--- a/Engine/Plugins/Experimental/Mutable/Source/MutableTools/Private/MuT/NodeMeshFragment.cpp
+++ b/Engine/Plugins/Experimental/Mutable/Source/MutableTools/Private/MuT/NodeMeshFragment.cpp
@@ -118,14 +118,14 @@ namespace mu
 	//---------------------------------------------------------------------------------------------
 	void NodeMeshFragment::SetBlockCount( int t )
 	{
-		m_pD->m_blocks.resize(t);
+		m_pD->m_blocks.SetNum(t);
 	}
 
 
 	//---------------------------------------------------------------------------------------------
 	int NodeMeshFragment::GetBlock( int t ) const
 	{
-		check( t>=0 && t<(int)m_pD->m_blocks.size() );
+		check( t>=0 && t<m_pD->m_blocks.Num() );
 		return m_pD->m_blocks[t];
 	}
 
@@ -133,7 +133,7 @@ namespace mu
 	//---------------------------------------------------------------------------------------------
 	void NodeMeshFragment::SetBlock( int t, int b )
 	{
-		check( t>=0 && t<(int)m_pD->m_blocks.size() );
+		check( t>=0 && t<m_pD->m_blocks.Num() );
 		m_pD->m_blocks[t] = b;
 	}
 
diff --git a/Engine/Plugins/Experimental/Mutable/Source/MutableTools/Private/MuT/NodeMeshFragmentPrivate.h b/Engine/Plugins/Experimental/Mutable/Source/MutableTools/Private/MuT/NodeMeshFragmentPrivate.h
--- a/Engine/Plugins/Experimental/Mutable/Source/MutableTools/Private/MuT/NodeMeshFragmentPrivate.h
+++ b/Engine/Plugins/Experimental/Mutable/Source/MutableTools/Private/MuT/NodeMeshFragmentPrivate.h
@@ -29,7 +29,7 @@ namespace mu
 
 		NodeMeshPtr m_pMesh;
         int m_layoutOrGroup;
-		vector<int> m_blocks;
+		TArray<int> m_blocks;
 
         FRAGMENT_TYPE m_fragmentType;
 
diff --git a/Engine/Plugins/Experimental/Mutable/Source/MutableTools/Private/MuT/NodeMeshInterpolate.cpp b/Engine/Plugins/Experimental/Mutable/Source/MutableTools/Private/MuT/NodeMeshInterpolate.cpp
--- a/Engine/Plugins/Experimental/Mutable/Source/MutableTools/Private/MuT/NodeMeshInterpolate.cpp
+++ b/Engine/Plugins/Experimental/Mutable/Source/MutableTools/Private/MuT/NodeMeshInterpolate.cpp
@@ -15,9 +15,6 @@
 #include <utility>
 
 
-#define NODE_INPUT_COUNT 	4
-
-
 namespace mu
 {
 
@@ -40,7 +37,7 @@ namespace mu
 	//---------------------------------------------------------------------------------------------
 	int NodeMeshInterpolate::GetInputCount() const
 	{
-		return 1 + (int)m_pD->m_targets.size();
+		return 1 + m_pD->m_targets.Num();
 	}
 
 
@@ -103,7 +100,7 @@ namespace mu
 	//---------------------------------------------------------------------------------------------
 	NodeMeshPtr NodeMeshInterpolate::GetTarget( int t ) const
 	{
-		check( t>=0 && t<(int)m_pD->m_targets.size() );
+		check( t>=0 && t<(int)m_pD->m_targets.Num() );
 
 		return m_pD->m_targets[t].get();
 	}
@@ -112,7 +109,7 @@ namespace mu
 	//---------------------------------------------------------------------------------------------
 	void NodeMeshInterpolate::SetTarget( int t, NodeMeshPtr pNode )
 	{
-		check( t>=0 && t<(int)m_pD->m_targets.size() );
+		check( t>=0 && t<(int)m_pD->m_targets.Num() );
 		m_pD->m_targets[t] = pNode;
 	}
 
@@ -120,30 +117,30 @@ namespace mu
 	//---------------------------------------------------------------------------------------------
 	void NodeMeshInterpolate::SetTargetCount( int t )
 	{
-		m_pD->m_targets.resize(t);
+		m_pD->m_targets.SetNum(t);
 	}
 
 
 	//---------------------------------------------------------------------------------------------
 	int NodeMeshInterpolate::GetTargetCount() const
 	{
-		return (int)m_pD->m_targets.size();
+		return m_pD->m_targets.Num();
 	}
 
 
 	//---------------------------------------------------------------------------------------------
 	void NodeMeshInterpolate::SetChannelCount( int t )
 	{
-		m_pD->m_channels.resize(t);
+		m_pD->m_channels.SetNum(t);
 	}
 
 
 	//---------------------------------------------------------------------------------------------
 	void NodeMeshInterpolate::SetChannel( int i, MESH_BUFFER_SEMANTIC semantic, int semanticIndex )
 	{
-		check( i>=0 && i<(int)m_pD->m_channels.size() );
+		check( i>=0 && i<m_pD->m_channels.Num() );
 
-		if ( i>=0 && i<(int)m_pD->m_channels.size() )
+		if ( i>=0 && i<m_pD->m_channels.Num() )
 		{
 			m_pD->m_channels[i].semantic = semantic;
 			m_pD->m_channels[i].semanticIndex = semanticIndex;
@@ -159,7 +156,7 @@ namespace mu
 		NodeLayoutPtr pResult;
 
 		// TODO: Substract layouts too? Usually they are ignored.
-		if ( m_targets.size()>0 && m_targets[0] )
+		if ( m_targets.Num()>0 && m_targets[0] )
 		{
 			NodeMesh::Private* pPrivate =
 					dynamic_cast<NodeMesh::Private*>( m_targets[0]->GetBasePrivate() );
@@ -172,5 +169,3 @@ namespace mu
 
 
 }
-
-#undef NODE_INPUT_COUNT
diff --git a/Engine/Plugins/Experimental/Mutable/Source/MutableTools/Private/MuT/NodeMeshInterpolatePrivate.h b/Engine/Plugins/Experimental/Mutable/Source/MutableTools/Private/MuT/NodeMeshInterpolatePrivate.h
--- a/Engine/Plugins/Experimental/Mutable/Source/MutableTools/Private/MuT/NodeMeshInterpolatePrivate.h
+++ b/Engine/Plugins/Experimental/Mutable/Source/MutableTools/Private/MuT/NodeMeshInterpolatePrivate.h
@@ -25,7 +25,7 @@ namespace mu
 		static NODE_TYPE s_type;
 
 		NodeScalarPtr m_pFactor;
-		vector<NodeMeshPtr> m_targets;
+		TArray<NodeMeshPtr> m_targets;
 
 		//!
 		struct CHANNEL
@@ -68,7 +68,7 @@ namespace mu
 			}
 		};
 
-		vector<CHANNEL> m_channels;
+		TArray<CHANNEL> m_channels;
 
 		//!
 		void Serialise( OutputArchive& arch ) const
diff --git a/Engine/Plugins/Experimental/Mutable/Source/MutableTools/Private/MuT/NodeMeshMorph.cpp b/Engine/Plugins/Experimental/Mutable/Source/MutableTools/Private/MuT/NodeMeshMorph.cpp
--- a/Engine/Plugins/Experimental/Mutable/Source/MutableTools/Private/MuT/NodeMeshMorph.cpp
+++ b/Engine/Plugins/Experimental/Mutable/Source/MutableTools/Private/MuT/NodeMeshMorph.cpp
@@ -17,7 +17,6 @@
 #include <utility>
 
 
-
 namespace mu
 {
 
@@ -40,7 +39,7 @@ namespace mu
 	//---------------------------------------------------------------------------------------------
 	int NodeMeshMorph::GetInputCount() const
 	{
-		return 2 + (int)m_pD->m_morphs.size();
+		return 2 + m_pD->m_morphs.Num();
 	}
 
 
@@ -125,7 +124,7 @@ namespace mu
 	//---------------------------------------------------------------------------------------------
 	NodeMeshPtr NodeMeshMorph::GetMorph( int t ) const
 	{
-		check( t>=0 && t<(int)m_pD->m_morphs.size() );
+		check( t>=0 && t<m_pD->m_morphs.Num() );
 
 		return m_pD->m_morphs[t].get();
 	}
@@ -134,20 +133,20 @@ namespace mu
 	//---------------------------------------------------------------------------------------------
 	void NodeMeshMorph::SetMorph( int t, NodeMeshPtr pNode )
 	{
-		check( t>=0 && t<(int)m_pD->m_morphs.size() );
+		check( t>=0 && t<m_pD->m_morphs.Num() );
 		m_pD->m_morphs[t] = pNode;
 	}
 
 	//---------------------------------------------------------------------------------------------
 	void NodeMeshMorph::SetMorphCount( int t )
 	{
-		m_pD->m_morphs.resize(t);
+		m_pD->m_morphs.SetNum(t);
 	}
 
 	//---------------------------------------------------------------------------------------------
 	int NodeMeshMorph::GetMorphCount() const
 	{
-		return (int)m_pD->m_morphs.size();
+		return m_pD->m_morphs.Num();
 	}
 
 	//---------------------------------------------------------------------------------------------
diff --git a/Engine/Plugins/Experimental/Mutable/Source/MutableTools/Private/MuT/NodeMeshMorphPrivate.h b/Engine/Plugins/Experimental/Mutable/Source/MutableTools/Private/MuT/NodeMeshMorphPrivate.h
--- a/Engine/Plugins/Experimental/Mutable/Source/MutableTools/Private/MuT/NodeMeshMorphPrivate.h
+++ b/Engine/Plugins/Experimental/Mutable/Source/MutableTools/Private/MuT/NodeMeshMorphPrivate.h
@@ -23,7 +23,7 @@ namespace mu
 
 		NodeScalarPtr m_pFactor;
 		NodeMeshPtr m_pBase;
-		vector<NodeMeshPtr> m_morphs;
+		TArray<NodeMeshPtr> m_morphs;
         bool m_vertexIndicesAreRelative = false;
 		
 		bool m_reshapeSkeleton = false;
diff --git a/Engine/Plugins/Experimental/Mutable/Source/MutableTools/Private/MuT/NodeMeshSwitch.cpp b/Engine/Plugins/Experimental/Mutable/Source/MutableTools/Private/MuT/NodeMeshSwitch.cpp
--- a/Engine/Plugins/Experimental/Mutable/Source/MutableTools/Private/MuT/NodeMeshSwitch.cpp
+++ b/Engine/Plugins/Experimental/Mutable/Source/MutableTools/Private/MuT/NodeMeshSwitch.cpp
@@ -37,7 +37,7 @@ namespace mu
 	//---------------------------------------------------------------------------------------------
 	int NodeMeshSwitch::GetInputCount() const
 	{
-		return 1 + (int)m_pD->m_options.size();
+		return 1 + m_pD->m_options.Num();
 	}
 
 
@@ -100,14 +100,14 @@ namespace mu
 	//---------------------------------------------------------------------------------------------
 	void NodeMeshSwitch::SetOptionCount( int t )
 	{
-		m_pD->m_options.resize(t);
+		m_pD->m_options.SetNum(t);
 	}
 
 
 	//---------------------------------------------------------------------------------------------
 	NodeMeshPtr NodeMeshSwitch::GetOption( int t ) const
 	{
-		check( t>=0 && t<(int)m_pD->m_options.size() );
+		check( t>=0 && t<m_pD->m_options.Num() );
 		return m_pD->m_options[t].get();
 	}
 
@@ -115,7 +115,7 @@ namespace mu
 	//---------------------------------------------------------------------------------------------
 	void NodeMeshSwitch::SetOption( int t, NodeMeshPtr pNode )
 	{
-		check( t>=0 && t<(int)m_pD->m_options.size() );
+		check( t>=0 && t<m_pD->m_options.Num() );
 		m_pD->m_options[t] = pNode;
 	}
 
@@ -125,7 +125,7 @@ namespace mu
 	{
 		NodeLayoutPtr pResult;
 
-		if (m_options.size()>0 && m_options[0] )
+		if (m_options.Num()>0 && m_options[0] )
 		{
 			NodeMesh::Private* pPrivate =
 					dynamic_cast<NodeMesh::Private*>( m_options[0]->GetBasePrivate() );
diff --git a/Engine/Plugins/Experimental/Mutable/Source/MutableTools/Private/MuT/NodeMeshSwitchPrivate.h b/Engine/Plugins/Experimental/Mutable/Source/MutableTools/Private/MuT/NodeMeshSwitchPrivate.h
--- a/Engine/Plugins/Experimental/Mutable/Source/MutableTools/Private/MuT/NodeMeshSwitchPrivate.h
+++ b/Engine/Plugins/Experimental/Mutable/Source/MutableTools/Private/MuT/NodeMeshSwitchPrivate.h
@@ -23,7 +23,7 @@ namespace mu
 		static NODE_TYPE s_type;
 
 		NodeScalarPtr m_pParameter;
-		vector<NodeMeshPtr> m_options;
+		TArray<NodeMeshPtr> m_options;
 
 		//!
 		void Serialise( OutputArchive& arch ) const
diff --git a/Engine/Plugins/Experimental/Mutable/Source/MutableTools/Private/MuT/NodeMeshTable.cpp b/Engine/Plugins/Experimental/Mutable/Source/MutableTools/Private/MuT/NodeMeshTable.cpp
--- a/Engine/Plugins/Experimental/Mutable/Source/MutableTools/Private/MuT/NodeMeshTable.cpp
+++ b/Engine/Plugins/Experimental/Mutable/Source/MutableTools/Private/MuT/NodeMeshTable.cpp
@@ -106,14 +106,14 @@ namespace mu
 	//---------------------------------------------------------------------------------------------
 	int NodeMeshTable::GetLayoutCount() const
 	{
-		return (int)m_pD->m_layouts.size();
+		return m_pD->m_layouts.Num();
 	}
 
 
 	//---------------------------------------------------------------------------------------------
 	void NodeMeshTable::SetLayoutCount( int i )
 	{
-		m_pD->m_layouts.resize( i );
+		m_pD->m_layouts.SetNum( i );
 	}
 
 
@@ -138,7 +138,7 @@ namespace mu
 	{
 		NodeLayoutPtr pResult;
 
-		if ( i>=0 && i<(int)m_layouts.size() )
+		if ( i>=0 && i<m_layouts.Num() )
 		{
 			pResult = m_layouts[i];
 		}
diff --git a/Engine/Plugins/Experimental/Mutable/Source/MutableTools/Private/MuT/NodeMeshTablePrivate.h b/Engine/Plugins/Experimental/Mutable/Source/MutableTools/Private/MuT/NodeMeshTablePrivate.h
--- a/Engine/Plugins/Experimental/Mutable/Source/MutableTools/Private/MuT/NodeMeshTablePrivate.h
+++ b/Engine/Plugins/Experimental/Mutable/Source/MutableTools/Private/MuT/NodeMeshTablePrivate.h
@@ -29,12 +29,12 @@ namespace mu
 		TablePtr m_pTable;
 		string m_columnName;
 
-		vector<NodeLayoutPtr> m_layouts;
+		TArray<NodeLayoutPtr> m_layouts;
 
 		//!
 		void Serialise( OutputArchive& arch ) const
 		{
-            uint32_t ver = 1;
+            uint32 ver = 1;
 			arch << ver;
 
 			arch << m_parameterName;
@@ -46,7 +46,7 @@ namespace mu
 		//!
 		void Unserialise( InputArchive& arch )
 		{
-            uint32_t ver;
+            uint32 ver;
 			arch >> ver;
             check(ver==1);
 
diff --git a/Engine/Plugins/Experimental/Mutable/Source/MutableTools/Private/MuT/NodeMeshVariation.cpp b/Engine/Plugins/Experimental/Mutable/Source/MutableTools/Private/MuT/NodeMeshVariation.cpp
--- a/Engine/Plugins/Experimental/Mutable/Source/MutableTools/Private/MuT/NodeMeshVariation.cpp
+++ b/Engine/Plugins/Experimental/Mutable/Source/MutableTools/Private/MuT/NodeMeshVariation.cpp
@@ -30,7 +30,7 @@ namespace mu
 
 
     //---------------------------------------------------------------------------------------------
-    int NodeMeshVariation::GetInputCount() const { return 1 + int( m_pD->m_variations.size() ); }
+    int NodeMeshVariation::GetInputCount() const { return 1 + m_pD->m_variations.Num(); }
 
 
     //---------------------------------------------------------------------------------------------
@@ -44,11 +44,11 @@ namespace mu
         }
         i -= 1;
 
-        if ( i < int( m_pD->m_variations.size() ) )
+        if ( i < m_pD->m_variations.Num() )
         {
             return m_pD->m_variations[i].m_mesh.get();
         }
-        i -= int( m_pD->m_variations.size() );
+        i -= m_pD->m_variations.Num();
 
         return nullptr;
     }
@@ -66,37 +66,43 @@ namespace mu
         }
 
         i -= 1;
-        if ( i < int( m_pD->m_variations.size() ) )
+        if ( i < m_pD->m_variations.Num() )
         {
 
             m_pD->m_variations[i].m_mesh = dynamic_cast<NodeMesh*>( pNode.get() );
             return;
         }
-        i -= (int)m_pD->m_variations.size();
+        i -= m_pD->m_variations.Num();
     }
 
 
     //---------------------------------------------------------------------------------------------
     // Own Interface
     //---------------------------------------------------------------------------------------------
-    void NodeMeshVariation::SetDefaultMesh( NodeMesh* p ) { m_pD->m_defaultMesh = p; }
+    void NodeMeshVariation::SetDefaultMesh( NodeMesh* p ) 
+	{ 
+		m_pD->m_defaultMesh = p; 
+	}
 
 
     //---------------------------------------------------------------------------------------------
-    int NodeMeshVariation::GetVariationCount() const { return int( m_pD->m_variations.size() ); }
+    int NodeMeshVariation::GetVariationCount() const 
+	{ 
+		return  m_pD->m_variations.Num(); 
+	}
 
 
     //---------------------------------------------------------------------------------------------
     void NodeMeshVariation::SetVariationCount( int num )
     {
         check( num >= 0 );
-        m_pD->m_variations.resize( num );
+        m_pD->m_variations.SetNum( num );
     }
 
     //---------------------------------------------------------------------------------------------
     void NodeMeshVariation::SetVariationTag( int index, const char* strTag )
     {
-        check( index >= 0 && index < (int)m_pD->m_variations.size() );
+        check( index >= 0 && index < m_pD->m_variations.Num() );
         check( strTag );
 
         if ( strTag )
@@ -113,7 +119,7 @@ namespace mu
     //---------------------------------------------------------------------------------------------
     void NodeMeshVariation::SetVariationMesh( int index, NodeMesh* pNode )
     {
-        check( index >= 0 && index < (int)m_pD->m_variations.size() );
+        check( index >= 0 && index < m_pD->m_variations.Num() );
 
         m_pD->m_variations[index].m_mesh = pNode;
     }
diff --git a/Engine/Plugins/Experimental/Mutable/Source/MutableTools/Private/MuT/NodeModifier.cpp b/Engine/Plugins/Experimental/Mutable/Source/MutableTools/Private/MuT/NodeModifier.cpp
--- a/Engine/Plugins/Experimental/Mutable/Source/MutableTools/Private/MuT/NodeModifier.cpp
+++ b/Engine/Plugins/Experimental/Mutable/Source/MutableTools/Private/MuT/NodeModifier.cpp
@@ -94,7 +94,7 @@ namespace mu
 		NodeModifier::Private* pD = dynamic_cast<NodeModifier::Private*>(GetBasePrivate());
 		check(pD);
 
-		pD->m_tags.push_back(tagName);
+		pD->m_tags.Add(tagName);
 	}
 
 
@@ -104,7 +104,7 @@ namespace mu
         NodeModifier::Private* pD = dynamic_cast<NodeModifier::Private*>(GetBasePrivate());
         check(pD);
 
-        return int( pD->m_tags.size());
+        return pD->m_tags.Num();
     }
 
 
diff --git a/Engine/Plugins/Experimental/Mutable/Source/MutableTools/Private/MuT/NodeModifierMeshClipMorphPlane.cpp b/Engine/Plugins/Experimental/Mutable/Source/MutableTools/Private/MuT/NodeModifierMeshClipMorphPlane.cpp
--- a/Engine/Plugins/Experimental/Mutable/Source/MutableTools/Private/MuT/NodeModifierMeshClipMorphPlane.cpp
+++ b/Engine/Plugins/Experimental/Mutable/Source/MutableTools/Private/MuT/NodeModifierMeshClipMorphPlane.cpp
@@ -10,9 +10,6 @@
 #include "MuT/NodePrivate.h"
 
 
-#define NODE_INPUT_COUNT 	0
-
-
 namespace mu
 {
 
@@ -35,14 +32,14 @@ namespace mu
 	//---------------------------------------------------------------------------------------------
     int NodeModifierMeshClipMorphPlane::GetInputCount() const
 	{
-		return NODE_INPUT_COUNT;
+		return 0;
 	}
 
 
 	//---------------------------------------------------------------------------------------------
     Node* NodeModifierMeshClipMorphPlane::GetInputNode( int i ) const
 	{
-		check( i>=0 && i<NODE_INPUT_COUNT );
+		check( i>=0 && i< GetInputCount());
         (void)i;
 		return nullptr;
 	}
@@ -51,7 +48,7 @@ namespace mu
 	//---------------------------------------------------------------------------------------------
     void NodeModifierMeshClipMorphPlane::SetInputNode( int i, NodePtr )
 	{
-		check( i>=0 && i<NODE_INPUT_COUNT );
+		check( i>=0 && i< GetInputCount());
         (void)i;
     }
 
@@ -100,5 +97,3 @@ namespace mu
 
 }
 
-#undef NODE_INPUT_COUNT
-
diff --git a/Engine/Plugins/Experimental/Mutable/Source/MutableTools/Private/MuT/NodeModifierPrivate.h b/Engine/Plugins/Experimental/Mutable/Source/MutableTools/Private/MuT/NodeModifierPrivate.h
--- a/Engine/Plugins/Experimental/Mutable/Source/MutableTools/Private/MuT/NodeModifierPrivate.h
+++ b/Engine/Plugins/Experimental/Mutable/Source/MutableTools/Private/MuT/NodeModifierPrivate.h
@@ -16,7 +16,7 @@ namespace mu
 	public:
 
 		// Tags that target surface need to have enabled to receive this modifier.
-		vector<mu::string> m_tags;
+		TArray<mu::string> m_tags;
 
 		// Wether the modifier has to be applied after the normal node operations or before
 		bool m_applyBeforeNormalOperations = true;
diff --git a/Engine/Plugins/Experimental/Mutable/Source/MutableTools/Private/MuT/NodeObjectGroup.cpp b/Engine/Plugins/Experimental/Mutable/Source/MutableTools/Private/MuT/NodeObjectGroup.cpp
--- a/Engine/Plugins/Experimental/Mutable/Source/MutableTools/Private/MuT/NodeObjectGroup.cpp
+++ b/Engine/Plugins/Experimental/Mutable/Source/MutableTools/Private/MuT/NodeObjectGroup.cpp
@@ -38,7 +38,7 @@ namespace mu
 	//---------------------------------------------------------------------------------------------
 	int NodeObjectGroup::GetInputCount() const
 	{
-		return (int)m_pD->m_children.size();
+		return m_pD->m_children.Num();
 	}
 
 
@@ -49,7 +49,7 @@ namespace mu
 
 		Node* pResult = 0;
 
-		if ( i<(int)m_pD->m_children.size() )
+		if ( i<m_pD->m_children.Num() )
 		{
 			pResult = m_pD->m_children[i].get();
 		}
@@ -63,7 +63,7 @@ namespace mu
 	{
 		check( i>=0 && i<GetInputCount() );
 
-		if ( i<(int)m_pD->m_children.size() )
+		if ( i<m_pD->m_children.Num() )
 		{
 			m_pD->m_children[i] = dynamic_cast<NodeObject*>( pNode.get() );
 		}
@@ -130,7 +130,7 @@ namespace mu
 	//---------------------------------------------------------------------------------------------
 	int NodeObjectGroup::GetChildCount() const
 	{
-		return (int)m_pD->m_children.size();
+		return m_pD->m_children.Num();
 	}
 
 
@@ -138,14 +138,14 @@ namespace mu
 	void NodeObjectGroup::SetChildCount( int num )
 	{
 		check( num >=0 );
-		m_pD->m_children.resize( num );
+		m_pD->m_children.SetNum( num );
 	}
 
 
 	//---------------------------------------------------------------------------------------------
 	NodeObjectPtr NodeObjectGroup::GetChild( int index ) const
 	{
-		check( index >=0 && index < (int)m_pD->m_children.size() );
+		check( index >=0 && index < m_pD->m_children.Num() );
 
 		return m_pD->m_children[ index ].get();
 	}
@@ -154,7 +154,7 @@ namespace mu
 	//---------------------------------------------------------------------------------------------
 	void NodeObjectGroup::SetChild( int index, NodeObjectPtr pObject )
 	{
-		check( index >=0 && index < (int)m_pD->m_children.size() );
+		check( index >=0 && index < m_pD->m_children.Num() );
 
 		m_pD->m_children[ index ] = pObject;
 	}
@@ -165,7 +165,7 @@ namespace mu
 	{
 		NodeLayoutPtr pLayout;
 
-		for ( size_t i=0; !pLayout && i<m_children.size(); ++i )
+		for ( int32 i=0; !pLayout && i<m_children.Num(); ++i )
 		{
 			if (m_children[i])
 			{
diff --git a/Engine/Plugins/Experimental/Mutable/Source/MutableTools/Private/MuT/NodeObjectGroupPrivate.h b/Engine/Plugins/Experimental/Mutable/Source/MutableTools/Private/MuT/NodeObjectGroupPrivate.h
--- a/Engine/Plugins/Experimental/Mutable/Source/MutableTools/Private/MuT/NodeObjectGroupPrivate.h
+++ b/Engine/Plugins/Experimental/Mutable/Source/MutableTools/Private/MuT/NodeObjectGroupPrivate.h
@@ -34,7 +34,7 @@ namespace mu
 		//! Set the child selection type
 		void SetSelectionType( CHILD_SELECTION );
 
-		vector<NodeObjectPtr> m_children;
+		TArray<NodeObjectPtr> m_children;
 
 		//!
 		void Serialise( OutputArchive& arch ) const
diff --git a/Engine/Plugins/Experimental/Mutable/Source/MutableTools/Private/MuT/NodeObjectNew.cpp b/Engine/Plugins/Experimental/Mutable/Source/MutableTools/Private/MuT/NodeObjectNew.cpp
--- a/Engine/Plugins/Experimental/Mutable/Source/MutableTools/Private/MuT/NodeObjectNew.cpp
+++ b/Engine/Plugins/Experimental/Mutable/Source/MutableTools/Private/MuT/NodeObjectNew.cpp
@@ -46,7 +46,7 @@ namespace mu
 	//---------------------------------------------------------------------------------------------
 	int NodeObjectNew::GetInputCount() const
 	{
-		return (int)(m_pD->m_lods.size() + m_pD->m_children.size());
+		return (m_pD->m_lods.Num() + m_pD->m_children.Num());
 	}
 
 
@@ -57,13 +57,13 @@ namespace mu
 
 		Node* pResult = 0;
 
-		if ( i<(int)m_pD->m_lods.size() )
+		if ( i<(int)m_pD->m_lods.Num() )
 		{
 			pResult = m_pD->m_lods[i].get();
 		}
 		else
 		{
-			i -= (int)m_pD->m_lods.size();
+			i -= (int)m_pD->m_lods.Num();
 			pResult = m_pD->m_children[i].get();
 		}
 
@@ -76,13 +76,13 @@ namespace mu
 	{
 		check( i>=0 && i<GetInputCount() );
 
-		if ( i<(int)m_pD->m_lods.size() )
+		if ( i<(int)m_pD->m_lods.Num() )
 		{
 			m_pD->m_lods[i] = dynamic_cast<NodeLOD*>( pNode.get() );
 		}
 		else
 		{
-			i -= (int)m_pD->m_lods.size();
+			i -= (int)m_pD->m_lods.Num();
 			m_pD->m_children[i] = dynamic_cast<NodeObject*>( pNode.get() );
 		}
 	}
@@ -134,7 +134,7 @@ namespace mu
 	//---------------------------------------------------------------------------------------------
 	int NodeObjectNew::GetLODCount() const
 	{
-		return (int)m_pD->m_lods.size();
+		return m_pD->m_lods.Num();
 	}
 
 
@@ -142,14 +142,14 @@ namespace mu
 	void NodeObjectNew::SetLODCount( int num )
 	{
 		check( num >=0 );
-		m_pD->m_lods.resize( num );
+		m_pD->m_lods.SetNum( num );
 	}
 
 
 	//---------------------------------------------------------------------------------------------
 	NodeLODPtr NodeObjectNew::GetLOD( int index ) const
 	{
-		check( index >=0 && index < (int)m_pD->m_lods.size() );
+		check( index >=0 && index < m_pD->m_lods.Num() );
 
 		return m_pD->m_lods[ index ].get();
 	}
@@ -158,7 +158,7 @@ namespace mu
 	//---------------------------------------------------------------------------------------------
 	void NodeObjectNew::SetLOD( int index, NodeLODPtr pLOD )
 	{
-		check( index >=0 && index < (int)m_pD->m_lods.size() );
+		check( index >=0 && index < m_pD->m_lods.Num() );
 
 		m_pD->m_lods[ index ] = pLOD;
 	}
@@ -167,7 +167,7 @@ namespace mu
 	//---------------------------------------------------------------------------------------------
 	int NodeObjectNew::GetChildCount() const
 	{
-		return (int)m_pD->m_children.size();
+		return m_pD->m_children.Num();
 	}
 
 
@@ -175,14 +175,14 @@ namespace mu
 	void NodeObjectNew::SetChildCount( int num )
 	{
 		check( num >=0 );
-		m_pD->m_children.resize( num );
+		m_pD->m_children.SetNum( num );
 	}
 
 
 	//---------------------------------------------------------------------------------------------
 	NodeObjectPtr NodeObjectNew::GetChild( int index ) const
 	{
-		check( index >=0 && index < (int)m_pD->m_children.size() );
+		check( index >=0 && index < m_pD->m_children.Num() );
 
 		return m_pD->m_children[ index ].get();
 	}
@@ -191,7 +191,7 @@ namespace mu
 	//---------------------------------------------------------------------------------------------
 	void NodeObjectNew::SetChild( int index, NodeObjectPtr pObject )
 	{
-		check( index >=0 && index < (int)m_pD->m_children.size() );
+		check( index >=0 && index < m_pD->m_children.Num() );
 
 		m_pD->m_children[ index ] = pObject;
 	}
@@ -204,7 +204,7 @@ namespace mu
                                                      int //texture
                                                      ) const
 	{
-		check( lod>=0 && lod<(int)m_lods.size() );
+		check( lod>=0 && lod<m_lods.Num() );
 
 		NodeLayoutPtr pLayout;
 
@@ -245,12 +245,9 @@ namespace mu
     bool NodeObjectNew::Private::HasComponent( const NodeComponent* pComponent ) const
     {
         bool found = false;
-        for ( size_t l=0; !found && l<m_lods.size(); ++l )
+        for ( int32 l=0; !found && l<m_lods.Num(); ++l )
         {
-            found = std::find( m_lods[l]->GetPrivate()->m_components.begin(),
-                               m_lods[l]->GetPrivate()->m_components.end(),
-                               pComponent )
-                    != m_lods[l]->GetPrivate()->m_components.end();
+			found = m_lods[l]->GetPrivate()->m_components.FindByPredicate([pComponent](const NodeComponentPtr& e) { return e.get() == pComponent; }) != nullptr;
         }
 
         return found;
@@ -260,14 +257,14 @@ namespace mu
 	//---------------------------------------------------------------------------------------------
 	int NodeObjectNew::GetStateCount() const
 	{
-		return (int)m_pD->m_states.size();
+		return m_pD->m_states.Num();
 	}
 
 
 	//---------------------------------------------------------------------------------------------
 	void NodeObjectNew::SetStateCount( int c )
 	{
-		m_pD->m_states.resize( c );
+		m_pD->m_states.SetNum( c );
 	}
 
 
@@ -291,11 +288,7 @@ namespace mu
 	bool NodeObjectNew::HasStateParam( int s, const char* param ) const
 	{
 		check( s>=0 && s<GetStateCount() );
-		return std::find( m_pD->m_states[s].m_runtimeParams.begin(),
-						  m_pD->m_states[s].m_runtimeParams.end(),
-						  param )
-				!=
-				m_pD->m_states[s].m_runtimeParams.end();
+		return m_pD->m_states[s].m_runtimeParams.Contains( param );
 	}
 
 
@@ -306,7 +299,7 @@ namespace mu
 
 		if (!HasStateParam(s,param))
 		{
-			m_pD->m_states[s].m_runtimeParams.push_back( param );
+			m_pD->m_states[s].m_runtimeParams.Add( param );
 		}
 	}
 
@@ -316,12 +309,10 @@ namespace mu
 	{
 		check( s>=0 && s<GetStateCount() );
 
-		vector<string>::iterator it = std::find( m_pD->m_states[s].m_runtimeParams.begin(),
-														   m_pD->m_states[s].m_runtimeParams.end(),
-														   param );
-		if ( it != m_pD->m_states[s].m_runtimeParams.end() )
+		int32 it = m_pD->m_states[s].m_runtimeParams.Find( param );
+		if ( it != INDEX_NONE )
 		{
-			m_pD->m_states[s].m_runtimeParams.erase( it );
+			m_pD->m_states[s].m_runtimeParams.RemoveAt( it );
 		}
 	}
 
diff --git a/Engine/Plugins/Experimental/Mutable/Source/MutableTools/Private/MuT/NodeObjectNewPrivate.h b/Engine/Plugins/Experimental/Mutable/Source/MutableTools/Private/MuT/NodeObjectNewPrivate.h
--- a/Engine/Plugins/Experimental/Mutable/Source/MutableTools/Private/MuT/NodeObjectNewPrivate.h
+++ b/Engine/Plugins/Experimental/Mutable/Source/MutableTools/Private/MuT/NodeObjectNewPrivate.h
@@ -29,13 +29,13 @@ namespace mu
 		string m_name;
 		string m_uid;
 
-		vector<NodeLODPtr> m_lods;
+		TArray<NodeLODPtr> m_lods;
 
-		vector<NodeObjectPtr> m_children;
+		TArray<NodeObjectPtr> m_children;
 
 
 		//! List of states
-        vector<OBJECT_STATE> m_states;
+        TArray<OBJECT_STATE> m_states;
 
 
 		//!
diff --git a/Engine/Plugins/Experimental/Mutable/Source/MutableTools/Private/MuT/NodeObjectState.cpp b/Engine/Plugins/Experimental/Mutable/Source/MutableTools/Private/MuT/NodeObjectState.cpp
--- a/Engine/Plugins/Experimental/Mutable/Source/MutableTools/Private/MuT/NodeObjectState.cpp
+++ b/Engine/Plugins/Experimental/Mutable/Source/MutableTools/Private/MuT/NodeObjectState.cpp
@@ -171,11 +171,7 @@ namespace mu
 	//---------------------------------------------------------------------------------------------
 	bool NodeObjectState::HasStateParam( const char* param ) const
 	{
-		return std::find( m_pD->m_state.m_runtimeParams.begin(),
-						  m_pD->m_state.m_runtimeParams.end(),
-						  param )
-				!=
-				m_pD->m_state.m_runtimeParams.end();
+		return m_pD->m_state.m_runtimeParams.Contains( param );
 	}
 
 
@@ -184,7 +180,7 @@ namespace mu
 	{
 		if (!HasStateParam(param))
 		{
-			m_pD->m_state.m_runtimeParams.push_back( param );
+			m_pD->m_state.m_runtimeParams.Add( param );
 		}
 	}
 
@@ -192,12 +188,10 @@ namespace mu
 	//---------------------------------------------------------------------------------------------
 	void NodeObjectState::RemoveStateParam( const char* param )
 	{
-		vector<string>::iterator it = std::find( m_pD->m_state.m_runtimeParams.begin(),
-														   m_pD->m_state.m_runtimeParams.end(),
-														   param );
-		if ( it != m_pD->m_state.m_runtimeParams.end() )
+		int32 it = m_pD->m_state.m_runtimeParams.Find( param );
+		if ( it != INDEX_NONE )
 		{
-			m_pD->m_state.m_runtimeParams.erase( it );
+			m_pD->m_state.m_runtimeParams.RemoveAt( it );
 		}
 	}
 
diff --git a/Engine/Plugins/Experimental/Mutable/Source/MutableTools/Private/MuT/NodePatchImage.cpp b/Engine/Plugins/Experimental/Mutable/Source/MutableTools/Private/MuT/NodePatchImage.cpp
--- a/Engine/Plugins/Experimental/Mutable/Source/MutableTools/Private/MuT/NodePatchImage.cpp
+++ b/Engine/Plugins/Experimental/Mutable/Source/MutableTools/Private/MuT/NodePatchImage.cpp
@@ -109,21 +109,21 @@ namespace mu
     //---------------------------------------------------------------------------------------------
     void NodePatchImage::SetBlockCount( int c )
     {
-        m_pD->m_blocks.resize( c );
+        m_pD->m_blocks.SetNum( c );
     }
 
 
     //---------------------------------------------------------------------------------------------
     int NodePatchImage::GetBlockCount() const
     {
-        return (int)m_pD->m_blocks.size();
+        return (int)m_pD->m_blocks.Num();
     }
 
 
     //---------------------------------------------------------------------------------------------
     void NodePatchImage::SetBlock( int index, int block )
     {
-        check( index>=0 && index<(int)m_pD->m_blocks.size() );
+        check( index>=0 && index<(int)m_pD->m_blocks.Num() );
         m_pD->m_blocks[ index ] = block;
     }
 
@@ -131,7 +131,7 @@ namespace mu
     //---------------------------------------------------------------------------------------------
     int NodePatchImage::GetBlock( int index ) const
     {
-        check( index>=0 && index<(int)m_pD->m_blocks.size() );
+        check( index>=0 && index<(int)m_pD->m_blocks.Num() );
         return m_pD->m_blocks[ index ];
     }
 
diff --git a/Engine/Plugins/Experimental/Mutable/Source/MutableTools/Private/MuT/NodePatchImagePrivate.h b/Engine/Plugins/Experimental/Mutable/Source/MutableTools/Private/MuT/NodePatchImagePrivate.h
--- a/Engine/Plugins/Experimental/Mutable/Source/MutableTools/Private/MuT/NodePatchImagePrivate.h
+++ b/Engine/Plugins/Experimental/Mutable/Source/MutableTools/Private/MuT/NodePatchImagePrivate.h
@@ -26,7 +26,7 @@ namespace mu
         NodeImagePtr m_pImage;
         NodeImagePtr m_pMask;
 
-        vector<int> m_blocks;
+        TArray<int> m_blocks;
 
         EBlendType m_blendType = EBlendType::BT_BLEND;
 
diff --git a/Engine/Plugins/Experimental/Mutable/Source/MutableTools/Private/MuT/NodePrivate.h b/Engine/Plugins/Experimental/Mutable/Source/MutableTools/Private/MuT/NodePrivate.h
--- a/Engine/Plugins/Experimental/Mutable/Source/MutableTools/Private/MuT/NodePrivate.h
+++ b/Engine/Plugins/Experimental/Mutable/Source/MutableTools/Private/MuT/NodePrivate.h
@@ -22,9 +22,6 @@ namespace mu
     {
     public:
 
-        //!
-        static void CloneChildren( const Node* pSource, NodePtr pTarget, NodeMapPtr pMap );
-
         //! This is an opaque context used to attach to reported error messages.
 		const void* m_errorContext = nullptr;
 
@@ -33,14 +30,6 @@ namespace mu
     };
 
 
-    class NodeMap::Private : public Base
-    {
-    public:
-
-        map< const void*, NodePtr > m_map;
-
-    };
-
 
 #define MUTABLE_IMPLEMENT_NODE( N, T, PN, PT )				\
     N::N()													\
@@ -68,30 +57,13 @@ namespace mu
 		N::Serialise(this, arch);							\
 	}														\
                                                             \
-    Ptr<N> N::StaticUnserialise( InputArchive& arch )				\
+    Ptr<N> N::StaticUnserialise( InputArchive& arch )		\
     {														\
         Ptr<N> pResult = new N();							\
         arch >> *pResult->m_pD;								\
         return pResult;										\
     }														\
                                                             \
-    NodePtr N::Clone( NodeMapPtr pMap ) const				\
-    {														\
-        Ptr<N> pCloned = new N();							\
-                                                            \
-        *pCloned->m_pD = *m_pD;								\
-                                                            \
-        if (!pMap)											\
-        {													\
-            pMap = new NodeMap;								\
-        }													\
-        pMap->Add( this, pCloned );							\
-                                                            \
-        Private::CloneChildren( this, pCloned, pMap );		\
-                                                            \
-        return pCloned;										\
-    }														\
-                                                            \
     N::Private* N::GetPrivate() const						\
     {														\
         return m_pD;										\
diff --git a/Engine/Plugins/Experimental/Mutable/Source/MutableTools/Private/MuT/NodeProjectorParameter.cpp b/Engine/Plugins/Experimental/Mutable/Source/MutableTools/Private/MuT/NodeProjectorParameter.cpp
--- a/Engine/Plugins/Experimental/Mutable/Source/MutableTools/Private/MuT/NodeProjectorParameter.cpp
+++ b/Engine/Plugins/Experimental/Mutable/Source/MutableTools/Private/MuT/NodeProjectorParameter.cpp
@@ -39,7 +39,7 @@ namespace mu
 	//---------------------------------------------------------------------------------------------
 	int NodeProjectorParameter::GetInputCount() const
 	{
-        return int( m_pD->m_ranges.size() );
+        return m_pD->m_ranges.Num();
     }
 
 
@@ -160,15 +160,15 @@ namespace mu
     void NodeProjectorParameter::SetRangeCount( int i )
     {
         check(i>=0);
-        m_pD->m_ranges.resize(i);
+        m_pD->m_ranges.SetNum(i);
     }
 
 
     //---------------------------------------------------------------------------------------------
     void NodeProjectorParameter::SetRange( int i, NodeRangePtr pRange )
     {
-        check( i>=0 && i<int(m_pD->m_ranges.size()) );
-        if ( i>=0 && i<int(m_pD->m_ranges.size()) )
+        check( i>=0 && i<m_pD->m_ranges.Num() );
+        if ( i>=0 && i<m_pD->m_ranges.Num() )
         {
             m_pD->m_ranges[i] = pRange;
         }
diff --git a/Engine/Plugins/Experimental/Mutable/Source/MutableTools/Private/MuT/NodeRangeFromScalar.cpp b/Engine/Plugins/Experimental/Mutable/Source/MutableTools/Private/MuT/NodeRangeFromScalar.cpp
--- a/Engine/Plugins/Experimental/Mutable/Source/MutableTools/Private/MuT/NodeRangeFromScalar.cpp
+++ b/Engine/Plugins/Experimental/Mutable/Source/MutableTools/Private/MuT/NodeRangeFromScalar.cpp
@@ -11,9 +11,6 @@
 #include "MuT/NodeScalar.h"
 
 
-#define NODE_INPUT_COUNT 	1
-
-
 namespace mu
 {
 
@@ -36,14 +33,14 @@ namespace mu
 	//---------------------------------------------------------------------------------------------
     int NodeRangeFromScalar::GetInputCount() const
 	{
-		return NODE_INPUT_COUNT;
+		return 1;
 	}
 
 
 	//---------------------------------------------------------------------------------------------
     Node* NodeRangeFromScalar::GetInputNode( int i ) const
 	{
-		check( i>=0 && i<NODE_INPUT_COUNT );
+		check( i>=0 && i< GetInputCount());
 
         Node* pResult = nullptr;
 
@@ -59,7 +56,7 @@ namespace mu
 	//---------------------------------------------------------------------------------------------
     void NodeRangeFromScalar::SetInputNode( int i, NodePtr pNode )
 	{
-		check( i>=0 && i<NODE_INPUT_COUNT );
+		check( i>=0 && i< GetInputCount());
 
 		switch (i)
 		{
@@ -106,5 +103,3 @@ namespace mu
 
 }
 
-#undef NODE_INPUT_COUNT
-
diff --git a/Engine/Plugins/Experimental/Mutable/Source/MutableTools/Private/MuT/NodeScalarArithmeticOperation.cpp b/Engine/Plugins/Experimental/Mutable/Source/MutableTools/Private/MuT/NodeScalarArithmeticOperation.cpp
--- a/Engine/Plugins/Experimental/Mutable/Source/MutableTools/Private/MuT/NodeScalarArithmeticOperation.cpp
+++ b/Engine/Plugins/Experimental/Mutable/Source/MutableTools/Private/MuT/NodeScalarArithmeticOperation.cpp
@@ -8,10 +8,6 @@
 #include "MuT/NodeScalar.h"
 #include "MuT/NodeScalarArithmeticOperationPrivate.h"
 
-
-#define NODE_INPUT_COUNT 	2
-
-
 namespace mu
 {
 
@@ -44,14 +40,14 @@ namespace mu
 	//---------------------------------------------------------------------------------------------
     int NodeScalarArithmeticOperation::GetInputCount() const
 	{
-		return NODE_INPUT_COUNT;
+		return 2;
 	}
 
 
 	//---------------------------------------------------------------------------------------------
     Node* NodeScalarArithmeticOperation::GetInputNode( int i ) const
 	{
-		check( i>=0 && i<NODE_INPUT_COUNT );
+		check( i>=0 && i< GetInputCount());
 
 		Node* pResult = 0;
 
@@ -68,7 +64,7 @@ namespace mu
 	//---------------------------------------------------------------------------------------------
     void NodeScalarArithmeticOperation::SetInputNode( int i, NodePtr pNode )
 	{
-		check( i>=0 && i<NODE_INPUT_COUNT );
+		check( i>=0 && i< GetInputCount());
 
 		switch (i)
 		{
@@ -122,5 +118,3 @@ namespace mu
 
 }
 
-#undef NODE_INPUT_COUNT
-
diff --git a/Engine/Plugins/Experimental/Mutable/Source/MutableTools/Private/MuT/NodeScalarConstant.cpp b/Engine/Plugins/Experimental/Mutable/Source/MutableTools/Private/MuT/NodeScalarConstant.cpp
--- a/Engine/Plugins/Experimental/Mutable/Source/MutableTools/Private/MuT/NodeScalarConstant.cpp
+++ b/Engine/Plugins/Experimental/Mutable/Source/MutableTools/Private/MuT/NodeScalarConstant.cpp
@@ -8,9 +8,6 @@
 #include "MuT/NodeScalarConstantPrivate.h"
 
 
-#define NODE_INPUT_COUNT 	0
-
-
 namespace mu
 {
 
@@ -33,7 +30,7 @@ namespace mu
 	//---------------------------------------------------------------------------------------------
 	int NodeScalarConstant::GetInputCount() const
 	{
-		return NODE_INPUT_COUNT;
+		return 0;
 	}
 
 
@@ -69,5 +66,3 @@ namespace mu
 
 }
 
-#undef NODE_INPUT_COUNT
-
diff --git a/Engine/Plugins/Experimental/Mutable/Source/MutableTools/Private/MuT/NodeScalarCurve.cpp b/Engine/Plugins/Experimental/Mutable/Source/MutableTools/Private/MuT/NodeScalarCurve.cpp
--- a/Engine/Plugins/Experimental/Mutable/Source/MutableTools/Private/MuT/NodeScalarCurve.cpp
+++ b/Engine/Plugins/Experimental/Mutable/Source/MutableTools/Private/MuT/NodeScalarCurve.cpp
@@ -14,9 +14,6 @@
 #include <utility>
 
 
-#define NODE_INPUT_COUNT 	1
-
-
 namespace mu
 {
 
@@ -40,7 +37,7 @@ namespace mu
 	//---------------------------------------------------------------------------------------------
     int NodeScalarCurve::GetInputCount() const
 	{
-		return NODE_INPUT_COUNT;
+		return 1;
 	}
 
 
@@ -69,7 +66,7 @@ namespace mu
 	//---------------------------------------------------------------------------------------------
     int NodeScalarCurve::GetKeyFrameCount() const
 	{
-		return (int)m_pD->m_curve.keyFrames.size();
+		return (int)m_pD->m_curve.keyFrames.Num();
 	}
 
 
@@ -77,7 +74,7 @@ namespace mu
     void NodeScalarCurve::SetKeyFrameCount( int num )
 	{
 		check( num >=0 );
-		m_pD->m_curve.keyFrames.resize( num );
+		m_pD->m_curve.keyFrames.SetNum( num );
 	}
 
 	//---------------------------------------------------------------------------------------------
@@ -94,7 +91,7 @@ namespace mu
             float out_tangent, float out_tangent_weight,
             uint8_t interp_mode, uint8_t tangent_mode, uint8_t tangent_weight_mode )
 	{
-		check( index >=0 && index < (int)m_pD->m_curve.keyFrames.size() );
+		check( index >=0 && index < m_pD->m_curve.keyFrames.Num() );
 
 		CurveKeyFrame& key_frame = m_pD->m_curve.keyFrames[index];
 
@@ -126,5 +123,3 @@ namespace mu
 
 }
 
-#undef NODE_INPUT_COUNT
-
diff --git a/Engine/Plugins/Experimental/Mutable/Source/MutableTools/Private/MuT/NodeScalarEnumParameter.cpp b/Engine/Plugins/Experimental/Mutable/Source/MutableTools/Private/MuT/NodeScalarEnumParameter.cpp
--- a/Engine/Plugins/Experimental/Mutable/Source/MutableTools/Private/MuT/NodeScalarEnumParameter.cpp
+++ b/Engine/Plugins/Experimental/Mutable/Source/MutableTools/Private/MuT/NodeScalarEnumParameter.cpp
@@ -32,7 +32,7 @@ namespace mu
 	//---------------------------------------------------------------------------------------------
 	int NodeScalarEnumParameter::GetInputCount() const
 	{
-        return int( m_pD->m_ranges.size() );
+        return m_pD->m_ranges.Num();
 	}
 
 
@@ -117,21 +117,21 @@ namespace mu
 	//---------------------------------------------------------------------------------------------
 	void NodeScalarEnumParameter::SetValueCount( int i )
 	{
-		m_pD->m_options.resize(i);
+		m_pD->m_options.SetNum(i);
 	}
 
 
 	//---------------------------------------------------------------------------------------------
 	int NodeScalarEnumParameter::GetValueCount() const
 	{
-		return (int)m_pD->m_options.size();
+		return (int)m_pD->m_options.Num();
 	}
 
 
 	//---------------------------------------------------------------------------------------------
 	void NodeScalarEnumParameter::SetValue( int i, float value, const char* strName )
 	{
-		check( i>=0 && i<(int)m_pD->m_options.size() );
+		check( i>=0 && i<(int)m_pD->m_options.Num() );
 		m_pD->m_options[i].name = strName;
 		m_pD->m_options[i].value = value;
 	}
@@ -141,15 +141,15 @@ namespace mu
     void NodeScalarEnumParameter::SetRangeCount( int i )
     {
         check(i>=0);
-        m_pD->m_ranges.resize(i);
+        m_pD->m_ranges.SetNum(i);
     }
 
 
     //---------------------------------------------------------------------------------------------
     void NodeScalarEnumParameter::SetRange( int i, NodeRangePtr pRange )
     {
-        check( i>=0 && i<int(m_pD->m_ranges.size()) );
-        if ( i>=0 && i<int(m_pD->m_ranges.size()) )
+        check( i>=0 && i<int(m_pD->m_ranges.Num()) );
+        if ( i>=0 && i<int(m_pD->m_ranges.Num()) )
         {
             m_pD->m_ranges[i] = pRange;
         }
diff --git a/Engine/Plugins/Experimental/Mutable/Source/MutableTools/Private/MuT/NodeScalarEnumParameterPrivate.h b/Engine/Plugins/Experimental/Mutable/Source/MutableTools/Private/MuT/NodeScalarEnumParameterPrivate.h
--- a/Engine/Plugins/Experimental/Mutable/Source/MutableTools/Private/MuT/NodeScalarEnumParameterPrivate.h
+++ b/Engine/Plugins/Experimental/Mutable/Source/MutableTools/Private/MuT/NodeScalarEnumParameterPrivate.h
@@ -56,9 +56,9 @@ namespace mu
 			}
 		};
 
-		vector<OPTION> m_options;
+		TArray<OPTION> m_options;
 
-        vector<Ptr<NodeRange>> m_ranges;
+		TArray<Ptr<NodeRange>> m_ranges;
 
 		//!
 		void Serialise( OutputArchive& arch ) const
