[mutable] Optimize mesh clip versus switching of clipping mesh.

diff --git a/Engine/Plugins/Experimental/Mutable/Source/MutableTools/Private/MuT/NodeScalarParameter.cpp b/Engine/Plugins/Experimental/Mutable/Source/MutableTools/Private/MuT/NodeScalarParameter.cpp
--- a/Engine/Plugins/Experimental/Mutable/Source/MutableTools/Private/MuT/NodeScalarParameter.cpp
+++ b/Engine/Plugins/Experimental/Mutable/Source/MutableTools/Private/MuT/NodeScalarParameter.cpp
@@ -34,9 +34,9 @@ namespace mu
 	//---------------------------------------------------------------------------------------------
 	int NodeScalarParameter::GetInputCount() const
 	{
-        return int( m_pD->m_additionalImages.size()
+        return int( m_pD->m_additionalImages.Num()
                     +
-                    m_pD->m_ranges.size() );
+                    m_pD->m_ranges.Num() );
 	}
 
 
@@ -45,8 +45,8 @@ namespace mu
 	{
         check( i<GetInputCount() );
 
-        int imageCount = int(m_pD->m_additionalImages.size());
-        int rangeCount = int(m_pD->m_ranges.size());
+        int imageCount = int(m_pD->m_additionalImages.Num());
+        int rangeCount = int(m_pD->m_ranges.Num());
         if (i<imageCount)
         {
             return m_pD->m_additionalImages[i].get();
@@ -64,8 +64,8 @@ namespace mu
     void NodeScalarParameter::SetInputNode( int i, NodePtr n )
 	{
         check( i<GetInputCount() );
-        int imageCount = int(m_pD->m_additionalImages.size());
-        int rangeCount = int(m_pD->m_ranges.size());
+        int imageCount = int(m_pD->m_additionalImages.Num());
+        int rangeCount = int(m_pD->m_ranges.Num());
         if (i<imageCount)
         {
             m_pD->m_additionalImages[i] = dynamic_cast<NodeImage*>(n.get());
@@ -151,14 +151,14 @@ namespace mu
 	void NodeScalarParameter::SetDescImageCount( int i )
 	{
         check(i>=0);
-        m_pD->m_additionalImages.resize(i);
+        m_pD->m_additionalImages.SetNum(i);
 	}
 
 
 	//---------------------------------------------------------------------------------------------
 	void NodeScalarParameter::SetDescImage( int i, NodeImagePtr pImage )
 	{
-		check( i>=0 && i<(int)m_pD->m_additionalImages.size() );
+		check( i>=0 && i<(int)m_pD->m_additionalImages.Num() );
 		m_pD->m_additionalImages[i] = pImage;
 	}
 
@@ -167,15 +167,15 @@ namespace mu
     void NodeScalarParameter::SetRangeCount( int i )
     {
         check(i>=0);
-        m_pD->m_ranges.resize(i);
+        m_pD->m_ranges.SetNum(i);
     }
 
 
     //---------------------------------------------------------------------------------------------
     void NodeScalarParameter::SetRange( int i, NodeRangePtr pRange )
     {
-        check( i>=0 && i<int(m_pD->m_ranges.size()) );
-        if ( i>=0 && i<int(m_pD->m_ranges.size()) )
+        check( i>=0 && i<int(m_pD->m_ranges.Num()) );
+        if ( i>=0 && i<int(m_pD->m_ranges.Num()) )
         {
             m_pD->m_ranges[i] = pRange;
         }
diff --git a/Engine/Plugins/Experimental/Mutable/Source/MutableTools/Private/MuT/NodeScalarParameterPrivate.h b/Engine/Plugins/Experimental/Mutable/Source/MutableTools/Private/MuT/NodeScalarParameterPrivate.h
--- a/Engine/Plugins/Experimental/Mutable/Source/MutableTools/Private/MuT/NodeScalarParameterPrivate.h
+++ b/Engine/Plugins/Experimental/Mutable/Source/MutableTools/Private/MuT/NodeScalarParameterPrivate.h
@@ -30,9 +30,9 @@ namespace mu
 		string m_uid;
 		PARAMETER_DETAILED_TYPE m_detailedType = PARAMETER_DETAILED_TYPE::UNKNOWN;
 
-        vector<Ptr<NodeImage>> m_additionalImages;
+        TArray<Ptr<NodeImage>> m_additionalImages;
 
-        vector<Ptr<NodeRange>> m_ranges;
+		TArray<Ptr<NodeRange>> m_ranges;
 
 		//!
 		void Serialise( OutputArchive& arch ) const
diff --git a/Engine/Plugins/Experimental/Mutable/Source/MutableTools/Private/MuT/NodeScalarSwitch.cpp b/Engine/Plugins/Experimental/Mutable/Source/MutableTools/Private/MuT/NodeScalarSwitch.cpp
--- a/Engine/Plugins/Experimental/Mutable/Source/MutableTools/Private/MuT/NodeScalarSwitch.cpp
+++ b/Engine/Plugins/Experimental/Mutable/Source/MutableTools/Private/MuT/NodeScalarSwitch.cpp
@@ -35,7 +35,7 @@ namespace mu
 	//---------------------------------------------------------------------------------------------
     int NodeScalarSwitch::GetInputCount() const
 	{
-		return 1 + (int)m_pD->m_options.size();
+		return 1 + (int)m_pD->m_options.Num();
 	}
 
 
@@ -98,14 +98,14 @@ namespace mu
 	//---------------------------------------------------------------------------------------------
     void NodeScalarSwitch::SetOptionCount( int t )
 	{
-		m_pD->m_options.resize(t);
+		m_pD->m_options.SetNum(t);
 	}
 
 
 	//---------------------------------------------------------------------------------------------
     NodeScalarPtr NodeScalarSwitch::GetOption( int t ) const
 	{
-		check( t>=0 && t<(int)m_pD->m_options.size() );
+		check( t>=0 && t<(int)m_pD->m_options.Num() );
 		return m_pD->m_options[t].get();
 	}
 
@@ -113,7 +113,7 @@ namespace mu
 	//---------------------------------------------------------------------------------------------
     void NodeScalarSwitch::SetOption( int t, NodeScalarPtr pNode )
 	{
-		check( t>=0 && t<(int)m_pD->m_options.size() );
+		check( t>=0 && t<(int)m_pD->m_options.Num() );
 		m_pD->m_options[t] = pNode;
 	}
 
diff --git a/Engine/Plugins/Experimental/Mutable/Source/MutableTools/Private/MuT/NodeScalarSwitchPrivate.h b/Engine/Plugins/Experimental/Mutable/Source/MutableTools/Private/MuT/NodeScalarSwitchPrivate.h
--- a/Engine/Plugins/Experimental/Mutable/Source/MutableTools/Private/MuT/NodeScalarSwitchPrivate.h
+++ b/Engine/Plugins/Experimental/Mutable/Source/MutableTools/Private/MuT/NodeScalarSwitchPrivate.h
@@ -23,7 +23,7 @@ namespace mu
 		static NODE_TYPE s_type;
 
 		NodeScalarPtr m_pParameter;
-        vector<NodeScalarPtr> m_options;
+        TArray<NodeScalarPtr> m_options;
 
 		//!
 		void Serialise( OutputArchive& arch ) const
diff --git a/Engine/Plugins/Experimental/Mutable/Source/MutableTools/Private/MuT/NodeScalarVariation.cpp b/Engine/Plugins/Experimental/Mutable/Source/MutableTools/Private/MuT/NodeScalarVariation.cpp
--- a/Engine/Plugins/Experimental/Mutable/Source/MutableTools/Private/MuT/NodeScalarVariation.cpp
+++ b/Engine/Plugins/Experimental/Mutable/Source/MutableTools/Private/MuT/NodeScalarVariation.cpp
@@ -30,7 +30,7 @@ namespace mu
 
 
     //---------------------------------------------------------------------------------------------
-    int NodeScalarVariation::GetInputCount() const { return 1 + int( m_pD->m_variations.size() ); }
+    int NodeScalarVariation::GetInputCount() const { return 1 + int( m_pD->m_variations.Num() ); }
 
 
     //---------------------------------------------------------------------------------------------
@@ -44,11 +44,11 @@ namespace mu
         }
         i -= 1;
 
-        if ( i < int( m_pD->m_variations.size() ) )
+        if ( i < int( m_pD->m_variations.Num() ) )
         {
             return m_pD->m_variations[i].m_scalar.get();
         }
-        i -= int( m_pD->m_variations.size() );
+        i -= int( m_pD->m_variations.Num() );
 
         return nullptr;
     }
@@ -66,13 +66,13 @@ namespace mu
         }
 
         i -= 1;
-        if ( i < int( m_pD->m_variations.size() ) )
+        if ( i < int( m_pD->m_variations.Num() ) )
         {
 
             m_pD->m_variations[i].m_scalar = dynamic_cast<NodeScalar*>( pNode.get() );
             return;
         }
-        i -= (int)m_pD->m_variations.size();
+        i -= (int)m_pD->m_variations.Num();
     }
 
 
@@ -83,20 +83,20 @@ namespace mu
 
 
     //---------------------------------------------------------------------------------------------
-    int NodeScalarVariation::GetVariationCount() const { return int( m_pD->m_variations.size() ); }
+    int NodeScalarVariation::GetVariationCount() const { return int( m_pD->m_variations.Num() ); }
 
 
     //---------------------------------------------------------------------------------------------
     void NodeScalarVariation::SetVariationCount( int num )
     {
         check( num >= 0 );
-        m_pD->m_variations.resize( num );
+        m_pD->m_variations.SetNum( num );
     }
 
     //---------------------------------------------------------------------------------------------
     void NodeScalarVariation::SetVariationTag( int index, const char* strTag )
     {
-        check( index >= 0 && index < (int)m_pD->m_variations.size() );
+        check( index >= 0 && index < (int)m_pD->m_variations.Num() );
         check( strTag );
 
         if ( strTag )
@@ -113,7 +113,7 @@ namespace mu
     //---------------------------------------------------------------------------------------------
     void NodeScalarVariation::SetVariationScalar( int index, NodeScalar* pNode )
     {
-        check( index >= 0 && index < (int)m_pD->m_variations.size() );
+        check( index >= 0 && index < (int)m_pD->m_variations.Num() );
 
         m_pD->m_variations[index].m_scalar = pNode;
     }
diff --git a/Engine/Plugins/Experimental/Mutable/Source/MutableTools/Private/MuT/NodeStringConstant.cpp b/Engine/Plugins/Experimental/Mutable/Source/MutableTools/Private/MuT/NodeStringConstant.cpp
--- a/Engine/Plugins/Experimental/Mutable/Source/MutableTools/Private/MuT/NodeStringConstant.cpp
+++ b/Engine/Plugins/Experimental/Mutable/Source/MutableTools/Private/MuT/NodeStringConstant.cpp
@@ -9,9 +9,6 @@
 #include "MuT/NodeStringConstantPrivate.h"
 
 
-#define NODE_INPUT_COUNT 	0
-
-
 namespace mu
 {
 
@@ -34,7 +31,7 @@ namespace mu
 	//---------------------------------------------------------------------------------------------
 	int NodeStringConstant::GetInputCount() const
 	{
-		return NODE_INPUT_COUNT;
+		return 0;
 	}
 
 
@@ -70,5 +67,3 @@ namespace mu
 
 }
 
-#undef NODE_INPUT_COUNT
-
diff --git a/Engine/Plugins/Experimental/Mutable/Source/MutableTools/Private/MuT/NodeStringParameter.cpp b/Engine/Plugins/Experimental/Mutable/Source/MutableTools/Private/MuT/NodeStringParameter.cpp
--- a/Engine/Plugins/Experimental/Mutable/Source/MutableTools/Private/MuT/NodeStringParameter.cpp
+++ b/Engine/Plugins/Experimental/Mutable/Source/MutableTools/Private/MuT/NodeStringParameter.cpp
@@ -34,9 +34,9 @@ namespace mu
 	//---------------------------------------------------------------------------------------------
 	int NodeStringParameter::GetInputCount() const
 	{
-        return int( m_pD->m_additionalImages.size()
+        return int( m_pD->m_additionalImages.Num()
                     +
-                    m_pD->m_ranges.size() );
+                    m_pD->m_ranges.Num() );
 	}
 
 
@@ -45,8 +45,8 @@ namespace mu
 	{
         check( i<GetInputCount() );
 
-        int imageCount = int(m_pD->m_additionalImages.size());
-        int rangeCount = int(m_pD->m_ranges.size());
+        int imageCount = int(m_pD->m_additionalImages.Num());
+        int rangeCount = int(m_pD->m_ranges.Num());
         if (i<imageCount)
         {
             return m_pD->m_additionalImages[i].get();
@@ -64,8 +64,8 @@ namespace mu
     void NodeStringParameter::SetInputNode( int i, NodePtr n )
 	{
         check( i<GetInputCount() );
-        int imageCount = int(m_pD->m_additionalImages.size());
-        int rangeCount = int(m_pD->m_ranges.size());
+        int imageCount = int(m_pD->m_additionalImages.Num());
+        int rangeCount = int(m_pD->m_ranges.Num());
         if (i<imageCount)
         {
             m_pD->m_additionalImages[i] = dynamic_cast<NodeImage*>(n.get());
@@ -151,15 +151,15 @@ namespace mu
     void NodeStringParameter::SetRangeCount( int i )
     {
         check(i>=0);
-        m_pD->m_ranges.resize(i);
+        m_pD->m_ranges.SetNum(i);
     }
 
 
     //---------------------------------------------------------------------------------------------
     void NodeStringParameter::SetRange( int i, NodeRangePtr pRange )
     {
-        check( i>=0 && i<int(m_pD->m_ranges.size()) );
-        if ( i>=0 && i<int(m_pD->m_ranges.size()) )
+        check( i>=0 && i<int(m_pD->m_ranges.Num()) );
+        if ( i>=0 && i<int(m_pD->m_ranges.Num()) )
         {
             m_pD->m_ranges[i] = pRange;
         }
diff --git a/Engine/Plugins/Experimental/Mutable/Source/MutableTools/Private/MuT/NodeStringParameterPrivate.h b/Engine/Plugins/Experimental/Mutable/Source/MutableTools/Private/MuT/NodeStringParameterPrivate.h
--- a/Engine/Plugins/Experimental/Mutable/Source/MutableTools/Private/MuT/NodeStringParameterPrivate.h
+++ b/Engine/Plugins/Experimental/Mutable/Source/MutableTools/Private/MuT/NodeStringParameterPrivate.h
@@ -30,9 +30,9 @@ namespace mu
 		string m_uid;
 		PARAMETER_DETAILED_TYPE m_detailedType = PARAMETER_DETAILED_TYPE::UNKNOWN;
 
-        vector<Ptr<NodeImage>> m_additionalImages;
+        TArray<Ptr<NodeImage>> m_additionalImages;
 
-        vector<Ptr<NodeRange>> m_ranges;
+		TArray<Ptr<NodeRange>> m_ranges;
 
 		//!
 		void Serialise( OutputArchive& arch ) const
diff --git a/Engine/Plugins/Experimental/Mutable/Source/MutableTools/Private/MuT/NodeSurfaceEdit.cpp b/Engine/Plugins/Experimental/Mutable/Source/MutableTools/Private/MuT/NodeSurfaceEdit.cpp
--- a/Engine/Plugins/Experimental/Mutable/Source/MutableTools/Private/MuT/NodeSurfaceEdit.cpp
+++ b/Engine/Plugins/Experimental/Mutable/Source/MutableTools/Private/MuT/NodeSurfaceEdit.cpp
@@ -43,7 +43,7 @@ namespace mu
 		return
 			3 	// Extend mesh, morph, factor
 			+
-			(int)m_pD->m_textures.size()*2
+			m_pD->m_textures.Num()*2
 			;
 	}
 
@@ -136,14 +136,14 @@ namespace mu
     //---------------------------------------------------------------------------------------------
     void NodeSurfaceEdit::AddTag(const char* tagName)
     {
-        m_pD->m_tags.push_back(tagName);
+        m_pD->m_tags.Add(tagName);
     }
 
 
     //---------------------------------------------------------------------------------------------
     int NodeSurfaceEdit::GetTagCount() const
     {
-        return int( m_pD->m_tags.size());
+        return m_pD->m_tags.Num();
     }
 
 
@@ -203,7 +203,7 @@ namespace mu
 	//---------------------------------------------------------------------------------------------
     int NodeSurfaceEdit::GetImageCount() const
 	{
-		return (int)m_pD->m_textures.size();
+		return m_pD->m_textures.Num();
 	}
 
 
@@ -211,14 +211,14 @@ namespace mu
     void NodeSurfaceEdit::SetImageCount( int num )
 	{
 		check( num >=0 );
-		m_pD->m_textures.resize( num );
+		m_pD->m_textures.SetNum( num );
 	}
 
 
 	//---------------------------------------------------------------------------------------------
     NodeImagePtr NodeSurfaceEdit::GetImage( int index ) const
 	{
-		check( index >=0 && index < (int)m_pD->m_textures.size() );
+		check( index >=0 && index < m_pD->m_textures.Num() );
 
 		return m_pD->m_textures[ index ].m_pExtend.get();
 	}
@@ -227,7 +227,7 @@ namespace mu
 	//---------------------------------------------------------------------------------------------
     void NodeSurfaceEdit::SetImage( int index, NodeImagePtr pNode )
 	{
-		check( index >=0 && index < (int)m_pD->m_textures.size() );
+		check( index >=0 && index < m_pD->m_textures.Num() );
 
 		m_pD->m_textures[ index ].m_pExtend = pNode;
 	}
@@ -236,7 +236,7 @@ namespace mu
 	//---------------------------------------------------------------------------------------------
     NodePatchImagePtr NodeSurfaceEdit::GetPatch( int index ) const
 	{
-		check( index >=0 && index < (int)m_pD->m_textures.size() );
+		check( index >=0 && index < m_pD->m_textures.Num() );
 
 		return m_pD->m_textures[ index ].m_pPatch.get();
 	}
@@ -245,7 +245,7 @@ namespace mu
 	//---------------------------------------------------------------------------------------------
     void NodeSurfaceEdit::SetPatch( int index, NodePatchImagePtr pNode )
 	{
-		check( index >=0 && index < (int)m_pD->m_textures.size() );
+		check( index >=0 && index < m_pD->m_textures.Num() );
 
 		m_pD->m_textures[ index ].m_pPatch = pNode;
 	}
diff --git a/Engine/Plugins/Experimental/Mutable/Source/MutableTools/Private/MuT/NodeSurfaceEditPrivate.h b/Engine/Plugins/Experimental/Mutable/Source/MutableTools/Private/MuT/NodeSurfaceEditPrivate.h
--- a/Engine/Plugins/Experimental/Mutable/Source/MutableTools/Private/MuT/NodeSurfaceEditPrivate.h
+++ b/Engine/Plugins/Experimental/Mutable/Source/MutableTools/Private/MuT/NodeSurfaceEditPrivate.h
@@ -58,10 +58,10 @@ namespace mu
 			}
 		};
 
-		vector<TEXTURE> m_textures;
+		TArray<TEXTURE> m_textures;
 
         //! Tags in this surface edit
-        vector<mu::string> m_tags;
+		TArray<mu::string> m_tags;
 
 		//! Factor of the morph
 		NodeScalarPtr m_pFactor;
diff --git a/Engine/Plugins/Experimental/Mutable/Source/MutableTools/Private/MuT/NodeSurfaceNew.cpp b/Engine/Plugins/Experimental/Mutable/Source/MutableTools/Private/MuT/NodeSurfaceNew.cpp
--- a/Engine/Plugins/Experimental/Mutable/Source/MutableTools/Private/MuT/NodeSurfaceNew.cpp
+++ b/Engine/Plugins/Experimental/Mutable/Source/MutableTools/Private/MuT/NodeSurfaceNew.cpp
@@ -41,8 +41,8 @@ namespace mu
 	//---------------------------------------------------------------------------------------------
     int NodeSurfaceNew::GetInputCount() const
 	{
-        return (int)( m_pD->m_meshes.size() + m_pD->m_images.size() + m_pD->m_vectors.size() +
-                      m_pD->m_scalars.size() + m_pD->m_strings.size() );
+        return m_pD->m_meshes.Num() + m_pD->m_images.Num() + m_pD->m_vectors.Num() +
+                      m_pD->m_scalars.Num() + m_pD->m_strings.Num();
     }
 
 
@@ -51,51 +51,51 @@ namespace mu
 	{
 		check( i >=0 && i < GetInputCount() );
 
-        std::size_t index = size_t(i);
+        int32 index = i;
 
-        if ( index < m_pD->m_meshes.size() )
+        if ( index < m_pD->m_meshes.Num() )
         {
             return m_pD->m_meshes[index].m_pMesh.get();
         }
 		else
 		{
-            index -= m_pD->m_meshes.size();
+            index -= m_pD->m_meshes.Num();
         }
 
-        if ( index < m_pD->m_images.size() )
+        if ( index < m_pD->m_images.Num() )
         {
             return m_pD->m_images[index].m_pImage.get();
         }
         else 
 		{
-            index -= m_pD->m_images.size();
+            index -= m_pD->m_images.Num();
         }
 
-        if ( index < m_pD->m_vectors.size() )
+        if ( index < m_pD->m_vectors.Num() )
         {
             return m_pD->m_vectors[index].m_pVector.get();
         }
         else
         {
-            index -= m_pD->m_vectors.size();
+            index -= m_pD->m_vectors.Num();
         }
 
-        if ( index < m_pD->m_scalars.size() )
+        if ( index < m_pD->m_scalars.Num() )
         {
             return m_pD->m_scalars[index].m_pScalar.get();
         }
         else
         {
-            index -= m_pD->m_scalars.size();
+            index -= m_pD->m_scalars.Num();
         }
 
-        if ( index < m_pD->m_strings.size() )
+        if ( index < m_pD->m_strings.Num() )
         {
             return m_pD->m_strings[index].m_pString.get();
         }
         else
         {
-            index -= m_pD->m_strings.size();
+            index -= m_pD->m_strings.Num();
         }
 
         return nullptr;
@@ -109,55 +109,55 @@ namespace mu
 
         std::size_t index = size_t( i );
 
-        if ( index<m_pD->m_meshes.size() )
+        if ( index<m_pD->m_meshes.Num() )
 		{
 			m_pD->m_meshes[ index ].m_pMesh = dynamic_cast<NodeMesh*>(pNode.get());
             return;
         }
 		else
 		{
-            index -= m_pD->m_meshes.size();
+            index -= m_pD->m_meshes.Num();
         }
 
 
-        if ( index < m_pD->m_images.size())
+        if ( index < m_pD->m_images.Num())
         {
 			m_pD->m_images[index].m_pImage = dynamic_cast<NodeImage*>(pNode.get());
             return;
         }
         else
         {
-            index -= m_pD->m_images.size();
+            index -= m_pD->m_images.Num();
         }
         
-		if ( index < m_pD->m_vectors.size() )
+		if ( index < m_pD->m_vectors.Num() )
         {
             m_pD->m_vectors[index].m_pVector = dynamic_cast<NodeColour*>(pNode.get());
             return;
         }
         else
         {
-            index -= m_pD->m_vectors.size();
+            index -= m_pD->m_vectors.Num();
         }
 
-        if ( index < m_pD->m_scalars.size() )
+        if ( index < m_pD->m_scalars.Num() )
         {
             m_pD->m_scalars[index].m_pScalar = dynamic_cast<NodeScalar*>( pNode.get() );
             return;
         }
         else
         {
-            index -= m_pD->m_scalars.size();
+            index -= m_pD->m_scalars.Num();
         }
 
-        if ( index < m_pD->m_strings.size() )
+        if ( index < m_pD->m_strings.Num() )
         {
             m_pD->m_strings[index].m_pString = dynamic_cast<NodeString*>( pNode.get() );
             return;
         }
         else
         {
-            index -= m_pD->m_strings.size();
+            index -= m_pD->m_strings.Num();
         }
     }
 
@@ -197,7 +197,7 @@ namespace mu
 	//---------------------------------------------------------------------------------------------
     int NodeSurfaceNew::GetMeshCount() const
 	{
-		return (int)m_pD->m_meshes.size();
+		return (int)m_pD->m_meshes.Num();
 	}
 
 
@@ -205,14 +205,14 @@ namespace mu
     void NodeSurfaceNew::SetMeshCount( int num )
 	{
 		check( num >=0 );
-		m_pD->m_meshes.resize( num );
+		m_pD->m_meshes.SetNum( num );
 	}
 
 
 	//---------------------------------------------------------------------------------------------
     NodeMeshPtr NodeSurfaceNew::GetMesh( int index ) const
 	{
-		check( index >=0 && index < (int)m_pD->m_meshes.size() );
+		check( index >=0 && index < (int)m_pD->m_meshes.Num() );
 
 		return m_pD->m_meshes[ index ].m_pMesh.get();
 	}
@@ -221,7 +221,7 @@ namespace mu
 	//---------------------------------------------------------------------------------------------
     void NodeSurfaceNew::SetMesh( int index, NodeMeshPtr pNode )
 	{
-		check( index >=0 && index < (int)m_pD->m_meshes.size() );
+		check( index >=0 && index < (int)m_pD->m_meshes.Num() );
 
 		m_pD->m_meshes[ index ].m_pMesh = pNode;
 	}
@@ -230,7 +230,7 @@ namespace mu
 	//---------------------------------------------------------------------------------------------
     const char* NodeSurfaceNew::GetMeshName( int index ) const
 	{
-		check( index >=0 && index < (int)m_pD->m_meshes.size() );
+		check( index >=0 && index < (int)m_pD->m_meshes.Num() );
 
 		const char* strResult = m_pD->m_meshes[ index ].m_name.c_str();
 
@@ -241,7 +241,7 @@ namespace mu
 	//---------------------------------------------------------------------------------------------
     void NodeSurfaceNew::SetMeshName( int index, const char* strName )
 	{
-		check( index >=0 && index < (int)m_pD->m_meshes.size() );
+		check( index >=0 && index < (int)m_pD->m_meshes.Num() );
 
 		if (strName)
 		{
@@ -257,7 +257,7 @@ namespace mu
 	//---------------------------------------------------------------------------------------------
     int NodeSurfaceNew::GetImageCount() const
 	{
-		return (int)m_pD->m_images.size();
+		return (int)m_pD->m_images.Num();
 	}
 
 
@@ -265,14 +265,14 @@ namespace mu
     void NodeSurfaceNew::SetImageCount( int num )
 	{
 		check( num >=0 );
-		m_pD->m_images.resize( num );
+		m_pD->m_images.SetNum( num );
 	}
 
 
 	//---------------------------------------------------------------------------------------------
     NodeImagePtr NodeSurfaceNew::GetImage( int index ) const
 	{
-		check( index >=0 && index < (int)m_pD->m_images.size() );
+		check( index >=0 && index < (int)m_pD->m_images.Num() );
 
 		return m_pD->m_images[ index ].m_pImage.get();
 	}
@@ -281,7 +281,7 @@ namespace mu
 	//---------------------------------------------------------------------------------------------
     void NodeSurfaceNew::SetImage( int index, NodeImagePtr pNode )
 	{
-		check( index >=0 && index < (int)m_pD->m_images.size() );
+		check( index >=0 && index < (int)m_pD->m_images.Num() );
 
 		m_pD->m_images[ index ].m_pImage = pNode;
 	}
@@ -290,7 +290,7 @@ namespace mu
 	//---------------------------------------------------------------------------------------------
     const char* NodeSurfaceNew::GetImageName( int index ) const
 	{
-		check( index >=0 && index < (int)m_pD->m_images.size() );
+		check( index >=0 && index < (int)m_pD->m_images.Num() );
 
 		const char* strResult = m_pD->m_images[ index ].m_name.c_str();
 
@@ -301,7 +301,7 @@ namespace mu
 	//---------------------------------------------------------------------------------------------
     void NodeSurfaceNew::SetImageName( int index, const char* strName )
 	{
-		check( index >=0 && index < (int)m_pD->m_images.size() );
+		check( index >=0 && index < (int)m_pD->m_images.Num() );
 
 		if (strName)
 		{
@@ -317,7 +317,7 @@ namespace mu
     //---------------------------------------------------------------------------------------------
     int NodeSurfaceNew::GetImageLayoutIndex( int index ) const
     {
-        check( index >= 0 && index < (int)m_pD->m_images.size() );
+        check( index >= 0 && index < (int)m_pD->m_images.Num() );
 
         return int(m_pD->m_images[index].m_layoutIndex);
     }
@@ -326,7 +326,7 @@ namespace mu
 	//---------------------------------------------------------------------------------------------
 	void NodeSurfaceNew::SetImageLayoutIndex(int index, int layoutIndex)
 	{
-		check(index >= 0 && index < (int)m_pD->m_images.size());
+		check(index >= 0 && index < (int)m_pD->m_images.Num());
 
 		m_pD->m_images[index].m_layoutIndex = int8_t(layoutIndex);
 	}
@@ -335,7 +335,7 @@ namespace mu
 	//---------------------------------------------------------------------------------------------
 	void NodeSurfaceNew::SetImageAdditionalNames(int index, const char* strMaterialName, const char* strMaterialParameterName )
 	{
-		if (index >= 0 && index < (int)m_pD->m_images.size())
+		if (index >= 0 && index < (int)m_pD->m_images.Num())
 		{
 			m_pD->m_images[index].m_materialName = strMaterialName;
 			m_pD->m_images[index].m_materialParameterName = strMaterialParameterName;
@@ -350,7 +350,7 @@ namespace mu
     //---------------------------------------------------------------------------------------------
 	int NodeSurfaceNew::GetVectorCount() const
 	{
-		return (int)m_pD->m_vectors.size();
+		return (int)m_pD->m_vectors.Num();
 	}
 
 
@@ -358,14 +358,14 @@ namespace mu
 	void NodeSurfaceNew::SetVectorCount(int num)
 	{
 		check(num >= 0);
-		m_pD->m_vectors.resize(num);
+		m_pD->m_vectors.SetNum(num);
 	}
 
 
 	//---------------------------------------------------------------------------------------------
 	NodeColourPtr NodeSurfaceNew::GetVector(int index) const
 	{
-		check(index >= 0 && index < (int)m_pD->m_vectors.size());
+		check(index >= 0 && index < (int)m_pD->m_vectors.Num());
 
 		return m_pD->m_vectors[index].m_pVector.get();
 	}
@@ -374,7 +374,7 @@ namespace mu
 	//---------------------------------------------------------------------------------------------
 	void NodeSurfaceNew::SetVector(int index, NodeColourPtr pNode)
 	{
-		check(index >= 0 && index < (int)m_pD->m_vectors.size());
+		check(index >= 0 && index < (int)m_pD->m_vectors.Num());
 
 		m_pD->m_vectors[index].m_pVector = pNode;
 	}
@@ -383,7 +383,7 @@ namespace mu
 	//---------------------------------------------------------------------------------------------
 	const char* NodeSurfaceNew::GetVectorName(int index) const
 	{
-        check(index >= 0 && index < (int)m_pD->m_vectors.size());
+        check(index >= 0 && index < (int)m_pD->m_vectors.Num());
 
 		const char* strResult = m_pD->m_vectors[index].m_name.c_str();
 
@@ -394,7 +394,7 @@ namespace mu
 	//---------------------------------------------------------------------------------------------
 	void NodeSurfaceNew::SetVectorName(int index, const char* strName)
 	{
-        check(index >= 0 && index < (int)m_pD->m_vectors.size());
+        check(index >= 0 && index < (int)m_pD->m_vectors.Num());
 
 		if (strName)
 		{
@@ -410,7 +410,7 @@ namespace mu
     //---------------------------------------------------------------------------------------------
     int NodeSurfaceNew::GetScalarCount() const
     {
-        return (int)m_pD->m_scalars.size();
+        return (int)m_pD->m_scalars.Num();
     }
 
 
@@ -418,14 +418,14 @@ namespace mu
     void NodeSurfaceNew::SetScalarCount(int num)
     {
         check(num >= 0);
-        m_pD->m_scalars.resize(num);
+        m_pD->m_scalars.SetNum(num);
     }
 
 
     //---------------------------------------------------------------------------------------------
     NodeScalarPtr NodeSurfaceNew::GetScalar(int index) const
     {
-        check(index >= 0 && index < (int)m_pD->m_scalars.size());
+        check(index >= 0 && index < (int)m_pD->m_scalars.Num());
 
         return m_pD->m_scalars[index].m_pScalar.get();
     }
@@ -434,7 +434,7 @@ namespace mu
     //---------------------------------------------------------------------------------------------
     void NodeSurfaceNew::SetScalar(int index, NodeScalarPtr pNode)
     {
-        check(index >= 0 && index < (int)m_pD->m_scalars.size());
+        check(index >= 0 && index < (int)m_pD->m_scalars.Num());
 
         m_pD->m_scalars[index].m_pScalar = pNode;
     }
@@ -443,7 +443,7 @@ namespace mu
     //---------------------------------------------------------------------------------------------
     const char* NodeSurfaceNew::GetScalarName(int index) const
     {
-        check(index >= 0 && index < (int)m_pD->m_scalars.size());
+        check(index >= 0 && index < (int)m_pD->m_scalars.Num());
 
         const char* strResult = m_pD->m_scalars[index].m_name.c_str();
 
@@ -454,7 +454,7 @@ namespace mu
     //---------------------------------------------------------------------------------------------
     void NodeSurfaceNew::SetScalarName(int index, const char* strName)
     {
-        check(index >= 0 && index < (int)m_pD->m_scalars.size());
+        check(index >= 0 && index < (int)m_pD->m_scalars.Num());
 
         if (strName)
         {
@@ -468,21 +468,21 @@ namespace mu
 
 
     //---------------------------------------------------------------------------------------------
-    int NodeSurfaceNew::GetStringCount() const { return (int)m_pD->m_strings.size(); }
+    int NodeSurfaceNew::GetStringCount() const { return (int)m_pD->m_strings.Num(); }
 
 
     //---------------------------------------------------------------------------------------------
     void NodeSurfaceNew::SetStringCount( int num )
     {
         check( num >= 0 );
-        m_pD->m_strings.resize( num );
+        m_pD->m_strings.SetNum( num );
     }
 
 
     //---------------------------------------------------------------------------------------------
     NodeStringPtr NodeSurfaceNew::GetString( int index ) const
     {
-        check( index >= 0 && index < (int)m_pD->m_strings.size() );
+        check( index >= 0 && index < (int)m_pD->m_strings.Num() );
 
         return m_pD->m_strings[index].m_pString.get();
     }
@@ -491,7 +491,7 @@ namespace mu
     //---------------------------------------------------------------------------------------------
     void NodeSurfaceNew::SetString( int index, NodeStringPtr pNode )
     {
-        check( index >= 0 && index < (int)m_pD->m_strings.size() );
+        check( index >= 0 && index < (int)m_pD->m_strings.Num() );
 
         m_pD->m_strings[index].m_pString = pNode;
     }
@@ -500,7 +500,7 @@ namespace mu
     //---------------------------------------------------------------------------------------------
     const char* NodeSurfaceNew::GetStringName( int index ) const
     {
-        check( index >= 0 && index < (int)m_pD->m_strings.size() );
+        check( index >= 0 && index < (int)m_pD->m_strings.Num() );
 
         const char* strResult = m_pD->m_strings[index].m_name.c_str();
 
@@ -511,7 +511,7 @@ namespace mu
     //---------------------------------------------------------------------------------------------
     void NodeSurfaceNew::SetStringName( int index, const char* strName )
     {
-        check( index >= 0 && index < (int)m_pD->m_strings.size() );
+        check( index >= 0 && index < (int)m_pD->m_strings.Num() );
 
         if ( strName )
         {
@@ -526,13 +526,10 @@ namespace mu
 
     //---------------------------------------------------------------------------------------------
     int NodeSurfaceNew::Private::FindImage( const char* strName ) const
-	{
-		int i = 0;
-		for ( vector<IMAGE>::const_iterator it = m_images.begin()
-			; it!=m_images.end()
-			; ++it, ++i )
+	{		
+		for ( int32 i = 0; i<m_images.Num(); ++i )
 		{
-			if ( it->m_name == strName )
+			if ( m_images[i].m_name == strName)
 			{
 				return i;
 			}
@@ -545,12 +542,9 @@ namespace mu
 	//---------------------------------------------------------------------------------------------
 	int NodeSurfaceNew::Private::FindMesh(const char* strName) const
 	{
-		int i = 0;
-		for (vector<MESH>::const_iterator it = m_meshes.begin()
-			; it != m_meshes.end()
-			; ++it, ++i)
+		for ( int32 i = 0; i < m_meshes.Num(); ++i)
 		{
-			if (it->m_name == strName)
+			if (m_meshes[i].m_name == strName)
 			{
 				return i;
 			}
@@ -563,12 +557,9 @@ namespace mu
     //---------------------------------------------------------------------------------------------
     int NodeSurfaceNew::Private::FindVector(const char* strName) const
     {
-        int i = 0;
-        for (vector<VECTOR>::const_iterator it = m_vectors.begin()
-            ; it != m_vectors.end()
-            ; ++it, ++i)
+        for ( int32 i = 0; i < m_vectors.Num(); ++i)
         {
-            if (it->m_name == strName)
+            if (m_vectors[i].m_name == strName)
             {
                 return i;
             }
@@ -581,11 +572,9 @@ namespace mu
     //---------------------------------------------------------------------------------------------
     int NodeSurfaceNew::Private::FindScalar( const char* strName ) const
     {
-        int i = 0;
-        for ( vector<SCALAR>::const_iterator it = m_scalars.begin(); it != m_scalars.end();
-              ++it, ++i )
+        for (int32 i = 0; i < m_scalars.Num(); ++i)
         {
-            if ( it->m_name == strName )
+            if ( m_scalars[i].m_name == strName )
             {
                 return i;
             }
@@ -598,11 +587,9 @@ namespace mu
     //---------------------------------------------------------------------------------------------
     int NodeSurfaceNew::Private::FindString( const char* strName ) const
     {
-        int i = 0;
-        for ( vector<STRING>::const_iterator it = m_strings.begin(); it != m_strings.end();
-              ++it, ++i )
+        for (int32 i = 0; i < m_strings.Num(); ++i)
         {
-            if ( it->m_name == strName )
+            if ( m_strings[i].m_name == strName )
             {
                 return i;
             }
@@ -615,14 +602,14 @@ namespace mu
     //---------------------------------------------------------------------------------------------
 	void NodeSurfaceNew::AddTag(const char* tagName)
 	{
-		m_pD->m_tags.push_back(tagName);
+		m_pD->m_tags.Add(tagName);
 	}
 
 
     //---------------------------------------------------------------------------------------------
     int NodeSurfaceNew::GetTagCount() const
     {
-        return int( m_pD->m_tags.size());
+        return m_pD->m_tags.Num();
     }
 
 
diff --git a/Engine/Plugins/Experimental/Mutable/Source/MutableTools/Private/MuT/NodeSurfaceNewPrivate.h b/Engine/Plugins/Experimental/Mutable/Source/MutableTools/Private/MuT/NodeSurfaceNewPrivate.h
--- a/Engine/Plugins/Experimental/Mutable/Source/MutableTools/Private/MuT/NodeSurfaceNewPrivate.h
+++ b/Engine/Plugins/Experimental/Mutable/Source/MutableTools/Private/MuT/NodeSurfaceNewPrivate.h
@@ -54,7 +54,7 @@ namespace mu
 			}
 		};
 
-		vector<MESH> m_meshes;
+		TArray<MESH> m_meshes;
 
 		struct IMAGE
 		{
@@ -86,7 +86,7 @@ namespace mu
             }
         };
 
-		vector<IMAGE> m_images;
+		TArray<IMAGE> m_images;
 
 		struct VECTOR
 		{
@@ -107,7 +107,7 @@ namespace mu
 			}
 		};
 
-		vector<VECTOR> m_vectors;
+		TArray<VECTOR> m_vectors;
 
         struct SCALAR
         {
@@ -128,7 +128,7 @@ namespace mu
             }
         };
 
-        vector<SCALAR> m_scalars;
+        TArray<SCALAR> m_scalars;
 
         struct STRING
         {
@@ -149,10 +149,10 @@ namespace mu
             }
         };
 
-        vector<STRING> m_strings;
+		TArray<STRING> m_strings;
 
         //! Tags in this surface
-		vector<mu::string> m_tags;
+		TArray<mu::string> m_tags;
 
 		//! Find an image node index by name or return -1
 		int FindImage( const char* strName ) const;
diff --git a/Engine/Plugins/Experimental/Mutable/Source/MutableTools/Private/MuT/NodeSurfaceVariation.cpp b/Engine/Plugins/Experimental/Mutable/Source/MutableTools/Private/MuT/NodeSurfaceVariation.cpp
--- a/Engine/Plugins/Experimental/Mutable/Source/MutableTools/Private/MuT/NodeSurfaceVariation.cpp
+++ b/Engine/Plugins/Experimental/Mutable/Source/MutableTools/Private/MuT/NodeSurfaceVariation.cpp
@@ -37,12 +37,12 @@ namespace mu
 	//---------------------------------------------------------------------------------------------
     int NodeSurfaceVariation::GetInputCount() const
 	{
-        size_t c = m_pD->m_defaultSurfaces.size();
-        c += m_pD->m_defaultModifiers.size();
+        int32 c = m_pD->m_defaultSurfaces.Num();
+        c += m_pD->m_defaultModifiers.Num();
         for (const auto& v : m_pD->m_variations)
 		{
-            c += v.m_surfaces.size();
-            c += v.m_modifiers.size();
+            c += v.m_surfaces.Num();
+            c += v.m_modifiers.Num();
         }
 
 		return (int)c;
@@ -54,31 +54,31 @@ namespace mu
 	{
 		check( i >=0 && i < GetInputCount() );
 
-        if ( i<(int)m_pD->m_defaultSurfaces.size())
+        if ( i<m_pD->m_defaultSurfaces.Num())
 		{
             return m_pD->m_defaultSurfaces[i].get();
 		}
-        i -= (int)m_pD->m_defaultSurfaces.size();
+        i -= m_pD->m_defaultSurfaces.Num();
 
-        if ( i<(int)m_pD->m_defaultModifiers.size())
+        if ( i<m_pD->m_defaultModifiers.Num())
         {
             return m_pD->m_defaultModifiers[i].get();
         }
-        i -= (int)m_pD->m_defaultModifiers.size();
+        i -= m_pD->m_defaultModifiers.Num();
 
         for (const auto& v : m_pD->m_variations)
         {
-            if (i < (int)v.m_surfaces.size())
+            if (i < (int)v.m_surfaces.Num())
             {
                 return v.m_surfaces[i].get();
             }
-            i -= (int)v.m_surfaces.size();
+            i -= (int)v.m_surfaces.Num();
 
-            if (i < (int)v.m_modifiers.size())
+            if (i < (int)v.m_modifiers.Num())
             {
                 return v.m_modifiers[i].get();
             }
-            i -= (int)v.m_modifiers.size();
+            i -= (int)v.m_modifiers.Num();
         }
 
 		return nullptr;
@@ -90,36 +90,36 @@ namespace mu
 	{
 		check( i >=0 && i <  GetInputCount());
 
-        if (i<(int)m_pD->m_defaultSurfaces.size())
+        if (i<m_pD->m_defaultSurfaces.Num())
 		{
             m_pD->m_defaultSurfaces[i] = dynamic_cast<NodeSurface*>(pNode.get());
 			return;
 		}
 
-        i -= (int)m_pD->m_defaultSurfaces.size();
-        if (i<(int)m_pD->m_defaultModifiers.size())
+        i -= m_pD->m_defaultSurfaces.Num();
+        if (i<m_pD->m_defaultModifiers.Num())
         {
 
             m_pD->m_defaultModifiers[i] = dynamic_cast<NodeModifier*>(pNode.get());
             return;
         }
-        i -= (int)m_pD->m_defaultModifiers.size();
+        i -= m_pD->m_defaultModifiers.Num();
 
         for (auto& v : m_pD->m_variations)
         {
-            if (i < (int)v.m_surfaces.size())
+            if (i < (int)v.m_surfaces.Num())
             {
                 v.m_surfaces[i] = dynamic_cast<NodeSurface*>(pNode.get());
                 return;
             }
-            i -= (int)v.m_surfaces.size();
+            i -= (int)v.m_surfaces.Num();
 
-            if (i < (int)v.m_modifiers.size())
+            if (i < (int)v.m_modifiers.Num())
             {
                 v.m_modifiers[i] = dynamic_cast<NodeModifier*>(pNode.get());
                 return;
             }
-            i -= (int)v.m_modifiers.size();
+            i -= (int)v.m_modifiers.Num();
         }
 	}
 
@@ -129,21 +129,21 @@ namespace mu
 	//---------------------------------------------------------------------------------------------
     void NodeSurfaceVariation::AddDefaultSurface( NodeSurface* p )
     {
-        m_pD->m_defaultSurfaces.push_back(p);
+        m_pD->m_defaultSurfaces.Add(p);
     }
 
 
     //---------------------------------------------------------------------------------------------
     void NodeSurfaceVariation::AddDefaultModifier( NodeModifier* p )
     {
-        m_pD->m_defaultModifiers.push_back(p);
+        m_pD->m_defaultModifiers.Add(p);
     }
 
 
     //---------------------------------------------------------------------------------------------
     int NodeSurfaceVariation::GetVariationCount() const
 	{
-		return (int)m_pD->m_variations.size();
+		return m_pD->m_variations.Num();
 	}
 
 
@@ -151,7 +151,7 @@ namespace mu
     void NodeSurfaceVariation::SetVariationCount( int num )
 	{
 		check( num >=0 );
-		m_pD->m_variations.resize( num );
+		m_pD->m_variations.SetNum( num );
 	}
 
 
@@ -165,7 +165,7 @@ namespace mu
 	//---------------------------------------------------------------------------------------------
 	void NodeSurfaceVariation::SetVariationTag(int index, const char* strTag)
 	{
-		check(index >= 0 && index < (int)m_pD->m_variations.size());
+		check(index >= 0 && index < m_pD->m_variations.Num());
 		check(strTag);
 
         if (strTag)
@@ -182,18 +182,18 @@ namespace mu
 	//---------------------------------------------------------------------------------------------
 	void NodeSurfaceVariation::AddVariationSurface(int index, NodeSurface* pNode)
 	{
-		check(index >= 0 && index < (int)m_pD->m_variations.size());
+		check(index >= 0 && index < m_pD->m_variations.Num());
 
-		m_pD->m_variations[index].m_surfaces.push_back(pNode);
+		m_pD->m_variations[index].m_surfaces.Add(pNode);
 	}
 
 
     //---------------------------------------------------------------------------------------------
     void NodeSurfaceVariation::AddVariationModifier(int index, NodeModifier* pModifier)
     {
-        check(index >= 0 && index < (int)m_pD->m_variations.size());
+        check(index >= 0 && index < m_pD->m_variations.Num());
 
-        m_pD->m_variations[index].m_modifiers.push_back(pModifier);
+        m_pD->m_variations[index].m_modifiers.Add(pModifier);
     }
 
 }
diff --git a/Engine/Plugins/Experimental/Mutable/Source/MutableTools/Private/MuT/NodeSurfaceVariationPrivate.h b/Engine/Plugins/Experimental/Mutable/Source/MutableTools/Private/MuT/NodeSurfaceVariationPrivate.h
--- a/Engine/Plugins/Experimental/Mutable/Source/MutableTools/Private/MuT/NodeSurfaceVariationPrivate.h
+++ b/Engine/Plugins/Experimental/Mutable/Source/MutableTools/Private/MuT/NodeSurfaceVariationPrivate.h
@@ -27,13 +27,13 @@ namespace mu
 
 		static NODE_TYPE s_type;
 
-        vector<NodeSurfacePtr> m_defaultSurfaces;
-        vector<NodeModifierPtr> m_defaultModifiers;
+        TArray<NodeSurfacePtr> m_defaultSurfaces;
+		TArray<NodeModifierPtr> m_defaultModifiers;
 
 		struct VARIATION
 		{
-            vector<NodeSurfacePtr> m_surfaces;
-            vector<NodeModifierPtr> m_modifiers;
+			TArray<NodeSurfacePtr> m_surfaces;
+			TArray<NodeModifierPtr> m_modifiers;
             string m_tag;
 
 			//!
@@ -61,7 +61,7 @@ namespace mu
 
         NodeSurfaceVariation::VariationType m_type = NodeSurfaceVariation::VariationType::Tag;
 
-		vector<VARIATION> m_variations;
+		TArray<VARIATION> m_variations;
 
 		//!
 		void Serialise( OutputArchive& arch ) const
diff --git a/Engine/Plugins/Experimental/Mutable/Source/MutableTools/Private/MuT/Platform.h b/Engine/Plugins/Experimental/Mutable/Source/MutableTools/Private/MuT/Platform.h
--- a/Engine/Plugins/Experimental/Mutable/Source/MutableTools/Private/MuT/Platform.h
+++ b/Engine/Plugins/Experimental/Mutable/Source/MutableTools/Private/MuT/Platform.h
@@ -5,12 +5,12 @@
 #include "MuR/Platform.h"
 #include "MuR/MemoryPrivate.h"
 
+#include <vector>
+#include <map>
+#include <set>
+#include <unordered_map>
 #include <sys/stat.h>
 
-//#if defined(_MSC_VER) || defined(__MINGW32__)
-    //#include <windef.h>
-    //#include <winbase.h>
-//#endif
 
 // This file contains some hacks to solve differences between platforms necessary for the tools
 // library and not present in the run-time library.
@@ -19,6 +19,28 @@
 namespace mu
 {
 
+
+	//! STL-like containers using this allocator
+	template< typename T >
+	using vector = std::vector<T>;
+
+	template< typename T >
+	using basic_string = std::basic_string<T, std::char_traits<T>>;
+
+	using string = std::basic_string<char, std::char_traits<char> >;
+
+	template< typename K, typename T >
+	using map = std::map< K, T, std::less<K> >;
+
+	template< typename T >
+	using set = std::set< T, std::less<T> >;
+
+	template< typename T >
+	using multiset = std::multiset< T, std::less<T> >;
+
+	template< typename K, typename T >
+	using pair = std::pair<K, T>;
+
     //-------------------------------------------------------------------------------------------------
     inline FILE* mutable_fopen
 		(
@@ -68,22 +90,5 @@ namespace mu
     }
 
 
-    //-------------------------------------------------------------------------------------------------
-    inline string mutable_path_separator()
-    {
-    #ifdef _WIN32
-        return "\\";
-    #else
-        return "/";
-    #endif
-    }
-
-
-    //-------------------------------------------------------------------------------------------------
-    inline bool mutable_file_exists( const string& path )
-    {
-        struct stat buffer;
-        return (stat (path.c_str(), &buffer) == 0);
-    }
-
 }
+
diff --git a/Engine/Plugins/Experimental/Mutable/Source/MutableTools/Private/MuT/Streams.cpp b/Engine/Plugins/Experimental/Mutable/Source/MutableTools/Private/MuT/Streams.cpp
--- a/Engine/Plugins/Experimental/Mutable/Source/MutableTools/Private/MuT/Streams.cpp
+++ b/Engine/Plugins/Experimental/Mutable/Source/MutableTools/Private/MuT/Streams.cpp
@@ -91,7 +91,7 @@ namespace mu
         uint8_t* pDest = (uint8_t*) pData;
         while (size && m_pD->m_pFile )
 		{			
-			int available = std::min( int(m_pD->m_size) - int(m_pD->m_pos), int(size) );
+			int available = FMath::Min( int(m_pD->m_size) - int(m_pD->m_pos), int(size) );
 
 			// Copy from what is already loaded
 			if ( available>0 )
@@ -137,7 +137,7 @@ namespace mu
 
 		while (size)
 		{
-            int available = std::min( int(MUTABLE_STREAM_BUFFER_SIZE) - int(m_pD->m_pos), int(size) );
+            int available = FMath::Min( int(MUTABLE_STREAM_BUFFER_SIZE) - int(m_pD->m_pos), int(size) );
 
 			// Copy from what fits in the current buffer
 			if ( available>0 )
diff --git a/Engine/Plugins/Experimental/Mutable/Source/MutableTools/Private/MuT/StreamsPrivate.h b/Engine/Plugins/Experimental/Mutable/Source/MutableTools/Private/MuT/StreamsPrivate.h
--- a/Engine/Plugins/Experimental/Mutable/Source/MutableTools/Private/MuT/StreamsPrivate.h
+++ b/Engine/Plugins/Experimental/Mutable/Source/MutableTools/Private/MuT/StreamsPrivate.h
@@ -137,13 +137,14 @@ namespace mu
             }
             else if (!m_fileName.IsEmpty())
             {
-                std::vector<char> buf(m_fileSize);
+				TArray<char> buf;
+				buf.SetNumUninitialized(m_fileSize);
                 auto resourceFile = FPlatformFileManager::Get().GetPlatformFile().OpenRead( *m_fileName );
                 check(resourceFile);
-                resourceFile->Read((uint8*)buf.data(), m_fileSize);
+                resourceFile->Read((uint8*)buf.GetData(), m_fileSize);
                 delete resourceFile;
 
-                InputMemoryStream stream( buf.data(), m_fileSize );
+                InputMemoryStream stream( buf.GetData(), m_fileSize );
                 InputArchive arch(&stream);
                 r =  R::StaticUnserialise( arch );
             }
diff --git a/Engine/Plugins/Experimental/Mutable/Source/MutableTools/Private/MuT/Table.cpp b/Engine/Plugins/Experimental/Mutable/Source/MutableTools/Private/MuT/Table.cpp
--- a/Engine/Plugins/Experimental/Mutable/Source/MutableTools/Private/MuT/Table.cpp
+++ b/Engine/Plugins/Experimental/Mutable/Source/MutableTools/Private/MuT/Table.cpp
@@ -77,7 +77,7 @@ namespace mu
 	//---------------------------------------------------------------------------------------------
 	int Table::AddColumn( const char* strName, TABLE_COLUMN_TYPE type )
 	{
-		int result = (int)m_pD->m_columns.size();
+		int result = m_pD->m_columns.Num();
 
 		TABLE_COLUMN c;
 
@@ -87,12 +87,12 @@ namespace mu
 		}
 		c.m_type = type;
 
-		m_pD->m_columns.push_back(c);
+		m_pD->m_columns.Add(c);
 
 		// Add it to all rows
-		for ( std::size_t r=0; r<m_pD->m_rows.size(); ++r )
+		for ( std::size_t r=0; r<m_pD->m_rows.Num(); ++r )
 		{
-			m_pD->m_rows[r].m_values.push_back( TABLE_VALUE() );
+			m_pD->m_rows[r].m_values.Add( TABLE_VALUE() );
 		}
 
 		return result;
@@ -104,7 +104,7 @@ namespace mu
 	{
 		int res = -1;
 
-		for ( std::size_t c=0; c<m_pD->m_columns.size(); ++c )
+		for ( std::size_t c=0; c<m_pD->m_columns.Num(); ++c )
 		{
 			if ( m_pD->m_columns[c].m_name == strName )
 			{
@@ -123,8 +123,8 @@ namespace mu
 
 		TABLE_ROW r;
 		r.m_id = id;
-		r.m_values.resize( m_pD->m_columns.size() );
-		m_pD->m_rows.push_back( r );
+		r.m_values.SetNum( m_pD->m_columns.Num() );
+		m_pD->m_rows.Add( r );
 	}
 
 
@@ -140,8 +140,8 @@ namespace mu
 	{
 		int row = m_pD->FindRow(rowId);
 		check( row>=0 );
-		check( column < (int)m_pD->m_columns.size() );
-		check( column < (int)m_pD->m_rows[row].m_values.size() );
+		check( column < (int)m_pD->m_columns.Num() );
+		check( column < (int)m_pD->m_rows[row].m_values.Num() );
 		check( m_pD->m_columns[column].m_type == TCT_SCALAR );
 
 		m_pD->m_rows[ row ].m_values[column].m_scalar = value;
@@ -153,8 +153,8 @@ namespace mu
 	{
 		int row = m_pD->FindRow(rowId);
 		check( row>=0 );
-		check( column < (int)m_pD->m_columns.size() );
-		check( column < (int)m_pD->m_rows[row].m_values.size() );
+		check( column < (int)m_pD->m_columns.Num() );
+		check( column < (int)m_pD->m_rows[row].m_values.Num() );
 		check( m_pD->m_columns[column].m_type == TCT_COLOUR );
 
 		m_pD->m_rows[ row ].m_values[column].m_colour = vec3<float>( r, g, b );
@@ -166,8 +166,8 @@ namespace mu
 	{
 		int row = m_pD->FindRow(rowId);
 		check( row>=0 );
-		check( column < (int)m_pD->m_columns.size() );
-		check( column < (int)m_pD->m_rows[row].m_values.size() );
+		check( column < (int)m_pD->m_columns.Num() );
+		check( column < (int)m_pD->m_rows[row].m_values.Num() );
 		check( m_pD->m_columns[column].m_type == TCT_IMAGE );
 
 		m_pD->m_rows[ row ].m_values[column].m_pProxyImage = new ResourceProxyMemory<Image>(pImage);
@@ -179,8 +179,8 @@ namespace mu
 	{
 		int row = m_pD->FindRow(rowId);
 		check( row>=0 );
-		check( column < (int)m_pD->m_columns.size() );
-		check( column < (int)m_pD->m_rows[row].m_values.size() );
+		check( column < (int)m_pD->m_columns.Num() );
+		check( column < (int)m_pD->m_rows[row].m_values.Num() );
 		check( m_pD->m_columns[column].m_type == TCT_MESH );
 
 		m_pD->m_rows[ row ].m_values[column].m_pMesh = pMesh;
@@ -192,8 +192,8 @@ namespace mu
 	{
 		int row = m_pD->FindRow(rowId);
 		check( row>=0 );
-		check( column < (int)m_pD->m_columns.size() );
-		check( column < (int)m_pD->m_rows[row].m_values.size() );
+		check( column < (int)m_pD->m_columns.Num() );
+		check( column < (int)m_pD->m_rows[row].m_values.Num() );
 		check( m_pD->m_columns[column].m_type == TCT_STRING );
 
 		m_pD->m_rows[ row ].m_values[column].m_string = strValue;
diff --git a/Engine/Plugins/Experimental/Mutable/Source/MutableTools/Private/MuT/TablePrivate.h b/Engine/Plugins/Experimental/Mutable/Source/MutableTools/Private/MuT/TablePrivate.h
--- a/Engine/Plugins/Experimental/Mutable/Source/MutableTools/Private/MuT/TablePrivate.h
+++ b/Engine/Plugins/Experimental/Mutable/Source/MutableTools/Private/MuT/TablePrivate.h
@@ -40,7 +40,7 @@ namespace mu
 	struct TABLE_ROW
 	{
         uint32_t m_id;
-		vector<TABLE_VALUE> m_values;
+		TArray<TABLE_VALUE> m_values;
 	};
 
 
@@ -50,8 +50,8 @@ namespace mu
 	public:
 
 		string m_name;
-		vector<TABLE_COLUMN> m_columns;
-		vector<TABLE_ROW> m_rows;
+		TArray<TABLE_COLUMN> m_columns;
+		TArray<TABLE_ROW> m_rows;
 		bool m_NoneOption = false;
 
 
@@ -61,19 +61,19 @@ namespace mu
             uint32_t ver = 1;
 			arch << ver;
 
-            arch << (uint32_t)m_columns.size();
-			for ( std::size_t c=0; c<m_columns.size(); ++c )
+            arch << (uint32_t)m_columns.Num();
+			for ( int32 c=0; c<m_columns.Num(); ++c )
 			{
 				arch << m_columns[c].m_name;
 				arch << m_columns[c].m_type;
 			}
 
-            arch << (uint32_t)m_rows.size();
-			for ( std::size_t r=0; r<m_rows.size(); ++r )
+            arch << (uint32_t)m_rows.Num();
+			for (int32 r=0; r<m_rows.Num(); ++r )
 			{
 				arch << m_rows[r].m_id;
 
-				for ( std::size_t c=0; c<m_columns.size(); ++c )
+				for (int32 c=0; c<m_columns.Num(); ++c )
 				{
 					const TABLE_VALUE& v = m_rows[r].m_values[c];
 
@@ -112,7 +112,7 @@ namespace mu
 
             uint32_t columnCount;
 			arch >> columnCount;
-			m_columns.resize( columnCount );
+			m_columns.SetNum( columnCount );
             for (uint32_t c=0; c<columnCount; ++c )
 			{
 				arch >> m_columns[c].m_name;
@@ -121,11 +121,11 @@ namespace mu
 
             uint32_t rowCount;
 			arch >> rowCount;
-			m_rows.resize( rowCount );
+			m_rows.SetNum( rowCount );
             for ( uint32_t r=0; r<rowCount; ++r )
 			{
 				arch >> m_rows[r].m_id;
-				m_rows[r].m_values.resize( columnCount );
+				m_rows[r].m_values.SetNum( columnCount );
 
                 for (uint32_t c=0; c<columnCount; ++c )
 				{
@@ -167,7 +167,7 @@ namespace mu
 		{
 			int res = -1;
 
-			for ( std::size_t r=0; res<0 && r<m_rows.size(); ++r )
+			for ( std::size_t r=0; res<0 && r<m_rows.Num(); ++r )
 			{
 				if ( m_rows[r].m_id==id )
 				{
diff --git a/Engine/Plugins/Experimental/Mutable/Source/MutableTools/Public/MuT/Compiler.h b/Engine/Plugins/Experimental/Mutable/Source/MutableTools/Public/MuT/Compiler.h
--- a/Engine/Plugins/Experimental/Mutable/Source/MutableTools/Public/MuT/Compiler.h
+++ b/Engine/Plugins/Experimental/Mutable/Source/MutableTools/Public/MuT/Compiler.h
@@ -155,14 +155,6 @@ namespace mu
         //! Compile the expression into a run-time model.
         ModelPtr Compile( const Ptr<Node>& pNode );
 
-        //! After a call to compile a model, this method will return the additional information on
-        //! the result. See the ModelReport class for more information.
-        //! If called without calling Compile first, a null pointer is returned.
-        ModelReportPtrConst GetModelReport() const;
-
-        //! Compile the transform object into a Transform that can be applied to a model expression.
-        TransformPtr Compile( const NodeTransformedObject* pNode );
-
         //! Return the log of messages of all the compile operations executed so far.
         ErrorLogPtrConst GetLog() const;
 
diff --git a/Engine/Plugins/Experimental/Mutable/Source/MutableTools/Public/MuT/Node.h b/Engine/Plugins/Experimental/Mutable/Source/MutableTools/Public/MuT/Node.h
--- a/Engine/Plugins/Experimental/Mutable/Source/MutableTools/Public/MuT/Node.h
+++ b/Engine/Plugins/Experimental/Mutable/Source/MutableTools/Public/MuT/Node.h
@@ -95,11 +95,6 @@ namespace mu
 		// Own Interface
 		//-----------------------------------------------------------------------------------------
 
-		//! Clone the node tree.
-		//! \param pMap is an optional parameter that will return a map from the source tree nodes
-		//! to the cloned tree nodes.
-        virtual NodePtr Clone( NodeMapPtr pMap = nullptr ) const = 0;
-
 		//! Node type hierarchy data.
         virtual const NODE_TYPE* GetType() const;
 		static const NODE_TYPE* GetStaticType();
@@ -128,42 +123,4 @@ namespace mu
 	};
 
 
-	//! Container mapping nodes to nodes.
-	class MUTABLETOOLS_API NodeMap : public RefCounted
-	{
-	public:
-
-		NodeMap();
-
-		//-----------------------------------------------------------------------------------------
-		// Own Interface
-		//-----------------------------------------------------------------------------------------
-
-		//! Return the number of elements in the map
-		int GetSize() const;
-
-		//! Insert a new key-value pair or overwrite the previews value with that key
-		void Add( const void* key, NodePtr value );
-
-		//! Get the value for a particular key. Returns 0 if the key is not in the map.
-		NodePtr Get( const void* key ) const;
-
-		//-----------------------------------------------------------------------------------------
-		// Interface pattern
-		//-----------------------------------------------------------------------------------------
-		class Private;
-		virtual Private* GetPrivate() const;
-
-	protected:
-
-		//! Forbidden. Manage with the Ptr<> template.
-		~NodeMap();
-
-	private:
-
-		Private* m_pD;
-
-	};
-
-
 }
diff --git a/Engine/Plugins/Experimental/Mutable/Source/MutableTools/Public/MuT/NodeBool.h b/Engine/Plugins/Experimental/Mutable/Source/MutableTools/Public/MuT/NodeBool.h
--- a/Engine/Plugins/Experimental/Mutable/Source/MutableTools/Public/MuT/NodeBool.h
+++ b/Engine/Plugins/Experimental/Mutable/Source/MutableTools/Public/MuT/NodeBool.h
@@ -116,8 +116,6 @@ namespace mu
 		// Node Interface
 		//-----------------------------------------------------------------------------------------
 
-        NodePtr Clone( NodeMapPtr pMap = 0 ) const override;
-
         const NODE_TYPE* GetType() const override;
 		static const NODE_TYPE* GetStaticType();
 
@@ -177,8 +175,6 @@ namespace mu
 		// Node Interface
 		//-----------------------------------------------------------------------------------------
 
-        NodePtr Clone( NodeMapPtr pMap = 0 ) const override;
-
         const NODE_TYPE* GetType() const override;
 		static const NODE_TYPE* GetStaticType();
 
@@ -250,7 +246,7 @@ namespace mu
 		// Node Interface
 		//-----------------------------------------------------------------------------------------
 
-        NodePtr Clone( NodeMapPtr pMap = 0 ) const override;
+        
 
         const NODE_TYPE* GetType() const override;
 		static const NODE_TYPE* GetStaticType();
@@ -305,7 +301,7 @@ namespace mu
 		// Node Interface
 		//-----------------------------------------------------------------------------------------
 
-        NodePtr Clone( NodeMapPtr pMap = 0 ) const override;
+        
 
         const NODE_TYPE* GetType() const override;
 		static const NODE_TYPE* GetStaticType();
@@ -364,7 +360,7 @@ namespace mu
 		// Node Interface
 		//-----------------------------------------------------------------------------------------
 
-        NodePtr Clone( NodeMapPtr pMap = 0 ) const override;
+        
 
         const NODE_TYPE* GetType() const override;
 		static const NODE_TYPE* GetStaticType();
diff --git a/Engine/Plugins/Experimental/Mutable/Source/MutableTools/Public/MuT/NodeColourArithmeticOperation.h b/Engine/Plugins/Experimental/Mutable/Source/MutableTools/Public/MuT/NodeColourArithmeticOperation.h
--- a/Engine/Plugins/Experimental/Mutable/Source/MutableTools/Public/MuT/NodeColourArithmeticOperation.h
+++ b/Engine/Plugins/Experimental/Mutable/Source/MutableTools/Public/MuT/NodeColourArithmeticOperation.h
@@ -38,8 +38,6 @@ namespace mu
 		// Node Interface
 		//-----------------------------------------------------------------------------------------
 
-		virtual NodePtr Clone( NodeMapPtr pMap = 0 ) const override;
-
 		virtual const NODE_TYPE* GetType() const override;
 		static const NODE_TYPE* GetStaticType();
 
diff --git a/Engine/Plugins/Experimental/Mutable/Source/MutableTools/Public/MuT/NodeColourConstant.h b/Engine/Plugins/Experimental/Mutable/Source/MutableTools/Public/MuT/NodeColourConstant.h
--- a/Engine/Plugins/Experimental/Mutable/Source/MutableTools/Public/MuT/NodeColourConstant.h
+++ b/Engine/Plugins/Experimental/Mutable/Source/MutableTools/Public/MuT/NodeColourConstant.h
@@ -36,7 +36,7 @@ namespace mu
 		// Node Interface
 		//-----------------------------------------------------------------------------------------
 
-        NodePtr Clone( NodeMapPtr pMap = 0 ) const override;
+        
 
         const NODE_TYPE* GetType() const override;
 		static const NODE_TYPE* GetStaticType();
diff --git a/Engine/Plugins/Experimental/Mutable/Source/MutableTools/Public/MuT/NodeColourFromScalars.h b/Engine/Plugins/Experimental/Mutable/Source/MutableTools/Public/MuT/NodeColourFromScalars.h
--- a/Engine/Plugins/Experimental/Mutable/Source/MutableTools/Public/MuT/NodeColourFromScalars.h
+++ b/Engine/Plugins/Experimental/Mutable/Source/MutableTools/Public/MuT/NodeColourFromScalars.h
@@ -42,8 +42,6 @@ namespace mu
 		// Node Interface
 		//-----------------------------------------------------------------------------------------
 
-        NodePtr Clone( NodeMapPtr pMap = nullptr ) const override;
-
 		const NODE_TYPE* GetType() const override;
 		static const NODE_TYPE* GetStaticType();
 
diff --git a/Engine/Plugins/Experimental/Mutable/Source/MutableTools/Public/MuT/NodeColourParameter.h b/Engine/Plugins/Experimental/Mutable/Source/MutableTools/Public/MuT/NodeColourParameter.h
--- a/Engine/Plugins/Experimental/Mutable/Source/MutableTools/Public/MuT/NodeColourParameter.h
+++ b/Engine/Plugins/Experimental/Mutable/Source/MutableTools/Public/MuT/NodeColourParameter.h
@@ -41,7 +41,7 @@ namespace mu
 		// Node Interface
 		//-----------------------------------------------------------------------------------------
 
-		NodePtr Clone( NodeMapPtr pMap = 0 ) const override;
+		
 
 		const NODE_TYPE* GetType() const override;
 		static const NODE_TYPE* GetStaticType();
diff --git a/Engine/Plugins/Experimental/Mutable/Source/MutableTools/Public/MuT/NodeImageInvert.h b/Engine/Plugins/Experimental/Mutable/Source/MutableTools/Public/MuT/NodeImageInvert.h
--- a/Engine/Plugins/Experimental/Mutable/Source/MutableTools/Public/MuT/NodeImageInvert.h
+++ b/Engine/Plugins/Experimental/Mutable/Source/MutableTools/Public/MuT/NodeImageInvert.h
@@ -38,8 +38,6 @@ namespace mu
 		// Node Interface
 		//-----------------------------------------------------------------------------------------
 
-		NodePtr Clone(NodeMapPtr pMap = 0) const override;
-
 		const NODE_TYPE* GetType() const override;
 		static const NODE_TYPE* GetStaticType();
 
diff --git a/Engine/Plugins/Experimental/Mutable/Source/MutableTools/Public/MuT/NodeImageNormalComposite.h b/Engine/Plugins/Experimental/Mutable/Source/MutableTools/Public/MuT/NodeImageNormalComposite.h
--- a/Engine/Plugins/Experimental/Mutable/Source/MutableTools/Public/MuT/NodeImageNormalComposite.h
+++ b/Engine/Plugins/Experimental/Mutable/Source/MutableTools/Public/MuT/NodeImageNormalComposite.h
@@ -40,8 +40,6 @@ namespace mu
 		// Node Interface
 		//-----------------------------------------------------------------------------------------
 
-		NodePtr Clone(NodeMapPtr pMap = 0) const override;
-
 		const NODE_TYPE* GetType() const override;
 		static const NODE_TYPE* GetStaticType();
 
diff --git a/Engine/Plugins/Experimental/Mutable/Source/MutableTools/Public/MuT/NodeMeshClipMorphPlane.h b/Engine/Plugins/Experimental/Mutable/Source/MutableTools/Public/MuT/NodeMeshClipMorphPlane.h
--- a/Engine/Plugins/Experimental/Mutable/Source/MutableTools/Public/MuT/NodeMeshClipMorphPlane.h
+++ b/Engine/Plugins/Experimental/Mutable/Source/MutableTools/Public/MuT/NodeMeshClipMorphPlane.h
@@ -37,8 +37,6 @@ namespace mu
 		// Node Interface
 		//-----------------------------------------------------------------------------------------
 
-		virtual NodePtr Clone( NodeMapPtr pMap = 0 ) const;
-
 		virtual const NODE_TYPE* GetType() const;
 		static const NODE_TYPE* GetStaticType();
 
diff --git a/Engine/Plugins/Experimental/Mutable/Source/MutableTools/Public/MuT/NodeMeshClipWithMesh.h b/Engine/Plugins/Experimental/Mutable/Source/MutableTools/Public/MuT/NodeMeshClipWithMesh.h
--- a/Engine/Plugins/Experimental/Mutable/Source/MutableTools/Public/MuT/NodeMeshClipWithMesh.h
+++ b/Engine/Plugins/Experimental/Mutable/Source/MutableTools/Public/MuT/NodeMeshClipWithMesh.h
@@ -37,8 +37,6 @@ namespace mu
 		// Node Interface
 		//-----------------------------------------------------------------------------------------
 
-		virtual NodePtr Clone( NodeMapPtr pMap = 0 ) const;
-
 		virtual const NODE_TYPE* GetType() const;
 		static const NODE_TYPE* GetStaticType();
 
diff --git a/Engine/Plugins/Experimental/Mutable/Source/MutableTools/Public/MuT/NodeMeshFormat.h b/Engine/Plugins/Experimental/Mutable/Source/MutableTools/Public/MuT/NodeMeshFormat.h
--- a/Engine/Plugins/Experimental/Mutable/Source/MutableTools/Public/MuT/NodeMeshFormat.h
+++ b/Engine/Plugins/Experimental/Mutable/Source/MutableTools/Public/MuT/NodeMeshFormat.h
@@ -39,7 +39,7 @@ namespace mu
 		// Node Interface
 		//-----------------------------------------------------------------------------------------
 
-		NodePtr Clone( NodeMapPtr pMap = 0 ) const override;
+		
 
 		const NODE_TYPE* GetType() const override;
 		static const NODE_TYPE* GetStaticType();
@@ -71,10 +71,6 @@ namespace mu
 		//! structure. If this is null, the vertex buffers will not be changed.
 		FMeshBufferSet& GetFaceBuffers();
 
-		//! Enable or disable the rebuilding of the tangent space of the mesh.
-		//! By default it is disabled.
-		void SetRebuildTangents( bool );
-
 		//-----------------------------------------------------------------------------------------
 		// Interface pattern
 		//-----------------------------------------------------------------------------------------
diff --git a/Engine/Plugins/Experimental/Mutable/Source/MutableTools/Public/MuT/NodeProjector.h b/Engine/Plugins/Experimental/Mutable/Source/MutableTools/Public/MuT/NodeProjector.h
--- a/Engine/Plugins/Experimental/Mutable/Source/MutableTools/Public/MuT/NodeProjector.h
+++ b/Engine/Plugins/Experimental/Mutable/Source/MutableTools/Public/MuT/NodeProjector.h
@@ -96,7 +96,7 @@ namespace mu
 		// Node Interface
 		//-----------------------------------------------------------------------------------------
 
-        NodePtr Clone( NodeMapPtr pMap = 0 ) const override;
+        
 
         const NODE_TYPE* GetType() const override;
 		static const NODE_TYPE* GetStaticType();
@@ -165,7 +165,7 @@ namespace mu
 		// Node Interface
 		//-----------------------------------------------------------------------------------------
 
-        NodePtr Clone( NodeMapPtr pMap = 0 ) const override;
+        
 
         const NODE_TYPE* GetType() const override;
 		static const NODE_TYPE* GetStaticType();
diff --git a/Engine/Plugins/Experimental/Mutable/Source/MutableTools/Public/MuT/NodeRangeFromScalar.h b/Engine/Plugins/Experimental/Mutable/Source/MutableTools/Public/MuT/NodeRangeFromScalar.h
--- a/Engine/Plugins/Experimental/Mutable/Source/MutableTools/Public/MuT/NodeRangeFromScalar.h
+++ b/Engine/Plugins/Experimental/Mutable/Source/MutableTools/Public/MuT/NodeRangeFromScalar.h
@@ -34,8 +34,6 @@ namespace mu
 		// Node Interface
 		//-----------------------------------------------------------------------------------------
 
-        NodePtr Clone( NodeMapPtr pMap = nullptr ) const override;
-
 		const NODE_TYPE* GetType() const override;
 		static const NODE_TYPE* GetStaticType();
 
diff --git a/Engine/Plugins/Experimental/Mutable/Source/MutableTools/Public/MuT/NodeScalarCurve.h b/Engine/Plugins/Experimental/Mutable/Source/MutableTools/Public/MuT/NodeScalarCurve.h
--- a/Engine/Plugins/Experimental/Mutable/Source/MutableTools/Public/MuT/NodeScalarCurve.h
+++ b/Engine/Plugins/Experimental/Mutable/Source/MutableTools/Public/MuT/NodeScalarCurve.h
@@ -36,7 +36,6 @@ namespace mu
 		//-----------------------------------------------------------------------------------------
 		// Node Interface
 		//-----------------------------------------------------------------------------------------
-        NodePtr Clone( NodeMapPtr pMap = nullptr ) const override;
 
         const NODE_TYPE* GetType() const override;
 		static const NODE_TYPE* GetStaticType();
