[mutable] Now data table rows can be excluded from the CO compilation by adding a Bool parameter to the struct and setting the value true.

+ Fixed information loss when a table node was refreshed.

#rnx

diff --git a/Engine/Plugins/Experimental/Mutable/Source/CustomizableObjectEditor/Classes/MuCOE/Nodes/CustomizableObjectNodeTable.h b/Engine/Plugins/Experimental/Mutable/Source/CustomizableObjectEditor/Classes/MuCOE/Nodes/CustomizableObjectNodeTable.h
--- a/Engine/Plugins/Experimental/Mutable/Source/CustomizableObjectEditor/Classes/MuCOE/Nodes/CustomizableObjectNodeTable.h
+++ b/Engine/Plugins/Experimental/Mutable/Source/CustomizableObjectEditor/Classes/MuCOE/Nodes/CustomizableObjectNodeTable.h
@@ -94,14 +94,14 @@ public:
 	GENERATED_BODY()
 
 	// Name of the property parameter
-	UPROPERTY(EditAnywhere, Category = CustomizableObject)
+	UPROPERTY(EditAnywhere, Category = TableProperties)
 	FString ParameterName = "Default Name";
 
-	UPROPERTY(EditAnywhere, Category = CustomizableObject)
+	UPROPERTY(EditAnywhere, Category = TableProperties)
 	bool bAddNoneOption;
 
 	// Pointer to the Data Table Asset represented in this node
-	UPROPERTY(EditAnywhere, Category = CustomizableObject, meta = (DontUpdateWhileEditing))
+	UPROPERTY(EditAnywhere, Category = TableProperties, meta = (DontUpdateWhileEditing))
 	TObjectPtr<UDataTable> Table = nullptr;
 	
 	UPROPERTY(EditAnywhere, Category = UI, meta = (DisplayName = "Parameter UI Metadata"))
@@ -122,9 +122,11 @@ public:
 	void AllocateDefaultPins(UCustomizableObjectNodeRemapPins* RemapPins) override;
 	bool IsNodeOutDatedAndNeedsRefresh() override;
 	FString GetRefreshMessage() const override;
-	virtual void RemapPins(const TMap<UEdGraphPin*, UEdGraphPin*>& PinsToRemap) override;
 	virtual bool ProvidesCustomPinRelevancyTest() const override;
 	virtual bool IsPinRelevant(const UEdGraphPin* Pin) const override;
+
+	/*** Allows to perform work when remapping the pin data. */
+	virtual void RemapPinsData(const TMap<UEdGraphPin*, UEdGraphPin*>& PinsToRemap) override;
 	
 	// Own interface
 	/** Given an Image pin, return the PinMode. */
@@ -229,12 +231,16 @@ public:
 	// We should do this in a template!
 	USkeletalMesh* GetSkeletalMeshAt(const UEdGraphPin* Pin, int32 RowIndex) const;
 
+	// Return the name of the enabled rows in the data table.
+	// Returns the name if the row has a bool column set as false (true == disabled)
+	TArray<FName> GetRowNames() const;
+
 private:
-	
+
 	/** Number of properties to know when the node needs an update */
 	UPROPERTY()
 	int32 NumProperties;
-
+	
 	FDelegateHandle OnTableChangedDelegateHandle;
 	
 	// Generates a mesh pin for each LOD and Material Surface of the reference SkeletalMesh
diff --git a/Engine/Plugins/Experimental/Mutable/Source/CustomizableObjectEditor/Private/MuCOE/GenerateMutableSource/GenerateMutableSourceMesh.cpp b/Engine/Plugins/Experimental/Mutable/Source/CustomizableObjectEditor/Private/MuCOE/GenerateMutableSource/GenerateMutableSourceMesh.cpp
--- a/Engine/Plugins/Experimental/Mutable/Source/CustomizableObjectEditor/Private/MuCOE/GenerateMutableSource/GenerateMutableSourceMesh.cpp
+++ b/Engine/Plugins/Experimental/Mutable/Source/CustomizableObjectEditor/Private/MuCOE/GenerateMutableSource/GenerateMutableSourceMesh.cpp
@@ -1929,7 +1929,7 @@ mu::NodeMeshPtr GenerateMorphMesh(const UEdGraphPin* Pin,
 	if (const UCustomizableObjectNodeTable* TypedNodeTable = Cast<UCustomizableObjectNodeTable>(Pin->GetOwningNode()))
 	{
 		// Generate a new Column for each morph
-		int32 NumRows = TypedNodeTable->Table->GetRowNames().Num();
+		int32 NumRows = TypedNodeTable->GetRowNames().Num();
 
 		// Should exist
 		mu::TablePtr Table = GenerationContext.GeneratedTables[TypedNodeTable->Table->GetName()];
diff --git a/Engine/Plugins/Experimental/Mutable/Source/CustomizableObjectEditor/Private/MuCOE/Nodes/CustomizableObjectNodeTable.cpp b/Engine/Plugins/Experimental/Mutable/Source/CustomizableObjectEditor/Private/MuCOE/Nodes/CustomizableObjectNodeTable.cpp
--- a/Engine/Plugins/Experimental/Mutable/Source/CustomizableObjectEditor/Private/MuCOE/Nodes/CustomizableObjectNodeTable.cpp
+++ b/Engine/Plugins/Experimental/Mutable/Source/CustomizableObjectEditor/Private/MuCOE/Nodes/CustomizableObjectNodeTable.cpp
@@ -522,7 +522,7 @@ FString UCustomizableObjectNodeTable::GetRefreshMessage() const
 }
 
 
-void UCustomizableObjectNodeTable::RemapPins(const TMap<UEdGraphPin*, UEdGraphPin*>& PinsToRemap)
+void UCustomizableObjectNodeTable::RemapPinsData(const TMap<UEdGraphPin*, UEdGraphPin*>& PinsToRemap)
 {
 	const UEdGraphSchema_CustomizableObject* Schema = GetDefault<UEdGraphSchema_CustomizableObject>();
 
@@ -530,20 +530,37 @@ void UCustomizableObjectNodeTable::RemapPins(const TMap<UEdGraphPin*, UEdGraphPi
 	{
 		if (Pair.Key->PinType.PinCategory == Schema->PC_Mesh)
 		{
-			UCustomizableObjectNodeTableMeshPinData* PinDataOldPin =Cast<UCustomizableObjectNodeTableMeshPinData>(GetPinData(*(Pair.Key)));
-			UCustomizableObjectNodeTableMeshPinData* PinDataNewPin =Cast<UCustomizableObjectNodeTableMeshPinData>(GetPinData(*(Pair.Value)));
-			
+			UCustomizableObjectNodeTableMeshPinData* PinDataOldPin = Cast<UCustomizableObjectNodeTableMeshPinData>(GetPinData(*(Pair.Key)));
+			UCustomizableObjectNodeTableMeshPinData* PinDataNewPin = Cast<UCustomizableObjectNodeTableMeshPinData>(GetPinData(*(Pair.Value)));
+
 			if (PinDataOldPin && PinDataNewPin)
 			{
 				PinDataOldPin->ColumnName = PinDataNewPin->ColumnName;
 				PinDataOldPin->MutableColumnName = PinDataNewPin->MutableColumnName;
 				PinDataOldPin->LOD = PinDataNewPin->LOD;
 				PinDataOldPin->Material = PinDataNewPin->Material;
+
+				// Keeping information added in layout editor if the layout is the same
+				for (TObjectPtr<UCustomizableObjectLayout>& NewLayout : PinDataNewPin->Layouts)
+				{
+					for (TObjectPtr<UCustomizableObjectLayout>& OldLayout : PinDataOldPin->Layouts)
+					{
+						if (NewLayout->GetLayoutName() == OldLayout->GetLayoutName())
+						{
+							NewLayout->Blocks = OldLayout->Blocks;
+							NewLayout->SetGridSize(OldLayout->GetGridSize());
+							NewLayout->SetMaxGridSize(OldLayout->GetMaxGridSize());
+							NewLayout->SetPackingStrategy(OldLayout->GetPackingStrategy());
+
+							break;
+						}
+					}
+				}
+
+				PinDataOldPin->Layouts = PinDataNewPin->Layouts;
 			}
 		}
 	}
-
-	Super::RemapPins(PinsToRemap);
 }
 
 
@@ -876,4 +893,56 @@ USkeletalMesh* UCustomizableObjectNodeTable::GetSkeletalMeshAt(const UEdGraphPin
 }
 
 
+TArray<FName> UCustomizableObjectNodeTable::GetRowNames() const
+{
+	TArray<FName> RowNames;
+
+	if (Table)
+	{
+		const UScriptStruct* TableStruct = Table->GetRowStruct();
+
+		if (!TableStruct)
+		{
+			return RowNames;
+		}
+
+		TArray<FName> TableRowNames = Table->GetRowNames();
+		FBoolProperty* BoolProperty = nullptr;
+
+		for (TFieldIterator<FProperty> PropertyIt(TableStruct); PropertyIt; ++PropertyIt)
+		{
+			BoolProperty = CastField<FBoolProperty>(*PropertyIt);
+
+			if (BoolProperty)
+			{
+				for (const FName& RowName : TableRowNames)
+				{
+					if (uint8* RowData = Table->FindRowUnchecked(RowName))
+					{
+						if (uint8* CellData = BoolProperty->ContainerPtrToValuePtr<uint8>(RowData, 0))
+						{
+							if (!BoolProperty->GetPropertyValue(CellData))
+							{
+								RowNames.Add(RowName);
+							}
+						}
+					}
+				}
+
+				// There should be only one Bool column
+				break;
+			}
+		}
+
+		// There is no Bool column
+		if (!BoolProperty)
+		{
+			return TableRowNames;
+		}
+	}
+
+	return RowNames;
+}
+
+
 #undef LOCTEXT_NAMESPACE
