[mutable] Improved table node properties tab usability:
+ Added two new sections: "Animation Properties" and "Default Mesh Layout Editor"
+ Fixed visual bug in animation comboboxes.
+ Added a clear button to reset animation combobox selections.
+ Improved tooltips.
+ Added a tooltip widget to the layout editor with all the shortcuts available.
+ Added clipping rect to the layout editor canvas.

#rnx

diff --git a/Engine/Plugins/Experimental/Mutable/Source/CustomizableObjectEditor/Private/MuCOE/CustomizableObjectEditor.cpp b/Engine/Plugins/Experimental/Mutable/Source/CustomizableObjectEditor/Private/MuCOE/CustomizableObjectEditor.cpp
--- a/Engine/Plugins/Experimental/Mutable/Source/CustomizableObjectEditor/Private/MuCOE/CustomizableObjectEditor.cpp
+++ b/Engine/Plugins/Experimental/Mutable/Source/CustomizableObjectEditor/Private/MuCOE/CustomizableObjectEditor.cpp
@@ -1735,36 +1735,7 @@ void FCustomizableObjectEditor::OnSelectedGraphNodesChanged(const FGraphPanelSel
 		// Set the size of the child widget to show properly the node widgets
 		FVector2D SplitWeights = { 0.15f,0.85f };
 		
-		// Create the specialised user interface for each node type
-		if ( UCustomizableObjectNodeLayoutBlocks* CustomizableObjectNodeLayoutBlocks = Cast<UCustomizableObjectNodeLayoutBlocks>(FirstNode) )
-		{
-			if (!LayoutBlocksEditor.IsValid())
-			{
-				LayoutBlocksEditor = SNew(SCustomizableObjectNodeLayoutBlocksEditor)
-					.CustomizableObjectEditor(SharedThis(this));
-			}
-
-			LayoutBlocksEditor->SetCurrentLayout(CustomizableObjectNodeLayoutBlocks->Layout);
-			SplitWeights = { 0.0f,1.0f };
-
-			ChildWidget = LayoutBlocksEditor;
-		}
-
-		else if (UCustomizableObjectNodeTable* CustomizableObjectNodeTable = Cast<UCustomizableObjectNodeTable>(FirstNode))
-		{
-			if (!LayoutBlocksEditor.IsValid())
-			{
-				LayoutBlocksEditor = SNew(SCustomizableObjectNodeLayoutBlocksEditor)
-					.CustomizableObjectEditor(SharedThis(this));
-			}
-
-			LayoutBlocksEditor->SetCurrentLayout(nullptr);
-			SplitWeights = { 0.32f,0.68f };
-
-			ChildWidget = LayoutBlocksEditor;
-		}
-
-		else if ( UCustomizableObjectNodeEditMaterial* CustomizableObjectNodeEditMaterial = Cast<UCustomizableObjectNodeEditMaterial>(FirstNode) )
+		if ( UCustomizableObjectNodeEditMaterial* CustomizableObjectNodeEditMaterial = Cast<UCustomizableObjectNodeEditMaterial>(FirstNode) )
 		{
 			if (!LayoutBlocksSelector.IsValid())
 			{
diff --git a/Engine/Plugins/Experimental/Mutable/Source/CustomizableObjectEditor/Private/MuCOE/Nodes/CustomizableObjectNodeLayoutBlocksDetails.cpp b/Engine/Plugins/Experimental/Mutable/Source/CustomizableObjectEditor/Private/MuCOE/Nodes/CustomizableObjectNodeLayoutBlocksDetails.cpp
--- a/Engine/Plugins/Experimental/Mutable/Source/CustomizableObjectEditor/Private/MuCOE/Nodes/CustomizableObjectNodeLayoutBlocksDetails.cpp
+++ b/Engine/Plugins/Experimental/Mutable/Source/CustomizableObjectEditor/Private/MuCOE/Nodes/CustomizableObjectNodeLayoutBlocksDetails.cpp
@@ -11,6 +11,7 @@
 #include "Math/IntRect.h"
 #include "MuCOE/CustomizableObjectLayout.h"
 #include "MuCOE/Nodes/CustomizableObjectNodeLayoutBlocks.h"
+#include "MuCOE/SCustomizableObjectNodeLayoutBlocksEditor.h"
 #include "Templates/Casts.h"
 #include "UObject/ObjectPtr.h"
 #include "UObject/WeakObjectPtr.h"
@@ -43,11 +44,19 @@ void FCustomizableObjectNodeLayoutBlocksDetails::CustomizeDetails( IDetailLayout
 
 	if (Node)
 	{
+		LayoutBlocksEditor = SNew(SCustomizableObjectNodeLayoutBlocksEditor);
+
 		BlocksCategory.AddCustomRow( LOCTEXT("BlocksDetails_BlockInstructions", "BlockInstructions") )
 		[
-			SNew( STextBlock )
-			.Text(  LOCTEXT("Use the node preview panel to set up the blocks.","Use the node preview panel to set up the blocks.") )
+			SNew(SBox)
+			.HeightOverride(700.0f)
+			.WidthOverride(700.0f)
+			[
+				LayoutBlocksEditor.ToSharedRef()
+			]
 		];
+		
+		LayoutBlocksEditor->SetCurrentLayout(Node->Layout);
 	}
 	else
 	{
diff --git a/Engine/Plugins/Experimental/Mutable/Source/CustomizableObjectEditor/Private/MuCOE/Nodes/CustomizableObjectNodeLayoutBlocksDetails.h b/Engine/Plugins/Experimental/Mutable/Source/CustomizableObjectEditor/Private/MuCOE/Nodes/CustomizableObjectNodeLayoutBlocksDetails.h
--- a/Engine/Plugins/Experimental/Mutable/Source/CustomizableObjectEditor/Private/MuCOE/Nodes/CustomizableObjectNodeLayoutBlocksDetails.h
+++ b/Engine/Plugins/Experimental/Mutable/Source/CustomizableObjectEditor/Private/MuCOE/Nodes/CustomizableObjectNodeLayoutBlocksDetails.h
@@ -26,6 +26,10 @@ public:
 private:
 
 	class UCustomizableObjectNodeLayoutBlocks* Node;
+
+	// Layout block editor widget
+	TSharedPtr<class SCustomizableObjectNodeLayoutBlocksEditor> LayoutBlocksEditor;
+
 	TArray< TSharedPtr<FString> > GridComboOptions;
 
 	void OnGridComboBoxSelectionChanged(TSharedPtr<FString> Selection, ESelectInfo::Type SelectInfo);
diff --git a/Engine/Plugins/Experimental/Mutable/Source/CustomizableObjectEditor/Private/MuCOE/Nodes/CustomizableObjectNodeTableDetails.cpp b/Engine/Plugins/Experimental/Mutable/Source/CustomizableObjectEditor/Private/MuCOE/Nodes/CustomizableObjectNodeTableDetails.cpp
--- a/Engine/Plugins/Experimental/Mutable/Source/CustomizableObjectEditor/Private/MuCOE/Nodes/CustomizableObjectNodeTableDetails.cpp
+++ b/Engine/Plugins/Experimental/Mutable/Source/CustomizableObjectEditor/Private/MuCOE/Nodes/CustomizableObjectNodeTableDetails.cpp
@@ -18,11 +18,13 @@
 #include "IDetailsView.h"
 #include "Internationalization/Internationalization.h"
 #include "Internationalization/Text.h"
+#include "Input/Reply.h"
 #include "Layout/Margin.h"
 #include "Layout/Visibility.h"
 #include "Misc/Attribute.h"
 #include "MuCOE/CustomizableObjectLayout.h"
 #include "MuCOE/Nodes/CustomizableObjectNodeTable.h"
+#include "MuCOE/SCustomizableObjectNodeLayoutBlocksEditor.h"
 #include "SlotBase.h"
 #include "Styling/AppStyle.h"
 #include "Templates/Casts.h"
@@ -37,20 +39,24 @@
 #include "Widgets/Layout/SBorder.h"
 #include "Widgets/SBoxPanel.h"
 #include "Widgets/Text/STextBlock.h"
+#include "Widgets/Input/SButton.h"
 
 #define LOCTEXT_NAMESPACE "CustomizableObjectDetails"
 
+
 TSharedRef<IDetailCustomization> FCustomizableObjectNodeTableDetails::MakeInstance()
 {
 	return MakeShareable(new FCustomizableObjectNodeTableDetails);
 }
 
 
-void FCustomizableObjectNodeTableDetails::CustomizeDetails(IDetailLayoutBuilder& DetailBuilder)
+void FCustomizableObjectNodeTableDetails::CustomizeDetails(const TSharedPtr<IDetailLayoutBuilder>& DetailBuilder)
 {
 	Node = 0;
+	DetailBuilderPtr = DetailBuilder;
 
-	const IDetailsView* DetailsView = DetailBuilder.GetDetailsView();
+	const IDetailsView* DetailsView = DetailBuilder->GetDetailsView();
+	
 	if (DetailsView->GetSelectedObjects().Num())
 	{
 		Node = Cast<UCustomizableObjectNodeTable>(DetailsView->GetSelectedObjects()[0].Get());
@@ -58,13 +64,17 @@ void FCustomizableObjectNodeTableDetails::CustomizeDetails(IDetailLayoutBuilder&
 
 	if (Node)
 	{
-		IDetailCategoryBuilder& CustomizableObjectCategory = DetailBuilder.EditCategory("CustomizableObject");
-		IDetailCategoryBuilder& UICategory = DetailBuilder.EditCategory("UI");
-		IDetailCategoryBuilder& LayoutCategory = DetailBuilder.EditCategory("LayoutEditor");
+		IDetailCategoryBuilder& CustomizableObjectCategory = DetailBuilder->EditCategory("TableProperties");
+		IDetailCategoryBuilder& UICategory = DetailBuilder->EditCategory("UI");
+		IDetailCategoryBuilder& AnimationCategory = DetailBuilder->EditCategory("AnimationProperties");
+		IDetailCategoryBuilder& LayoutCategory = DetailBuilder->EditCategory("DefaultMeshLayoutEditor");
 
-		GenerateColumnComboBoxOptions();
+		// Attaching the Posrecontruct delegate to force a refresh of the details
+		Node->PostReconstructNodeDelegate.AddSP(this, &FCustomizableObjectNodeTableDetails::OnNodePinValueChanged);
 
-		LayoutCategory.AddCustomRow(LOCTEXT("TestName", "Test Name"))
+		GenerateMeshColumnComboBoxOptions();
+
+		AnimationCategory.AddCustomRow(LOCTEXT("AnimationProperties", "Animation Properties"))
 		[
 			SNew(SVerticalBox)
 			
@@ -79,35 +89,8 @@ void FCustomizableObjectNodeTableDetails::CustomizeDetails(IDetailLayoutBuilder&
 				.Padding(0.0f, 5.0f, 6.0f, 0.0f)
 				[
 					SNew(STextBlock)
-					.Text(LOCTEXT("ColumnText", "Column: "))
-				]
-
-				+ SHorizontalBox::Slot()
-				[
-					SNew(SBorder)
-					.BorderImage(FAppStyle::GetBrush("NoBorder"))
-					.Padding(FMargin(0.0f, 0.0f, 10.0f, 0.0f))
-					[
-						SAssignNew(ColumnComboBox, STextComboBox)
-						.OptionsSource(&ColumnOptionNames)
-						.InitiallySelectedItem(nullptr)
-						.OnSelectionChanged(this, &FCustomizableObjectNodeTableDetails::OnColumnComboBoxSelectionChanged)
-					]
-				]
-			]
-			
-			// Layout selection widget
-			+SVerticalBox::Slot()
-			.Padding(0.0f, 5.0f, 6.0f, 0.0f)
-			.AutoHeight()
-			[
-				SNew(SHorizontalBox)
-				+ SHorizontalBox::Slot()
-				.AutoWidth()
-				.Padding(0.0f, 5.0f, 6.0f, 0.0f)
-				[
-					SNew(STextBlock)
-					.Text(LOCTEXT("LayoutText", "Layout: "))
+					.Text(LOCTEXT("AnimMeshColumnText", "Mesh Column: "))
+					.ToolTipText(LOCTEXT("AnimMeshColumnTooltip","Select a mesh column from the Data Table to edit its animation options (Applied to all LODs)."))
 				]
 
 				+ SHorizontalBox::Slot()
@@ -116,10 +99,10 @@ void FCustomizableObjectNodeTableDetails::CustomizeDetails(IDetailLayoutBuilder&
 					.BorderImage(FAppStyle::GetBrush("NoBorder"))
 					.Padding(FMargin(0.0f, 0.0f, 10.0f, 0.0f))
 					[
-						SAssignNew(LayoutComboBox, STextComboBox)
-						.OptionsSource(&LayoutOptionNames)
+						SAssignNew(AnimMeshColumnComboBox, STextComboBox)
+						.OptionsSource(&AnimMeshColumnOptionNames)
 						.InitiallySelectedItem(nullptr)
-						.OnSelectionChanged(this, &FCustomizableObjectNodeTableDetails::OnLayoutComboBoxSelectionChanged)
+						.OnSelectionChanged(this, &FCustomizableObjectNodeTableDetails::OnAnimMeshColumnComboBoxSelectionChanged)
 					]
 				]
 			]
@@ -135,7 +118,8 @@ void FCustomizableObjectNodeTableDetails::CustomizeDetails(IDetailLayoutBuilder&
 				.Padding(0.0f, 5.0f, 6.0f, 0.0f)
 				[
 					SNew(STextBlock)
-					.Text(LOCTEXT("AnimBPText", "Animation Column: "))
+					.Text(LOCTEXT("AnimBPText", "Animation Blueprint Column: "))
+					.ToolTipText(LOCTEXT("AnimBlueprintColumnTooltip", "Select an animation blueprint column from the Data Table that will be applied to the mesh selected"))
 					.Visibility_Lambda([this]() -> EVisibility
 					{
 						if (!AnimComboBox.IsValid())
@@ -173,6 +157,7 @@ void FCustomizableObjectNodeTableDetails::CustomizeDetails(IDetailLayoutBuilder&
 				[
 					SNew(STextBlock)
 					.Text(LOCTEXT("AnimSlotText", "Animation Slot Column: "))
+					.ToolTipText(LOCTEXT("AnimSlotColumnTooltip", "Select an animation slot column from the Data Table that will set to the slot value of the animation blueprint"))
 					.Visibility_Lambda([this]() -> EVisibility 
 						{
 							if (!AnimSlotComboBox.IsValid())
@@ -210,6 +195,7 @@ void FCustomizableObjectNodeTableDetails::CustomizeDetails(IDetailLayoutBuilder&
 				[
 					SNew(STextBlock)
 					.Text(LOCTEXT("AnimTagsText", "Animation Tags Column: "))
+					.ToolTipText(LOCTEXT("AnimTagColumnTooltip", "Select an animation tag column from the Data Table that will set to the animation tags of the animation blueprint"))
 					.Visibility_Lambda([this]() -> EVisibility 
 						{
 							if (!AnimTagsComboBox.IsValid())
@@ -234,15 +220,82 @@ void FCustomizableObjectNodeTableDetails::CustomizeDetails(IDetailLayoutBuilder&
 					]
 				]
 			]
+
+			+SVerticalBox::Slot()
+			.Padding(0.0f, 5.0f, 16.0f,5.0f)
+			.AutoHeight()
+			.HAlign(EHorizontalAlignment::HAlign_Right)
+			[
+				SNew(SHorizontalBox)
+				+ SHorizontalBox::Slot()
+				.AutoWidth()
+				[
+					SAssignNew(ClearButton,SButton)
+					.Text(FText::FromString("Clear"))
+					.ToolTipText(FText::FromString("Clear options of animation parameters"))
+					.HAlign(EHorizontalAlignment::HAlign_Center)
+					.Visibility(EVisibility::Collapsed)
+					.OnClicked(this, &FCustomizableObjectNodeTableDetails::OnClearButtonPressed)
+				]
+			]
+		];
+
+		LayoutBlocksEditor = SNew(SCustomizableObjectNodeLayoutBlocksEditor);
+
+		LayoutCategory.AddCustomRow(LOCTEXT("LayoutEditor", "Layout Editor"))
+		[
+			SNew(SVerticalBox)
+			
+			// Mesh Column selection widget
+			+SVerticalBox::Slot()
+			.Padding(0.0f, 5.0f, 6.0f, 0.0f)
+			.AutoHeight()
+			[
+				SNew(SHorizontalBox)
+				+ SHorizontalBox::Slot()
+				.AutoWidth()
+				.Padding(0.0f, 5.0f, 6.0f, 0.0f)
+				[
+					SNew(STextBlock)
+					.Text(LOCTEXT("LayoutMeshColumnText", "Mesh Column: "))
+					.ToolTipText(LOCTEXT("LayoutMeshColumnTooltip", "Select a mesh from the Data Table to edit its layout blocks."))
+				]
+		
+				+ SHorizontalBox::Slot()
+				[
+					SNew(SBorder)
+					.BorderImage(FAppStyle::GetBrush("NoBorder"))
+					.Padding(FMargin(0.0f, 0.0f, 10.0f, 0.0f))
+					[
+						SAssignNew(LayoutMeshColumnComboBox, STextComboBox)
+						.OptionsSource(&LayoutMeshColumnOptionNames)
+						.InitiallySelectedItem(nullptr)
+						.OnSelectionChanged(this, &FCustomizableObjectNodeTableDetails::OnLayoutMeshColumnComboBoxSelectionChanged)
+					]
+				]
+			]
+
+			+SVerticalBox::Slot()
+			.Padding(0.0f,20.0f,0.0f,0.0f)
+			[
+				SNew(SBox)
+				.HeightOverride(700.0f)
+				.WidthOverride(700.0f)
+				[
+					LayoutBlocksEditor.ToSharedRef()
+				]
+			]
 		];
-	}
 
+		LayoutBlocksEditor->SetCurrentLayout(nullptr);
+	}
 }
 
 
-void FCustomizableObjectNodeTableDetails::GenerateColumnComboBoxOptions()
+void FCustomizableObjectNodeTableDetails::GenerateMeshColumnComboBoxOptions()
 {
-	ColumnOptionNames.Empty();
+	AnimMeshColumnOptionNames.Empty();
+	LayoutMeshColumnOptionNames.Empty();
 
 	if (Node->Table)
 	{
@@ -253,44 +306,46 @@ void FCustomizableObjectNodeTableDetails::GenerateColumnComboBoxOptions()
 		{
 			FProperty* ColumnProperty = *It;
 
-			if (ColumnProperty)
+			if (!ColumnProperty)
 			{
-				if (const FSoftObjectProperty* SoftObjectProperty = CastField<FSoftObjectProperty>(ColumnProperty))
-				{
-					if (SoftObjectProperty->PropertyClass->IsChildOf(USkeletalMesh::StaticClass())
-						|| SoftObjectProperty->PropertyClass->IsChildOf(UStaticMesh::StaticClass()))
-					{
-						ColumnOptionNames.Add(MakeShareable(new FString(DataTableUtils::GetPropertyExportName(ColumnProperty))));
-					}
-				}
+				continue;
 			}
-		}
-	}
-}
 
-
-void FCustomizableObjectNodeTableDetails::GenerateLayoutComboBoxOptions()
-{
-	LayoutOptionNames.Empty();
-	
-	for (const UEdGraphPin* Pin : Node->Pins)
-	{
-		const UCustomizableObjectNodeTableMeshPinData* PinData = Cast<UCustomizableObjectNodeTableMeshPinData >(Node->GetPinData(*Pin));
-
-		if (PinData && PinData->ColumnName == *ColumnComboBox->GetSelectedItem().Get())
-		{
-			if (PinData->Layouts.Num() > 1)
+			if (const FSoftObjectProperty* SoftObjectProperty = CastField<FSoftObjectProperty>(ColumnProperty))
 			{
-				for (int32 LayoutIndex = 0; LayoutIndex < PinData->Layouts.Num(); ++LayoutIndex)
+				if (SoftObjectProperty->PropertyClass->IsChildOf(USkeletalMesh::StaticClass())
+					|| SoftObjectProperty->PropertyClass->IsChildOf(UStaticMesh::StaticClass()))
 				{
-					FString LayoutName = Pin->PinFriendlyName.ToString() + FString::Printf(TEXT(" UV_%d"), LayoutIndex);
-					LayoutOptionNames.Add(MakeShareable(new FString(LayoutName)));
+					FString MeshColumnName = DataTableUtils::GetPropertyExportName(ColumnProperty);
+					AnimMeshColumnOptionNames.Add(MakeShareable(new FString(MeshColumnName)));
+
+					for (const UEdGraphPin* Pin : Node->GetAllNonOrphanPins())
+					{
+						const UCustomizableObjectNodeTableMeshPinData* PinData = Cast<UCustomizableObjectNodeTableMeshPinData >(Node->GetPinData(*Pin));
+
+						if (!PinData || PinData->ColumnName != MeshColumnName)
+						{
+							continue;
+						}
+
+						if (PinData && PinData->ColumnName == MeshColumnName)
+						{
+							if (PinData->Layouts.Num() > 1)
+							{
+								for (int32 LayoutIndex = 0; LayoutIndex < PinData->Layouts.Num(); ++LayoutIndex)
+								{
+									FString LayoutName = Pin->PinFriendlyName.ToString() + FString::Printf(TEXT(" UV_%d"), LayoutIndex);
+									LayoutMeshColumnOptionNames.Add(MakeShareable(new FString(LayoutName)));
+								}
+							}
+							else
+							{
+								LayoutMeshColumnOptionNames.Add(MakeShareable(new FString(Pin->PinFriendlyName.ToString())));
+							}
+						}
+					}
 				}
 			}
-			else
-			{
-				LayoutOptionNames.Add(MakeShareable(new FString(Pin->PinFriendlyName.ToString())));
-			}
 		}
 	}
 }
@@ -298,10 +353,26 @@ void FCustomizableObjectNodeTableDetails::GenerateLayoutComboBoxOptions()
 
 void FCustomizableObjectNodeTableDetails::GenerateAnimInstanceComboBoxOptions()
 {
+	// Options Reset
 	AnimOptionNames.Empty();
+	AnimSlotOptionNames.Empty();
+	AnimTagsOptionNames.Empty();
+
+	// Selection Reset
+	AnimComboBox->ClearSelection();
+	AnimSlotComboBox->ClearSelection();
+	AnimTagsComboBox->ClearSelection();
+
+	FString ColumnName;
+
+	if (AnimMeshColumnComboBox.IsValid())
+	{
+		ColumnName = *AnimMeshColumnComboBox->GetSelectedItem();
+	}
 
 	const UScriptStruct* TableStruct = Node->Table->GetRowStruct();
 
+	// Fillins name options arrays and setting the selected item if any
 	for (TFieldIterator<FProperty> It(TableStruct); It; ++It)
 	{
 		if (FProperty* ColumnProperty = *It)
@@ -313,11 +384,11 @@ void FCustomizableObjectNodeTableDetails::GenerateAnimInstanceComboBoxOptions()
 					TSharedPtr<FString> Option = MakeShareable(new FString(DataTableUtils::GetPropertyExportName(ColumnProperty)));
 					AnimOptionNames.Add(Option);
 
-					for (const UEdGraphPin* Pin : Node->Pins)
+					for (const UEdGraphPin* Pin : Node->GetAllNonOrphanPins())
 					{
 						const UCustomizableObjectNodeTableMeshPinData* PinData = Cast<UCustomizableObjectNodeTableMeshPinData >(Node->GetPinData(*Pin));
 
-						if (PinData && PinData->ColumnName == *ColumnComboBox->GetSelectedItem().Get() && PinData->AnimInstanceColumnName == *Option)
+						if (PinData && PinData->ColumnName == ColumnName && PinData->AnimInstanceColumnName == *Option)
 						{
 							AnimComboBox->SetSelectedItem(Option);
 							break;
@@ -325,80 +396,37 @@ void FCustomizableObjectNodeTableDetails::GenerateAnimInstanceComboBoxOptions()
 					}
 				}
 			}
-		}
-	}
-	
-}
-
-
-void FCustomizableObjectNodeTableDetails::GenerateAnimSlotComboBoxOptions()
-{
-	AnimSlotOptionNames.Empty();
-
-	if (Node->Table)
-	{
-		const UScriptStruct* TableStruct = Node->Table->GetRowStruct();
-
-		// we just need the mesh columns
-		for (TFieldIterator<FProperty> It(TableStruct); It; ++It)
-		{
-			FProperty* ColumnProperty = *It;
-
-			if (ColumnProperty)
+			else if (const FIntProperty* NumProperty = CastField<FIntProperty>(ColumnProperty))
 			{
-				if (const FIntProperty* NumProperty = CastField<FIntProperty>(ColumnProperty))
+				TSharedPtr<FString> Option = MakeShareable(new FString(DataTableUtils::GetPropertyExportName(ColumnProperty)));
+				AnimSlotOptionNames.Add(Option);
+
+				for (const UEdGraphPin* Pin : Node->GetAllNonOrphanPins())
 				{
-					TSharedPtr<FString> Option = MakeShareable(new FString(DataTableUtils::GetPropertyExportName(ColumnProperty)));
-					AnimSlotOptionNames.Add(Option);
+					const UCustomizableObjectNodeTableMeshPinData* PinData = Cast<UCustomizableObjectNodeTableMeshPinData >(Node->GetPinData(*Pin));
 
-					for (const UEdGraphPin* Pin : Node->Pins)
+					if (PinData && PinData->ColumnName == ColumnName && PinData->AnimSlotColumnName == *Option)
 					{
-						const UCustomizableObjectNodeTableMeshPinData* PinData = Cast<UCustomizableObjectNodeTableMeshPinData >(Node->GetPinData(*Pin));
-
-						if (PinData && PinData->ColumnName == *ColumnComboBox->GetSelectedItem().Get() && PinData->AnimSlotColumnName == *Option)
-						{
-							AnimSlotComboBox->SetSelectedItem(Option);
-							break;
-						}
+						AnimSlotComboBox->SetSelectedItem(Option);
+						break;
 					}
 				}
 			}
-		}
-	}
-}
-
-
-void FCustomizableObjectNodeTableDetails::GenerateAnimTagsComboBoxOptions()
-{
-	AnimTagsOptionNames.Empty();
-
-	if (Node->Table)
-	{
-		const UScriptStruct* TableStruct = Node->Table->GetRowStruct();
-
-		// we just need the mesh columns
-		for (TFieldIterator<FProperty> It(TableStruct); It; ++It)
-		{
-			FProperty* ColumnProperty = *It;
-
-			if (ColumnProperty)
+			else if (const FStructProperty* StructProperty = CastField<FStructProperty>(ColumnProperty))
 			{
-				if (const FStructProperty* StructProperty = CastField<FStructProperty>(ColumnProperty))
+				if (StructProperty->Struct == TBaseStructure<FGameplayTagContainer>::Get())
 				{
-					if (StructProperty->Struct == TBaseStructure<FGameplayTagContainer>::Get())
+					TSharedPtr<FString> Option = MakeShareable(new FString(DataTableUtils::GetPropertyExportName(ColumnProperty)));
+					AnimTagsOptionNames.Add(Option);
+
+					for (const UEdGraphPin* Pin : Node->GetAllNonOrphanPins())
 					{
-						TSharedPtr<FString> Option = MakeShareable(new FString(DataTableUtils::GetPropertyExportName(ColumnProperty)));
-						AnimTagsOptionNames.Add(Option);
+						const UCustomizableObjectNodeTableMeshPinData* PinData = Cast<UCustomizableObjectNodeTableMeshPinData >(Node->GetPinData(*Pin));
 
-						for (const UEdGraphPin* Pin : Node->Pins)
+						if (PinData && PinData->ColumnName == ColumnName && PinData->AnimTagColumnName == *Option)
 						{
-							const UCustomizableObjectNodeTableMeshPinData* PinData = Cast<UCustomizableObjectNodeTableMeshPinData >(Node->GetPinData(*Pin));
-
-							if (PinData && PinData->ColumnName == *ColumnComboBox->GetSelectedItem().Get() && PinData->AnimTagColumnName == *Option)
-							{
-								AnimTagsComboBox->SetSelectedItem(Option);
-								break;
-							}
+							AnimTagsComboBox->SetSelectedItem(Option);
+							break;
 						}
 					}
 				}
@@ -408,52 +436,41 @@ void FCustomizableObjectNodeTableDetails::GenerateAnimTagsComboBoxOptions()
 }
 
 
-void FCustomizableObjectNodeTableDetails::OnColumnComboBoxSelectionChanged(TSharedPtr<FString> Selection, ESelectInfo::Type SelectInfo)
+void FCustomizableObjectNodeTableDetails::OnAnimMeshColumnComboBoxSelectionChanged(TSharedPtr<FString> Selection, ESelectInfo::Type SelectInfo)
 {
 	if (Selection.IsValid())
 	{
-		if (LayoutComboBox.IsValid())
-		{
-			LayoutComboBox->ClearSelection();
-			LayoutComboBox->RefreshOptions();
-
-			GenerateLayoutComboBoxOptions();
-		}
-
-		if (AnimComboBox.IsValid() && AnimSlotComboBox.IsValid())
+		if (AnimComboBox.IsValid() && AnimSlotComboBox.IsValid() && AnimTagsComboBox.IsValid())
 		{
 			AnimComboBox->SetVisibility(EVisibility::Visible);
 			AnimSlotComboBox->SetVisibility(EVisibility::Visible);
 			AnimTagsComboBox->SetVisibility(EVisibility::Visible);
+			ClearButton->SetVisibility(EVisibility::Visible);
 
 			GenerateAnimInstanceComboBoxOptions();
-			GenerateAnimSlotComboBoxOptions();
-			GenerateAnimTagsComboBoxOptions();
 		}
-
-		Node->SetLayoutInLayoutEditor(nullptr);
 	}
 }
 
 
-void FCustomizableObjectNodeTableDetails::OnLayoutComboBoxSelectionChanged(TSharedPtr<FString> Selection, ESelectInfo::Type SelectInfo)
+void FCustomizableObjectNodeTableDetails::OnLayoutMeshColumnComboBoxSelectionChanged(TSharedPtr<FString> Selection, ESelectInfo::Type SelectInfo)
 {
 	if (Selection.IsValid())
 	{
-		FString ColumnName = *ColumnComboBox->GetSelectedItem();
-		FString LayoutName = *Selection;
+		FString ColumnName;
+		(*Selection).Split(" LOD_", &ColumnName, NULL);
 
-		for (const UEdGraphPin* Pin : Node->Pins)
+		for (const UEdGraphPin* Pin : Node->GetAllNonOrphanPins())
 		{
 			const UCustomizableObjectNodeTableMeshPinData* PinData = Cast<UCustomizableObjectNodeTableMeshPinData >(Node->GetPinData(*Pin));
-			
+
 			if (PinData && PinData->ColumnName == ColumnName)
 			{
 				for (int32 LayoutIndex = 0; LayoutIndex < PinData->Layouts.Num(); ++LayoutIndex)
 				{
-					if (PinData->Layouts[LayoutIndex]->GetLayoutName() == LayoutName)
+					if (PinData->Layouts[LayoutIndex]->GetLayoutName() == *Selection)
 					{
-						Node->SetLayoutInLayoutEditor(PinData->Layouts[LayoutIndex]);
+						LayoutBlocksEditor->SetCurrentLayout(PinData->Layouts[LayoutIndex]);
 					}
 				}
 			}
@@ -464,11 +481,11 @@ void FCustomizableObjectNodeTableDetails::OnLayoutComboBoxSelectionChanged(TShar
 
 void FCustomizableObjectNodeTableDetails::OnAnimInstanceComboBoxSelectionChanged(TSharedPtr<FString> Selection, ESelectInfo::Type SelectInfo)
 {
-	if (Selection.IsValid() && ColumnComboBox->GetSelectedItem().IsValid())
+	if (Selection.IsValid() && AnimMeshColumnComboBox->GetSelectedItem().IsValid())
 	{
-		FString ColumnName = *ColumnComboBox->GetSelectedItem();
+		FString ColumnName = *AnimMeshColumnComboBox->GetSelectedItem();
 
-		for (const UEdGraphPin* Pin : Node->Pins)
+		for (const UEdGraphPin* Pin : Node->GetAllNonOrphanPins())
 		{
 			UCustomizableObjectNodeTableMeshPinData* PinData = Cast< UCustomizableObjectNodeTableMeshPinData>(Node->GetPinData(*Pin));
 
@@ -485,11 +502,11 @@ void FCustomizableObjectNodeTableDetails::OnAnimInstanceComboBoxSelectionChanged
 
 void FCustomizableObjectNodeTableDetails::OnAnimSlotComboBoxSelectionChanged(TSharedPtr<FString> Selection, ESelectInfo::Type SelectInfo)
 {
-	if (Selection.IsValid() && ColumnComboBox->GetSelectedItem().IsValid())
+	if (Selection.IsValid() && AnimMeshColumnComboBox->GetSelectedItem().IsValid())
 	{
-		FString ColumnName = *ColumnComboBox->GetSelectedItem();
+		FString ColumnName = *AnimMeshColumnComboBox->GetSelectedItem();
 
-		for (const UEdGraphPin* Pin : Node->Pins)
+		for (const UEdGraphPin* Pin : Node->GetAllNonOrphanPins())
 		{
 			UCustomizableObjectNodeTableMeshPinData* PinData = Cast< UCustomizableObjectNodeTableMeshPinData>(Node->GetPinData(*Pin));
 
@@ -506,11 +523,11 @@ void FCustomizableObjectNodeTableDetails::OnAnimSlotComboBoxSelectionChanged(TSh
 
 void FCustomizableObjectNodeTableDetails::OnAnimTagsComboBoxSelectionChanged(TSharedPtr<FString> Selection, ESelectInfo::Type SelectInfo)
 {
-	if (Selection.IsValid() && ColumnComboBox->GetSelectedItem().IsValid())
+	if (Selection.IsValid() && AnimMeshColumnComboBox->GetSelectedItem().IsValid())
 	{
-		FString ColumnName = *ColumnComboBox->GetSelectedItem();
+		FString ColumnName = *AnimMeshColumnComboBox->GetSelectedItem();
 
-		for (const UEdGraphPin* Pin : Node->Pins)
+		for (const UEdGraphPin* Pin : Node->GetAllNonOrphanPins())
 		{
 			UCustomizableObjectNodeTableMeshPinData* PinData = Cast< UCustomizableObjectNodeTableMeshPinData>(Node->GetPinData(*Pin));
 
@@ -525,4 +542,58 @@ void FCustomizableObjectNodeTableDetails::OnAnimTagsComboBoxSelectionChanged(TSh
 }
 
 
+void FCustomizableObjectNodeTableDetails::OnNodePinValueChanged()
+{
+	if (IDetailLayoutBuilder* DetailBuilder = DetailBuilderPtr.Pin().Get()) // Raw because we don't want to keep alive the details builder when calling the force refresh details
+	{
+		DetailBuilder->ForceRefreshDetails();
+	}
+}
+
+
+FReply FCustomizableObjectNodeTableDetails::OnClearButtonPressed()
+{
+	bool bCleared = false;
+
+	if (AnimMeshColumnComboBox->GetSelectedItem().IsValid())
+	{
+		FString ColumnName = *AnimMeshColumnComboBox->GetSelectedItem();
+
+		for (const UEdGraphPin* Pin : Node->GetAllNonOrphanPins())
+		{
+			UCustomizableObjectNodeTableMeshPinData* PinData = Cast< UCustomizableObjectNodeTableMeshPinData>(Node->GetPinData(*Pin));
+
+			if (PinData && PinData->ColumnName == ColumnName)
+			{
+				PinData->AnimInstanceColumnName.Reset();
+				PinData->AnimSlotColumnName.Reset();
+				PinData->AnimTagColumnName.Reset();
+
+				bCleared = true;
+			}
+		}
+
+		if (bCleared)
+		{
+			if (AnimComboBox.IsValid())
+			{
+				AnimComboBox->ClearSelection();
+			}
+			
+			if (AnimSlotComboBox.IsValid())
+			{
+				AnimSlotComboBox->ClearSelection();
+			}
+			
+			if (AnimTagsComboBox.IsValid())
+			{
+				AnimTagsComboBox->ClearSelection();
+			}
+		}
+	}
+
+	return bCleared ? FReply::Handled() : FReply::Unhandled();
+}
+
+
 #undef LOCTEXT_NAMESPACE
diff --git a/Engine/Plugins/Experimental/Mutable/Source/CustomizableObjectEditor/Private/MuCOE/Nodes/CustomizableObjectNodeTableDetails.h b/Engine/Plugins/Experimental/Mutable/Source/CustomizableObjectEditor/Private/MuCOE/Nodes/CustomizableObjectNodeTableDetails.h
--- a/Engine/Plugins/Experimental/Mutable/Source/CustomizableObjectEditor/Private/MuCOE/Nodes/CustomizableObjectNodeTableDetails.h
+++ b/Engine/Plugins/Experimental/Mutable/Source/CustomizableObjectEditor/Private/MuCOE/Nodes/CustomizableObjectNodeTableDetails.h
@@ -10,6 +10,9 @@
 class FString;
 class IDetailLayoutBuilder;
 class STextComboBox;
+class FReply;
+class SButton;
+class SCustomizableObjectNodeLayoutBlocksEditor;
 
 /** Copy Material node details panel. Hides all properties from the inheret Material node. */
 class FCustomizableObjectNodeTableDetails : public IDetailCustomization
@@ -18,57 +21,67 @@ public:
 	
 	static TSharedRef<IDetailCustomization> MakeInstance();
 
+	// ILayoutDetails interface
+	/** Do not use. Add details customization in the other CustomizeDetails signature. */
+	virtual void CustomizeDetails(IDetailLayoutBuilder& DetailBuilder) override {};
+
 	/** Hides details copied from CustomizableObjectNodeMaterial. */
-	virtual void CustomizeDetails(IDetailLayoutBuilder& DetailBuilder) override;
+	virtual void CustomizeDetails(const TSharedPtr<IDetailLayoutBuilder>& DetailBuilder) override;
 
 private:
 
-	// Generates columns combobox options
-	void GenerateColumnComboBoxOptions();
-	
-	// Generates layouts combobox options
-	void GenerateLayoutComboBoxOptions();
-	
+	// Generates Mesh columns combobox options
+	void GenerateMeshColumnComboBoxOptions();
+
 	// Generates Animation Instance combobox options
 	void GenerateAnimInstanceComboBoxOptions();
 
 	// Generates Animation Slot combobox options
 	void GenerateAnimSlotComboBoxOptions();
-	
+
 	// Generates Animation Tags combobox options
 	void GenerateAnimTagsComboBoxOptions();
 
 	// OnComboBoxSelectionChanged Callback for Column ComboBox
 	void OnColumnComboBoxSelectionChanged(TSharedPtr<FString> Selection, ESelectInfo::Type SelectInfo);
-	
+
 	// OnComboBoxSelectionChanged Callback for Layout ComboBox
-	void OnLayoutComboBoxSelectionChanged(TSharedPtr<FString> Selection, ESelectInfo::Type SelectInfo);
-	
+	void OnAnimMeshColumnComboBoxSelectionChanged(TSharedPtr<FString> Selection, ESelectInfo::Type SelectInfo);
+
+	// OnComboBoxSelectionChanged Callback for Layout ComboBox
+	void OnLayoutMeshColumnComboBoxSelectionChanged(TSharedPtr<FString> Selection, ESelectInfo::Type SelectInfo);
+
 	// OnComboBoxSelectionChanged Callback for AnimInstance ComboBox
 	void OnAnimInstanceComboBoxSelectionChanged(TSharedPtr<FString> Selection, ESelectInfo::Type SelectInfo);
-	
+
 	// OnComboBoxSelectionChanged Callback for Anim Slot ComboBox
 	void OnAnimSlotComboBoxSelectionChanged(TSharedPtr<FString> Selection, ESelectInfo::Type SelectInfo);
 
 	// OnComboBoxSelectionChanged Callback for Anim Tags ComboBox
 	void OnAnimTagsComboBoxSelectionChanged(TSharedPtr<FString> Selection, ESelectInfo::Type SelectInfo);
 
+	// Function called when the table node has been refreshed
+	void OnNodePinValueChanged();
+
+	// Callback to clear the animation combobox selections
+	FReply OnClearButtonPressed();
+
 private:
 
 	// Pointer to the node represented in this details
 	class UCustomizableObjectNodeTable* Node;
 
 	// ComboBox widget to select a column from the NodeTable
-	TSharedPtr<STextComboBox> ColumnComboBox;
+	TSharedPtr<STextComboBox> LayoutMeshColumnComboBox;
 
-	// Array with the name of the table columns as combobox options
-	TArray<TSharedPtr<FString>> ColumnOptionNames;	
-	
 	// ComboBox widget to select a column from the NodeTable
-	TSharedPtr<STextComboBox> LayoutComboBox;
+	TSharedPtr<STextComboBox> AnimMeshColumnComboBox;
 
 	// Array with the name of the table columns as combobox options
-	TArray<TSharedPtr<FString>> LayoutOptionNames;
+	TArray<TSharedPtr<FString>> AnimMeshColumnOptionNames;
+
+	// Array with the name of the table columns as combobox options
+	TArray<TSharedPtr<FString>> LayoutMeshColumnOptionNames;
 
 	// ComboBox widget to select an Animation Instance column from the NodeTable
 	TSharedPtr<STextComboBox> AnimComboBox;
@@ -87,5 +100,14 @@ private:
 
 	// Array with the name of the Animation Tags columns as combobox options
 	TArray<TSharedPtr<FString>> AnimTagsOptionNames;
+	
+	// Button to clear selections of the animation comboboxes
+	TSharedPtr<SButton> ClearButton;
+
+	// Layout block editor widget
+	TSharedPtr<SCustomizableObjectNodeLayoutBlocksEditor> LayoutBlocksEditor;
+
+	// Pointer to the Detail Builder to force the refresh on recontruct the node
+	TWeakPtr<IDetailLayoutBuilder> DetailBuilderPtr = nullptr;
 
 };
\ No newline at end of file
diff --git a/Engine/Plugins/Experimental/Mutable/Source/CustomizableObjectEditor/Private/MuCOE/SCustomizableObjectLayoutGrid.cpp b/Engine/Plugins/Experimental/Mutable/Source/CustomizableObjectEditor/Private/MuCOE/SCustomizableObjectLayoutGrid.cpp
--- a/Engine/Plugins/Experimental/Mutable/Source/CustomizableObjectEditor/Private/MuCOE/SCustomizableObjectLayoutGrid.cpp
+++ b/Engine/Plugins/Experimental/Mutable/Source/CustomizableObjectEditor/Private/MuCOE/SCustomizableObjectLayoutGrid.cpp
@@ -86,8 +86,22 @@ public:
 		return ViewRect;
 	}
 
+	/** Sets the clipping rect for the render target */
+	void SetClippingRect(const FIntRect& InClippingRect)
+	{
+		ClippingRect = InClippingRect;
+	}
+
+	/** Gets the clipping rect for the render target */
+	const FIntRect& GetClippingRect() const
+	{
+		return ClippingRect;
+	}
+
 private:
+
 	FIntRect ViewRect;
+	FIntRect ClippingRect;
 };
 
 
@@ -104,11 +118,19 @@ public:
 	 * 
 	 * All data will be copied.
 	 */
-	void Initialize(const FIntRect& InCanvasRect, const FVector2D& InOrigin, const FVector2D& InSize, const TArray<FVector2f>& InUVLayout);
-	
+	void Initialize(const FIntRect& InCanvasRect, const FIntRect& InClippingRect, const FVector2D& InOrigin, const FVector2D& InSize, const FIntPoint& InGridSize, const float InCellSize);
+	void InitializeDrawingData(const TArray<FVector2f>& InUVLayout, const TArray<FVector2f>& InUnassignedUVs, const TArray<FCustomizableObjectLayoutBlock>& InBlocks, const TArray<FGuid>& InSelectedBlocks);
+
+	/** Sets the layout mode to know what to draw */
+	void SetLayoutMode(ELayoutGridMode Mode);
+
 private:
+
 	void DrawRenderThread(FRHICommandListImmediate& RHICmdList, const void* RenderTarget) override;
 
+	/** Basic function to draw a block in the canvas */
+	void DrawBlock(FBatchedElements* BatchedElements, const FHitProxyId HitProxyId, const FRect2D& BlockRect, FColor Color);
+
 	/** SlateElement initialized, can Draw during the DrawRenderThread call. */
 	bool Initialized = false;
 
@@ -118,8 +140,20 @@ private:
 	/** Drawing size. */
 	FVector2D Size;
 
-	/** Drawing UVLayout. */
+	/** Size of the Layout Grid */
+	FIntPoint GridSize;
+
+	/** Cell Size */
+	float CellSize;
+
+	/** Drawing Data. */
 	TArray<FVector2D> UVLayout;
+	TArray<FVector2D> UnassignedUVs;
+	TArray<FCustomizableObjectLayoutBlock> Blocks;
+	TArray<FGuid> SelectedBlocks;
+
+	/** Layout Mode */
+	ELayoutGridMode LayoutMode;
 
 	FSlateCanvasRenderTarget* RenderTarget = new FSlateCanvasRenderTarget();
 };
@@ -150,6 +184,7 @@ void SCustomizableObjectLayoutGrid::Construct( const FArguments& InArgs )
 	Zoom = 1;
 
 	UVCanvasDrawer = TSharedPtr<FUVCanvasDrawer, ESPMode::ThreadSafe>(new FUVCanvasDrawer());
+	UVCanvasDrawer->SetLayoutMode(Mode);
 }
 
 
@@ -198,137 +233,39 @@ int32 SCustomizableObjectLayoutGrid::OnPaint(const FPaintArgs& Args, const FGeom
 	float AuxCellSize = Size.X / GridSizePoint.X;
 	
 	// Drawing Offsets
-	FVector2D Offset = (AreaSize-Size)/2.0f;
+	FVector2D Offset = FVector2D((AreaSize - Size).X / 2.0f, 0.0f);
 	FVector2D ZoomOffset = ((Size - OldSize) / 2.0f);
 	
 	// Drawing Origin
 	FVector2D Origin = BorderPadding + Offset + PaddingAmount - DistanceFromOrigin;
 
-	// Create line points
-	TArray< FVector2D > LinePoints;
-	LinePoints.SetNum(2);
-	
-	// Drawing Vertical Lines
-	for( int32 LineIndex = 0; LineIndex < GridSizePoint.X + 1; LineIndex++ )
-	{
-		LinePoints[0] = FVector2D( Origin.X + LineIndex * AuxCellSize, Origin.Y );
-		LinePoints[1] = FVector2D( Origin.X + LineIndex * AuxCellSize, Origin.Y + Size.Y );
+	// Setting Canvas Drawing Rectangles
+	FSlateRect SlateCanvasRect = AllottedGeometry.GetLayoutBoundingRect();
+	FSlateRect ClippedCanvasRect = SlateCanvasRect.IntersectionWith(MyClippingRect);
 
-		FSlateDrawElement::MakeLines( OutDrawElements,RetLayerId,AllottedGeometry.ToPaintGeometry(),
-			LinePoints, ESlateDrawEffect::None,	FColor(150, 150, 150, 64), false, 2.0);
-	}
-
-	// Drawing Horizontal Lines
-	for( int32 LineIndex = 0; LineIndex < GridSizePoint.Y + 1; LineIndex++ )
-	{
-		LinePoints[0] = FVector2D( Origin.X, Origin.Y + LineIndex * AuxCellSize );
-		LinePoints[1] = FVector2D( Origin.X + Size.X, Origin.Y + LineIndex * AuxCellSize );
-
-		FSlateDrawElement::MakeLines( OutDrawElements, RetLayerId, AllottedGeometry.ToPaintGeometry(),
-			LinePoints,	ESlateDrawEffect::None, FColor(150, 150, 150, 64), false, 2.0);
-	}
+	FIntRect CanvasRect(
+		FMath::TruncToInt(FMath::Max(0.0f, SlateCanvasRect.Left)),
+		FMath::TruncToInt(FMath::Max(0.0f, SlateCanvasRect.Top)),
+		FMath::TruncToInt(FMath::Max(0.0f, SlateCanvasRect.Right)),
+		FMath::TruncToInt(FMath::Max(0.0f, SlateCanvasRect.Bottom)));
 
-	RetLayerId++;
+	FIntRect ClippingRect(
+		FMath::TruncToInt(FMath::Max(0.0f, ClippedCanvasRect.Left)),
+		FMath::TruncToInt(FMath::Max(0.0f, ClippedCanvasRect.Top)),
+		FMath::TruncToInt(FMath::Max(0.0f, ClippedCanvasRect.Right)),
+		FMath::TruncToInt(FMath::Max(0.0f, ClippedCanvasRect.Bottom)));
 
-	// Draw UV using a CustomSlateElement on the RenderThread
-	const float CanvasMinX = FMath::Max(0.0f, AllottedGeometry.AbsolutePosition.X);
-	const float CanvasMinY = FMath::Max(0.0f, AllottedGeometry.AbsolutePosition.Y);
-	const FIntRect CanvasRect(
-		FMath::RoundToInt(CanvasMinX),
-		FMath::RoundToInt(CanvasMinY),
-		FMath::RoundToInt(CanvasMinX + AllottedGeometry.GetLocalSize().X * AllottedGeometry.Scale),
-		FMath::RoundToInt(CanvasMinY + AllottedGeometry.GetLocalSize().Y * AllottedGeometry.Scale));
 	
-	UVCanvasDrawer->Initialize(CanvasRect, Origin * AllottedGeometry.Scale, Size * AllottedGeometry.Scale, UVLayout);
+	UVCanvasDrawer->InitializeDrawingData(UVLayout, UnassignedUVLayoutVertices, Blocks.Get(), SelectedBlocks);
+	UVCanvasDrawer->Initialize(CanvasRect, ClippingRect, Origin * AllottedGeometry.Scale, Size * AllottedGeometry.Scale, GridSizePoint, AuxCellSize * AllottedGeometry.Scale);
 	FSlateDrawElement::MakeCustom(OutDrawElements, RetLayerId, UVCanvasDrawer);
 
 	const auto MakeYellowSquareLine = [&](const TArray<FVector2D>& Points) -> void
 	{
-		FSlateDrawElement::MakeLines( OutDrawElements, RetLayerId, AllottedGeometry.ToPaintGeometry(),
-			Points,	ESlateDrawEffect::None,	FColor(250, 230, 43, 255), true, 2.0);
+		FSlateDrawElement::MakeLines(OutDrawElements, RetLayerId, AllottedGeometry.ToPaintGeometry(),
+			Points, ESlateDrawEffect::None, FColor(250, 230, 43, 255), true, 2.0);
 	};
 
-	TArray<FVector2D> SquareLinePoints;
-	SquareLinePoints.SetNum(2);
-
-	const FVector2D CrossSize = Size*0.01;
-	for (const FVector2f& Vertex : UnassignedUVLayoutVertices)
-	{
-		SquareLinePoints[0] = (Origin + FVector2D(Vertex) * Size + FVector2D(CrossSize));
-		SquareLinePoints[1]	= (Origin + FVector2D(Vertex) * Size - FVector2D(CrossSize) * FVector2D(1.0f, -1.0f));
-		MakeYellowSquareLine(SquareLinePoints);
-		
-		SquareLinePoints[0]	= (Origin + FVector2D(Vertex) * Size - FVector2D(CrossSize));
-		MakeYellowSquareLine(SquareLinePoints);
-
-		SquareLinePoints[1]	= (Origin + FVector2D(Vertex) * Size + FVector2D(CrossSize) * FVector2D(1.0f, -1.0f));
-		MakeYellowSquareLine(SquareLinePoints);
-		
-		SquareLinePoints[0]	= (Origin + FVector2D(Vertex) * Size + FVector2D(CrossSize));
-		MakeYellowSquareLine(SquareLinePoints);
-	}
-
-
-	// Blocks
-	const FSlateBrush* BlockBrush = FAppStyle::GetBrush("TextBlock.HighlightShape");
-	const FVector2f PaintGeomPosition = AllottedGeometry.ToPaintGeometry().DrawPosition;
-	const float PaintGeomScale = AllottedGeometry.ToPaintGeometry().DrawScale;
-	
-	for (const FCustomizableObjectLayoutBlock& Block : Blocks.Get())
-	{
-		FSlateRenderTransform GeomTransform = FSlateRenderTransform(1.0f, FVector2D(PaintGeomPosition + BlockRects[Block.Id].Rect.Min * PaintGeomScale));
-		FPaintGeometry Geom(FSlateLayoutTransform(), GeomTransform, FVector2D(BlockRects[Block.Id].Rect.Size * PaintGeomScale), false );
-		
-		FSlateDrawElement::MakeBox(	OutDrawElements, RetLayerId, Geom, BlockBrush,
-			DrawEffects, SelectedBlocks.Find(Block.Id)>=0? SelectionColor :FColor(230,199,75,155) );	
-
-		if (Mode==ELGM_Edit)
-		{
-			FSlateRenderTransform HandleGeomTransform = FSlateRenderTransform(1.0f, FVector2D(PaintGeomPosition + BlockRects[Block.Id].HandleRect.Min * PaintGeomScale));
-			FPaintGeometry HandleGeom(FSlateLayoutTransform(), HandleGeomTransform, FVector2D(BlockRects[Block.Id].HandleRect.Size * PaintGeomScale), false);
-			
-			FColor HandleRectColor = FColor(255, 96, 96, 255);
-			bool bCanResize = SelectedBlocks.Num() == 1 && SelectedBlocks.Contains(Block.Id) && MouseOnBlock(Block.Id, CurrentMousePosition, true);
-
-			if (bCanResize)
-			{
-				HandleRectColor = FColor(200, 0, 0, 255);
-			}
-
-			FSlateDrawElement::MakeBox(	OutDrawElements, RetLayerId, HandleGeom, BlockBrush,
-				DrawEffects, HandleRectColor);
-
-
-			//Selected block outline
-			if (SelectedBlocks.Find(Block.Id) >= 0)
-			{
-				TArray<FVector2D> SelectionSquarePoints;
-				SelectionSquarePoints.SetNum(2);
-
-				FVector2D RectMin = FVector2D(BlockRects[Block.Id].Rect.Min);
-				FVector2D RectMax = FVector2D(BlockRects[Block.Id].Rect.Min + BlockRects[Block.Id].Rect.Size);
-
-				FVector2D TopRightCorner = FVector2D(RectMax.X, RectMin.Y);
-				FVector2D BottomLeftCorner = FVector2D(RectMin.X, RectMax.Y);
-
-				SelectionSquarePoints[0] = RectMin;
-				SelectionSquarePoints[1] = TopRightCorner;
-				MakeYellowSquareLine(SelectionSquarePoints);
-
-				SelectionSquarePoints[0] = RectMax;
-				MakeYellowSquareLine(SelectionSquarePoints);
-				
-				SelectionSquarePoints[1] = BottomLeftCorner;
-				MakeYellowSquareLine(SelectionSquarePoints);
-				
-				SelectionSquarePoints[0] = RectMin;
-				MakeYellowSquareLine(SelectionSquarePoints);
-			}
-		}
-	}
-
-	RetLayerId++;
-
 	// Drawing Multi-Selection rect
 	if (Mode == ELGM_Edit && Selecting)
 	{
@@ -383,7 +320,7 @@ void SCustomizableObjectLayoutGrid::Tick( const FGeometry& AllottedGeometry, con
 	FVector2D OldSize = Size;
 	Size *= Zoom;
 	CellSize = Size.X/GridSize.Get().X;
-	FVector2D Offset = (AreaSize-Size)/2.0f;
+	FVector2D Offset = FVector2D((AreaSize - Size).X / 2.0f, 0.0f);
 	FVector2D ZoomOffset = (Size - OldSize) / 2.0f;
 	FVector2D Origin = BorderPadding + Offset + PaddingAmount - DistanceFromOrigin;
 	DrawOrigin = Origin;
@@ -404,7 +341,7 @@ void SCustomizableObjectLayoutGrid::Tick( const FGeometry& AllottedGeometry, con
 		BlockData.HandleRect.Size = FVector2f(CellSize) * HandleRectSize;
 		BlockData.HandleRect.Min = BlockData.Rect.Min + BlockData.Rect.Size - BlockData.HandleRect.Size;
 
-		BlockRects.Add(Block.Id, BlockData );
+		BlockRects.Add(Block.Id, BlockData);
 	}
 
 	// Update selection list
@@ -527,14 +464,14 @@ FReply SCustomizableObjectLayoutGrid::OnMouseButtonDown( const FGeometry& MyGeom
 						[
 							SNew(SNumericEntryBox<int32>)
 							.MinValue(0)
-						.MaxValue(INT_MAX)
-						.MaxSliderValue(100)
-						.AllowSpin(SelectedBlocks.Num() == 1)
-						.Value(this, &SCustomizableObjectLayoutGrid::GetBlockPriortyValue)
-						.UndeterminedString(LOCTEXT("MultipleValues", "Multiples Values"))
-						.OnValueChanged(this, &SCustomizableObjectLayoutGrid::OnBlockPriorityChanged)
-						.ToolTipText(LOCTEXT("SetBlockPriorityTooltip", "Sets the block priority for a Fixed Layout Strategy"))
-						.EditableTextBoxStyle(&FAppStyle::GetWidgetStyle<FEditableTextBoxStyle>("NormalEditableTextBox"))
+							.MaxValue(INT_MAX)
+							.MaxSliderValue(100)
+							.AllowSpin(SelectedBlocks.Num() == 1)
+							.Value(this, &SCustomizableObjectLayoutGrid::GetBlockPriortyValue)
+							.UndeterminedString(LOCTEXT("MultipleValues", "Multiples Values"))
+							.OnValueChanged(this, &SCustomizableObjectLayoutGrid::OnBlockPriorityChanged)
+							.ToolTipText(LOCTEXT("SetBlockPriorityTooltip", "Sets the block priority for a Fixed Layout Strategy"))
+							.EditableTextBoxStyle(&FAppStyle::GetWidgetStyle<FEditableTextBoxStyle>("NormalEditableTextBox"))
 						]
 					, FText::FromString("Block Priority"), true);
 				}
@@ -948,6 +885,10 @@ FReply SCustomizableObjectLayoutGrid::OnKeyDown(const FGeometry& MyGeometry, con
 			FVector2D MouseToCellPosition = (CurrentMousePosition - DrawOrigin) / CellSize;
 			GenerateNewBlock(MouseToCellPosition);
 		}
+		else if (InKeyEvent.GetKey() == EKeys::F)
+		{
+			SetBlockSizeToMax();
+		}
 	}
 
 	if (InKeyEvent.GetKey() == EKeys::Delete)
@@ -1034,6 +975,26 @@ void SCustomizableObjectLayoutGrid::DuplicateBlocks()
 }
 
 
+void SCustomizableObjectLayoutGrid::SetBlockSizeToMax()
+{
+	if (SelectedBlocks.Num())
+	{
+		for (const FCustomizableObjectLayoutBlock& Block : Blocks.Get())
+		{
+			if (SelectedBlocks.Find(Block.Id) != INDEX_NONE)
+			{
+				FIntRect FinalBlock;
+
+				FinalBlock.Min = FIntPoint(0, 0);
+				FinalBlock.Max = GridSize.Get();
+
+				BlockChangedDelegate.ExecuteIfBound(Block.Id, FinalBlock);
+			}
+		}
+	}
+}
+
+
 void SCustomizableObjectLayoutGrid::CalculateSelectionRect()
 {
 	if (InitSelectionRect.X <= CurrentMousePosition.X)
@@ -1153,30 +1114,51 @@ void SCustomizableObjectLayoutGrid::SetLayoutStrategy(ECustomizableObjectTexture
 
 // Canvas Drawer --------------------------------------------------------------
 
-void FUVCanvasDrawer::Initialize(const FIntRect& InCanvasRect, const FVector2D& InOrigin, const FVector2D& InSize, const TArray<FVector2f>& InUVLayout)
+void FUVCanvasDrawer::Initialize(const FIntRect& InCanvasRect, const FIntRect& InClippingRect, const FVector2D& InOrigin, const FVector2D& InSize, const FIntPoint& InGridSize, const float InCellSize)
 {
 	Initialized = InCanvasRect.Size().X > 0 && InCanvasRect.Size().Y > 0;
 	if (Initialized)
 	{
 		RenderTarget->SetViewRect(InCanvasRect);
+		RenderTarget->SetClippingRect(InClippingRect);
 
 		Origin = InOrigin;
 		Size = InSize;
+		CellSize = InCellSize;
+		GridSize = InGridSize;
+	}
+}
 
-		// Convert data
-		UVLayout.SetNum(InUVLayout.Num());
-		for (int32 Index=0; Index<InUVLayout.Num(); ++Index)
-		{
-			UVLayout[Index] = FVector2D(InUVLayout[Index]);
-		}
+
+void FUVCanvasDrawer::InitializeDrawingData(const TArray<FVector2f>& InUVLayout, const TArray<FVector2f>& InUnassignedUVs, const TArray<FCustomizableObjectLayoutBlock>& InBlocks, const TArray<FGuid>& InSelectedBlocks)
+{
+	Blocks = InBlocks;
+	SelectedBlocks = InSelectedBlocks;
+
+	// Convert data
+	UVLayout.SetNum(InUVLayout.Num());
+	for (int32 Index = 0; Index < InUVLayout.Num(); ++Index)
+	{
+		UVLayout[Index] = FVector2D(InUVLayout[Index]);
 	}
+
+	UnassignedUVs.SetNum(InUnassignedUVs.Num());
+	for (int32 Index = 0; Index < UnassignedUVs.Num(); ++Index)
+	{
+		UnassignedUVs[Index] = FVector2D(InUnassignedUVs[Index]);
+	}
+}
+
+
+void FUVCanvasDrawer::SetLayoutMode(ELayoutGridMode Mode)
+{
+	LayoutMode = Mode;
 }
 
 
 void FUVCanvasDrawer::DrawRenderThread(class FRHICommandListImmediate& RHICmdList, const void* InWindowBackBuffer)
 {
-	//UE crashes if there are no UVs to draw due to DX
-	if (Initialized && UVLayout.Num())
+	if (Initialized)
 	{
 		RenderTarget->SetRenderTargetTexture(*(FTexture2DRHIRef*)InWindowBackBuffer);
 
@@ -1185,29 +1167,143 @@ void FUVCanvasDrawer::DrawRenderThread(class FRHICommandListImmediate& RHICmdLis
 #else
 		FCanvas Canvas(RenderTarget, nullptr, 0, 0, 0, GMaxRHIFeatureLevel);
 #endif
+		
+		Canvas.SetRenderTargetRect(RenderTarget->GetViewRect());
+		Canvas.SetRenderTargetScissorRect(RenderTarget->GetClippingRect());
 
+		// Num Lines
 		const uint32 NumEdges = UVLayout.Num() / 2;
+		const uint32 NumGridLines = GridSize.X + GridSize.Y + 2;
+		const uint32 NumUnasignedUVs = UnassignedUVs.Num() * 4;
 
-		Canvas.SetRenderTargetRect(RenderTarget->GetViewRect());
+		// Num Vertices
+		const uint32 NumVertices = LayoutMode == ELayoutGridMode::ELGM_Edit ? Blocks.Num() * 8 : Blocks.Num() * 4;
+
+		// Num Triangles
+		const uint32 NumTriangles = LayoutMode == ELayoutGridMode::ELGM_Edit ? Blocks.Num() * 4 : Blocks.Num() * 2;
 
 		FBatchedElements* BatchedElements = Canvas.GetBatchedElements(FCanvas::ET_Line);
-		BatchedElements->AddReserveLines(NumEdges);
+		BatchedElements->AddReserveLines(NumEdges + NumGridLines + NumUnasignedUVs);
+		BatchedElements->AddReserveVertices(NumVertices);
+		BatchedElements->AddReserveTriangles(NumTriangles, GWhiteTexture, ESimpleElementBlendMode::SE_BLEND_Translucent);
+
+		// Color Definitions
+		const FColor GridLineColor = FColor(150, 150, 150, 64);
+		const FColor UVLineColor = FColor(255, 255, 255, 255);
+		const FColor UnassignedUVsColor = FColor::Yellow;
+		const FColor ResizeBlockColor = FColor(255, 96, 96, 255);
 
-		const FLinearColor LineColor = FLinearColor::White;
 		const FHitProxyId HitProxyId = Canvas.GetHitProxyId();
 
+		// Create line points
 		FVector LinePoints[2];
+
+		// Drawing Vertical Lines
+		for (int32 LineIndex = 0; LineIndex < GridSize.X + 1; LineIndex++)
+		{
+			LinePoints[0] = FVector(FVector2D(Origin.X + LineIndex * CellSize, Origin.Y), 0.0f);
+			LinePoints[1] = FVector(FVector2D(Origin.X + LineIndex * CellSize, Origin.Y + Size.Y), 0.0f);
+
+			BatchedElements->AddTranslucentLine(LinePoints[0], LinePoints[1], GridLineColor, HitProxyId, 2.0f);
+		}
+
+		// Drawing Horizontal Lines
+		for (int32 LineIndex = 0; LineIndex < GridSize.Y + 1; LineIndex++)
+		{
+			LinePoints[0] = FVector(FVector2D(Origin.X, Origin.Y + LineIndex * CellSize), 0.0f);
+			LinePoints[1] = FVector(FVector2D(Origin.X + Size.X, Origin.Y + LineIndex * CellSize), 0.0f);
+
+			BatchedElements->AddTranslucentLine(LinePoints[0], LinePoints[1], GridLineColor, HitProxyId, 2.0f);
+		}
+
+		// Drawing UV Lines
 		for (uint32 LineIndex = 0; LineIndex < NumEdges; ++LineIndex)
 		{
 			LinePoints[0] = FVector(Origin + UVLayout[LineIndex * 2 + 0] * Size, 0.0f);
 			LinePoints[1] = FVector(Origin + UVLayout[LineIndex * 2 + 1] * Size, 0.0f);
 
-			BatchedElements->AddLine(LinePoints[0], LinePoints[1], LineColor, HitProxyId);
+			BatchedElements->AddLine(LinePoints[0], LinePoints[1], UVLineColor, HitProxyId);
+		}
+
+		// Drawing Unassigned UVs
+		const FVector2D CrossSize = Size * 0.01;
+		for (const FVector2d& Vertex : UnassignedUVs)
+		{
+			LinePoints[0] = FVector(Origin + FVector2D(Vertex) * Size + FVector2D(CrossSize), 0.0f);
+			LinePoints[1] = FVector(Origin + FVector2D(Vertex) * Size - FVector2D(CrossSize) * FVector2D(1.0f, -1.0f), 0.0f);
+			BatchedElements->AddLine(LinePoints[0], LinePoints[1], UVLineColor, HitProxyId);
+
+			LinePoints[0] = FVector(Origin + FVector2D(Vertex) * Size - FVector2D(CrossSize), 0.0f);
+			BatchedElements->AddLine(LinePoints[0], LinePoints[1], UVLineColor, HitProxyId);
+
+			LinePoints[1] = FVector(Origin + FVector2D(Vertex) * Size + FVector2D(CrossSize) * FVector2D(1.0f, -1.0f), 0.0f);
+			BatchedElements->AddLine(LinePoints[0], LinePoints[1], UVLineColor, HitProxyId);
+
+			LinePoints[0] = FVector(Origin + FVector2D(Vertex) * Size + FVector2D(CrossSize), 0.0f);
+			BatchedElements->AddLine(LinePoints[0], LinePoints[1], UVLineColor, HitProxyId);
+		}
+
+		// Drawing Blocks
+		for (const FCustomizableObjectLayoutBlock& Block : Blocks)
+		{
+			const FColor SelectionBlockColor = SelectedBlocks.Contains(Block.Id) ? FColor(75, 106, 230, 155) : FColor(230, 199, 75, 155);
+
+			const FVector2f BlockMin(Block.Min);
+			const FVector2f BlockMax(Block.Max);
+
+			// Selection Block
+			FRect2D SelectionBlock;
+			SelectionBlock.Min = FVector2f(Origin) + BlockMin * CellSize + CellSize * 0.1f;
+			SelectionBlock.Size = (BlockMax - BlockMin) * CellSize - CellSize * 0.2f;
+
+			DrawBlock(BatchedElements, HitProxyId, SelectionBlock, SelectionBlockColor);
+
+			if (LayoutMode == ELayoutGridMode::ELGM_Edit)
+			{
+				// Resize Block
+				FRect2D ResizeBlock;;
+				float HandleRectSize = FMath::Log2(float(GridSize.X)) / 10.0f;
+				ResizeBlock.Size = FVector2f(CellSize) * HandleRectSize;
+				ResizeBlock.Min = SelectionBlock.Min + SelectionBlock.Size - ResizeBlock.Size;
+
+				DrawBlock(BatchedElements, HitProxyId, ResizeBlock, ResizeBlockColor);
+			}
 		}
 
 		Canvas.Flush_RenderThread(RHICmdList, true);
 
 		RenderTarget->ClearRenderTargetTexture();
+		RHICmdList.SetScissorRect(false, 0, 0, 0, 0);
+	}
+}
+
+
+void FUVCanvasDrawer::DrawBlock(FBatchedElements* BatchedElements, const FHitProxyId HitProxyId, const FRect2D& BlockRect, FColor Color)
+{
+	// Vertex positions
+	FVector4 Vert0(BlockRect.Min.X, BlockRect.Min.Y, 0, 1);
+	FVector4 Vert1(BlockRect.Min.X, BlockRect.Min.Y + BlockRect.Size.Y, 0, 1);
+	FVector4 Vert2(BlockRect.Min.X + BlockRect.Size.X, BlockRect.Min.Y, 0, 1);
+	FVector4 Vert3(BlockRect.Min.X + BlockRect.Size.X, BlockRect.Min.Y + BlockRect.Size.Y, 0, 1);
+
+	// Brush Paint triangle
+	{
+		int32 V0 = BatchedElements->AddVertex(Vert0, FVector2d(0.0f, 0.0f), Color, HitProxyId);
+		int32 V1 = BatchedElements->AddVertex(Vert1, FVector2d(0.0f, 1.0f), Color, HitProxyId);
+		int32 V2 = BatchedElements->AddVertex(Vert2, FVector2d(1.0f, 0.0f), Color, HitProxyId);
+		int32 V3 = BatchedElements->AddVertex(Vert3, FVector2d(1.0f, 1.0f), Color, HitProxyId);
+
+		BatchedElements->AddTriangle(V0, V1, V2, GWhiteTexture, EBlendMode::BLEND_Translucent);
+		BatchedElements->AddTriangle(V1, V3, V2, GWhiteTexture, EBlendMode::BLEND_Translucent);
+	}
+
+	// Drawing Outline to selected Blocks
+	if (Color == FColor(75, 106, 230, 155))
+	{
+		BatchedElements->AddLine(Vert0, Vert1, FColor(230, 199, 75, 155), HitProxyId, 4.0f);
+		BatchedElements->AddLine(Vert1, Vert3, FColor(230, 199, 75, 155), HitProxyId, 4.0f);
+		BatchedElements->AddLine(Vert3, Vert2, FColor(230, 199, 75, 155), HitProxyId, 4.0f);
+		BatchedElements->AddLine(Vert2, Vert0, FColor(230, 199, 75, 155), HitProxyId, 4.0f);
 	}
 }
 
diff --git a/Engine/Plugins/Experimental/Mutable/Source/CustomizableObjectEditor/Private/MuCOE/SCustomizableObjectLayoutGrid.h b/Engine/Plugins/Experimental/Mutable/Source/CustomizableObjectEditor/Private/MuCOE/SCustomizableObjectLayoutGrid.h
--- a/Engine/Plugins/Experimental/Mutable/Source/CustomizableObjectEditor/Private/MuCOE/SCustomizableObjectLayoutGrid.h
+++ b/Engine/Plugins/Experimental/Mutable/Source/CustomizableObjectEditor/Private/MuCOE/SCustomizableObjectLayoutGrid.h
@@ -42,6 +42,11 @@ typedef enum
 	ELGM_Select
 } ELayoutGridMode;
 
+struct FBlockWidgetData
+{
+	FRect2D Rect;
+	FRect2D HandleRect;
+};
 
 /** */
 class SCustomizableObjectLayoutGrid : public SCompoundWidget
@@ -106,6 +111,9 @@ public:
 	/** Duplicates the selected blocks */
 	void DuplicateBlocks();
 
+	/** Sets the size of the selected blocks to the size of the Grid */
+	void SetBlockSizeToMax();
+
 	void CalculateSelectionRect();
 
 	FColor SelectionColor;
@@ -152,11 +160,7 @@ private:
 
 	float CellSize;
 
-	struct FBlockWidgetData
-	{
-		FRect2D Rect;
-		FRect2D HandleRect;
-	};
+	/** Map to relate Block ids with blocks data */
 	TMap<FGuid,FBlockWidgetData> BlockRects;
 
 	/** Interaction status. */
diff --git a/Engine/Plugins/Experimental/Mutable/Source/CustomizableObjectEditor/Private/MuCOE/SCustomizableObjectNodeLayoutBlocksEditor.cpp b/Engine/Plugins/Experimental/Mutable/Source/CustomizableObjectEditor/Private/MuCOE/SCustomizableObjectNodeLayoutBlocksEditor.cpp
--- a/Engine/Plugins/Experimental/Mutable/Source/CustomizableObjectEditor/Private/MuCOE/SCustomizableObjectNodeLayoutBlocksEditor.cpp
+++ b/Engine/Plugins/Experimental/Mutable/Source/CustomizableObjectEditor/Private/MuCOE/SCustomizableObjectNodeLayoutBlocksEditor.cpp
@@ -36,6 +36,8 @@
 #include "Widgets/Layout/SBorder.h"
 #include "Widgets/SBoxPanel.h"
 #include "Widgets/Text/STextBlock.h"
+#include "Widgets/Layout/SGridPanel.h"
+#include "Widgets/SToolTip.h"
 
 class ISlateStyle;
 class SWidget;
@@ -77,10 +79,13 @@ public:
 };
 
 
+SCustomizableObjectNodeLayoutBlocksEditor::SCustomizableObjectNodeLayoutBlocksEditor() : UICommandList(new FUICommandList())
+{
+}
+
 
 void SCustomizableObjectNodeLayoutBlocksEditor::Construct(const FArguments& InArgs)
 {
-	CustomizableObjectEditorPtr = InArgs._CustomizableObjectEditor;
 	CurrentLayout = 0;
 	
 	BindCommands();
@@ -112,14 +117,20 @@ void SCustomizableObjectNodeLayoutBlocksEditor::SetCurrentLayout(UCustomizableOb
 		SNew(SVerticalBox)
 
 		+ SVerticalBox::Slot()
-		.Padding( 2.0f )
+		.Padding(0.0f,2.0f,0.0f,0.0f )
 		.AutoHeight()
 		[
 			BuildLayoutToolBar()
 		]
+
+		+ SVerticalBox::Slot()
+		.HAlign(EHorizontalAlignment::HAlign_Left)
 		.AutoHeight()
+		.Padding(3.0f,5.0f,0.0f,5.0f)
 		[
-			BuildLayoutToolBar()
+			SNew(SImage)
+			.Image(FAppStyle::GetBrush(TEXT("Icons.Info")))
+			.ToolTip(GenerateInfoToolTip())
 		]
 
 		+ SVerticalBox::Slot()
@@ -159,8 +170,7 @@ void SCustomizableObjectNodeLayoutBlocksEditor::AddReferencedObjects( FReference
 
 TSharedRef<SWidget> SCustomizableObjectNodeLayoutBlocksEditor::BuildLayoutToolBar()
 {
-	TSharedPtr<FUICommandList> CommandList = CustomizableObjectEditorPtr.Pin()->GetToolkitCommands();
-	FToolBarBuilder LayoutToolbarBuilder(CommandList, FMultiBoxCustomization::None);
+	FToolBarBuilder LayoutToolbarBuilder(UICommandList, FMultiBoxCustomization::None);
 
 	//Getting toolbar style
 	const ISlateStyle* const StyleSet = &FCoreStyle::Get();
@@ -272,8 +282,8 @@ TSharedRef<SWidget> SCustomizableObjectNodeLayoutBlocksEditor::BuildLayoutStrate
 		.IsEnabled(FSlateApplication::Get().GetNormalExecutionAttribute());
 
 	LayoutPackingStrategies.Empty();
-	LayoutPackingStrategies.Add(MakeShareable(new FString("Resizable Layout")));
-	LayoutPackingStrategies.Add(MakeShareable(new FString("Fixed Layout")));
+	LayoutPackingStrategies.Add(MakeShareable(new FString("Resizable")));
+	LayoutPackingStrategies.Add(MakeShareable(new FString("Fixed")));
 
 	LayoutPackingStrategyCombo = SNew(STextComboBox)
 		.OptionsSource(&LayoutPackingStrategies)
@@ -286,7 +296,7 @@ TSharedRef<SWidget> SCustomizableObjectNodeLayoutBlocksEditor::BuildLayoutStrate
 		[
 			SNew(STextBlock)
 			.Text(LOCTEXT("Layout Strategy:", "Layout Strategy:"))
-			.ToolTipText(LOCTEXT("LayoutStrategyTooltup","Selects the packing strategy."))
+			.ToolTipText(LOCTEXT("LayoutStrategyTooltup","Selects the packing strategy: Resizable Layout or Fixed Layout"))
 			.Margin(FMargin(0, 4, 0, 0))
 			.ShadowOffset(FVector2D::UnitVector)
 			.ColorAndOpacity(FLinearColor::Gray)
@@ -302,7 +312,6 @@ TSharedRef<SWidget> SCustomizableObjectNodeLayoutBlocksEditor::BuildLayoutStrate
 			]
 		];
 
-
 	FixedLayoutWidget = SNew(SHorizontalBox)
 		+ SHorizontalBox::Slot().Padding(2.0f).AutoWidth()
 		[
@@ -341,7 +350,7 @@ TSharedRef<SWidget> SCustomizableObjectNodeLayoutBlocksEditor::BuildLayoutStrate
 
 	return SNew(SVerticalBox)
 		//.IsEnabled(CurrentLayout ? !CurrentLayout->GetNode()->GetGraphEditor()->GetCustomizableObject()->bDisableTextureLayoutManagement : false)
-		+ SVerticalBox::Slot()
+		+ SVerticalBox::Slot().VAlign(EVerticalAlignment::VAlign_Center)
 		[
 			LayoutStrategyWidget.ToSharedRef()
 		]
@@ -567,8 +576,6 @@ void SCustomizableObjectNodeLayoutBlocksEditor::BindCommands()
 
 	const FLayoutEditorCommands& Commands = FLayoutEditorCommands::Get();
 
-	const TSharedRef<FUICommandList>& UICommandList = CustomizableObjectEditorPtr.Pin()->GetToolkitCommands();
-
 	UICommandList->MapAction(
 		Commands.AddBlock,
 		FExecuteAction::CreateSP( this, &SCustomizableObjectNodeLayoutBlocksEditor::OnAddBlock ),
@@ -589,5 +596,45 @@ void SCustomizableObjectNodeLayoutBlocksEditor::BindCommands()
 }
 
 
+TSharedPtr<IToolTip> SCustomizableObjectNodeLayoutBlocksEditor::GenerateInfoToolTip() const
+{
+	TSharedPtr<SGridPanel> ToolTipWidget = SNew(SGridPanel);
+	int32 SlotCount = 0;
+
+	auto BuildShortcutAndTooltip = [ToolTipWidget, &SlotCount](const FText& Shortcut, const FText& Tooltip)
+	{
+		// Command Shortcut
+		ToolTipWidget->AddSlot(0, SlotCount)
+		[
+			SNew(STextBlock)
+			.Text(Shortcut)
+		];
+
+		// Command Explanation
+		ToolTipWidget->AddSlot(1, SlotCount)
+		.Padding(15.0f, 0.0f, 0.0f, 0.0f)
+		[
+			SNew(STextBlock)
+			.Text(Tooltip)
+		];
+
+		++SlotCount;
+	};
+
+	// Duplicate command
+	BuildShortcutAndTooltip(LOCTEXT("ShortCut_DuplicateBlocks", "CTRL + D"), LOCTEXT("Tooltip_DuplicateBlocks", "Duplicate selected block/s"));
+	BuildShortcutAndTooltip(LOCTEXT("ShortCut_CreateNewBlock", "CTRL + N"), LOCTEXT("Tooltip_CreateNewBlock", "Create new block"));
+	BuildShortcutAndTooltip(LOCTEXT("ShortCut_FillGridSize", "CTRL + F"), LOCTEXT("Tooltip_FillGridSize", "Resize selected block/s to grid size"));
+	BuildShortcutAndTooltip(LOCTEXT("ShortCut_DeleteSelectedBlock","DEL"), LOCTEXT("Tooltip_DeleteSelectedBlock","Delete selected block/s"));
+	BuildShortcutAndTooltip(LOCTEXT("ShortCut_SelectMultipleBlocksOneByOne","SHIFT + L Click"), LOCTEXT("Tooltip_SelectMultipleBlocksOneByOne","Select multiple blocks one by one"));
+	BuildShortcutAndTooltip(LOCTEXT("ShortCut_SelectMultipleBlocks","L Click + Drag"), LOCTEXT("Tooltip_SelectMultipleBlocks","Select blocks that intersect with the yellow rectangle"));
+
+	return SNew(SToolTip)
+	[
+		ToolTipWidget.ToSharedRef()
+	];
+}
+
+
 #undef LOCTEXT_NAMESPACE
 
diff --git a/Engine/Plugins/Experimental/Mutable/Source/CustomizableObjectEditor/Private/MuCOE/SCustomizableObjectNodeLayoutBlocksEditor.h b/Engine/Plugins/Experimental/Mutable/Source/CustomizableObjectEditor/Private/MuCOE/SCustomizableObjectNodeLayoutBlocksEditor.h
--- a/Engine/Plugins/Experimental/Mutable/Source/CustomizableObjectEditor/Private/MuCOE/SCustomizableObjectNodeLayoutBlocksEditor.h
+++ b/Engine/Plugins/Experimental/Mutable/Source/CustomizableObjectEditor/Private/MuCOE/SCustomizableObjectNodeLayoutBlocksEditor.h
@@ -16,6 +16,12 @@ class FReferenceCollector;
 class ICustomizableObjectInstanceEditor;
 class ISlateStyle;
 class SWidget;
+class STextComboBox;
+class SVerticalBox;
+class SHorizontalBox;
+class FUICommandList;
+class SCustomizableObjectLayoutGrid;
+
 struct FCustomizableObjectLayoutBlock;
 struct FGuid;
 
@@ -27,11 +33,12 @@ class SCustomizableObjectNodeLayoutBlocksEditor : public SCompoundWidget, public
 {
 public:
 	SLATE_BEGIN_ARGS( SCustomizableObjectNodeLayoutBlocksEditor ){}
-		SLATE_ARGUMENT(TWeakPtr<ICustomizableObjectInstanceEditor>, CustomizableObjectEditor)
 	SLATE_END_ARGS()
 
+	SCustomizableObjectNodeLayoutBlocksEditor();
+	virtual ~SCustomizableObjectNodeLayoutBlocksEditor();
+
 	void Construct(const FArguments& InArgs);
-	~SCustomizableObjectNodeLayoutBlocksEditor();
 	
 	// FSerializableObject interface
 	void AddReferencedObjects( FReferenceCollector& Collector ) override;
@@ -50,20 +57,17 @@ public:
 
 private:
 
-	/** Pointer back to the editor tool that owns us */
-	TWeakPtr<ICustomizableObjectInstanceEditor> CustomizableObjectEditorPtr;
-
 	/** */
 	class UCustomizableObjectLayout* CurrentLayout;
 
 	/** */
-	TSharedPtr<class SCustomizableObjectLayoutGrid> LayoutGridWidget;
+	TSharedPtr<SCustomizableObjectLayoutGrid> LayoutGridWidget;
 
-	TSharedPtr< class SVerticalBox > LayoutGridSizeWidget;
+	TSharedPtr< SVerticalBox > LayoutGridSizeWidget;
 
 	/** Widget for displaying the available layout block grid sizes. */
-	TSharedPtr< class STextComboBox > LayoutGridSizeCombo;
-	TSharedPtr< class STextComboBox > MaxLayoutGridSizeCombo;
+	TSharedPtr< STextComboBox > LayoutGridSizeCombo;
+	TSharedPtr< STextComboBox > MaxLayoutGridSizeCombo;
 
 	/** List of available layout grid sizes. */
 	TArray< TSharedPtr< FString > > LayoutGridSizes;
@@ -75,16 +79,21 @@ private:
 	TArray< TSharedPtr< FString > > LayoutPackingStrategies;
 
 	/** Widget for displaying the available layout packing strategies. */
-	TSharedPtr< class STextComboBox > LayoutPackingStrategyCombo;
+	TSharedPtr< STextComboBox > LayoutPackingStrategyCombo;
 
 	/** Widget to select the layout packing strategy */
-	TSharedPtr<class SHorizontalBox> LayoutStrategyWidget;
+	TSharedPtr< SHorizontalBox> LayoutStrategyWidget;
 
 	/** Widget to select the fixed layout properties */
-	TSharedPtr<class SHorizontalBox> FixedLayoutWidget;
+	TSharedPtr< SHorizontalBox> FixedLayoutWidget;
 
 	TSharedPtr<SWidget> StrategyWidget;
 
+	/** The list of UI Commands executable */
+	TSharedRef< FUICommandList> UICommandList;
+
+private:
+
 	/** */
 	FIntPoint GetGridSize() const;
 	void OnBlockChanged(FGuid BlockId, FIntRect Block );
@@ -110,4 +119,6 @@ private:
 	/** Called when the packing strategy has changed. */
 	void OnLayoutPackingStrategyChanged(TSharedPtr<FString> NewSelection, ESelectInfo::Type SelectInfo);
 
+	TSharedPtr<IToolTip> GenerateInfoToolTip() const;
+
 };
diff --git a/Engine/Plugins/Experimental/Mutable/Source/CustomizableObjectEditor/Private/MuCOE/SCustomizableObjectNodeLayoutBlocksSelector.cpp b/Engine/Plugins/Experimental/Mutable/Source/CustomizableObjectEditor/Private/MuCOE/SCustomizableObjectNodeLayoutBlocksSelector.cpp
--- a/Engine/Plugins/Experimental/Mutable/Source/CustomizableObjectEditor/Private/MuCOE/SCustomizableObjectNodeLayoutBlocksSelector.cpp
+++ b/Engine/Plugins/Experimental/Mutable/Source/CustomizableObjectEditor/Private/MuCOE/SCustomizableObjectNodeLayoutBlocksSelector.cpp
@@ -180,7 +180,16 @@ TSharedRef<SWidget> SCustomizableObjectNodeLayoutBlocksSelector::BuildLayoutTool
 
 	LayoutToolbarBuilder.BeginSection("Label");
 	{
-		LayoutToolbarBuilder.AddWidget(BlocksLabel.ToSharedRef());
+		LayoutToolbarBuilder.AddWidget
+		(
+			SNew(SBox)
+			.WidthOverride(125.0f)
+			.VAlign(EVerticalAlignment::VAlign_Center)
+			.HAlign(EHorizontalAlignment::HAlign_Center)
+			[
+				BlocksLabel.ToSharedRef()
+			]
+		);
 	}
 	LayoutToolbarBuilder.EndSection();
 
@@ -189,7 +198,6 @@ TSharedRef<SWidget> SCustomizableObjectNodeLayoutBlocksSelector::BuildLayoutTool
 
 		+SHorizontalBox::Slot()
 		.Padding(4,0)
-		.AutoWidth()
 		[
 			SNew(SBorder)
 			.Padding(0)
diff --git a/Engine/Plugins/Experimental/Mutable/Source/CustomizableObjectEditor/Private/MuCOE/SMutableLayoutViewer.cpp b/Engine/Plugins/Experimental/Mutable/Source/CustomizableObjectEditor/Private/MuCOE/SMutableLayoutViewer.cpp
--- a/Engine/Plugins/Experimental/Mutable/Source/CustomizableObjectEditor/Private/MuCOE/SMutableLayoutViewer.cpp
+++ b/Engine/Plugins/Experimental/Mutable/Source/CustomizableObjectEditor/Private/MuCOE/SMutableLayoutViewer.cpp
@@ -85,6 +85,7 @@ void SMutableLayoutViewer::Construct(const FArguments& InArgs)
 		+ SVerticalBox::Slot()
 		[
 			SAssignNew(LayoutViewer, SCustomizableObjectLayoutGrid)
+			.Mode(ELayoutGridMode::ELGM_Show)
 		]
 	];
 
