[mutable] Now, if a mesh from a node table row has less LODs than the default mesh it will reuse the previous successfully generated LOD.

#rnx

diff --git a/Engine/Plugins/Experimental/Mutable/Source/CustomizableObjectEditor/Private/MuCOE/GenerateMutableSource/GenerateMutableSourceTable.cpp b/Engine/Plugins/Experimental/Mutable/Source/CustomizableObjectEditor/Private/MuCOE/GenerateMutableSource/GenerateMutableSourceTable.cpp
--- a/Engine/Plugins/Experimental/Mutable/Source/CustomizableObjectEditor/Private/MuCOE/GenerateMutableSource/GenerateMutableSourceTable.cpp
+++ b/Engine/Plugins/Experimental/Mutable/Source/CustomizableObjectEditor/Private/MuCOE/GenerateMutableSource/GenerateMutableSourceTable.cpp
@@ -184,20 +184,20 @@ void FillTableColumn(const UEdGraphPin* Pin, mu::TablePtr MutableTable, FString
 			if (NumLODs != ReferenceNumLODs)
 			{
 				FString Dif_1 = NumLODs > ReferenceNumLODs ? "more" : "less";
-				FString Dif_2 = NumLODs > ReferenceNumLODs ? "Some will be ignored" : "This can cause some Errors.";
+				FString Dif_2 = NumLODs > ReferenceNumLODs ? "Some will be ignored" : "LOD " + FString::FromInt(NumLODs - 1) + " will be used instead. This can cause some performance penalties.";
 
 				FString msg = FString::Printf(TEXT("Mesh from column [%s] row [%s] has %s LODs than the reference mesh. %s"), *ColumnName, *RowName, *Dif_1, *Dif_2);
 				GenerationContext.Compiler->CompilerLog(FText::FromString(msg), CustomNodeTable);
 			}
 
 			// Parameter used for LOD differences
-			int32 LODs = NumLODs <= ReferenceNumLODs ? NumLODs : ReferenceNumLODs;
+			int32 CurrentLOD = 0;
 
 			// Generating skeletal mesh columns
-			for (int32 LODIndex = 0; LODIndex < LODs; ++LODIndex)
+			for (int32 LODIndex = 0; LODIndex < ReferenceNumLODs; ++LODIndex)
 			{
-				int32 NumMaterials = Helper_GetImportedModel(SkeletalMesh)->LODModels[LODIndex].Sections.Num();
-				int32 ReferenceNumMaterials = Helper_GetImportedModel(ReferenceSkeletalMesh)->LODModels[LODIndex].Sections.Num();
+				int32 NumMaterials = Helper_GetImportedModel(SkeletalMesh)->LODModels[CurrentLOD].Sections.Num();
+				int32 ReferenceNumMaterials = Helper_GetImportedModel(ReferenceSkeletalMesh)->LODModels[CurrentLOD].Sections.Num();
 
 				if (NumMaterials > ReferenceNumMaterials)
 				{
@@ -218,7 +218,7 @@ void FillTableColumn(const UEdGraphPin* Pin, mu::TablePtr MutableTable, FString
 						CurrentColumn = MutableTable->AddColumn(TCHAR_TO_ANSI(*MutableColumnName), mu::TABLE_COLUMN_TYPE::TCT_MESH);
 					}
 
-					mu::MeshPtr MutableMesh = GenerateMutableMesh(SkeletalMesh, LODIndex, MatIndex, GenerationContext, CustomNodeTable);
+					mu::MeshPtr MutableMesh = GenerateMutableMesh(SkeletalMesh, CurrentLOD, MatIndex, GenerationContext, CustomNodeTable);
 
 					if (MutableMesh)
 					{
@@ -246,6 +246,8 @@ void FillTableColumn(const UEdGraphPin* Pin, mu::TablePtr MutableTable, FString
 						GenerationContext.Compiler->CompilerLog(FText::FromString(msg), CustomNodeTable);
 					}
 				}
+
+				CurrentLOD = (CurrentLOD + 1) >= NumLODs ? CurrentLOD : (CurrentLOD + 1);
 			}
 		}
 
@@ -268,18 +270,19 @@ void FillTableColumn(const UEdGraphPin* Pin, mu::TablePtr MutableTable, FString
 			if (NumLODs != ReferenceNumLODs)
 			{
 				FString Dif_1 = NumLODs > ReferenceNumLODs ? "more" : "less";
-				FString Dif_2 = NumLODs > ReferenceNumLODs ? "Some will be ignored" : "This can cause some Errors.";
+				FString Dif_2 = NumLODs > ReferenceNumLODs ? "Some will be ignored" : "LOD " + FString::FromInt(NumLODs - 1) + " will be used instead. This can cause some performance penalties.";
 
 				FString msg = FString::Printf(TEXT("Mesh from column [%s] row [%s] has %s LODs than the reference mesh. %s"), *ColumnName, *RowName, *Dif_1, *Dif_2);
 				GenerationContext.Compiler->CompilerLog(FText::FromString(msg), CustomNodeTable);
 			}
 
-			int32 LODs = NumLODs <= ReferenceNumLODs ? NumLODs : ReferenceNumLODs;
+			// Parameter used for LOD differences
+			int32 CurrentLOD = 0;
 
-			for (int32 LODIndex = 0; LODIndex < LODs; ++LODIndex)
+			for (int32 LODIndex = 0; LODIndex < ReferenceNumLODs; ++LODIndex)
 			{
-				int32 NumMaterials = StaticMesh->GetRenderData()->LODResources[LODIndex].Sections.Num();
-				int32 ReferenceNumMaterials = ReferenceStaticMesh->GetRenderData()->LODResources[LODIndex].Sections.Num();
+				int32 NumMaterials = StaticMesh->GetRenderData()->LODResources[CurrentLOD].Sections.Num();
+				int32 ReferenceNumMaterials = ReferenceStaticMesh->GetRenderData()->LODResources[CurrentLOD].Sections.Num();
 
 				if (NumMaterials > ReferenceNumMaterials)
 				{
@@ -300,7 +303,7 @@ void FillTableColumn(const UEdGraphPin* Pin, mu::TablePtr MutableTable, FString
 						CurrentColumn = MutableTable->AddColumn(TCHAR_TO_ANSI(*MutableColumnName), mu::TABLE_COLUMN_TYPE::TCT_MESH);
 					}
 
-					mu::MeshPtr MutableMesh = GenerateMutableMesh(StaticMesh, LODIndex, MatIndex, GenerationContext, CustomNodeTable);
+					mu::MeshPtr MutableMesh = GenerateMutableMesh(StaticMesh, CurrentLOD, MatIndex, GenerationContext, CustomNodeTable);
 
 					if (MutableMesh)
 					{
@@ -314,6 +317,8 @@ void FillTableColumn(const UEdGraphPin* Pin, mu::TablePtr MutableTable, FString
 						GenerationContext.Compiler->CompilerLog(FText::FromString(msg), CustomNodeTable);
 					}
 				}
+
+				CurrentLOD = (CurrentLOD + 1) >= NumLODs ? CurrentLOD : (CurrentLOD + 1);
 			}
 		}
 
@@ -641,7 +646,7 @@ mu::TablePtr GenerateMutableSourceTable(const FString& TableName, const UEdGraph
 			TArray<FString> ColumnHeaders = Table->GetColumnTitles();
 
 			// Getting names of the rows to access the information
-			TArray<FName> RowNames = Table->GetRowNames();
+			TArray<FName> RowNames = TypedTable->GetRowNames();
 
 			// index to iterate the column headers
 			int32 ColumnHeaderIndex = 1;
