[mutable] Fix static analysis warnings related to image channel count being constant.
#rnx

diff --git a/Engine/Plugins/Experimental/Mutable/Source/MutableRuntime/Private/MuR/OpImageLayer.h b/Engine/Plugins/Experimental/Mutable/Source/MutableRuntime/Private/MuR/OpImageLayer.h
--- a/Engine/Plugins/Experimental/Mutable/Source/MutableRuntime/Private/MuR/OpImageLayer.h
+++ b/Engine/Plugins/Experimental/Mutable/Source/MutableRuntime/Private/MuR/OpImageLayer.h
@@ -5,7 +5,6 @@
 #include "MuR/ImagePrivate.h"
 #include "MuR/ImageRLE.h"
 #include "Async/ParallelFor.h"
-
 #include "Templates/UnrealTemplate.h"
 
 namespace mu
@@ -209,13 +208,14 @@ namespace mu
 
         bool isUncompressed = ( maskFormat == EImageFormat::IF_L_UBYTE );
 
+		constexpr int32 NumColorChannels = FMath::Min(NC,3);
         if ( isUncompressed )
         {
             int32 pixelCount = pBase->CalculatePixelCount();
             for ( int i=0; i<pixelCount; ++i )
             {
                 unsigned mask = pMaskBuf[i];
-                for ( int c=0; c<NC && c<3; ++c )
+                for ( int32 c=0; c<NumColorChannels; ++c )
                 {
                     unsigned base = pBaseBuf[NC*i+c];
                     unsigned result = BLEND_FUNC_MASKED( base, top[c], mask );
@@ -272,7 +272,7 @@ namespace mu
                         {
                             for ( int i=0; i<equal; ++i )
                             {
-                                for ( int c=0; c<NC && c<3; ++c )
+                                for ( int32 c=0; c<NumColorChannels; ++c )
                                 {
                                     unsigned base = pBaseBuf[NC*i+c];
                                     unsigned result = BLEND_FUNC( base, top[c] );
@@ -297,7 +297,7 @@ namespace mu
                         {
                             for ( int i=0; i<equal; ++i )
                             {
-                                for ( int c=0; c<NC && c<3; ++c )
+                                for ( int32 c=0; c<NumColorChannels; ++c )
                                 {
                                     unsigned base = pBaseBuf[NC*i+c];
                                     unsigned result = BLEND_FUNC_MASKED( base, top[c], equalPixel );
@@ -333,7 +333,7 @@ namespace mu
 						check(pDestBuf + NC * different <= pStartDestBuf + pBase->GetDataSize());
                         for ( int i=0; i<different; ++i )
                         {
-                            for ( int c=0; c<NC && c<3; ++c )
+                            for ( int32 c=0; c<NumColorChannels; ++c )
                             {
                                 unsigned mask = pMaskBuf[i];
                                 unsigned base = pBaseBuf[NC*i+c];
@@ -643,7 +643,8 @@ namespace mu
 					] (uint32 i)
 					{
 						unsigned mask = pMaskBuf[i];
-						for (int c = 0; c < NC && c < 3; ++c)
+						constexpr int32 NumColorChannels = FMath::Min(NC, 3);
+						for (int c = 0; c <NumColorChannels; ++c)
 						{
 							unsigned base = pBaseBuf[NC * i + c];
 							unsigned blended = pBlendedBuf[NC * i + c];
