[mutable] Fix wrong memory copy size detected in static analysis (222991).
#rnx

diff --git a/Engine/Plugins/Experimental/Mutable/Source/MutableRuntime/Private/MuR/Model.cpp b/Engine/Plugins/Experimental/Mutable/Source/MutableRuntime/Private/MuR/Model.cpp
--- a/Engine/Plugins/Experimental/Mutable/Source/MutableRuntime/Private/MuR/Model.cpp
+++ b/Engine/Plugins/Experimental/Mutable/Source/MutableRuntime/Private/MuR/Model.cpp
@@ -599,7 +599,7 @@ namespace mu
 
             //! Floating point RGBA colour, with each channel ranging from 0.0 to 1.0
             case PARAMETER_TYPE::T_COLOUR:
-                dataSize = 4*sizeof(float);
+                dataSize = 3*sizeof(float);
                 parameterValuesBlob.SetNum( pos+dataSize );
 				FMemory::Memcpy( &parameterValuesBlob[pos],
                         &pParams->GetPrivate()->m_values[param].m_colour,
