Optimized GenerateMutableSourceMesh to reuse intermediate Skeletal Mesh objects where possible and look up morph target vertices more efficiently.

This significantly reduces Customizable Object compile times for complex meshes.

diff --git a/Engine/Plugins/Experimental/Mutable/Source/CustomizableObjectEditor/Private/MuCOE/GenerateMutableSource/GenerateMutableSource.h b/Engine/Plugins/Experimental/Mutable/Source/CustomizableObjectEditor/Private/MuCOE/GenerateMutableSource/GenerateMutableSource.h
--- a/Engine/Plugins/Experimental/Mutable/Source/CustomizableObjectEditor/Private/MuCOE/GenerateMutableSource/GenerateMutableSource.h
+++ b/Engine/Plugins/Experimental/Mutable/Source/CustomizableObjectEditor/Private/MuCOE/GenerateMutableSource/GenerateMutableSource.h
@@ -702,6 +702,20 @@ struct FMutableGraphGenerationContext
 
 	// Stores the parameters generated in the node tables
 	TMap<const class UCustomizableObjectNodeTable*, TArray<FGuid>> GeneratedParametersInTables;
+
+	struct FMeshWithBoneRemovalApplied
+	{
+		TObjectPtr<class USkeletalMesh> Mesh;
+		
+		// Key is LOD index
+		//
+		// This is a cache. Entries are added on demand. A processed LOD with no bones removed will have an empty entry.
+		TMap<int32, TMap<int32, int32>> RemovedBonesActiveParentIndicesPerLOD;
+
+		bool bHasBonesToRemove = false;
+	};
+
+	TMap<TObjectPtr<class USkeletalMesh>, FMeshWithBoneRemovalApplied> MeshesWithBoneRemovalApplied;
 };
 
 /** Pin Data scope wrapper. Pops the pin data on scope exit. */
diff --git a/Engine/Plugins/Experimental/Mutable/Source/CustomizableObjectEditor/Private/MuCOE/GenerateMutableSource/GenerateMutableSourceMesh.cpp b/Engine/Plugins/Experimental/Mutable/Source/CustomizableObjectEditor/Private/MuCOE/GenerateMutableSource/GenerateMutableSourceMesh.cpp
--- a/Engine/Plugins/Experimental/Mutable/Source/CustomizableObjectEditor/Private/MuCOE/GenerateMutableSource/GenerateMutableSourceMesh.cpp
+++ b/Engine/Plugins/Experimental/Mutable/Source/CustomizableObjectEditor/Private/MuCOE/GenerateMutableSource/GenerateMutableSourceMesh.cpp
@@ -2,6 +2,7 @@
 
 #include "MuCOE/GenerateMutableSource/GenerateMutableSourceMesh.h"
 
+#include "Algo/BinarySearch.h"
 #include "Animation/MorphTarget.h"
 #include "Animation/PoseAsset.h"
 #include "Animation/Skeleton.h"
@@ -132,153 +133,193 @@
 
 #define LOCTEXT_NAMESPACE "CustomizableObjectEditor"
 
-
-mu::MeshPtr ConvertSkeletalMeshToMutable(USkeletalMesh* InSkeletalMesh, int LOD, int MaterialIndex, FMutableGraphGenerationContext& GenerationContext, const UCustomizableObjectNode* CurrentNode)
+TObjectPtr<const USkeletalMesh> GetMeshWithBoneRemovalApplied(TObjectPtr<USkeletalMesh> InSkeletalMesh, int32 InLODIndex, FMutableGraphGenerationContext& GenerationContext, const UCustomizableObjectNode* CurrentNode, TMap<int32, int32>& OutRemovedBonesActiveParentIndices)
 {
-	// Get the mesh generation flags to use
-	uint32 CurrentFlags = GenerationContext.MeshGenerationFlags.Last();
-	bool bIgnoreSkeleton = CurrentFlags & uint32(EMutableMeshConversionFlags::IgnoreSkinning);
-	bool bIgnorePhysics = CurrentFlags & uint32(EMutableMeshConversionFlags::IgnorePhysics);
+	FMutableGraphGenerationContext::FMeshWithBoneRemovalApplied& CacheEntry = GenerationContext.MeshesWithBoneRemovalApplied.FindOrAdd(InSkeletalMesh);
+	if (CacheEntry.Mesh)
+	{
+		if (CacheEntry.bHasBonesToRemove)
+		{
+			TMap<int32, int32>* RemovedBonesActiveParentIndicesForLOD = CacheEntry.RemovedBonesActiveParentIndicesPerLOD.Find(InLODIndex);
+			if (RemovedBonesActiveParentIndicesForLOD != nullptr)
+			{
+				// Note that this copies the map. This function could provide a pointer into the cache instead, but
+				// it would be hard to track when the cache pointer is invalidated, so this copy is done for safety.
+				OutRemovedBonesActiveParentIndices = *RemovedBonesActiveParentIndicesForLOD;
+				return CacheEntry.Mesh;
+			}
+		}
+		else
+		{
+			OutRemovedBonesActiveParentIndices.Reset();
+			return CacheEntry.Mesh;
+		}
+	}
 
-	mu::MeshPtr MutableMesh = new mu::Mesh();
-		
-	if(!InSkeletalMesh)
+	CacheEntry.bHasBonesToRemove = false;
+	for (int32 LODIndex = 0; LODIndex < Helper_GetLODNum(InSkeletalMesh); ++LODIndex)
 	{
-		// Trying to convert a nullptr to mutable?
-		return nullptr;
+		if (Helper_GetLODInfoBonesToRemove(InSkeletalMesh, LODIndex).Num() > 0)
+		{
+			CacheEntry.bHasBonesToRemove = true;
+			break;
+		}
 	}
-	
-	USkeleton* InSkeleton = InSkeletalMesh->GetSkeleton();
-	check(InSkeleton);
 
-	TObjectPtr<USkeletalMesh> SkeletalMesh = InSkeletalMesh;
+	if (!CacheEntry.bHasBonesToRemove)
+	{
+		// No changes needed, just return the original mesh
+		CacheEntry.Mesh = InSkeletalMesh;
+		OutRemovedBonesActiveParentIndices.Reset();
+		return CacheEntry.Mesh;
+	}
 
-	FMutableComponentInfo& MutComponentInfo = GenerationContext.GetCurrentComponentInfo();
-	USkeletalMesh* ComponentRefSkeletalMesh = MutComponentInfo.RefSkeletalMesh;
-	USkeleton* ComponentRefSkeleton = MutComponentInfo.RefSkeleton;
-	check(ComponentRefSkeletalMesh);
-	check(ComponentRefSkeleton);
-	
+	if (!CacheEntry.Mesh)
+	{
+		UE_LOG(LogMutable, Log, TEXT("GetMeshWithBoneRemovalApplied duplicating mesh %s"), *InSkeletalMesh->GetPathName());
 
-	// Apply removed bones
-	bool bHasBonesToRemove = false;
-	TMap<int32, int32> RemovedBonesActiveParentIndices;
-	if (!bIgnoreSkeleton)
+		// Keep the name since it is used at least when gathering info about morph targets. 
+		CacheEntry.Mesh = (USkeletalMesh*)StaticDuplicateObject(InSkeletalMesh, GetTransientPackage(), InSkeletalMesh->GetFName(), EObjectFlags::RF_Transient);
+		check(CacheEntry.Mesh);
+
+		if (Helper_GetLODNum(InSkeletalMesh) < Helper_GetLODNum(CacheEntry.Mesh))
+		{
+			GenerationContext.Compiler->CompilerLog(LOCTEXT("ReferenceLODMismatch", "The reference mesh has less LODs than the generated mesh."),
+				CurrentNode,
+				EMessageSeverity::Warning);
+		}
+	}
+
+	UE_LOG(LogMutable, Log, TEXT("GetMeshWithBoneRemovalApplied calculating bone hierarchy for mesh %s LOD %i"), *InSkeletalMesh->GetPathName(), InLODIndex);
+
+	// If this code is reached, it should mean that the removed bone indices haven't been cached for this LOD yet
+	check(!CacheEntry.RemovedBonesActiveParentIndicesPerLOD.Contains(InLODIndex));
+
+	TMap<int32, int32>& CachedRemovedBonesActiveParentIndices = CacheEntry.RemovedBonesActiveParentIndicesPerLOD.Add(InLODIndex);
+	// The new entry should be empty
+	check(CachedRemovedBonesActiveParentIndices.Num() == 0);
+
+	TSet<int32> RemovedBoneIndices;
 	{
-		for (int32 LODIndex = 0; LODIndex < Helper_GetLODNum(InSkeletalMesh); ++LODIndex)
+		TArray<FName> BoneNamesToRemove;
+		BoneNamesToRemove.Reset(Helper_GetLODInfoBonesToRemove(InSkeletalMesh, InLODIndex).Num());
+		
+		for (const FBoneReference& BoneRef : Helper_GetLODInfoBonesToRemove(InSkeletalMesh, InLODIndex))
 		{
-			if (Helper_GetLODInfoBonesToRemove(InSkeletalMesh, LODIndex).Num() > 0)
-			{
-				bHasBonesToRemove = true;
-				break;
-			}
+			BoneNamesToRemove.Add(BoneRef.BoneName);
 		}
 
-		if (bHasBonesToRemove)
+		if (BoneNamesToRemove.Num() > 0)
 		{
-			IMeshUtilities& MeshUtilities = FModuleManager::Get().LoadModuleChecked<IMeshUtilities>("MeshUtilities");
-			// Keep the name since it is used at least when gathering info about morph targets. 
-			SkeletalMesh = (USkeletalMesh*)StaticDuplicateObject(SkeletalMesh, GetTransientPackage(), SkeletalMesh->GetFName(), EObjectFlags::RF_Transient);
-			check(SkeletalMesh);
-			
-			TMap<int32, bool> RemovedBoneIndices;
-			for (int32 LODIndex = 0; LODIndex < Helper_GetLODNum(SkeletalMesh); ++LODIndex)
-			{
-				TArray<FName> BoneNamesToRemove;
+			UE_LOG(LogMutable, Log, TEXT("GetMeshWithBoneRemovalApplied removing bones from mesh %s LOD %i"), *InSkeletalMesh->GetPathName(), InLODIndex);
 
-				if (Helper_GetLODNum(InSkeletalMesh) > LODIndex)
-				{
-					for (const FBoneReference& BoneRef : Helper_GetLODInfoBonesToRemove(InSkeletalMesh, LODIndex))
-					{
-						BoneNamesToRemove.Add(BoneRef.BoneName);
-					}
+			IMeshUtilities& MeshUtilities = FModuleManager::Get().LoadModuleChecked<IMeshUtilities>("MeshUtilities");
+			MeshUtilities.RemoveBonesFromMesh(CacheEntry.Mesh, InLODIndex, &BoneNamesToRemove);
+		}
 
-					if (BoneNamesToRemove.Num() > 0)
-					{
-						MeshUtilities.RemoveBonesFromMesh(SkeletalMesh, LODIndex, &BoneNamesToRemove);
-					}
+		for (const FName& BoneName : BoneNamesToRemove)
+		{
+			int32 RefBoneIndex = CacheEntry.Mesh->GetRefSkeleton().FindBoneIndex(BoneName);
+			if (RefBoneIndex != INDEX_NONE)
+			{
+				RemovedBoneIndices.Add(RefBoneIndex);
+			}
+		}
+	}
 
-					if (LODIndex == LOD)
-					{
-						// Find Removed Bones Indices
-						for (const FName& Bone : BoneNamesToRemove)
-						{
-							int32 RefBoneIndex = SkeletalMesh->GetRefSkeleton().FindBoneIndex(Bone);
-							if (RefBoneIndex != INDEX_NONE)
-							{
-								RemovedBoneIndices.Add(RefBoneIndex, true);
+	// Reconstruct RefSkeleton's bone hierarchy
+	TMap<int32, TArray<int32>> SkeletalTree;
+	FReferenceSkeleton& RefSkeleton = InSkeletalMesh->GetRefSkeleton();
+	for (int32 i = 0; i < RefSkeleton.GetRawBoneNum(); ++i)
+	{
+		int32 ParentBoneIndex = RefSkeleton.GetRawParentIndex(i);
+		if (ParentBoneIndex == INDEX_NONE)
+		{
+			continue;
+		}
 
-							}
-						}
+		if (TArray<int32>* Children = SkeletalTree.Find(ParentBoneIndex))
+		{
+			Children->Add(i);
+		}
+		else
+		{
+			TArray<int32> NewChildrenArray;
+			NewChildrenArray.Add(i);
+			SkeletalTree.Add(ParentBoneIndex, MoveTemp(NewChildrenArray));
+		}
+	}
 
-					}
-				}
-				else
-				{
-					GenerationContext.Compiler->CompilerLog(LOCTEXT("ReferenceLODMismatch", "The reference mesh has less LODs than the generated mesh."),
-						CurrentNode,
-						EMessageSeverity::Warning);
-					break;
-				}
-			}
+	// Search for the first active parent of each removed bone
+	// If a vertex is influenced by a removed bone the influences will be transferred to the parent bone.
+	for (const int32 RemovedBone : RemovedBoneIndices)
+	{
+		int32 ParentIndex = RefSkeleton.GetRawParentIndex(RemovedBone);
 
-			// Reconstruct RefSkeleton's bone hierarchy
-			TMap<int32, TArray<int32>> SkeletalTree;
-			FReferenceSkeleton& RefSkeleton = InSkeletalMesh->GetRefSkeleton();
-			for (int32 i = 0; i < RefSkeleton.GetRawBoneNum(); ++i)
+		if (TArray<int32>* Children = SkeletalTree.Find(RemovedBone))
+		{
+			for (int32 i = 0; i < Children->Num(); ++i)
 			{
-				int32 ParentBoneIndex = RefSkeleton.GetRawParentIndex(i);
-				if (ParentBoneIndex == INDEX_NONE)
+				int32& ChildIndex = (*Children)[i];
+				CachedRemovedBonesActiveParentIndices.Add(ChildIndex, ParentIndex);
+
+				if (RemovedBoneIndices.Find(ChildIndex))
 				{
 					continue;
 				}
 
-				if (TArray<int32>* Children = SkeletalTree.Find(ParentBoneIndex))
+				if (TArray<int32>* ChildChildren = SkeletalTree.Find(ChildIndex))
 				{
-					Children->Add(i);
-				}
-				else
-				{
-					TArray<int32> NewChildrenArray;
-					NewChildrenArray.Add(i);
-					SkeletalTree.Add(ParentBoneIndex, NewChildrenArray);
+					for (int j = 0; j < ChildChildren->Num(); ++j)
+					{
+						Children->Add((*ChildChildren)[j]);
+					}
 				}
 			}
+		}
 
-			// Search for the first active parent of each removed bone
-			// If a vertex is influenced by a removed bone the influences will be transferred to the parent bone.
-			for (const TPair<int32, bool>& RemovedBone : RemovedBoneIndices)
-			{
-				int32 ParentIndex = RefSkeleton.GetRawParentIndex(RemovedBone.Key);
+		if (ParentIndex >= 0)
+		{
+			CachedRemovedBonesActiveParentIndices.Add(RemovedBone, ParentIndex);
+		}
+	}
 
-				if (TArray<int32>* Children = SkeletalTree.Find(RemovedBone.Key))
-				{
-					for (int32 i = 0; i < Children->Num(); ++i)
-					{
-						int32& ChildIndex = (*Children)[i];
-						RemovedBonesActiveParentIndices.Add(ChildIndex, ParentIndex);
+	OutRemovedBonesActiveParentIndices = CachedRemovedBonesActiveParentIndices;
+	return CacheEntry.Mesh;
+}
 
-						if (RemovedBoneIndices.Find(ChildIndex))
-						{
-							continue;
-						}
+mu::MeshPtr ConvertSkeletalMeshToMutable(USkeletalMesh* InSkeletalMesh, int LOD, int MaterialIndex, FMutableGraphGenerationContext& GenerationContext, const UCustomizableObjectNode* CurrentNode)
+{
+	// Get the mesh generation flags to use
+	uint32 CurrentFlags = GenerationContext.MeshGenerationFlags.Last();
+	bool bIgnoreSkeleton = CurrentFlags & uint32(EMutableMeshConversionFlags::IgnoreSkinning);
+	bool bIgnorePhysics = CurrentFlags & uint32(EMutableMeshConversionFlags::IgnorePhysics);
 
-						if (TArray<int32>* ChildChildren = SkeletalTree.Find(ChildIndex))
-						{
-							for (int j = 0; j < ChildChildren->Num(); ++j)
-							{
-								Children->Add((*ChildChildren)[j]);
-							}
-						}
-					}
-				}
+	mu::MeshPtr MutableMesh = new mu::Mesh();
+		
+	if(!InSkeletalMesh)
+	{
+		// Trying to convert a nullptr to mutable?
+		return nullptr;
+	}
+	
+	USkeleton* InSkeleton = InSkeletalMesh->GetSkeleton();
+	check(InSkeleton);
 
-				if (ParentIndex >= 0)
-				{
-					RemovedBonesActiveParentIndices.Add(RemovedBone.Key, ParentIndex);
-				}
-			}
-		}
+	TObjectPtr<const USkeletalMesh> SkeletalMesh = InSkeletalMesh;
+
+	FMutableComponentInfo& MutComponentInfo = GenerationContext.GetCurrentComponentInfo();
+	USkeletalMesh* ComponentRefSkeletalMesh = MutComponentInfo.RefSkeletalMesh;
+	USkeleton* ComponentRefSkeleton = MutComponentInfo.RefSkeleton;
+	check(ComponentRefSkeletalMesh);
+	check(ComponentRefSkeleton);
+	
+
+	// Apply removed bones
+	TMap<int32, int32> RemovedBonesActiveParentIndices;
+	if (!bIgnoreSkeleton)
+	{
+		SkeletalMesh = GetMeshWithBoneRemovalApplied(InSkeletalMesh, LOD, GenerationContext, CurrentNode, RemovedBonesActiveParentIndices);
 	}
 
 	const FSkeletalMeshModel* ImportedModel = Helper_GetImportedModel(SkeletalMesh);
@@ -438,9 +479,6 @@ mu::MeshPtr ConvertSkeletalMeshToMutable(USkeletalMesh* InSkeletalMesh, int LOD,
 					InfluencesToReplaceMap.Add(BoneMapIndex, ParentBoneMapIndex != INDEX_NONE ? ParentBoneMapIndex : 0);
 				}
 			}
-
-			// Removed bones might not be used in the current section
-			bHasBonesToRemove = !InfluencesToReplaceMap.IsEmpty();
 		}
 	}
 
@@ -572,7 +610,7 @@ mu::MeshPtr ConvertSkeletalMeshToMutable(USkeletalMesh* InSkeletalMesh, int LOD,
 				int32 MaxOrderedWeighsIndices[MaxMutableWeights] = { -1, -1, -1, -1 };
 
 				// Transfer removed bones influences to parent bones
-				if (bHasBonesToRemove)
+				if (InfluencesToReplaceMap.Num() > 0)
 				{
 					for (int32 i = 0; i < MaxSectionInfluences; ++i)
 					{
@@ -685,7 +723,7 @@ mu::MeshPtr ConvertSkeletalMeshToMutable(USkeletalMesh* InSkeletalMesh, int LOD,
 				}
 			}
 		}
-		else if (bHasBonesToRemove)
+		else if (InfluencesToReplaceMap.Num() > 0)
 		{
 			// Transfer removed bones influences to parent bones
 			for (int32 VertexIndex = VertexStart; VertexIndex < VertexStart + VertexCount && VertexIndex < Vertices.Num(); ++VertexIndex)
@@ -727,9 +765,13 @@ mu::MeshPtr ConvertSkeletalMeshToMutable(USkeletalMesh* InSkeletalMesh, int LOD,
 	{
 		nextBufferIndex += 2;
 
+		// This call involves resolving every TObjectPtr<UMorphTarget> to a UMorphTarget*, so
+		// cache the result here to avoid calling it repeatedly.
+		const TArray<UMorphTarget*>& SkeletalMeshMorphTargets = SkeletalMesh->GetMorphTargets();
+
 		// Find realtime MorphTargets to be used.
-		TArray<UMorphTarget*> UsedMorphTargets;
-		UsedMorphTargets.Empty(SkeletalMesh->GetMorphTargets().Num());
+		TArray<const UMorphTarget*> UsedMorphTargets;
+		UsedMorphTargets.Reserve(SkeletalMeshMorphTargets.Num());
 
 		const UCustomizableObjectNodeSkeletalMesh* NodeTyped = Cast<UCustomizableObjectNodeSkeletalMesh>(CurrentNode);
 		check(NodeTyped);
@@ -737,7 +779,7 @@ mu::MeshPtr ConvertSkeletalMeshToMutable(USkeletalMesh* InSkeletalMesh, int LOD,
         // Add SkeletalMesh morphs to the usage override data structure.
         // Notice this will only be populated here, when compiling.
 		TArray<FRealTimeMorphSelectionOverride>& RealTimeMorphTargetOverrides = GenerationContext.RealTimeMorphTargetsOverrides;
-        for (const TObjectPtr<UMorphTarget>& MorphTarget : SkeletalMesh->GetMorphTargets())
+        for (const UMorphTarget* MorphTarget : SkeletalMeshMorphTargets)
 		{
             // Find if the morph target global override is already present
             // and add it if needed.
@@ -761,13 +803,11 @@ mu::MeshPtr ConvertSkeletalMeshToMutable(USkeletalMesh* InSkeletalMesh, int LOD,
         TArray<FName> UsedMorphTargetsNames = [&]()
         {
             TArray<FName> MorphTargetsNames;
-            MorphTargetsNames.Reserve(SkeletalMesh->GetMorphTargets().Num());
+            MorphTargetsNames.Reserve(SkeletalMeshMorphTargets.Num());
             
             if (NodeTyped->bUseAllRealTimeMorphs)
             {
-                TArray<TObjectPtr<UMorphTarget>>& SkeletalMeshMorphs = SkeletalMesh->GetMorphTargets();
-        
-                for (TObjectPtr<UMorphTarget>& MorphTarget : SkeletalMeshMorphs)
+                for (const UMorphTarget* MorphTarget : SkeletalMeshMorphTargets)
                 {
                     check(MorphTarget);
                     MorphTargetsNames.Add(MorphTarget->GetFName());
@@ -817,7 +857,7 @@ mu::MeshPtr ConvertSkeletalMeshToMutable(USkeletalMesh* InSkeletalMesh, int LOD,
             }
         }
 
-		for (UMorphTarget* MorphTarget : SkeletalMesh->GetMorphTargets())
+		for (const UMorphTarget* MorphTarget : SkeletalMeshMorphTargets)
 		{
 			if (!MorphTarget)
 			{
@@ -871,31 +911,42 @@ mu::MeshPtr ConvertSkeletalMeshToMutable(USkeletalMesh* InSkeletalMesh, int LOD,
 
 		if (UsedMorphTargets.Num())
 		{
+			const double StartTime = FPlatformTime::Seconds();
+
+			TArray<FMorphTargetVertexData> MorphsUsed;
 			for (int32 VertexIdx = VertexStart; VertexIdx < VertexStart + VertexCount && VertexIdx < Vertices.Num(); ++VertexIdx)
 			{
-				TArray<FMorphTargetVertexData> MorphsUsed;
+				MorphsUsed.Reset(UsedMorphTargets.Num());
+
 				for (const UMorphTarget* MorphTarget : UsedMorphTargets)
 				{
-					if (!MorphTarget) continue;
+					if (!MorphTarget)
+					{
+						continue;
+					}
 
 					const TArray<FMorphTargetLODModel>& MorphLODModels = MorphTarget->GetMorphLODModels();
 
-					if (!(MorphTarget && LOD < MorphLODModels.Num()))
+					if (LOD >= MorphLODModels.Num()
+						|| !MorphLODModels[LOD].SectionIndices.Contains(MaterialIndex))
 					{
 						continue;
 					}
 
-					const FMorphTargetDelta* VertexFound = MorphLODModels[LOD].Vertices.FindByPredicate(
-						[&VertexIdx](const FMorphTargetDelta& Elem) -> bool
-					{ 
-							return Elem.SourceIdx == VertexIdx; 
-					});
+					// The vertices should be sorted by SourceIdx
+					check(MorphLODModels[LOD].Vertices.Num() < 2 || MorphLODModels[LOD].Vertices[0].SourceIdx < MorphLODModels[LOD].Vertices.Last().SourceIdx);
 
-					if (!VertexFound)
+					const int32 VertexFoundIndex = Algo::BinarySearchBy(
+						MorphLODModels[LOD].Vertices,
+						(uint32)VertexIdx,
+						[](const FMorphTargetDelta& Element) { return Element.SourceIdx; });
+
+					if (VertexFoundIndex == INDEX_NONE)
 					{
 						continue;
 					}
 
+					const FMorphTargetDelta& VertexFound = MorphLODModels[LOD].Vertices[VertexFoundIndex];
 					const FName MorphTargetName = MorphTarget->GetFName();
 
 					TArray<FMorphTargetInfo>& ContributingMorphTargetsInfo = GenerationContext.ContributingMorphTargetsInfo;
@@ -910,7 +961,7 @@ mu::MeshPtr ConvertSkeletalMeshToMutable(USkeletalMesh* InSkeletalMesh, int LOD,
 					FMorphTargetInfo& MorphTargetInfo = ContributingMorphTargetsInfo[DestMorphTargetIdx];
 					MorphTargetInfo.LodNum = FMath::Max(MorphTargetInfo.LodNum, LOD + 1);
 
-					MorphsUsed.Emplace(FMorphTargetVertexData{ VertexFound->PositionDelta, VertexFound->TangentZDelta, DestMorphTargetIdx });
+					MorphsUsed.Emplace(FMorphTargetVertexData{ VertexFound.PositionDelta, VertexFound.TangentZDelta, DestMorphTargetIdx });
 				}
 
 				if (MorphsUsed.Num())
@@ -921,6 +972,8 @@ mu::MeshPtr ConvertSkeletalMeshToMutable(USkeletalMesh* InSkeletalMesh, int LOD,
 					GenerationContext.MorphTargetReconstructionData.Append(MorphsUsed);
 				}
 			}
+
+			UE_LOG(LogMutable, Log, TEXT("Processing morph targets took %.2f ms"), (FPlatformTime::Seconds() - StartTime) * 1000.0);
 		}
 	}
 
@@ -950,7 +1003,7 @@ mu::MeshPtr ConvertSkeletalMeshToMutable(USkeletalMesh* InSkeletalMesh, int LOD,
 		// clothing assets are shared among all LODs in a section
 		const int32 ClothingAssetIndex = [&]() -> int32
 		{
-			UClothingAssetBase* ClothingAssetBase = SkeletalMesh->GetSectionClothingAsset(LOD, MaterialIndex);
+			const UClothingAssetBase* ClothingAssetBase = SkeletalMesh->GetSectionClothingAsset(LOD, MaterialIndex);
 
 			if (!ClothingAssetBase)
 			{
@@ -965,7 +1018,7 @@ mu::MeshPtr ConvertSkeletalMeshToMutable(USkeletalMesh* InSkeletalMesh, int LOD,
 				return FoundIndex;
 			}
 
-			UClothingAssetCommon* Asset = Cast<UClothingAssetCommon>(ClothingAssetBase);
+			const UClothingAssetCommon* Asset = Cast<UClothingAssetCommon>(ClothingAssetBase);
 			if (!Asset)
 			{
 				return INDEX_NONE;
