[mutable] Fix mip generation for RLE comporessed formats.

- Changes are minimal from the previous method.
- Now the Base image data size is computed correctly for non-block-compressed images, before this size was always 0.
- A new Image method GetMipsDataSize() has been introduced to compute the total data of resident mips. This method has the same result value as CalculateDataSize for block-compressed images but it uses the image metadata to compute the actual size for non-block-compressed images.

diff --git a/Engine/Plugins/Experimental/Mutable/Source/MutableRuntime/Private/MuR/Image.cpp b/Engine/Plugins/Experimental/Mutable/Source/MutableRuntime/Private/MuR/Image.cpp
--- a/Engine/Plugins/Experimental/Mutable/Source/MutableRuntime/Private/MuR/Image.cpp
+++ b/Engine/Plugins/Experimental/Mutable/Source/MutableRuntime/Private/MuR/Image.cpp
@@ -470,6 +470,37 @@ namespace mu
 		return pResult;
 	}
 
+	int32 Image::GetMipsDataSize() const
+	{
+		int32 res = 0;
+
+        const FImageFormatData& fdata = GetImageFormatData( m_format );
+        if (fdata.m_bytesPerBlock)
+        {
+			return CalculateDataSize();
+        }
+		else if (m_format == EImageFormat::IF_L_UBYTE_RLE ||
+				 m_format == EImageFormat::IF_L_UBIT_RLE)
+		{	
+			if (m_data.Num())
+			{ 
+			    // Every mip has variable size, but it is stored in the first 4 bytes.
+				const uint8* pMipData = m_data.GetData();
+				for (int l = 0; l < m_lods; ++l)
+				{
+					const uint32_t mipSize = *(const uint32_t*)pMipData;
+					pMipData += mipSize;
+					res += mipSize;
+				}	
+			}
+		}
+		else
+		{
+			checkf(false, TEXT("Trying to get mips data size in an unsupported pixel format."));
+		}
+
+        return res;
+	}
 
     //---------------------------------------------------------------------------------------------
     vec4<float> Image::Sample( vec2<float> coords ) const
diff --git a/Engine/Plugins/Experimental/Mutable/Source/MutableRuntime/Private/MuR/ImageRLE.h b/Engine/Plugins/Experimental/Mutable/Source/MutableRuntime/Private/MuR/ImageRLE.h
--- a/Engine/Plugins/Experimental/Mutable/Source/MutableRuntime/Private/MuR/ImageRLE.h
+++ b/Engine/Plugins/Experimental/Mutable/Source/MutableRuntime/Private/MuR/ImageRLE.h
@@ -6,7 +6,7 @@
 #include "Math/UnrealMathSSE.h"
 #include "MuR/Image.h"
 /** This define enabled additional checks when using RLE compression. These checks have a lot of
-* of verhead so it should usually be disabled.
+* of overhead so it should usually be disabled.
 */
 //#define MUTABLE_DEBUG_RLE
 
diff --git a/Engine/Plugins/Experimental/Mutable/Source/MutableRuntime/Private/MuR/OpImageMipmap.h b/Engine/Plugins/Experimental/Mutable/Source/MutableRuntime/Private/MuR/OpImageMipmap.h
--- a/Engine/Plugins/Experimental/Mutable/Source/MutableRuntime/Private/MuR/OpImageMipmap.h
+++ b/Engine/Plugins/Experimental/Mutable/Source/MutableRuntime/Private/MuR/OpImageMipmap.h
@@ -509,179 +509,207 @@ namespace mu
 	};
 
 
-	//---------------------------------------------------------------------------------------------
-	inline void ImageMipmap_PrepareScratch(Image* pDest, const Image* pBase, int levelCount,
-		SCRATCH_IMAGE_MIPMAP* scratch)
-	{
-		int startLevel = pBase->GetLODCount() - 1;
-
-		check(pDest->GetLODCount() == levelCount);
-		check(pDest->GetSizeX() == pBase->GetSizeX());
-		check(pDest->GetSizeY() == pBase->GetSizeY());
-		check(pDest->GetFormat() == pBase->GetFormat());
+    //---------------------------------------------------------------------------------------------
+    inline void ImageMipmap_PrepareScratch( Image* pDest, const Image* pBase, int levelCount,
+                                            SCRATCH_IMAGE_MIPMAP* scratch )
+    {
+        int startLevel = pBase->GetLODCount() - 1;
+
+        check(pDest->GetLODCount() == levelCount);
+        check(pDest->GetSizeX() == pBase->GetSizeX());
+        check(pDest->GetSizeY() == pBase->GetSizeY());
+        check(pDest->GetFormat() == pBase->GetFormat());
 		(void)pDest;
 
-		switch (pBase->GetFormat())
-		{
-		case EImageFormat::IF_L_UBYTE:
-		case EImageFormat::IF_RGB_UBYTE:
-		case EImageFormat::IF_RGBA_UBYTE:
-		case EImageFormat::IF_BGRA_UBYTE:
-			break;
-
-			// Bad cases: we need to decompress, mip and then recompress. It may be necessary to
-			// generate the latests mips after composing, or in unoptimised code.
-			// TODO: Make sure that the code optimisations avoid this cases generating separate
-			// operations to generate the mip tail.
-		case EImageFormat::IF_BC1:
-		case EImageFormat::IF_BC2:
-		case EImageFormat::IF_BC3:
-		case EImageFormat::IF_BC4:
-		case EImageFormat::IF_BC5:
-		case EImageFormat::IF_ASTC_4x4_RGB_LDR:
-		case EImageFormat::IF_ASTC_4x4_RGBA_LDR:
-		case EImageFormat::IF_ASTC_4x4_RG_LDR:
-		{
-			// Uncompress the last mip that we already have
-			vec2<int> uncompressedSize = pBase->CalculateMipSize(startLevel);
-			scratch->pUncompressed = new Image(
-				(uint16)uncompressedSize[0],
-				(uint16)uncompressedSize[1],
-				1,
-				EImageFormat::IF_RGBA_UBYTE);
-
-			// Generate the mipmaps from there on
-			scratch->pUncompressedMips = new Image(
-				(uint16)uncompressedSize[0],
-				(uint16)uncompressedSize[1],
-				levelCount - startLevel,
-				EImageFormat::IF_RGBA_UBYTE);
-
-			// Compress the mipmapped image
-			scratch->pCompressedMips = new Image(
-				(uint16)uncompressedSize[0],
-				(uint16)uncompressedSize[1],
-				scratch->pUncompressedMips->GetLODCount(),
-				pBase->GetFormat());
+        switch (pBase->GetFormat())
+        {
+        case EImageFormat::IF_L_UBYTE:
+        case EImageFormat::IF_RGB_UBYTE:
+        case EImageFormat::IF_RGBA_UBYTE:
+        case EImageFormat::IF_BGRA_UBYTE:
+            break;
+
+        // Bad cases: we need to decompress, mip and then recompress. It may be necessary to
+        // generate the latests mips after composing, or in unoptimised code.
+        // TODO: Make sure that the code optimisations avoid this cases generating separate
+        // operations to generate the mip tail.
+        case EImageFormat::IF_BC1:
+        case EImageFormat::IF_BC2:
+        case EImageFormat::IF_BC3:
+        case EImageFormat::IF_BC4:
+        case EImageFormat::IF_BC5:
+        case EImageFormat::IF_ASTC_4x4_RGB_LDR:
+        case EImageFormat::IF_ASTC_4x4_RGBA_LDR:
+        case EImageFormat::IF_ASTC_4x4_RG_LDR:
+        {
+		// Uncompress the last mip that we already have
+		vec2<int> uncompressedSize = pBase->CalculateMipSize(startLevel);
+		scratch->pUncompressed = new Image(
+			(uint16)uncompressedSize[0],
+			(uint16)uncompressedSize[1],
+			1,
+			EImageFormat::IF_RGBA_UBYTE);
+
+
+		FImageSize uncompressedMipsSize = pBase->CalculateMipSize(startLevel + 1);
+        // Generate the mipmaps from there on
+        scratch->pUncompressedMips = new Image(
+			(uint16)uncompressedMipsSize[0],
+			(uint16)uncompressedMipsSize[1],
+			FMath::Max(1, levelCount - startLevel - 1),
+			EImageFormat::IF_RGBA_UBYTE);
+
+		// Compress the mipmapped image
+		scratch->pCompressedMips = new Image(
+			(uint16)uncompressedMipsSize[0],
+			(uint16)uncompressedMipsSize[1],
+			scratch->pUncompressedMips->GetLODCount(),
+			pBase->GetFormat());
+
+            break;
+        }
+
+        case EImageFormat::IF_L_UBIT_RLE:
+        case EImageFormat::IF_L_UBYTE_RLE:
+        {
+            // Uncompress the last mip that we already have
+            vec2<int> uncompressedSize = pBase->CalculateMipSize(startLevel);
+            scratch->pUncompressed = new Image(
+                (uint16)uncompressedSize[0],
+                (uint16)uncompressedSize[1],
+                1,
+                EImageFormat::IF_L_UBYTE);
+
+
+			FImageSize uncompressedMipsSize = pBase->CalculateMipSize(startLevel + 1);
+            // Generate the mipmaps from there on
+            scratch->pUncompressedMips = new Image(
+                (uint16)uncompressedMipsSize[0],
+                (uint16)uncompressedMipsSize[1],
+				FMath::Max(1, levelCount - startLevel - 1),
+                EImageFormat::IF_L_UBYTE);
+
+
+            // Compress the mipmapped image
+            scratch->pCompressedMips = new Image(
+                (uint16)uncompressedMipsSize[0],
+                (uint16)uncompressedMipsSize[1],
+                scratch->pUncompressedMips->GetLODCount(),
+                pBase->GetFormat());
+
+            // Preallocate ample memory for the compressed data
+            scratch->pCompressedMips->m_data.SetNum(scratch->pUncompressedMips->GetDataSize());
+
+            // Preallocate ample memory for the destination data
+            pDest->m_data.SetNum(pBase->GetDataSize() + scratch->pUncompressedMips->GetDataSize());
+
+            break;
+        }
+
+        default:
+            checkf( false, TEXT("Format not implemented in mipmap generation."));
+        }
+    }
+
+    //---------------------------------------------------------------------------------------------
+    //! Generate the mipmaps for images.
+	//! if bGenerateOnlyTail is true, generates the mips missing form pBase to levelCount and sets
+	//! them in pDest (the full chain is spit in two images). Otherwise generate the mips missing 
+	//! from pBase up to levelCount and append them in pDest to the already generated pBase's mips.
+    //---------------------------------------------------------------------------------------------
+    inline void ImageMipmap( int imageCompressionQuality, Image* pDest, const Image* pBase,
+                             int levelCount, SCRATCH_IMAGE_MIPMAP* scratch, 
+							 const FMipmapGenerationSettings& settings, bool bGenerateOnlyTail = false )
+    {
+        int startLevel = pBase->GetLODCount() - 1;
 
-			break;
-		}
+		check(!(pBase->m_flags & Image::IF_CANNOT_BE_SCALED));
 
-		case EImageFormat::IF_L_UBIT_RLE:
-		case EImageFormat::IF_L_UBYTE_RLE:
+	
+		if (!bGenerateOnlyTail)
 		{
-			// Uncompress the last mip that we already have
-			vec2<int> uncompressedSize = pBase->CalculateMipSize(startLevel);
-			scratch->pUncompressed = new Image(
-				(uint16)uncompressedSize[0],
-				(uint16)uncompressedSize[1],
-				1,
-				EImageFormat::IF_L_UBYTE);
-
-			// Generate the mipmaps from there on
-			scratch->pUncompressedMips = new Image(
-				(uint16)uncompressedSize[0],
-				(uint16)uncompressedSize[1],
-				levelCount - startLevel,
-				EImageFormat::IF_L_UBYTE);
-
-			// Compress the mipmapped image
-			scratch->pCompressedMips = new Image(
-				(uint16)uncompressedSize[0],
-				(uint16)uncompressedSize[1],
-				scratch->pUncompressedMips->GetLODCount(),
-				pBase->GetFormat());
-
-			// Preallocate ample memory for the compressed data
-			scratch->pCompressedMips->m_data.SetNum(scratch->pUncompressedMips->GetDataSize());
-
-			// Preallocate ample memory for the destination data
-			pDest->m_data.SetNum(pBase->GetDataSize()
-				+
-				scratch->pUncompressedMips->GetDataSize());
-
-			break;
+			check(pDest->GetLODCount() == levelCount);
+			check(pDest->GetSizeX() == pBase->GetSizeX());
+			check(pDest->GetSizeY() == pBase->GetSizeY());
 		}
+		else
+		{	
+			check(pDest->GetLODCount() + pBase->GetLODCount() == levelCount);
 
-		default:
-			checkf(false, TEXT("Format not implemented in mipmap generation."));
+			check(
+				[&]() -> bool
+				{
+					const FImageSize BaseImageNextMipSize = pBase->CalculateMipSize(startLevel + 1);
+					return BaseImageNextMipSize.x() == pDest->GetSizeX() && BaseImageNextMipSize.y() == pDest->GetSizeY();
+				}());
 		}
-	}
-
-
-	//---------------------------------------------------------------------------------------------
-	inline void ImageMipmap(int imageCompressionQuality, Image* pDest, const Image* pBase,
-		int levelCount, SCRATCH_IMAGE_MIPMAP* scratch,
-		FMipmapGenerationSettings& settings)
-	{
-		int startLevel = pBase->GetLODCount() - 1;
-
-		check(!(pBase->m_flags & Image::IF_CANNOT_BE_SCALED));
 
-		check(pDest->GetLODCount() == levelCount);
-		check(pDest->GetSizeX() == pBase->GetSizeX());
-		check(pDest->GetSizeY() == pBase->GetSizeY());
-		check(pDest->GetFormat() == pBase->GetFormat());
+        check(pDest->GetFormat() == pBase->GetFormat());
 
-		// Copy already-existing mips
-
-		// Calculate the data size, since the base could have more data allocated.
-		int32 mipDataSize = pBase->CalculateDataSize();
-		check(pDest->GetDataSize() >= mipDataSize);
-		FMemory::Memcpy(pDest->GetData(), pBase->GetData(), mipDataSize);
-
-		int mipsToBuild = levelCount - startLevel - 1;
-		if (!mipsToBuild)
+        // Calculate the data size, since the base could have more data allocated.
+		int32 BaseMipDataSize = 0;
+		if (!bGenerateOnlyTail)
 		{
-			return;
+			BaseMipDataSize = pBase->GetMipsDataSize();
+			check(pDest->GetDataSize() >= BaseMipDataSize);
+			FMemory::Memcpy(pDest->GetData(), pBase->GetData(), BaseMipDataSize);
 		}
 
-		uint8_t* pSourceBuf = pDest->GetMipData(startLevel);
-		uint8_t* pDestBuf = pDest->GetMipData(startLevel + 1);
-
-		vec2<int> sourceSize = pDest->CalculateMipSize(startLevel);
-
-		switch (pBase->GetFormat())
-		{
-		case EImageFormat::IF_L_UBYTE:
-			GenerateMipmapsUint8<1>(levelCount - startLevel - 1, pSourceBuf, pDestBuf, sourceSize, settings);
-			break;
-
-		case EImageFormat::IF_RGB_UBYTE:
-			GenerateMipmapsUint8<3>(levelCount - startLevel - 1, pSourceBuf, pDestBuf, sourceSize, settings);
-			break;
-
-		case EImageFormat::IF_BGRA_UBYTE:
-		case EImageFormat::IF_RGBA_UBYTE:
-			GenerateMipmapsUint8<4>(levelCount - startLevel - 1, pSourceBuf, pDestBuf, sourceSize, settings);
-			break;
-
-			// Bad cases: we need to decompress, mip and then recompress. It may be necessary to
-			// generate the latests mips after composing, or in unoptimised code.
-			// TODO: Make sure that the code optimisations avoid this cases generating separate
-			// operations to generate the mip tail.
-		case EImageFormat::IF_BC1:
-		case EImageFormat::IF_BC2:
-		case EImageFormat::IF_BC3:
-		case EImageFormat::IF_BC4:
-		case EImageFormat::IF_BC5:
-		case EImageFormat::IF_ASTC_4x4_RGB_LDR:
-		case EImageFormat::IF_ASTC_4x4_RGBA_LDR:
-		case EImageFormat::IF_ASTC_4x4_RG_LDR:
-		case EImageFormat::IF_L_UBIT_RLE:
-		case EImageFormat::IF_L_UBYTE_RLE:
-		{
-			// Uncompress the last mip that we already have
-			bool bSuccess = ImagePixelFormatInPlace(imageCompressionQuality, scratch->pUncompressed.get(), pBase, startLevel);
+        int mipsToBuild = levelCount - startLevel - 1;
+        if (!mipsToBuild)
+        {
+            return;
+        }
+
+        const uint8_t* pSourceBuf = !bGenerateOnlyTail ? pDest->GetMipData(startLevel) : pBase->GetMipData(startLevel);
+
+        uint8_t* pDestBuf = !bGenerateOnlyTail ? pDest->GetMipData(startLevel + 1) : pDest->GetMipData(0);
+	
+        vec2<int> sourceSize = pBase->CalculateMipSize(startLevel);
+
+        switch (pBase->GetFormat())
+        {
+        case EImageFormat::IF_L_UBYTE:
+            GenerateMipmapsUint8<1>(levelCount - startLevel - 1, pSourceBuf, pDestBuf, sourceSize, settings);
+            break; 
+
+        case EImageFormat::IF_RGB_UBYTE:
+            GenerateMipmapsUint8<3>(levelCount - startLevel - 1, pSourceBuf, pDestBuf, sourceSize, settings);
+            break;
+
+        case EImageFormat::IF_BGRA_UBYTE:
+        case EImageFormat::IF_RGBA_UBYTE:
+            GenerateMipmapsUint8<4>(levelCount - startLevel - 1, pSourceBuf, pDestBuf, sourceSize, settings);
+            break;
+
+        // Bad cases: we need to decompress, mip and then recompress. It may be necessary to
+        // generate the latests mips after composing, or in unoptimised code.
+        // TODO: Make sure that the code optimisations avoid this cases generating separate
+        // operations to generate the mip tail.
+        case EImageFormat::IF_BC1:
+        case EImageFormat::IF_BC2:
+        case EImageFormat::IF_BC3:
+        case EImageFormat::IF_BC4:
+        case EImageFormat::IF_BC5:
+        case EImageFormat::IF_ASTC_4x4_RGB_LDR:
+        case EImageFormat::IF_ASTC_4x4_RGBA_LDR:
+        case EImageFormat::IF_ASTC_4x4_RG_LDR:
+        case EImageFormat::IF_L_UBIT_RLE:
+        case EImageFormat::IF_L_UBYTE_RLE:
+        {
+            // Uncompress the last mip that we already have
+            bool bSuccess = ImagePixelFormatInPlace(imageCompressionQuality, scratch->pUncompressed.get(), pBase, startLevel);
 			check(bSuccess);
 
-			// Generate the mipmaps from there on
-			ImageMipmap(imageCompressionQuality, scratch->pUncompressedMips.get(),
-				scratch->pUncompressed.get(), levelCount - startLevel, 0, settings);
+            // Generate the mipmaps from there on
 
-			// Compress the mipmapped image
+			constexpr bool bGenerateOnlyTailForCompressed = true;
+            ImageMipmap(imageCompressionQuality, scratch->pUncompressedMips.get(),
+                        scratch->pUncompressed.get(), levelCount - startLevel, 0, settings, bGenerateOnlyTailForCompressed);
+
+            // Compress the mipmapped image
 			bSuccess = ImagePixelFormatInPlace(imageCompressionQuality, scratch->pCompressedMips.get(), scratch->pUncompressedMips.get());
-			int32 ExcessDataSize = scratch->pCompressedMips->GetDataSize();
+			int32 ExcessDataSize = FMath::Max(2, scratch->pCompressedMips->GetDataSize());
 			while (!bSuccess)
 			{
 				// Bad case: this should almost never happen.
@@ -692,25 +720,23 @@ namespace mu
 				ExcessDataSize *= 4;
 			}
 
-			int32 FinalDestSize = pSourceBuf - pDest->GetData() + scratch->pCompressedMips->GetDataSize();
+			check(!bGenerateOnlyTail || BaseMipDataSize == 0);
+			const int32 FinalDestSize = BaseMipDataSize + scratch->pCompressedMips->GetDataSize();
 
 			if (FinalDestSize > pDest->GetDataSize())
 			{
 				// Bad case: this should almost never happen.
 				pDest->m_data.SetNum(FinalDestSize);
-				pSourceBuf = pDest->GetMipData(startLevel);
+				pDestBuf = !bGenerateOnlyTail ? pDest->GetMipData(startLevel + 1) : pDest->GetMipData(0);
 			}
 
-			// TODO: We are recompressing the source level, which we could avoid.
-			FMemory::Memcpy(pSourceBuf,
-				scratch->pCompressedMips->GetData(),
-				scratch->pCompressedMips->GetDataSize());
-			break;
-		}
+            FMemory::Memcpy(pDestBuf, scratch->pCompressedMips->GetData(), scratch->pCompressedMips->GetDataSize());
+            break;
+        }
 
-		default:
-			checkf(false, TEXT("Format not implemented in mipmap generation."));
-		}
+        default:
+            checkf( false, TEXT("Format not implemented in mipmap generation."));
+        }
+    }
 
-	}
 }
diff --git a/Engine/Plugins/Experimental/Mutable/Source/MutableRuntime/Public/MuR/Image.h b/Engine/Plugins/Experimental/Mutable/Source/MutableRuntime/Public/MuR/Image.h
--- a/Engine/Plugins/Experimental/Mutable/Source/MutableRuntime/Public/MuR/Image.h
+++ b/Engine/Plugins/Experimental/Mutable/Source/MutableRuntime/Public/MuR/Image.h
@@ -348,6 +348,10 @@ namespace mu
 		//! Return a pointer to the beginning of the data for a particular mip.
 		uint8* GetMipData(int mip);
 		const uint8* GetMipData(int mip) const;
+		
+		//! Return the size of the resident mips. for block-compressed images the result is the same as in CalculateDataSize()
+		//! for non-block-compressed images, the sizes encoded in the image data is used to compute the final size.
+		int32 GetMipsDataSize() const;
 
 		//! See if all the pixels in the image are equal and return the colour.
 		bool IsPlainColour(vec4<float>& colour) const;
