[mutable] Fix CustomizableObjectEditor not updating Skeletal Meshes.

diff --git a/Engine/Plugins/Experimental/Mutable/Source/CustomizableObject/Classes/MuCO/CustomizableObject.h b/Engine/Plugins/Experimental/Mutable/Source/CustomizableObject/Classes/MuCO/CustomizableObject.h
--- a/Engine/Plugins/Experimental/Mutable/Source/CustomizableObject/Classes/MuCO/CustomizableObject.h
+++ b/Engine/Plugins/Experimental/Mutable/Source/CustomizableObject/Classes/MuCO/CustomizableObject.h
@@ -1303,6 +1303,8 @@ public:
 	bool bIsChildObject = false;
 
 	FPostCompileDelegate PostCompileDelegate;
+
+	void PostCompile();
 #endif
 
 	UFUNCTION(BlueprintCallable, Category = CustomizableObject)
@@ -1349,8 +1351,6 @@ private:
 
 #endif
 
-private:
-
 	/** Used to prevent GC of MaskOutCache and keep it in memory while it's needed */
 	UPROPERTY(Transient)
 	TObjectPtr<UMutableMaskOutCache> MaskOutCache_HardRef;
@@ -1359,11 +1359,18 @@ private:
 	UPROPERTY()
 	FGuid Identifier;
 
+	/** Unique identifier. Regenerated each time the object is compiled. */
+	UPROPERTY()
+	FGuid CompilationGuid;
+
+public:
+	FGuid GetCompilationGuid() const;
+
+private:
 	/** BulkData that stores all in-game resources used by Mutable when generating instances.
 	  * Only valid in packaged builds */
 	UPROPERTY()
 	TObjectPtr<UCustomizableObjectBulk> BulkData;
 
 	TSharedPtr<FCustomizableObjectPrivateData> PrivateData;
-
 };
diff --git a/Engine/Plugins/Experimental/Mutable/Source/CustomizableObject/Private/MuCO/CustomizableObject.cpp b/Engine/Plugins/Experimental/Mutable/Source/CustomizableObject/Private/MuCO/CustomizableObject.cpp
--- a/Engine/Plugins/Experimental/Mutable/Source/CustomizableObject/Private/MuCO/CustomizableObject.cpp
+++ b/Engine/Plugins/Experimental/Mutable/Source/CustomizableObject/Private/MuCO/CustomizableObject.cpp
@@ -1058,6 +1058,16 @@ FString UCustomizableObject::GetStateParameterName(int32 StateIndex, int32 Param
 }
 
 
+#if WITH_EDITORONLY_DATA
+void UCustomizableObject::PostCompile()
+{
+	CompilationGuid = FGuid::NewGuid();
+
+	PostCompileDelegate.Broadcast();
+}
+#endif
+
+
 UCustomizableObjectInstance* UCustomizableObject::CreateInstance()
 {
 	MUTABLE_CPUPROFILER_SCOPE(UCustomizableObject::CreateInstance)
@@ -1468,6 +1478,13 @@ void UCustomizableObject::PreSaveRoot(FObjectPreSaveRootContext ObjectSaveContex
 	bIsValidationTriggeredBySave = true;
 }
 
+
+FGuid UCustomizableObject::GetCompilationGuid() const
+{
+	return CompilationGuid;
+}
+
+
 EDataValidationResult UCustomizableObject::IsDataValid(FDataValidationContext& Context)
 {
 	// This method seems to be designed to check data errors (like variables with unexpected values).
diff --git a/Engine/Plugins/Experimental/Mutable/Source/CustomizableObject/Private/MuCO/CustomizableObjectInstanceDescriptor.cpp b/Engine/Plugins/Experimental/Mutable/Source/CustomizableObject/Private/MuCO/CustomizableObjectInstanceDescriptor.cpp
--- a/Engine/Plugins/Experimental/Mutable/Source/CustomizableObject/Private/MuCO/CustomizableObjectInstanceDescriptor.cpp
+++ b/Engine/Plugins/Experimental/Mutable/Source/CustomizableObject/Private/MuCO/CustomizableObjectInstanceDescriptor.cpp
@@ -2095,6 +2095,7 @@ uint32 GetTypeHash(const FCustomizableObjectInstanceDescriptor& Key)
 	if (Key.CustomizableObject) 
 	{
 		Hash = HashCombine(Hash, GetTypeHash(Key.CustomizableObject->GetPathName()));
+		Hash = HashCombine(Hash, GetTypeHash(Key.CustomizableObject->GetCompilationGuid()));
 	}
 
 	for (const FCustomizableObjectBoolParameterValue& Value : Key.BoolParameters)
diff --git a/Engine/Plugins/Experimental/Mutable/Source/CustomizableObjectEditor/Private/MuCOE/CustomizableObjectCompiler.cpp b/Engine/Plugins/Experimental/Mutable/Source/CustomizableObjectEditor/Private/MuCOE/CustomizableObjectCompiler.cpp
--- a/Engine/Plugins/Experimental/Mutable/Source/CustomizableObjectEditor/Private/MuCOE/CustomizableObjectCompiler.cpp
+++ b/Engine/Plugins/Experimental/Mutable/Source/CustomizableObjectEditor/Private/MuCOE/CustomizableObjectCompiler.cpp
@@ -138,7 +138,8 @@ bool FCustomizableObjectCompiler::Tick()
 		bUpdated = true;
 
 		FinishSavingDerivedData();
-		CurrentObject->PostCompileDelegate.Broadcast();
+	
+		CurrentObject->PostCompile();
 
 		UE_LOG(LogMutable, Verbose, TEXT("PROFILE: [ %16.8f ] Finished Saving Derived Data task."), FPlatformTime::Seconds());
 		UE_LOG(LogMutable, Verbose, TEXT("PROFILE: -----------------------------------------------------------"));
