[Mutable] Remove a message that logs every frame while the Customizable Object Instance editor is open

diff --git a/Engine/Plugins/Experimental/Mutable/Source/CustomizableObjectEditor/Private/MuCOE/CustomizableObjectInstanceEditor.cpp b/Engine/Plugins/Experimental/Mutable/Source/CustomizableObjectEditor/Private/MuCOE/CustomizableObjectInstanceEditor.cpp
--- a/Engine/Plugins/Experimental/Mutable/Source/CustomizableObjectEditor/Private/MuCOE/CustomizableObjectInstanceEditor.cpp
+++ b/Engine/Plugins/Experimental/Mutable/Source/CustomizableObjectEditor/Private/MuCOE/CustomizableObjectInstanceEditor.cpp
@@ -907,8 +907,6 @@ void FCustomizableObjectInstanceEditor::Tick(float InDeltaTime)
 		}
 	}
 
-	UE_LOG(LogMutable, Log, TEXT("PROFILE: [ %16.8f ] FCustomizableObjectInstanceEditor::CompileObject end."), FPlatformTime::Seconds());
-
 	// If we want to show the Relevant/Runtime parameters, we need to refresh the details view to make sure that the scroll bar appears		
 	if (bOnlyRelevantParameters != CustomizableObjectInstance->bShowOnlyRelevantParameters)
 	{
