[mutable] CustomizableObject added BlueprintType.

diff --git a/Engine/Plugins/Experimental/Mutable/Source/CustomizableObject/Classes/MuCO/CustomizableObject.h b/Engine/Plugins/Experimental/Mutable/Source/CustomizableObject/Classes/MuCO/CustomizableObject.h
--- a/Engine/Plugins/Experimental/Mutable/Source/CustomizableObject/Classes/MuCO/CustomizableObject.h
+++ b/Engine/Plugins/Experimental/Mutable/Source/CustomizableObject/Classes/MuCO/CustomizableObject.h
@@ -912,7 +912,7 @@ private:
 };
 
 
-UCLASS( config=Engine )
+UCLASS( BlueprintType, config=Engine )
 class CUSTOMIZABLEOBJECT_API UCustomizableObject : public UObject
 {
 public:
