[mutable] Fix bad recollection of mutable constants memory.

diff --git a/Engine/Plugins/Experimental/Mutable/Source/CustomizableObjectEditor/Private/MuCOE/SMutableConstantsWidget.cpp b/Engine/Plugins/Experimental/Mutable/Source/CustomizableObjectEditor/Private/MuCOE/SMutableConstantsWidget.cpp
--- a/Engine/Plugins/Experimental/Mutable/Source/CustomizableObjectEditor/Private/MuCOE/SMutableConstantsWidget.cpp
+++ b/Engine/Plugins/Experimental/Mutable/Source/CustomizableObjectEditor/Private/MuCOE/SMutableConstantsWidget.cpp
@@ -57,14 +57,14 @@ namespace
 	 * @return A text representing the value provided as Bytes, Kilobytes, Megabytes and Gigabytes. It will not return
 	 * size types with value 0.
 	 */
-	FText GetSizeInBytesAsBiggerUnits(const uint32 SizeInBytes)
+	FText GetSizeInBytesAsBiggerUnits(const uint64 SizeInBytes)
 	{
 		FString OutputString = "";
 		
-		uint32 Bytes = SizeInBytes;
-		uint32 KiloBytes = 0;
-		uint32 MegaBytes = 0;
-		uint32 GigaBytes = 0;
+		uint64 Bytes = SizeInBytes;
+		uint64 KiloBytes = 0;
+		uint64 MegaBytes = 0;
+		uint64 GigaBytes = 0;
 		
 		// B to KB
 		if (Bytes >= 1024)
@@ -1148,17 +1148,21 @@ void SMutableConstantsWidget::LoadConstantStrings()
 {
 	check (MutableProgramPtr);
 
-	ConstantStringElements.Empty();
-	uint32 ConstantStringsAccumulatedSize = 0;
+	const int32 ConstantsCount = MutableProgramPtr->m_constantStrings.Num();
+	ConstantStringElements.Empty(ConstantsCount);
+	uint64 ConstantStringsAccumulatedSize = 0;
 	
-	for (int32 StringAddressIndex = 0; StringAddressIndex < MutableProgramPtr->m_constantStrings.Num(); StringAddressIndex++)
+	for (int32 StringAddressIndex = 0; StringAddressIndex < ConstantsCount; StringAddressIndex++)
 	{
 		TSharedPtr<FMutableConstantStringElement> ConstantStringElement = MakeShared<FMutableConstantStringElement>();
 		ConstantStringElement->MutableString = &(MutableProgramPtr->m_constantStrings[StringAddressIndex]);
 		ConstantStringElement->IndexOnSourceVector = StringAddressIndex;
 		
+		// Cache resource size
+		// in case we change the type of the contents of the mu::string we check its size as if it was a vector<>
+		ConstantStringsAccumulatedSize += ConstantStringElement->MutableString->size() * sizeof (mu::string::value_type);
+		
 		ConstantStringElements.Add(ConstantStringElement);
-		ConstantStringsAccumulatedSize += sizeof(*ConstantStringElement->MutableString);
 	}
 
 	// Cache the size in memory of the constants as a formatted text so it is able to be be later used by the UI
@@ -1168,18 +1172,21 @@ void SMutableConstantsWidget::LoadConstantStrings()
 void SMutableConstantsWidget::LoadConstantImages()
 {
 	check (MutableProgramPtr);
+
+	const int32 ConstantsCount = MutableProgramPtr->m_constantImages.Num();
+	ConstantImageElements.Empty(ConstantsCount);
 	
-	ConstantImageElements.Empty();
-	uint32 ConstantImagesAccumulatedSize = 0;
+	uint64 ConstantImagesAccumulatedSize = 0;
 
-	for (int32 ImageIndex = 0; ImageIndex < MutableProgramPtr->m_constantImages.Num(); ImageIndex++)
+	for (int32 ImageIndex = 0; ImageIndex < ConstantsCount; ImageIndex++)
 	{
 		TSharedPtr<FMutableConstantImageElement> ConstantImageElement =MakeShared<FMutableConstantImageElement>();
 		MutableProgramPtr->GetConstant(ImageIndex, ConstantImageElement->ImagePtr, 0);
 		ConstantImageElement->IndexOnSourceVector = ImageIndex;
 		
-		ConstantImageElements.Add(ConstantImageElement);
 		ConstantImagesAccumulatedSize += ConstantImageElement->ImagePtr->GetDataSize();
+		
+		ConstantImageElements.Add(ConstantImageElement);
 	}
 
 	// Cache the size in memory of the constants as a formatted text so it is able to be be later used by the UI
@@ -1203,17 +1210,20 @@ void SMutableConstantsWidget::LoadConstantMeshes()
 {
 	check (MutableProgramPtr);
 	
-	ConstantMeshElements.Empty();
-	uint32 ConstantMeshesAccumulatedSize = 0;
+	const int32 ConstantsCount = MutableProgramPtr->m_constantMeshes.Num();
+	ConstantMeshElements.Empty(ConstantsCount);
+	
+	uint64 ConstantMeshesAccumulatedSize = 0;
 	
-	for (int32 MeshIndex = 0; MeshIndex < MutableProgramPtr->m_constantMeshes.Num(); MeshIndex++)
+	for (int32 MeshIndex = 0; MeshIndex < ConstantsCount; MeshIndex++)
 	{
 		TSharedPtr< FMutableConstantMeshElement> ConstantMeshElement = MakeShared<FMutableConstantMeshElement>();
 		ConstantMeshElement->MeshPtr = MutableProgramPtr->m_constantMeshes[MeshIndex].Value;
 		ConstantMeshElement->IndexOnSourceVector = MeshIndex;
 		
-		ConstantMeshElements.Add(ConstantMeshElement);
 		ConstantMeshesAccumulatedSize += ConstantMeshElement->MeshPtr->GetDataSize();
+		
+		ConstantMeshElements.Add(ConstantMeshElement);
 	}
 
 	// Cache the size in memory of the constants as a formatted text so it is able to be be later used by the UI
@@ -1223,130 +1233,152 @@ void SMutableConstantsWidget::LoadConstantMeshes()
 void SMutableConstantsWidget::LoadConstantLayouts()
 {
 	check (MutableProgramPtr);
-	ConstantLayoutElements.Empty();
-	uint32 ConstantLayoutsAccumulatedSize = 0;
+	const int32 ConstantsCount = MutableProgramPtr->m_constantLayouts.Num();
+	ConstantLayoutElements.Empty(ConstantsCount);
+	
+	mu::OutputMemoryStream Stream;
+	mu::OutputArchive Archive{&Stream};
 	
-	for (int32 LayoutIndex = 0; LayoutIndex < MutableProgramPtr->m_constantLayouts.Num(); LayoutIndex++)
+	for (int32 LayoutIndex = 0; LayoutIndex < ConstantsCount; LayoutIndex++)
 	{		
 		TSharedPtr<FMutableConstantLayoutElement> ConstantLayoutElement = MakeShared<FMutableConstantLayoutElement>();
 		ConstantLayoutElement->Layout = MutableProgramPtr->m_constantLayouts[LayoutIndex];
 		ConstantLayoutElement->IndexOnSourceVector = LayoutIndex;
 		
+		// Cache resource for later GetBufferSize() call
+		ConstantLayoutElement->Layout->Serialise(Archive);
+
 		ConstantLayoutElements.Add(ConstantLayoutElement);
-		ConstantLayoutsAccumulatedSize += sizeof(*ConstantLayoutElement->Layout);
 	}
 
 	// Cache the size in memory of the constants as a formatted text so it is able to be be later used by the UI
-	ConstantLayoutsFormattedSize = GetSizeInBytesAsBiggerUnits(ConstantLayoutsAccumulatedSize);
+	ConstantLayoutsFormattedSize = GetSizeInBytesAsBiggerUnits(Stream.GetBufferSize());
 
 }
 
 void SMutableConstantsWidget::LoadConstantSkeletons()
 {
 	check (MutableProgramPtr);
-	ConstantSkeletonElements.Empty();
-	uint32 ConstantSkeletonsAccumulatedSize = 0;
+	const int32 ConstantsCount = MutableProgramPtr->m_constantSkeletons.Num();
+	ConstantSkeletonElements.Empty(ConstantsCount);
+
+	mu::OutputMemoryStream Stream;
+	mu::OutputArchive Archive{&Stream};
 	
-	for (int32 SkeletonIndex = 0; SkeletonIndex < MutableProgramPtr->m_constantSkeletons.Num(); SkeletonIndex++)
+	for (int32 SkeletonIndex = 0; SkeletonIndex < ConstantsCount; SkeletonIndex++)
 	{
 		TSharedPtr<FMutableConstantSkeletonElement> ConstantSkeletonElement = MakeShared<FMutableConstantSkeletonElement>();
 		ConstantSkeletonElement->Skeleton = MutableProgramPtr->m_constantSkeletons[SkeletonIndex];
 		ConstantSkeletonElement->IndexOnSourceVector = SkeletonIndex;
 		
+		ConstantSkeletonElement->Skeleton->Serialise(Archive);
+
 		ConstantSkeletonElements.Add(ConstantSkeletonElement);
-		ConstantSkeletonsAccumulatedSize += sizeof(*ConstantSkeletonElement->Skeleton);
 	}
 
 	// Cache the size in memory of the constants as a formatted text so it is able to be be later used by the UI
-	ConstantSkeletonsFormattedSize = GetSizeInBytesAsBiggerUnits(ConstantSkeletonsAccumulatedSize);
+	ConstantSkeletonsFormattedSize = GetSizeInBytesAsBiggerUnits(Stream.GetBufferSize());
 }
 
 void SMutableConstantsWidget::LoadConstantProjectors()
 {
 	check (MutableProgramPtr);
-	ConstantProjectorElements.Empty();
-	uint32 ConstantProjectorsAccumulatedSize = 0;
+	const int32 ConstantsCount = MutableProgramPtr->m_constantProjectors.Num();
+	ConstantProjectorElements.Empty(ConstantsCount);
+
+	mu::OutputMemoryStream Stream;
+	mu::OutputArchive Archive{&Stream};
 	
-	for (int32 ProjectorIndex = 0; ProjectorIndex < MutableProgramPtr->m_constantProjectors.Num(); ProjectorIndex++)
+	for (int32 ProjectorIndex = 0; ProjectorIndex < ConstantsCount; ProjectorIndex++)
 	{
 		TSharedPtr<FMutableConstantProjectorElement> ConstantProjectorElement = MakeShared<FMutableConstantProjectorElement>();
 		ConstantProjectorElement->Projector = &(MutableProgramPtr->m_constantProjectors[ProjectorIndex]);
 		ConstantProjectorElement->IndexOnSourceVector = ProjectorIndex;
 		
+		ConstantProjectorElement->Projector->Serialise(Archive);
+		
 		ConstantProjectorElements.Add(ConstantProjectorElement);
-		ConstantProjectorsAccumulatedSize += sizeof(*ConstantProjectorElement->Projector);
 	}
 
 	// Cache the size in memory of the constants as a formatted text so it is able to be be later used by the UI
-	ConstantProjectorsFormattedSize = GetSizeInBytesAsBiggerUnits(ConstantProjectorsAccumulatedSize);
-
+	ConstantProjectorsFormattedSize = GetSizeInBytesAsBiggerUnits(Stream.GetBufferSize());
 }
 
 void SMutableConstantsWidget::LoadConstantMatrices()
 {
 	check (MutableProgramPtr);
-	ConstantMatrixElements.Empty();
-	uint32 ConstantMatricesAccumulatedSize = 0;
+	const int32 ConstantsCount = MutableProgramPtr->m_constantMatrices.Num();
+	ConstantMatrixElements.Empty(ConstantsCount);
+
+	mu::OutputMemoryStream Stream;
+	mu::OutputArchive Archive{&Stream};
 	
-	for (int32 MatrixIndex = 0; MatrixIndex < MutableProgramPtr->m_constantMatrices.Num(); MatrixIndex++)
+	for (int32 MatrixIndex = 0; MatrixIndex < ConstantsCount; MatrixIndex++)
 	{
 		TSharedPtr<FMutableConstantMatrixElement> ConstantMatrixElement = MakeShared<FMutableConstantMatrixElement>();
 		ConstantMatrixElement->Matrix = &(MutableProgramPtr->m_constantMatrices[MatrixIndex]);
 		ConstantMatrixElement->IndexOnSourceVector = MatrixIndex;
 		
+		ConstantMatrixElement->Matrix->Serialise(Archive);
+		
 		ConstantMatrixElements.Add(ConstantMatrixElement);
-		ConstantMatricesAccumulatedSize += sizeof(*ConstantMatrixElement->Matrix);
 	}
 
 	// Cache the size in memory of the constants as a formatted text so it is able to be be later used by the UI
-	ConstantMatricesFormattedSize = GetSizeInBytesAsBiggerUnits(ConstantMatricesAccumulatedSize);
-
+	ConstantMatricesFormattedSize = GetSizeInBytesAsBiggerUnits(Stream.GetBufferSize());
 }
 
 void SMutableConstantsWidget::LoadConstantShapes()
 {
 	check (MutableProgramPtr);
-	ConstantShapeElements.Empty();
-	uint32 ConstantShapesAccumulatedSize = 0;
+	const int32 ConstantsCount = MutableProgramPtr->m_constantShapes.Num();
+	ConstantShapeElements.Empty(ConstantsCount);
+
+	mu::OutputMemoryStream Stream;
+	mu::OutputArchive Archive{&Stream};	
 	
-	for (int32 ShapeIndex = 0; ShapeIndex < MutableProgramPtr->m_constantShapes.Num(); ShapeIndex++)
+	for (int32 ShapeIndex = 0; ShapeIndex < ConstantsCount; ShapeIndex++)
 	{
 		TSharedPtr<FMutableConstantShapeElement> ConstantShapeElement = MakeShared<FMutableConstantShapeElement>();
 		ConstantShapeElement->Shape = &(MutableProgramPtr->m_constantShapes[ShapeIndex]);
 		ConstantShapeElement->IndexOnSourceVector = ShapeIndex;
 		
+		ConstantShapeElement->Shape->Serialise(Archive);
+
 		ConstantShapeElements.Add(ConstantShapeElement);
-		ConstantShapesAccumulatedSize += sizeof(*ConstantShapeElement->Shape);
 	}
 	
 	// Cache the size in memory of the constants as a formatted text so it is able to be be later used by the UI
-	ConstantShapesFormattedSize = GetSizeInBytesAsBiggerUnits(ConstantShapesAccumulatedSize);
-
+	ConstantShapesFormattedSize = GetSizeInBytesAsBiggerUnits(Stream.GetBufferSize());
 }
 
 void SMutableConstantsWidget::LoadConstantCurves()
 {
 	check (MutableProgramPtr);
-	ConstantCurveElements.Empty();
-	uint32 ConstantCurvesAccumulatedSize = 0;
+	const int32 ConstantsCount = MutableProgramPtr->m_constantCurves.Num();
+	ConstantCurveElements.Empty(ConstantsCount);
+	
+	mu::OutputMemoryStream Stream;
+	mu::OutputArchive Archive{&Stream};
 	
-	for (int32 CurveIndex = 0; CurveIndex < MutableProgramPtr->m_constantCurves.Num(); CurveIndex++)
+	for (int32 CurveIndex = 0; CurveIndex < ConstantsCount; CurveIndex++)
 	{
 		TSharedPtr<FMutableConstantCurveElement> ConstantCurveElement = MakeShared<FMutableConstantCurveElement>();
 		ConstantCurveElement->Curve = &(MutableProgramPtr->m_constantCurves[CurveIndex]);
 		ConstantCurveElement->IndexOnSourceVector = CurveIndex;
+
+		ConstantCurveElement->Curve->Serialise(Archive);
 		
 		ConstantCurveElements.Add(ConstantCurveElement);
-		ConstantCurvesAccumulatedSize += sizeof(*ConstantCurveElement->Curve);
 	}
 
 	// Cache the size in memory of the constants as a formatted text so it is able to be be later used by the UI
-	ConstantCurvesFormattedSize = GetSizeInBytesAsBiggerUnits(ConstantCurvesAccumulatedSize);
-
+	ConstantCurvesFormattedSize = GetSizeInBytesAsBiggerUnits(Stream.GetBufferSize());
 }
 
 #pragma  endregion
 
+
 #pragma region Previewer invocation methods
 
 
