[mutable] Added checkbox to disable table node row compilation restrictions.

#rnx

diff --git a/Engine/Plugins/Experimental/Mutable/Source/CustomizableObjectEditor/Classes/MuCOE/Nodes/CustomizableObjectNodeTable.h b/Engine/Plugins/Experimental/Mutable/Source/CustomizableObjectEditor/Classes/MuCOE/Nodes/CustomizableObjectNodeTable.h
--- a/Engine/Plugins/Experimental/Mutable/Source/CustomizableObjectEditor/Classes/MuCOE/Nodes/CustomizableObjectNodeTable.h
+++ b/Engine/Plugins/Experimental/Mutable/Source/CustomizableObjectEditor/Classes/MuCOE/Nodes/CustomizableObjectNodeTable.h
@@ -97,6 +97,7 @@ public:
 	UPROPERTY(EditAnywhere, Category = TableProperties)
 	FString ParameterName = "Default Name";
 
+	/** If true, adds a "None" parameter option */
 	UPROPERTY(EditAnywhere, Category = TableProperties)
 	bool bAddNoneOption;
 
@@ -104,6 +105,10 @@ public:
 	UPROPERTY(EditAnywhere, Category = TableProperties, meta = (DontUpdateWhileEditing))
 	TObjectPtr<UDataTable> Table = nullptr;
 	
+	/** If there is a bool column in the table, checked rows will not be compiled */
+	UPROPERTY(EditAnywhere, Category = TableProperties)
+	bool bDisableCheckedRows = true;
+	
 	UPROPERTY(EditAnywhere, Category = UI, meta = (DisplayName = "Parameter UI Metadata"))
 	FMutableParamUIMetadata ParamUIMetadata;
 
diff --git a/Engine/Plugins/Experimental/Mutable/Source/CustomizableObjectEditor/Private/MuCOE/Nodes/CustomizableObjectNodeTable.cpp b/Engine/Plugins/Experimental/Mutable/Source/CustomizableObjectEditor/Private/MuCOE/Nodes/CustomizableObjectNodeTable.cpp
--- a/Engine/Plugins/Experimental/Mutable/Source/CustomizableObjectEditor/Private/MuCOE/Nodes/CustomizableObjectNodeTable.cpp
+++ b/Engine/Plugins/Experimental/Mutable/Source/CustomizableObjectEditor/Private/MuCOE/Nodes/CustomizableObjectNodeTable.cpp
@@ -909,7 +909,7 @@ TArray<FName> UCustomizableObjectNodeTable::GetRowNames() const
 		TArray<FName> TableRowNames = Table->GetRowNames();
 		FBoolProperty* BoolProperty = nullptr;
 
-		for (TFieldIterator<FProperty> PropertyIt(TableStruct); PropertyIt; ++PropertyIt)
+		for (TFieldIterator<FProperty> PropertyIt(TableStruct); PropertyIt && bDisableCheckedRows; ++PropertyIt)
 		{
 			BoolProperty = CastField<FBoolProperty>(*PropertyIt);
 
@@ -934,7 +934,7 @@ TArray<FName> UCustomizableObjectNodeTable::GetRowNames() const
 			}
 		}
 
-		// There is no Bool column
+		// There is no Bool column or we don't want to disable rows
 		if (!BoolProperty)
 		{
 			return TableRowNames;
