[mutable] Fix image resize producing size 0 images.

#rnx

diff --git a/Engine/Plugins/Experimental/Mutable/Source/MutableRuntime/Private/MuR/CodeRunner.cpp b/Engine/Plugins/Experimental/Mutable/Source/MutableRuntime/Private/MuR/CodeRunner.cpp
--- a/Engine/Plugins/Experimental/Mutable/Source/MutableRuntime/Private/MuR/CodeRunner.cpp
+++ b/Engine/Plugins/Experimental/Mutable/Source/MutableRuntime/Private/MuR/CodeRunner.cpp
@@ -2924,8 +2924,8 @@ namespace mu
                 Ptr<const Image> pBase = GetMemory().GetImage( CACHE_ADDRESS(args.source,item) );
 
                 FImageSize destSize(
-                            uint16( pBase->GetSizeX()*args.factor[0] + 0.5f ),
-                            uint16( pBase->GetSizeY()*args.factor[1] + 0.5f ) );
+                            uint16( FMath::Max(1.0, pBase->GetSizeX()*args.factor[0] + 0.5f ) ),
+                            uint16( FMath::Max(1.0, pBase->GetSizeY()*args.factor[1] + 0.5f ) ) );
 
                 //pResult = ImageResize( pBase.get(), destSize );
                 ImagePtr pResult = ImageResizeLinear(
