[Mutable] Improve physics assets merging. Now the resulting asset will be created form core data.

- Added physics tags to tables.
- Now the default behaviour is to recreate the physics asset from the data from core. If the resulting mesh does not have any data, then the reference asset will be used.

diff --git a/Engine/Plugins/Experimental/Mutable/Source/CustomizableObject/Private/MuCO/CustomizableInstancePrivateData.h b/Engine/Plugins/Experimental/Mutable/Source/CustomizableObject/Private/MuCO/CustomizableInstancePrivateData.h
--- a/Engine/Plugins/Experimental/Mutable/Source/CustomizableObject/Private/MuCO/CustomizableInstancePrivateData.h
+++ b/Engine/Plugins/Experimental/Mutable/Source/CustomizableObject/Private/MuCO/CustomizableInstancePrivateData.h
@@ -251,8 +251,7 @@ private:
 	USkeleton* MergeSkeletons(UCustomizableObjectInstance* Public, const FMutableRefSkeletalMeshData* RefSkeletalMeshData, int32 ComponentIndex);
 
 	//
-	UPhysicsAsset* MergePhysicalAssets(class UPhysicsAsset* BaseAsset, int32 ComponentIndex);
-	void RefitPhysicsBodies(UPhysicsAsset* InOutPhysicsAsset, const mu::PhysicsBody* PhysicsBody);
+	UPhysicsAsset* BuildPhysicsAsset(TObjectPtr<class UPhysicsAsset> TamplateAsset, const mu::PhysicsBody* PhysicsBody, int32 ComponentIndex, bool bDisableCollisionBetweenAssets);
 	
 	
 	int32 GetLastMeshId(int32 ComponentIndex, int32 LODIndex) const;
diff --git a/Engine/Plugins/Experimental/Mutable/Source/CustomizableObject/Private/MuCO/CustomizableObjectInstance.cpp b/Engine/Plugins/Experimental/Mutable/Source/CustomizableObject/Private/MuCO/CustomizableObjectInstance.cpp
--- a/Engine/Plugins/Experimental/Mutable/Source/CustomizableObject/Private/MuCO/CustomizableObjectInstance.cpp
+++ b/Engine/Plugins/Experimental/Mutable/Source/CustomizableObject/Private/MuCO/CustomizableObjectInstance.cpp
@@ -126,7 +126,7 @@
 #include "MuR/OpImageSwizzle.h"
 #include "MuR/Parameters.h"
 #include "MuR/System.h"
-
+#include "PhysicsEngine/PhysicsConstraintTemplate.h"
 
 UTexture2D* UCustomizableInstancePrivateData::CreateTexture()
 {
@@ -991,19 +991,62 @@ USkeleton* UCustomizableInstancePrivateData::MergeSkeletons(UCustomizableObjectI
 	return FinalSkeleton;
 }
 
-
-void UCustomizableInstancePrivateData::RefitPhysicsBodies(UPhysicsAsset* InOutPhysicsAsset, const mu::PhysicsBody* MutablePhysicsBody)
+UPhysicsAsset* UCustomizableInstancePrivateData::BuildPhysicsAsset(
+		TObjectPtr<UPhysicsAsset> TemplateAsset,
+		const mu::PhysicsBody* MutablePhysics, int32 ComponentId,
+		bool bDisableCollisionsBetweenDifferentAssets)
 {
-	if (!(InOutPhysicsAsset && MutablePhysicsBody))
+
+	MUTABLE_CPUPROFILER_SCOPE(MergePhysicsAssets);
+
+	check(MutablePhysics);
+
+	UPhysicsAsset* Result = nullptr;
+
+	FCustomizableInstanceComponentData* ComponentData =	GetComponentData(ComponentId);
+	check(ComponentData);
+
+	TArray<TObjectPtr<UPhysicsAsset>>& PhysicsAssets = ComponentData->PhysicsAssets.PhysicsAssetsToMerge;
+
+	TArray<TObjectPtr<UPhysicsAsset>> ValidAssets;		
+
+	const int32 NumPhysicsAssets = ComponentData->PhysicsAssets.PhysicsAssetsToMerge.Num();
+	for (int32 I = 0; I < NumPhysicsAssets; ++I)
 	{
-		return;
+		const TObjectPtr<UPhysicsAsset>& PhysicsAsset = ComponentData->PhysicsAssets.PhysicsAssetsToMerge[I];
+
+		if (PhysicsAsset)
+		{
+			ValidAssets.AddUnique(PhysicsAsset);
+		}
+	}
+	
+	if (!ValidAssets.Num())
+	{
+		return Result;
 	}
 
-	const int32 MeshPhysicsBodyCount = MutablePhysicsBody->GetBodyCount();
+	if (!TemplateAsset)
+	{
+		Result = DuplicateObject(ValidAssets[0], nullptr);
+	}
+	else
+	{
+		Result = DuplicateObject(TemplateAsset, nullptr);	
+	}
 
-	for (int32 B = 0; B < MeshPhysicsBodyCount; ++B)
+	if (!Result)
 	{
-		const FString BodyBoneName( MutablePhysicsBody->GetBodyBoneName( B ) );
+		return nullptr;
+	}
+	// TODO: We are duplicating just to remove the data that will be replaced.
+	Result->CollisionDisableTable.Empty();
+	Result->SkeletalBodySetups.Empty();
+	Result->ConstraintSetup.Empty();
+
+	auto MakeAggGeomFromMutablePhysics = [](int32 BodyIndex, const mu::PhysicsBody* MutablePhysicsBody) -> FKAggregateGeom
+	{
+		FKAggregateGeom BodyAggGeom;
 
 		auto GetCollisionEnabledFormFlags = [](uint32 Flags) -> ECollisionEnabled::Type
 		{
@@ -1014,216 +1057,356 @@ void UCustomizableInstancePrivateData::RefitPhysicsBodies(UPhysicsAsset* InOutPh
 		{
 			return static_cast<bool>( (Flags >> 8 ) & 1 );
 		};
+				
+		const int32 NumSpheres = MutablePhysicsBody->GetSphereCount( BodyIndex );
+		TArray<FKSphereElem>& AggSpheres = BodyAggGeom.SphereElems;
+		AggSpheres.Empty(NumSpheres);
+		for (int32 I = 0; I < NumSpheres; ++I)
+		{
+			uint32 Flags = MutablePhysicsBody->GetSphereFlags( BodyIndex, I );
+			FString Name = MutablePhysicsBody->GetSphereName( BodyIndex, I );
+
+			FVector3f Position;
+			float Radius;
+
+			MutablePhysicsBody->GetSphere( BodyIndex, I, Position, Radius );
+			FKSphereElem& NewElem = AggSpheres.AddDefaulted_GetRef();
+			
+			NewElem.Center = FVector(Position);
+			NewElem.Radius = Radius;
+			NewElem.SetContributeToMass( GetContributeToMassFromFlags( Flags ) );
+			NewElem.SetCollisionEnabled( GetCollisionEnabledFormFlags( Flags ) );
+			NewElem.SetName(FName(*Name));
+		}
 		
-		TObjectPtr<USkeletalBodySetup>* Found = InOutPhysicsAsset->SkeletalBodySetups.FindByPredicate(
-		[&BodyBoneName](TObjectPtr<USkeletalBodySetup>& BodySetup)
+		const int32 NumBoxes = MutablePhysicsBody->GetBoxCount( BodyIndex );
+		TArray<FKBoxElem>& AggBoxes = BodyAggGeom.BoxElems;
+		AggBoxes.Empty(NumBoxes);
+		for (int32 I = 0; I < NumBoxes; ++I)
 		{
-			return BodyBoneName.Equals(BodySetup->BoneName.ToString());
-		});
+			uint32 Flags = MutablePhysicsBody->GetBoxFlags( BodyIndex, I );
+			FString Name = MutablePhysicsBody->GetBoxName( BodyIndex, I );
 
-		if (Found)
+			FVector3f Position;
+			FQuat4f Orientation;
+			FVector3f Size;
+			MutablePhysicsBody->GetBox( BodyIndex, I, Position, Orientation, Size );
+
+			FKBoxElem& NewElem = AggBoxes.AddDefaulted_GetRef();
+			
+			NewElem.Center = FVector( Position );
+			NewElem.Rotation = FRotator( Orientation.Rotator() );
+			NewElem.X = Size.X;
+			NewElem.Y = Size.Y;
+			NewElem.Z = Size.Z;
+			NewElem.SetContributeToMass( GetContributeToMassFromFlags( Flags ) );
+			NewElem.SetCollisionEnabled( GetCollisionEnabledFormFlags( Flags ) );
+			NewElem.SetName( FName(*Name) );
+		}
+
+		const int32 NumConvexes = MutablePhysicsBody->GetConvexCount( BodyIndex );
+		TArray<FKConvexElem>& AggConvexes = BodyAggGeom.ConvexElems;
+		AggConvexes.Empty();
+		for (int32 I = 0; I < NumConvexes; ++I)
 		{
-			FKAggregateGeom& BodyAggGeom = (*Found)->AggGeom;
+			uint32 Flags = MutablePhysicsBody->GetConvexFlags( BodyIndex, I );
+			FString Name = MutablePhysicsBody->GetConvexName( BodyIndex, I );
+
+			const FVector3f* Vertices;
+			const int32* Indices;
+			int32 NumVertices;
+			int32 NumIndices;
+			FTransform3f Transform;
+
+			MutablePhysicsBody->GetConvex( BodyIndex, I, Vertices, NumVertices, Indices, NumIndices, Transform );
 			
-			const int32 NumSpheres = MutablePhysicsBody->GetSphereCount( B );
-			TArray<FKSphereElem>& AggSpheres = BodyAggGeom.SphereElems;
-			AggSpheres.Empty(NumSpheres);
-			for (int32 I = 0; I < NumSpheres; ++I)
-			{
-				uint32 Flags = MutablePhysicsBody->GetSphereFlags( B, I );
-				FString Name = MutablePhysicsBody->GetSphereName( B, I );
+			TArrayView<const FVector3f> VerticesView( Vertices, NumVertices );
+			TArrayView<const int32> IndicesView( Indices, NumIndices );
+		}
 
-				FVector3f Position;
-				float Radius;
+		TArray<FKSphylElem>& AggSphyls = BodyAggGeom.SphylElems;
+		const int32 NumSphyls = MutablePhysicsBody->GetSphylCount( BodyIndex );
+		AggSphyls.Empty(NumSphyls);
 
-				MutablePhysicsBody->GetSphere( B, I, Position, Radius );
-				FKSphereElem& NewElem = AggSpheres.AddDefaulted_GetRef();
-				
-				NewElem.Center = FVector(Position);
-				NewElem.Radius = Radius;
-				NewElem.SetContributeToMass( GetContributeToMassFromFlags( Flags ) );
-				NewElem.SetCollisionEnabled( GetCollisionEnabledFormFlags( Flags ) );
-				NewElem.SetName(FName(*Name));
-			}
+		for (int32 I = 0; I < NumSphyls; ++I)
+		{
+			uint32 Flags = MutablePhysicsBody->GetSphylFlags( BodyIndex, I );
+			FString Name = MutablePhysicsBody->GetSphylName( BodyIndex, I );
+
+			FVector3f Position;
+			FQuat4f Orientation;
+			float Radius;
+			float Length;
+
+			MutablePhysicsBody->GetSphyl( BodyIndex, I, Position, Orientation, Radius, Length );
+
+			FKSphylElem& NewElem = AggSphyls.AddDefaulted_GetRef();
 			
-			const int32 NumBoxes = MutablePhysicsBody->GetBoxCount( B );
-			TArray<FKBoxElem>& AggBoxes = BodyAggGeom.BoxElems;
-			AggBoxes.Empty(NumBoxes);
-			for (int32 I = 0; I < NumBoxes; ++I)
-			{
-				uint32 Flags = MutablePhysicsBody->GetBoxFlags( B, I );
-				FString Name = MutablePhysicsBody->GetBoxName( B, I );
+			NewElem.Center = FVector( Position );
+			NewElem.Rotation = FRotator( Orientation.Rotator() );
+			NewElem.Radius = Radius;
+			NewElem.Length = Length;
+			
+			NewElem.SetContributeToMass( GetContributeToMassFromFlags( Flags ) );
+			NewElem.SetCollisionEnabled( GetCollisionEnabledFormFlags( Flags ) );
+			NewElem.SetName( FName(*Name) );
+		}	
 
-				FVector3f Position;
-				FQuat4f Orientation;
-				FVector3f Size;
-				MutablePhysicsBody->GetBox( B, I, Position, Orientation, Size );
+		TArray<FKTaperedCapsuleElem>& AggTaperedCapsules = BodyAggGeom.TaperedCapsuleElems;
+		const int32 NumTaperedCapsules = MutablePhysicsBody->GetTaperedCapsuleCount( BodyIndex );
+		AggTaperedCapsules.Empty(NumTaperedCapsules);
 
-				FKBoxElem& NewElem = AggBoxes.AddDefaulted_GetRef();
-				
-				NewElem.Center = FVector( Position );
-				NewElem.Rotation = FRotator( Orientation.Rotator() );
-				NewElem.X = Size.X;
-				NewElem.Y = Size.Y;
-				NewElem.Z = Size.Z;
-				NewElem.SetContributeToMass( GetContributeToMassFromFlags( Flags ) );
-				NewElem.SetCollisionEnabled( GetCollisionEnabledFormFlags( Flags ) );
-				NewElem.SetName( FName(*Name) );
-			}
+		for (int32 I = 0; I < NumTaperedCapsules; ++I)
+		{
+			uint32 Flags = MutablePhysicsBody->GetTaperedCapsuleFlags( BodyIndex, I );
+			FString Name = MutablePhysicsBody->GetTaperedCapsuleName( BodyIndex, I );
 
-			const int32 NumConvexes = MutablePhysicsBody->GetConvexCount( B );
-			TArray<FKConvexElem>& AggConvexes = BodyAggGeom.ConvexElems;
-			AggConvexes.Empty();
-			for (int32 I = 0; I < NumConvexes; ++I)
-			{
-				uint32 Flags = MutablePhysicsBody->GetConvexFlags( B, I );
-				FString Name = MutablePhysicsBody->GetConvexName( B, I );
-
-				const FVector3f* Vertices;
-				const int32* Indices;
-				int32 NumVertices;
-				int32 NumIndices;
-				FTransform3f Transform;
-	
-				MutablePhysicsBody->GetConvex( B, I, Vertices, NumVertices, Indices, NumIndices, Transform );
-				
-				TArrayView<const FVector3f> VerticesView( Vertices, NumVertices );
-				TArrayView<const int32> IndicesView( Indices, NumIndices );
-			}
+			FVector3f Position;
+			FQuat4f Orientation;
+			float Radius0;
+			float Radius1;
+			float Length;
 
-			TArray<FKSphylElem>& AggSphyls = BodyAggGeom.SphylElems;
-			const int32 NumSphyls = MutablePhysicsBody->GetSphylCount( B );
-			AggSphyls.Empty(NumSphyls);
+			MutablePhysicsBody->GetTaperedCapsule( BodyIndex, I, Position, Orientation, Radius0, Radius1, Length );
 
-			for (int32 I = 0; I < NumSphyls; ++I)
-			{
-				uint32 Flags = MutablePhysicsBody->GetSphylFlags( B, I );
-				FString Name = MutablePhysicsBody->GetSphylName( B, I );
+			FKTaperedCapsuleElem& NewElem = AggTaperedCapsules.AddDefaulted_GetRef();
+			
+			NewElem.Center = FVector( Position );
+			NewElem.Rotation = FRotator( Orientation.Rotator() );
+			NewElem.Radius0 = Radius0;
+			NewElem.Radius1 = Radius1;
+			NewElem.Length = Length;
+			
+			NewElem.SetContributeToMass( GetContributeToMassFromFlags( Flags ) );
+			NewElem.SetCollisionEnabled( GetCollisionEnabledFormFlags( Flags ) );
+			NewElem.SetName( FName(*Name) );	
+		}	
 
-				FVector3f Position;
-				FQuat4f Orientation;
-				float Radius;
-				float Length;
+		return BodyAggGeom;
+	};
 
-				MutablePhysicsBody->GetSphyl( B, I, Position, Orientation, Radius, Length );
+	TMap<FName, int32> BonesInUse;
 
-				FKSphylElem& NewElem = AggSphyls.AddDefaulted_GetRef();
-				
-				NewElem.Center = FVector( Position );
-				NewElem.Rotation = FRotator( Orientation.Rotator() );
-				NewElem.Radius = Radius;
-				NewElem.Length = Length;
-				
-				NewElem.SetContributeToMass( GetContributeToMassFromFlags( Flags ) );
-				NewElem.SetCollisionEnabled( GetCollisionEnabledFormFlags( Flags ) );
-				NewElem.SetName( FName(*Name) );
-			}	
+	const int32 MutablePhysicsBodyCount = MutablePhysics->GetBodyCount();
+	BonesInUse.Reserve(MutablePhysicsBodyCount);
+	for ( int32 I = 0; I < MutablePhysicsBodyCount; ++I )
+	{
+		BonesInUse.Add(FName(MutablePhysics->GetBodyBoneName(I)), I);
+	}
+
+	// Each array is a set of elements that can collide  
+	TArray<TArray<int32, TInlineAllocator<8>>> CollisionSets;
+
+	// {SetIndex, ElementInSetIndex, BodyIndex}
+	using CollisionSetEntryType = TTuple<int32, int32, int32>;	
+	// Map from BodyName/BoneName to set and index in set.
+	TMap<FName, CollisionSetEntryType> BodySetupSetMap;
+	
+	// {SetIndex0, SetIndex1, BodyIndex}
+	using MultiSetCollisionEnableType = TTuple<int32, int32, int32>;
 
-			TArray<FKTaperedCapsuleElem>& AggTaperedCapsules = BodyAggGeom.TaperedCapsuleElems;
-			const int32 NumTaperedCapsules = MutablePhysicsBody->GetTaperedCapsuleCount( B );
-			AggTaperedCapsules.Empty(NumTaperedCapsules);
+	// Only for elements that belong to two or more differnet sets. 
+	// Contains in which set the elements belong.
+	using MultiSetArrayType = TArray<int32, TInlineAllocator<4>>;
+	TMap<int32, MultiSetArrayType> MultiCollisionSets;
+	TArray<TArray<int32>> SetsIndexMap;
+
+	CollisionSets.SetNum(ValidAssets.Num());
+	SetsIndexMap.SetNum(CollisionSets.Num());
+
+	TMap<FRigidBodyIndexPair, bool> CollisionDisableTable;
+
+	// New body index
+	int32 CurrentBodyIndex = 0;
+	for (int32 CollisionSetIndex = 0;  CollisionSetIndex < ValidAssets.Num(); ++CollisionSetIndex)
+	{
+		const int32 PhysicsAssetBodySetupNum = ValidAssets[CollisionSetIndex]->SkeletalBodySetups.Num();
+		SetsIndexMap[CollisionSetIndex].Init(-1, PhysicsAssetBodySetupNum);
 
-			for (int32 I = 0; I < NumTaperedCapsules; ++I)
+		for (int32 BodySetupIndex = 0; BodySetupIndex < PhysicsAssetBodySetupNum; ++BodySetupIndex) 
+		{
+			const TObjectPtr<USkeletalBodySetup>& BodySetup = ValidAssets[CollisionSetIndex]->SkeletalBodySetups[BodySetupIndex];
+			
+			int32* MutableBodyIndex = BonesInUse.Find(BodySetup->BoneName);
+			if (!MutableBodyIndex)
 			{
-				uint32 Flags = MutablePhysicsBody->GetTaperedCapsuleFlags( B, I );
-				FString Name = MutablePhysicsBody->GetTaperedCapsuleName( B, I );
+				continue;
+			}
 
-				FVector3f Position;
-				FQuat4f Orientation;
-				float Radius0;
-				float Radius1;
-				float Length;
+			CollisionSetEntryType* Found = BodySetupSetMap.Find(BodySetup->BoneName);
 
-				MutablePhysicsBody->GetTaperedCapsule( B, I, Position, Orientation, Radius0, Radius1, Length );
+			if (!Found)
+			{
+				TObjectPtr<USkeletalBodySetup> NewBodySetup = NewObject<USkeletalBodySetup>(Result);
+				NewBodySetup->BodySetupGuid = FGuid::NewGuid();
+			
+				// Copy Body properties 	
+				NewBodySetup->BoneName = BodySetup->BoneName;
+				NewBodySetup->PhysicsType = BodySetup->PhysicsType;
+				NewBodySetup->bConsiderForBounds = BodySetup->bConsiderForBounds;
+				NewBodySetup->bMeshCollideAll = BodySetup->bMeshCollideAll;
+				NewBodySetup->bDoubleSidedGeometry = BodySetup->bDoubleSidedGeometry;
+				NewBodySetup->bGenerateNonMirroredCollision = BodySetup->bGenerateNonMirroredCollision;
+				NewBodySetup->bSharedCookedData = BodySetup->bSharedCookedData;
+				NewBodySetup->bGenerateMirroredCollision = BodySetup->bGenerateMirroredCollision;
+				NewBodySetup->PhysMaterial = BodySetup->PhysMaterial;
+				NewBodySetup->CollisionReponse = BodySetup->CollisionReponse;
+				NewBodySetup->CollisionTraceFlag = BodySetup->CollisionTraceFlag;
+				NewBodySetup->DefaultInstance = BodySetup->DefaultInstance;
+				NewBodySetup->WalkableSlopeOverride = BodySetup->WalkableSlopeOverride;
+				NewBodySetup->BuildScale3D = BodySetup->BuildScale3D;	
+				NewBodySetup->bSkipScaleFromAnimation = BodySetup->bSkipScaleFromAnimation;
+				NewBodySetup->AggGeom = MakeAggGeomFromMutablePhysics(*MutableBodyIndex, MutablePhysics);
+
+				Result->SkeletalBodySetups.Add(NewBodySetup);
+				
+				int32 IndexInSet = CollisionSets[CollisionSetIndex].Add(CurrentBodyIndex);
+				BodySetupSetMap.Add(BodySetup->BoneName, {CollisionSetIndex, IndexInSet, CurrentBodyIndex});
+				SetsIndexMap[CollisionSetIndex][IndexInSet] = CurrentBodyIndex;
 
-				FKTaperedCapsuleElem& NewElem = AggTaperedCapsules.AddDefaulted_GetRef();
+				++CurrentBodyIndex;
+			}
+			else
+			{
+				int32 FoundCollisionSetIndex = Found->Get<0>();
+				int32 FoundCollisionSetElemIndex = Found->Get<1>();
+				int32 FoundBodyIndex = Found->Get<2>();
 				
-				NewElem.Center = FVector( Position );
-				NewElem.Rotation = FRotator( Orientation.Rotator() );
-				NewElem.Radius0 = Radius0;
-				NewElem.Radius1 = Radius1;
-				NewElem.Length = Length;
+				// No need to add the body again. Volumes that come form mutable are already merged.
+				// here we only need to merge properies.
+				// TODO: check if there is other properties worth merging. In case of conflict select the more restrivtive one? 
+				Result->SkeletalBodySetups[FoundBodyIndex]->bConsiderForBounds |= BodySetup->bConsiderForBounds;
+
+				// Mark as removed so no indices are invalidated.
+				CollisionSets[FoundCollisionSetIndex][FoundCollisionSetElemIndex] = INDEX_NONE;
+				// Add Elem to the set but mark it as removed so we have an index for remapping.
+				int32 IndexInSet = CollisionSets[CollisionSetIndex].Add(INDEX_NONE);	
+				SetsIndexMap[CollisionSetIndex][IndexInSet] = FoundBodyIndex;
 				
-				NewElem.SetContributeToMass( GetContributeToMassFromFlags( Flags ) );
-				NewElem.SetCollisionEnabled( GetCollisionEnabledFormFlags( Flags ) );
-				NewElem.SetName( FName(*Name) );	
-			}	
+				MultiSetArrayType& Sets = MultiCollisionSets.FindOrAdd(FoundBodyIndex);
+
+				// The first time there is a collision (MultSet is empty), add the colliding element set
+				// as well as the current set.
+				if (!Sets.Num())
+				{
+					Sets.Add(FoundCollisionSetIndex);
+				}
+				
+				Sets.Add(CollisionSetIndex);
+			}
 		}
-	}
-}
+	
+		// Remap collision indices removing invalid ones.
 
-UPhysicsAsset* UCustomizableInstancePrivateData::MergePhysicalAssets(UPhysicsAsset* BaseAsset, int32 ComponentId)
-{
-	UPhysicsAsset* result = nullptr;
-	bool bMergeSucced = false;
+		CollisionDisableTable.Reserve(CollisionDisableTable.Num() + ValidAssets[CollisionSetIndex]->CollisionDisableTable.Num());
+		for (const TPair<FRigidBodyIndexPair, bool>& DisabledCollision : ValidAssets[CollisionSetIndex]->CollisionDisableTable)
+		{
+			int32 MappedIdx0 = SetsIndexMap[CollisionSetIndex][DisabledCollision.Key.Indices[0]];
+			int32 MappedIdx1 = SetsIndexMap[CollisionSetIndex][DisabledCollision.Key.Indices[1]];
 
-	if (BaseAsset)
-	{
-		result = DuplicateObject(BaseAsset, NULL);
+			// This will generate correct disables for the case when two shapes from different sets
+			// are meged to the same setup. Will introduce repeated pairs, but this is not a problem.
+
+			// Currenly if two bodies / bones have disabled collision in one of the merged assets, the collision
+			// will remain disabled even if other merges allow it.   
+			if ( MappedIdx0 != INDEX_NONE && MappedIdx1 != INDEX_NONE )
+			{
+				CollisionDisableTable.Add({MappedIdx0, MappedIdx1}, DisabledCollision.Value);
+			}
+		}
+
+		// Only add constraints that are part of the bones used for the mutable physics volumes description.
+		for (const TObjectPtr<UPhysicsConstraintTemplate>& Constrain : ValidAssets[CollisionSetIndex]->ConstraintSetup)
+		{
+			FName BoneA = Constrain->DefaultInstance.ConstraintBone1;
+			FName BoneB = Constrain->DefaultInstance.ConstraintBone2;
+
+			if (BonesInUse.Find(BoneA) && BonesInUse.Find(BoneB))
+			{
+				Result->ConstraintSetup.Add(Constrain);
+			}
+		}
+
+#if WITH_EDITORONLY_DATA
+		Result->ConstraintProfiles.Append(ValidAssets[CollisionSetIndex]->ConstraintProfiles);
+#endif
 	}
 
-	FCustomizableInstanceComponentData* ComponentData =	GetComponentData(ComponentId);
-	check(ComponentData);
-	
-	for (const UPhysicsAsset* PhysicsAsset : ComponentData->PhysicsAssets.PhysicsAssetsToMerge)
+	if (bDisableCollisionsBetweenDifferentAssets)
 	{
-		if (!result && PhysicsAsset)
+		// Compute collision disable table size upperbound to reduce number of alloactions.
+		int32 CollisionDisableTableSize = 0;
+		for (int32 S0 = 1; S0 < CollisionSets.Num(); ++S0)
 		{
-			result = DuplicateObject(PhysicsAsset, NULL);
+			for (int32 S1 = 0; S1 < S0; ++S1)
+			{	
+				CollisionDisableTableSize += CollisionSets[S1].Num() * CollisionSets[S0].Num();
+			}
 		}
-		else if (PhysicsAsset && PhysicsAsset != BaseAsset)
-		{
-			bMergeSucced = true;
 
-			for (const TObjectPtr<USkeletalBodySetup>& BodySetup : PhysicsAsset->SkeletalBodySetups)
-			{
-				if (!result->SkeletalBodySetups.Contains(BodySetup))
+		// We already may have elements in the table, but at the moment of 
+		// addition we don't know yet the final number of elements.
+		// Now a good number of elements will be added and because we know the final number of elements
+		// an upperbound to the number of interactions can be computed and reserved. 
+		CollisionDisableTable.Reserve(CollisionDisableTableSize);
+
+		// Generate disable collision entry for every element in Set S0 for every element in Set S1 
+		// that are not in multiple sets.
+		for (int32 S0 = 1; S0 < CollisionSets.Num(); ++S0)
+		{
+			for (int32 S1 = 0; S1 < S0; ++S1)
+			{	
+				for (int32 Set0Elem : CollisionSets[S0])
 				{
-					bool bFoundBone = false;
+					// Element present in more than one set, will be treated later.
+					if (Set0Elem == INDEX_NONE)
+					{
+						continue;
+					}
 
-					for (const TObjectPtr<USkeletalBodySetup>& OrigBodySetup : result->SkeletalBodySetups)
+					for (int32 Set1Elem : CollisionSets[S1])
 					{
-						if (OrigBodySetup->BoneName == BodySetup->BoneName)
+						// Element present in more than one set, will be treated later.
+						if (Set1Elem == INDEX_NONE)
 						{
-							OrigBodySetup->AddCollisionFrom(BodySetup);
-							bFoundBone = true;
-							break;
+							continue;
 						}
+						CollisionDisableTable.Add(FRigidBodyIndexPair{Set0Elem, Set1Elem}, false);
 					}
+				}
+			}
+		}
 
-					if (!bFoundBone)
+		// Process elements that belong to multiple sets that have been merged to the same element.
+		for ( const TPair<int32, MultiSetArrayType>& Sets : MultiCollisionSets )
+		{
+			for (int32 S = 0; S < CollisionSets.Num(); ++S)
+			{
+				if (!Sets.Value.Contains(S))
+				{	
+					for (int32 SetElem : CollisionSets[S])
 					{
-						result->SkeletalBodySetups.Add(BodySetup);
+						if (SetElem != INDEX_NONE)
+						{
+							CollisionDisableTable.Add(FRigidBodyIndexPair{Sets.Key, SetElem}, false);
+						}
 					}
 				}
 			}
-
-			result->ConstraintSetup.Append(PhysicsAsset->ConstraintSetup);
-			result->CollisionDisableTable.Append(PhysicsAsset->CollisionDisableTable);
-#if WITH_EDITORONLY_DATA
-			result->PhysicalAnimationProfiles.Append(PhysicsAsset->PhysicalAnimationProfiles);
-			result->ConstraintProfiles.Append(PhysicsAsset->ConstraintProfiles);
-#endif // WITH_EDITORONLY_DATA
-			// TODO: Merge the rest of the properties
 		}
-	}
 
-	if (bMergeSucced)
-	{
-			result->UpdateBodySetupIndexMap();
-			result->UpdateBoundsBodiesArray();
-	}
-	else if (!result)
-	{
-		result = NewObject<UPhysicsAsset>();
+		CollisionDisableTable.Shrink();
 	}
 
+	Result->CollisionDisableTable = MoveTemp(CollisionDisableTable);
+	Result->UpdateBodySetupIndexMap();
+	Result->UpdateBoundsBodiesArray();
+
 	ComponentData->PhysicsAssets.PhysicsAssetsToMerge.Empty();
 
-	return result;
+	return Result;
 }
 
-
 bool UCustomizableInstancePrivateData::UpdateSkeletalMesh_PostBeginUpdate0(UCustomizableObjectInstance* Public, const TSharedPtr<FMutableOperationData>& OperationData )
 {
 	MUTABLE_CPUPROFILER_SCOPE(UCustomizableInstancePrivateData::UpdateSkeletalMesh_PostBeginUpdate0)
@@ -1323,33 +1506,24 @@ bool UCustomizableInstancePrivateData::UpdateSkeletalMesh_PostBeginUpdate0(UCust
 				}
 
 				// Merge Physics Assets coming from SubMeshes of the newly generated Mesh
-				if (SkeletalMesh->GetPhysicsAsset() && OperationData->InstanceUpdateData.LODs.Num() )
+				if ( OperationData->InstanceUpdateData.LODs.Num() )
 				{
 					mu::MeshPtrConst MutableMesh = Component.Mesh;
-					mu::Ptr<const mu::PhysicsBody> MutablePhysicsBodies = MutableMesh->GetPhysicsBody();
-					const bool MutablePhysicsBodiesModified = MutablePhysicsBodies ? MutablePhysicsBodies->bBodiesModified : false;
+					mu::Ptr<const mu::PhysicsBody> MutablePhysics = MutableMesh ? MutableMesh->GetPhysicsBody() : nullptr;
 
-					if( HasCOInstanceFlags(ReplacePhysicsAssets) || MutablePhysicsBodiesModified )
+					if (MutablePhysics)
 					{
-						bool bBaseAssetIsValid = true;
-
-						for (const TObjectPtr<USkeletalBodySetup>& BodySetup : SkeletalMesh->GetPhysicsAsset()->SkeletalBodySetups)
-						{
-							if (BodySetup->BoneName != NAME_None && SkeletalMesh->GetRefSkeleton().FindBoneIndex(BodySetup->BoneName) == INDEX_NONE)
-							{
-								bBaseAssetIsValid = false;
-								break;
-							}
-						}
-
-						UPhysicsAsset* MergedPhysicsAsset = MergePhysicalAssets(bBaseAssetIsValid ? SkeletalMesh->GetPhysicsAsset() : nullptr, SkeletalMeshIndex);
+						constexpr bool bDisallowCollisionBetweenAssets = true;
+						UPhysicsAsset* PhysicsAssetResult = BuildPhysicsAsset( 
+								RefSkeletalMeshData->PhysicsAsset.Get(), MutablePhysics.get(), SkeletalMeshIndex, bDisallowCollisionBetweenAssets);
 					
-						if ( MutablePhysicsBodies )
+						SkeletalMesh->SetPhysicsAsset(PhysicsAssetResult);
+
+						if (PhysicsAssetResult)
 						{
-							RefitPhysicsBodies( MergedPhysicsAsset, MutablePhysicsBodies.get());
+							// This is only called in editor, no need to add editor guards.	
+							PhysicsAssetResult->SetPreviewMesh(SkeletalMesh);
 						}
-
-						SkeletalMesh->SetPhysicsAsset(MergedPhysicsAsset);
 					}
 				}
 
@@ -3966,7 +4140,7 @@ FGraphEventRef UCustomizableInstancePrivateData::LoadAdditionalAssetsAsync(const
 			for (int32 TagIndex = 0; TagIndex < MutableMesh->GetTagCount(); ++TagIndex)
 			{
 				FString Tag = MutableMesh->GetTag(TagIndex);
-				if (bReplacePhysicsAssets && Tag.RemoveFromStart("__PhysicsAsset:"))
+				if (Tag.RemoveFromStart("__PhysicsAsset:"))
 				{
 					TSoftObjectPtr<UPhysicsAsset>* PhysicsAsset = CustomizableObject->PhysicsAssetsMap.Find(Tag);
 
diff --git a/Engine/Plugins/Experimental/Mutable/Source/CustomizableObjectEditor/Private/MuCOE/CustomizableObjectCompiler.cpp b/Engine/Plugins/Experimental/Mutable/Source/CustomizableObjectEditor/Private/MuCOE/CustomizableObjectCompiler.cpp
--- a/Engine/Plugins/Experimental/Mutable/Source/CustomizableObjectEditor/Private/MuCOE/CustomizableObjectCompiler.cpp
+++ b/Engine/Plugins/Experimental/Mutable/Source/CustomizableObjectEditor/Private/MuCOE/CustomizableObjectCompiler.cpp
@@ -455,32 +455,32 @@ void FCustomizableObjectCompiler::DisplayUnnamedNodeObjectWarning(FMutableGraphG
 }
 
 
-void FCustomizableObjectCompiler::DisplayDiscardedPhysicsAssetSingleWarning(FMutableGraphGenerationContext& GenerationContext)
-{
-	if (GenerationContext.DiscartedPhysicsAssetMap.Num())
-	{
-		FString PhysicsAssetName;
-		for (int32 ComponentIndex = 0; ComponentIndex < GenerationContext.NumMeshComponentsInRoot; ++ComponentIndex)
-		{
-			USkeletalMesh* RefSkeletalMesh = GenerationContext.ComponentInfos.IsValidIndex(ComponentIndex) ? GenerationContext.ComponentInfos[ComponentIndex].RefSkeletalMesh : nullptr;
-			if (RefSkeletalMesh && GenerationContext.DiscartedPhysicsAssetMap.Find(RefSkeletalMesh->GetPhysicsAsset()))
-			{
-				PhysicsAssetName = RefSkeletalMesh->GetPhysicsAsset()->GetName();
-				break;
-			}
-		}
-		
-		if(PhysicsAssetName.IsEmpty())
-		{
-			PhysicsAssetName = (*GenerationContext.DiscartedPhysicsAssetMap.FindKey(0))->GetName();
-		}
-
-		FText Message = FText::Format( LOCTEXT("Discarted PhysicsAssets", "{0} and {1} other PhysicsAssets have been discarted because one or more Bodies have no corresponding bones in the SkeletalMesh.\nFor more infromation on the assets that need attention check the Output log."),
-										FText::FromString(PhysicsAssetName), FText::AsNumber(GenerationContext.DiscartedPhysicsAssetMap.Num()));
-		
-		CompilerLog(Message, nullptr, EMessageSeverity::Warning, false);
-	}
-}
+//void FCustomizableObjectCompiler::DisplayDiscardedPhysicsAssetSingleWarning(FMutableGraphGenerationContext& GenerationContext)
+//{
+//	if (GenerationContext.DiscartedPhysicsAssetMap.Num())
+//	{
+//		FString PhysicsAssetName;
+//		for (int32 ComponentIndex = 0; ComponentIndex < GenerationContext.NumMeshComponentsInRoot; ++ComponentIndex)
+//		{
+//			USkeletalMesh* RefSkeletalMesh = GenerationContext.ComponentInfos.IsValidIndex(ComponentIndex) ? GenerationContext.ComponentInfos[ComponentIndex].RefSkeletalMesh : nullptr;
+//			if (RefSkeletalMesh && GenerationContext.DiscartedPhysicsAssetMap.Find(RefSkeletalMesh->GetPhysicsAsset()))
+//			{
+//				PhysicsAssetName = RefSkeletalMesh->GetPhysicsAsset()->GetName();
+//				break;
+//			}
+//		}
+//		
+//		if(PhysicsAssetName.IsEmpty())
+//		{
+//			PhysicsAssetName = (*GenerationContext.DiscartedPhysicsAssetMap.FindKey(0))->GetName();
+//		}
+//
+//		FText Message = FText::Format( LOCTEXT("Discarted PhysicsAssets", "{0} and {1} other PhysicsAssets have been discarted because one or more Bodies have no corresponding bones in the SkeletalMesh.\nFor more infromation on the assets that need attention check the Output log."),
+//										FText::FromString(PhysicsAssetName), FText::AsNumber(GenerationContext.DiscartedPhysicsAssetMap.Num()));
+//		
+//		CompilerLog(Message, nullptr, EMessageSeverity::Warning, false);
+//	}
+//}
 
 
 void FCustomizableObjectCompiler::DisplayOrphanNodesWarning(FMutableGraphGenerationContext& GenerationContext)
@@ -701,15 +701,15 @@ mu::NodeObjectPtr FCustomizableObjectCompiler::GenerateMutableRoot(
 	// Generate ReferenceSkeletalMeshes data;
 	PopulateReferenceSkeletalMeshesData(GenerationContext);
 	
-	for (const USkeletalMesh* RefSkeletalMesh : Object->ReferenceSkeletalMeshes)
-	{
-		GenerationContext.CheckPhysicsAssetInSkeletalMesh(RefSkeletalMesh);
-	}
+	//for (const USkeletalMesh* RefSkeletalMesh : Object->ReferenceSkeletalMeshes)
+	//{
+	//	GenerationContext.CheckPhysicsAssetInSkeletalMesh(RefSkeletalMesh);
+	//}
 
 	DisplayParameterWarning( GenerationContext );
 	DisplayUnnamedNodeObjectWarning( GenerationContext );
 	DisplayDuplicatedNodeIdsWarning( GenerationContext );
-	DisplayDiscardedPhysicsAssetSingleWarning( GenerationContext );
+	//DisplayDiscardedPhysicsAssetSingleWarning( GenerationContext );
 	DisplayOrphanNodesWarning( GenerationContext );
 
 	if (GenerationContext.CustomizableObjectWithCycle)
diff --git a/Engine/Plugins/Experimental/Mutable/Source/CustomizableObjectEditor/Private/MuCOE/CustomizableObjectCompiler.h b/Engine/Plugins/Experimental/Mutable/Source/CustomizableObjectEditor/Private/MuCOE/CustomizableObjectCompiler.h
--- a/Engine/Plugins/Experimental/Mutable/Source/CustomizableObjectEditor/Private/MuCOE/CustomizableObjectCompiler.h
+++ b/Engine/Plugins/Experimental/Mutable/Source/CustomizableObjectEditor/Private/MuCOE/CustomizableObjectCompiler.h
@@ -160,7 +160,7 @@ private:
 	
 	/** Display warnings from discarded PhysicsAssets due to SkeletalBodySetups with no corresponding bones
 	in the SkeletalMesh's RefSkeleton */
-	void DisplayDiscardedPhysicsAssetSingleWarning(struct FMutableGraphGenerationContext& GenerationContext);
+	//void DisplayDiscardedPhysicsAssetSingleWarning(struct FMutableGraphGenerationContext& GenerationContext);
 	
 	/** Display a warning for each node contains an orphan pin. */
 	void DisplayOrphanNodesWarning(struct FMutableGraphGenerationContext& GenerationContext);
diff --git a/Engine/Plugins/Experimental/Mutable/Source/CustomizableObjectEditor/Private/MuCOE/GenerateMutableSource/GenerateMutableSource.cpp b/Engine/Plugins/Experimental/Mutable/Source/CustomizableObjectEditor/Private/MuCOE/GenerateMutableSource/GenerateMutableSource.cpp
--- a/Engine/Plugins/Experimental/Mutable/Source/CustomizableObjectEditor/Private/MuCOE/GenerateMutableSource/GenerateMutableSource.cpp
+++ b/Engine/Plugins/Experimental/Mutable/Source/CustomizableObjectEditor/Private/MuCOE/GenerateMutableSource/GenerateMutableSource.cpp
@@ -266,20 +266,20 @@ void FMutableGraphGenerationContext::GenerateClippingCOInternalTags()
 }
 
 
-void FMutableGraphGenerationContext::CheckPhysicsAssetInSkeletalMesh(const USkeletalMesh* SkeletalMesh)
-{
-	if (SkeletalMesh && SkeletalMesh->GetPhysicsAsset() && !DiscartedPhysicsAssetMap.Find(SkeletalMesh->GetPhysicsAsset()))
-	{
-		for (TObjectPtr<USkeletalBodySetup> Bodies : SkeletalMesh->GetPhysicsAsset()->SkeletalBodySetups)
-		{
-			if (Bodies->BoneName != NAME_None && SkeletalMesh->GetRefSkeleton().FindBoneIndex(Bodies->BoneName) == INDEX_NONE)
-			{
-				DiscartedPhysicsAssetMap.Add(SkeletalMesh->GetPhysicsAsset(), DiscartedPhysicsAssetMap.Num());
-				break;
-			}
-		}
-	}
-}
+//void FMutableGraphGenerationContext::CheckPhysicsAssetInSkeletalMesh(const USkeletalMesh* SkeletalMesh)
+//{
+//	if (SkeletalMesh && SkeletalMesh->GetPhysicsAsset() && !DiscartedPhysicsAssetMap.Find(SkeletalMesh->GetPhysicsAsset()))
+//	{
+//		for (TObjectPtr<USkeletalBodySetup> Bodies : SkeletalMesh->GetPhysicsAsset()->SkeletalBodySetups)
+//		{
+//			if (Bodies->BoneName != NAME_None && SkeletalMesh->GetRefSkeleton().FindBoneIndex(Bodies->BoneName) == INDEX_NONE)
+//			{
+//				DiscartedPhysicsAssetMap.Add(SkeletalMesh->GetPhysicsAsset(), DiscartedPhysicsAssetMap.Num());
+//				break;
+//			}
+//		}
+//	}
+//}
 
 
 FMutableComponentInfo& FMutableGraphGenerationContext::GetCurrentComponentInfo()
diff --git a/Engine/Plugins/Experimental/Mutable/Source/CustomizableObjectEditor/Private/MuCOE/GenerateMutableSource/GenerateMutableSource.h b/Engine/Plugins/Experimental/Mutable/Source/CustomizableObjectEditor/Private/MuCOE/GenerateMutableSource/GenerateMutableSource.h
--- a/Engine/Plugins/Experimental/Mutable/Source/CustomizableObjectEditor/Private/MuCOE/GenerateMutableSource/GenerateMutableSource.h
+++ b/Engine/Plugins/Experimental/Mutable/Source/CustomizableObjectEditor/Private/MuCOE/GenerateMutableSource/GenerateMutableSource.h
@@ -579,7 +579,7 @@ struct FMutableGraphGenerationContext
 
 	/** Check if the PhysicsAsset of a given SkeletalMesh has any SkeletalBodySetup with BoneNames not present in the
 	* InSkeletalMesh's RefSkeleton, if so, adds the PhysicsAsset to the DiscartedPhysicsAssetMap to display a warning later on */
-	void CheckPhysicsAssetInSkeletalMesh(const USkeletalMesh* InSkeletalMesh);
+	//void CheckPhysicsAssetInSkeletalMesh(const USkeletalMesh* InSkeletalMesh);
 
 	/** Get the reference skeletal mesh associated to the current mesh component being generated */
 	FMutableComponentInfo& GetCurrentComponentInfo();
@@ -598,7 +598,7 @@ struct FMutableGraphGenerationContext
 	TMap<FString, FParameterUIData> ParameterUIDataMap;
 	TMap<FString, FParameterUIData> StateUIDataMap;
 	TMultiMap<const UCustomizableObjectNodeObjectGroup*, FGroupProjectorTempData> ProjectorGroupMap;
-	TMap<UPhysicsAsset*, uint32> DiscartedPhysicsAssetMap;
+	//TMap<UPhysicsAsset*, uint32> DiscartedPhysicsAssetMap;
 
 	TArray<const USkeleton*> ReferencedSkeletons;
 
diff --git a/Engine/Plugins/Experimental/Mutable/Source/CustomizableObjectEditor/Private/MuCOE/GenerateMutableSource/GenerateMutableSourceMesh.cpp b/Engine/Plugins/Experimental/Mutable/Source/CustomizableObjectEditor/Private/MuCOE/GenerateMutableSource/GenerateMutableSourceMesh.cpp
--- a/Engine/Plugins/Experimental/Mutable/Source/CustomizableObjectEditor/Private/MuCOE/GenerateMutableSource/GenerateMutableSourceMesh.cpp
+++ b/Engine/Plugins/Experimental/Mutable/Source/CustomizableObjectEditor/Private/MuCOE/GenerateMutableSource/GenerateMutableSourceMesh.cpp
@@ -3,6 +3,7 @@
 #include "MuCOE/GenerateMutableSource/GenerateMutableSourceMesh.h"
 
 #include "Algo/BinarySearch.h"
+#include "Algo/Count.h"
 #include "Animation/MorphTarget.h"
 #include "Animation/PoseAsset.h"
 #include "Animation/Skeleton.h"
@@ -1232,9 +1233,9 @@ mu::MeshPtr ConvertSkeletalMeshToMutable(USkeletalMesh* InSkeletalMesh, int LOD,
 				bSkinned = false;
 				int32 NewBoneIndex = 0;
 				if (SkippedBoneMapIndices.Dequeue(NewBoneIndex)) // use skipped index
-					{
+				{
 					InverseBoneMap.Add(RefSkelIndex, NewBoneIndex);
-					}
+				}
 				else
 				{
 					InverseBoneMap.Add(RefSkelIndex, InverseBoneMap.Num());
@@ -1275,8 +1276,17 @@ mu::MeshPtr ConvertSkeletalMeshToMutable(USkeletalMesh* InSkeletalMesh, int LOD,
 		TArray<TObjectPtr<USkeletalBodySetup>> RelevantBodySetups;
 		RelevantBodySetups.Reserve(SkeletalBodySetups.Num());
 
-		for ( TObjectPtr<USkeletalBodySetup>& BodySetup : SkeletalBodySetups )
+		TArray<uint8> DiscardedBodySetups;
+		DiscardedBodySetups.Init(1, SkeletalBodySetups.Num());
+
+		for (int32 BodySetupIndex = 0; BodySetupIndex < SkeletalBodySetups.Num(); ++BodySetupIndex )
 		{
+			TObjectPtr<USkeletalBodySetup>& BodySetup = SkeletalBodySetups[BodySetupIndex];
+			if (!BodySetup)
+			{
+				continue;
+			}
+
 			FString BodyBoneName = BodySetup->BoneName.ToString();
 
 			const int32 SkeletonBoneCount = MutableMesh->GetSkeleton()->GetBoneCount();
@@ -1287,8 +1297,42 @@ mu::MeshPtr ConvertSkeletalMeshToMutable(USkeletalMesh* InSkeletalMesh, int LOD,
 				if (SkeletonBoneName.Equals(BodyBoneName))
 				{
 					RelevantBodySetups.Add(BodySetup);
+					DiscardedBodySetups[BodySetupIndex] = 0;
+				}
+			}
+		}
+
+		const int32 NumDiscardedSetups = Algo::CountIf(DiscardedBodySetups, [](const uint8& V) { return V; });
+
+
+		constexpr bool bOptOutOfIncompleteBodyWarnings = true;
+		if (NumDiscardedSetups > 0 && !bOptOutOfIncompleteBodyWarnings)
+		{
+			FString PhysicsSetupsRemovedMsg = 
+					FString::Printf(TEXT("PhysicsBodySetups in %s attached to bones"), 
+					*(SkeletalMesh->GetPhysicsAsset()->GetName()));
+
+			constexpr int32 MaxNumDiscardedShown = 3;
+			
+			int32 NumDiscardedShown = 0;
+			for (int32 I = 0; I < SkeletalBodySetups.Num() && NumDiscardedShown < MaxNumDiscardedShown; ++I)
+			{
+				if (DiscardedBodySetups[I] && SkeletalBodySetups[I])
+				{
+					PhysicsSetupsRemovedMsg += (NumDiscardedShown <= 0 ? " " : ", ") + SkeletalBodySetups[I]->BoneName.ToString();
+					++NumDiscardedShown;
 				}
 			}
+	
+			if (NumDiscardedShown < NumDiscardedSetups)
+			{
+				PhysicsSetupsRemovedMsg += FString::Printf(TEXT("... and %d more "), NumDiscardedSetups - MaxNumDiscardedShown);
+			}
+	
+			PhysicsSetupsRemovedMsg += FString::Printf(TEXT("have been discarded because they are not present in the SkeletalMesh [%s] Skeleton."),
+				*SkeletalMesh->GetName());
+					
+			GenerationContext.Compiler->CompilerLog(FText::FromString(PhysicsSetupsRemovedMsg), CurrentNode, EMessageSeverity::Warning);
 		}
 
 		mu::Ptr<mu::PhysicsBody> PhysicsBody = new mu::PhysicsBody;
@@ -2224,37 +2268,15 @@ mu::NodeMeshPtr GenerateMutableSourceMesh(const UEdGraphPin * Pin,
 			{
 				MeshNode->SetValue(MutableMesh);
 
-				if (TypedNodeSkel->SkeletalMesh->GetPhysicsAsset() && !GenerationContext.DiscartedPhysicsAssetMap.Find(TypedNodeSkel->SkeletalMesh->GetPhysicsAsset()))
+				if (TypedNodeSkel->SkeletalMesh->GetPhysicsAsset() && 
+					MutableMesh->GetPhysicsBody() && 
+					MutableMesh->GetPhysicsBody()->GetBodyCount())
 				{
-					bool bBoneMismatch = false;
-
-					// should be driven by build settings in some way, probably, or a configuration
-					const bool bOptedOutOfIncompleteBodyWarnings = false;
-
-					for (TObjectPtr<USkeletalBodySetup> BodySetup : TypedNodeSkel->SkeletalMesh->GetPhysicsAsset()->SkeletalBodySetups)
-					{
-						if (BodySetup->BoneName != NAME_None && TypedNodeSkel->SkeletalMesh->GetRefSkeleton().FindBoneIndex(BodySetup->BoneName) == INDEX_NONE)
-						{
-							if (!bOptedOutOfIncompleteBodyWarnings)
-							{
-								GenerationContext.DiscartedPhysicsAssetMap.Add(TypedNodeSkel->SkeletalMesh->GetPhysicsAsset(), GenerationContext.DiscartedPhysicsAssetMap.Num());
-								FString Msg = FString::Printf(TEXT("The PhysicsAsset %s has Bodies with no corresponding bones in the SkeletalMesh %s. It will not be used."), *TypedNodeSkel->SkeletalMesh->GetPhysicsAsset()->GetName(), *TypedNodeSkel->SkeletalMesh->GetName());
-								UE_LOG(LogMutable, Warning, TEXT("%s"), *Msg);
-							}
-
-							bBoneMismatch = true;
-							break;
-						}
-					}
-
-					if (!bBoneMismatch)
-					{
-						TSoftObjectPtr<UPhysicsAsset> PhysicsAsset = TypedNodeSkel->SkeletalMesh->GetPhysicsAsset();
-						GenerationContext.PhysicsAssetMap.Add(PhysicsAsset.ToString(), PhysicsAsset);
-						FString PhysicsAssetTag = FString("__PhysicsAsset:") + PhysicsAsset.ToString();
-						MutableMesh->SetTagCount(1);
-						MutableMesh->SetTag(0, TCHAR_TO_ANSI(*PhysicsAssetTag));
-					}
+					TSoftObjectPtr<UPhysicsAsset> PhysicsAsset = TypedNodeSkel->SkeletalMesh->GetPhysicsAsset();
+					GenerationContext.PhysicsAssetMap.Add(PhysicsAsset.ToString(), PhysicsAsset);
+					FString PhysicsAssetTag = FString("__PhysicsAsset:") + PhysicsAsset.ToString();
+					MutableMesh->SetTagCount(1);
+					MutableMesh->SetTag(0, TCHAR_TO_ANSI(*PhysicsAssetTag));
 				}
 
 				if (GenerationContext.Options.bClothingEnabled)
diff --git a/Engine/Plugins/Experimental/Mutable/Source/CustomizableObjectEditor/Private/MuCOE/GenerateMutableSource/GenerateMutableSourceTable.cpp b/Engine/Plugins/Experimental/Mutable/Source/CustomizableObjectEditor/Private/MuCOE/GenerateMutableSource/GenerateMutableSourceTable.cpp
--- a/Engine/Plugins/Experimental/Mutable/Source/CustomizableObjectEditor/Private/MuCOE/GenerateMutableSource/GenerateMutableSourceTable.cpp
+++ b/Engine/Plugins/Experimental/Mutable/Source/CustomizableObjectEditor/Private/MuCOE/GenerateMutableSource/GenerateMutableSourceTable.cpp
@@ -222,6 +222,16 @@ void FillTableColumn(const UEdGraphPin* Pin, mu::TablePtr MutableTable, FString
 
 					if (MutableMesh)
 					{
+ 						if (SkeletalMesh->GetPhysicsAsset() && MutableMesh->GetPhysicsBody() && MutableMesh->GetPhysicsBody()->GetBodyCount())
+						{	
+							TSoftObjectPtr<UPhysicsAsset> PhysicsAsset = SkeletalMesh->GetPhysicsAsset();
+							int32 CurrentTagCount = MutableMesh->GetTagCount();
+							MutableMesh->SetTagCount(CurrentTagCount + 1);
+							GenerationContext.PhysicsAssetMap.Add(PhysicsAsset.ToString(), PhysicsAsset);
+							FString PhysicsAssetTag = FString("__PhysicsAsset:") + PhysicsAsset.ToString();
+							MutableMesh->SetTag(CurrentTagCount, TCHAR_TO_ANSI(*PhysicsAssetTag));						
+						}
+
 						if (!AnimBPAssetTag.IsEmpty())
 						{
 							int32 CurrentTagCount = MutableMesh->GetTagCount();
