[mutable] Fix morphs not having any effect.
Refactor: When generating code from mutable nodes, use less state in the CodeGenerator objects, and store it locally in the stack using a struct passed as parameter.
Fix uses of "using namespace" not conformant with the coding standards.
Move the error log to use FStrings.

diff --git a/Engine/Plugins/Experimental/Mutable/Source/CustomizableObjectEditor/Private/MuCOE/CustomizableObjectCompileRunnable.cpp b/Engine/Plugins/Experimental/Mutable/Source/CustomizableObjectEditor/Private/MuCOE/CustomizableObjectCompileRunnable.cpp
--- a/Engine/Plugins/Experimental/Mutable/Source/CustomizableObjectEditor/Private/MuCOE/CustomizableObjectCompileRunnable.cpp
+++ b/Engine/Plugins/Experimental/Mutable/Source/CustomizableObjectEditor/Private/MuCOE/CustomizableObjectCompileRunnable.cpp
@@ -113,7 +113,7 @@ uint32 FCustomizableObjectCompileRunnable::Run()
 	mu::Ptr<const mu::ErrorLog> pLog = Compiler->GetLog();
 	for (int i = 0; i < pLog->GetMessageCount(); ++i)
 	{
-		const char* strMessage = pLog->GetMessageText(i);
+		const FString& Message = pLog->GetMessageText(i);
 		mu::ErrorLogMessageType MessageType = pLog->GetMessageType(i);
 		mu::ErrorLogMessageAttachedDataView MessageAttachedData = pLog->GetMessageAttachedData(i);
 
@@ -124,11 +124,11 @@ uint32 FCustomizableObjectCompileRunnable::Run()
 				TSharedPtr<FErrorAttachedData> ErrorAttachedData = MakeShared<FErrorAttachedData>();
 				ErrorAttachedData->UnassignedUVs.Reset();
 				ErrorAttachedData->UnassignedUVs.Append(MessageAttachedData.m_unassignedUVs, MessageAttachedData.m_unassignedUVsSize);
-				ArrayWarning.Add(FError(FText::AsCultureInvariant(strMessage), ErrorAttachedData, pLog->GetMessageContext(i)));
+				ArrayWarning.Add(FError(FText::AsCultureInvariant(Message), ErrorAttachedData, pLog->GetMessageContext(i)));
 			}
 			else
 			{
-				ArrayWarning.Add(FError(FText::AsCultureInvariant(strMessage), pLog->GetMessageContext(i)));
+				ArrayWarning.Add(FError(FText::AsCultureInvariant(Message), pLog->GetMessageContext(i)));
 			}
 		}
 		else if (MessageType == mu::ELMT_ERROR)
@@ -138,11 +138,11 @@ uint32 FCustomizableObjectCompileRunnable::Run()
 				TSharedPtr<FErrorAttachedData> ErrorAttachedData = MakeShared<FErrorAttachedData>();
 				ErrorAttachedData->UnassignedUVs.Reset();
 				ErrorAttachedData->UnassignedUVs.Append(MessageAttachedData.m_unassignedUVs, MessageAttachedData.m_unassignedUVsSize);
-				ArrayError.Add(FError(FText::AsCultureInvariant(strMessage), ErrorAttachedData, pLog->GetMessageContext(i)));
+				ArrayError.Add(FError(FText::AsCultureInvariant(Message), ErrorAttachedData, pLog->GetMessageContext(i)));
 			}
 			else
 			{
-				ArrayError.Add(FError(FText::AsCultureInvariant(strMessage), pLog->GetMessageContext(i)));
+				ArrayError.Add(FError(FText::AsCultureInvariant(Message), pLog->GetMessageContext(i)));
 			}
 		}
 	}
diff --git a/Engine/Plugins/Experimental/Mutable/Source/MutableRuntime/Private/MuR/Image.cpp b/Engine/Plugins/Experimental/Mutable/Source/MutableRuntime/Private/MuR/Image.cpp
--- a/Engine/Plugins/Experimental/Mutable/Source/MutableRuntime/Private/MuR/Image.cpp
+++ b/Engine/Plugins/Experimental/Mutable/Source/MutableRuntime/Private/MuR/Image.cpp
@@ -200,12 +200,7 @@ namespace mu
     //---------------------------------------------------------------------------------------------
     uint8_t* Image::GetData()
     {
-        uint8_t* pResult = 0;
-        if ( m_data.Num() )
-        {
-            pResult = &m_data[0];
-        }
-        return pResult;
+		return m_data.GetData();
     }
 
 
diff --git a/Engine/Plugins/Experimental/Mutable/Source/MutableRuntime/Private/MuR/Mesh.cpp b/Engine/Plugins/Experimental/Mutable/Source/MutableRuntime/Private/MuR/Mesh.cpp
--- a/Engine/Plugins/Experimental/Mutable/Source/MutableRuntime/Private/MuR/Mesh.cpp
+++ b/Engine/Plugins/Experimental/Mutable/Source/MutableRuntime/Private/MuR/Mesh.cpp
@@ -942,7 +942,7 @@ void Mesh::ResetBufferIndices()
 
 
 //---------------------------------------------------------------------------------------------
-bool Mesh::IsSimilar(const Mesh& o) const
+bool Mesh::IsSimilar(const Mesh& o, bool bCompareLayouts) const
 {
 	// Some meshes are just vertex indices (masks) we don't consider them for similarity,
 	// because the kind of vertex channel data they store is the kind that is ignored.
@@ -953,7 +953,7 @@ bool Mesh::IsSimilar(const Mesh& o) const
 
 	bool equal = m_IndexBuffers == o.m_IndexBuffers;
 	if (equal) equal = m_FaceBuffers == o.m_FaceBuffers;
-	if (equal) equal = (m_layouts.Num() == o.m_layouts.Num());
+	if (equal && bCompareLayouts) equal = (m_layouts.Num() == o.m_layouts.Num());
 	if (equal && m_pSkeleton != o.m_pSkeleton)
 	{
 		if (m_pSkeleton && o.m_pSkeleton)
@@ -983,15 +983,18 @@ bool Mesh::IsSimilar(const Mesh& o) const
 	if (equal) equal = (m_tags == o.m_tags);
 
 	// Special comparison for layouts
-	for (int32 i = 0; equal && i < m_layouts.Num(); ++i)
+	if (bCompareLayouts)
 	{
-		equal &= m_layouts[i]->IsSimilar(*o.m_layouts[i]);
+		for (int32 i = 0; equal && i < m_layouts.Num(); ++i)
+		{
+			equal &= m_layouts[i]->IsSimilar(*o.m_layouts[i]);
+		}
 	}
 
 	// Special comparison for vertex buffers
 	if (equal)
 	{
-		equal = m_VertexBuffers.IsSimilarRobust(o.m_VertexBuffers);
+		equal = m_VertexBuffers.IsSimilarRobust(o.m_VertexBuffers, bCompareLayouts);
 	}
 
 	return equal;
diff --git a/Engine/Plugins/Experimental/Mutable/Source/MutableRuntime/Private/MuR/MeshBufferSet.cpp b/Engine/Plugins/Experimental/Mutable/Source/MutableRuntime/Private/MuR/MeshBufferSet.cpp
--- a/Engine/Plugins/Experimental/Mutable/Source/MutableRuntime/Private/MuR/MeshBufferSet.cpp
+++ b/Engine/Plugins/Experimental/Mutable/Source/MutableRuntime/Private/MuR/MeshBufferSet.cpp
@@ -416,9 +416,8 @@ namespace mu
 	}
 
 	//-----------------------------------------------------------------------------------------
-	bool FMeshBufferSet::IsSimilarRobust(const FMeshBufferSet& Other) const
+	bool FMeshBufferSet::IsSimilarRobust(const FMeshBufferSet& Other, bool bCompareUVs) const
 	{
-
 		MUTABLE_CPUPROFILER_SCOPE(FMeshBufferSet::IsSimilarRobust);
 
 		if (m_elementCount != Other.m_elementCount)
@@ -454,18 +453,23 @@ namespace mu
 			const int32 ThisNumChannels = m_buffers[I].m_channels.Num();
 			const int32 OtherNumChannels = m_buffers[J].m_channels.Num();
 
-			if (ThisNumChannels != OtherNumChannels)
+			const MESH_BUFFER& ThisBuffer = m_buffers[I];
+			const MESH_BUFFER& OtherBuffer = Other.m_buffers[J];
+
+			if (!(ThisBuffer.m_channels == OtherBuffer.m_channels && ThisBuffer.m_elementSize==OtherBuffer.m_elementSize))
 			{
 				return false;
 			}
 
-			const MESH_BUFFER& ThisBuffer = m_buffers[I];
-			const MESH_BUFFER& OtherBuffer = Other.m_buffers[J];
-
 			for (uint32 Elem = 0; Elem < m_elementCount; ++Elem)
 			{
 				for (int32 C = 0; C < ThisNumChannels; ++C)
 				{
+					if (!bCompareUVs && ThisBuffer.m_channels[C].m_semantic==MBS_TEXCOORDS)
+					{						
+						continue;
+					}
+
 					const SIZE_T SizeA = GetMeshFormatData(ThisBuffer.m_channels[C].m_format).m_size * ThisBuffer.m_channels[C].m_componentCount;
 					const SIZE_T SizeB = GetMeshFormatData(OtherBuffer.m_channels[C].m_format).m_size * OtherBuffer.m_channels[C].m_componentCount;
 
@@ -488,83 +492,83 @@ namespace mu
 			++I;
 		}
 
-		if (I + 1 < ThisNumBuffers || J + 1 < OtherNumBuffers)
+		// Whatever buffers are left should be irrelevant
+		while (I < ThisNumBuffers)
 		{
-			return false;
-		}
-
-		if (I < ThisNumBuffers && !IsSpecialBufferToIgnoreInSimilar(m_buffers[I]))
-		{
-			return false;
+			if (!IsSpecialBufferToIgnoreInSimilar(m_buffers[I]))
+			{
+				return false;
+			}
+			++I;
 		}
 
-		if (J < OtherNumBuffers && !IsSpecialBufferToIgnoreInSimilar(Other.m_buffers[J]))
+		while (J < OtherNumBuffers)
 		{
-			return false;
+			if (!IsSpecialBufferToIgnoreInSimilar(Other.m_buffers[J]))
+			{
+				return false;
+			}
+			++J;
 		}
 
 		return true;
 	}
 
 	//-----------------------------------------------------------------------------------------
-	bool FMeshBufferSet::IsSimilar(const FMeshBufferSet& o) const
+	bool FMeshBufferSet::IsSimilar(const FMeshBufferSet& Other) const
 	{
 		MUTABLE_CPUPROFILER_SCOPE(FMeshBufferSet::IsSimilar);
 
-		if (m_elementCount != o.m_elementCount) return false;
+		if (m_elementCount != Other.m_elementCount) return false;
 
 		// Compare all buffers except the vertex index channel, which should always be alone in
 		// the last buffer
-		int32 i = 0;
-		int32 oi = 0;
+		int32 Index = 0;
+		int32 OtherIndex = 0;
 
-		while (i < m_buffers.Num() && oi < o.m_buffers.Num())
+		const int32 ThisNumBuffers = m_buffers.Num();
+		const int32 OtherNumBuffers = Other.m_buffers.Num();
+
+		while (Index < ThisNumBuffers && OtherIndex < OtherNumBuffers)
 		{
 			// Is it a special buffer that we should ignore?
-			if (IsSpecialBufferToIgnoreInSimilar(m_buffers[i]))
+			if (IsSpecialBufferToIgnoreInSimilar(m_buffers[Index]))
 			{
-				++i;
+				++Index;
+				continue;
 			}
-			else
+
+			if (IsSpecialBufferToIgnoreInSimilar(Other.m_buffers[OtherIndex]))
 			{
-				if (IsSpecialBufferToIgnoreInSimilar(o.m_buffers[oi]))
-				{
-					++oi;
-				}
-				else
-				{
-					if (!(m_buffers[i] == o.m_buffers[oi]))
-					{
-						return false;
-					}
-					++i;
-					++oi;
-				}
+				++OtherIndex;
+				continue;
 			}
-		}
 
-		// Remaining buffers?
-		if (i + 1 < m_buffers.Num()
-			||
-			oi + 1 < o.m_buffers.Num())
-		{
-			return false;
+			if (!(m_buffers[Index] == Other.m_buffers[OtherIndex]))
+			{
+				return false;
+			}
+			++Index;
+			++OtherIndex;
 		}
 
-		if (i < m_buffers.Num()
-			&&
-			!IsSpecialBufferToIgnoreInSimilar(m_buffers[i])
-			)
+		// Whatever buffers are left should be irrelevant
+		while (Index < ThisNumBuffers)
 		{
-			return false;
+			if (!IsSpecialBufferToIgnoreInSimilar(m_buffers[Index]))
+			{
+				return false;
+			}
+			++Index;
 		}
 
-		if (oi < o.m_buffers.Num()
-			&&
-			!IsSpecialBufferToIgnoreInSimilar(o.m_buffers[oi])
-			)
+		while (OtherIndex < OtherNumBuffers)
 		{
-			return false;
+			if (!IsSpecialBufferToIgnoreInSimilar(Other.m_buffers[OtherIndex]))
+			{
+				return false;
+			}
+			++OtherIndex;
 		}
 
 		return true;
diff --git a/Engine/Plugins/Experimental/Mutable/Source/MutableRuntime/Private/MuR/MeshPrivate.h b/Engine/Plugins/Experimental/Mutable/Source/MutableRuntime/Private/MuR/MeshPrivate.h
--- a/Engine/Plugins/Experimental/Mutable/Source/MutableRuntime/Private/MuR/MeshPrivate.h
+++ b/Engine/Plugins/Experimental/Mutable/Source/MutableRuntime/Private/MuR/MeshPrivate.h
@@ -414,7 +414,7 @@ namespace mu
 				m_format = MBF_NONE;
 				m_components = 0;
 				m_elementSize = 0;
-				m_pBuf = 0;
+				m_pBuf = nullptr;
 			}
 		}
 
diff --git a/Engine/Plugins/Experimental/Mutable/Source/MutableRuntime/Private/MuR/OpMeshDifference.h b/Engine/Plugins/Experimental/Mutable/Source/MutableRuntime/Private/MuR/OpMeshDifference.h
--- a/Engine/Plugins/Experimental/Mutable/Source/MutableRuntime/Private/MuR/OpMeshDifference.h
+++ b/Engine/Plugins/Experimental/Mutable/Source/MutableRuntime/Private/MuR/OpMeshDifference.h
@@ -169,8 +169,8 @@ struct MORPH_VERTEX
 				);
 
 			// Source vertex index channel
-			UntypedMeshBufferIteratorConst
-					itBaseId( pBase->GetVertexBuffers(), MBS_VERTEXINDEX, 0 );
+			UntypedMeshBufferIteratorConst itBaseId( pBase->GetVertexBuffers(), MBS_VERTEXINDEX, 0 );
+			check(itBaseId.GetElementSize());
 
 			// Set the data
 			uint8_t* pDestData = pDest->GetVertexBuffers().GetBufferData( 0 );
diff --git a/Engine/Plugins/Experimental/Mutable/Source/MutableRuntime/Private/MuR/OpMeshMorph.h b/Engine/Plugins/Experimental/Mutable/Source/MutableRuntime/Private/MuR/OpMeshMorph.h
--- a/Engine/Plugins/Experimental/Mutable/Source/MutableRuntime/Private/MuR/OpMeshMorph.h
+++ b/Engine/Plugins/Experimental/Mutable/Source/MutableRuntime/Private/MuR/OpMeshMorph.h
@@ -5,7 +5,6 @@
 #include "MuR/MeshPrivate.h"
 #include "MuR/ConvertData.h"
 #include "MuR/Platform.h"
-
 #include "MuR/SparseIndexMap.h"
 
 namespace mu
@@ -21,19 +20,19 @@ namespace mu
 		if (!pBase) return nullptr;
 
         const auto ApplyMorph = []
-            ( auto BaseIdIter, const TArray< UntypedMeshBufferIterator >& baseChannelsIters, const int32_t baseSize,
-              auto MorphIdIter, const TArray< UntypedMeshBufferIteratorConst >& morphChannelsIters, const int32_t morphSize,
+            ( auto BaseIdIter, const TArray< UntypedMeshBufferIterator >& baseChannelsIters, const int32 baseSize,
+              auto MorphIdIter, const TArray< UntypedMeshBufferIteratorConst >& morphChannelsIters, const int32 morphSize,
               const float factor )
         -> void
         {
-            uint32_t minBaseId = std::numeric_limits<uint32_t>::max();
-            uint32_t maxBaseId = 0;
+            uint32 minBaseId = std::numeric_limits<uint32>::max();
+            uint32 maxBaseId = 0;
 
 			{
                 auto limitsBaseIdIter = BaseIdIter;
-				for ( int32_t i = 0; i < baseSize; ++i, ++limitsBaseIdIter )
+				for ( int32 i = 0; i < baseSize; ++i, ++limitsBaseIdIter )
 				{
-					const uint32_t id = limitsBaseIdIter.GetAsUINT32();
+					const uint32 id = limitsBaseIdIter.GetAsUINT32();
 					minBaseId = FMath::Min( id, minBaseId );
 					maxBaseId = FMath::Max( id, maxBaseId );
 				}
@@ -42,30 +41,30 @@ namespace mu
             SparseIndexMap IndexMap( minBaseId, maxBaseId );
             
             auto usedMorphIdIter = MorphIdIter;
-            for ( uint32_t i = 0; i < static_cast<uint32_t>(morphSize); ++i, ++usedMorphIdIter)
+            for ( uint32 i = 0; i < static_cast<uint32>(morphSize); ++i, ++usedMorphIdIter)
             {
-                const uint32_t morphId = (*usedMorphIdIter)[0];
+                const uint32 morphId = (*usedMorphIdIter)[0];
                 
                 IndexMap.Insert( morphId, i );
             }
 
-            for ( int32_t v = 0; v < baseSize; ++v )
+            for ( int32 v = 0; v < baseSize; ++v )
             {
-                const uint32_t baseId = (BaseIdIter + v).GetAsUINT32();
-                const uint32_t morphIdx = IndexMap.Find( baseId );
+                const uint32 baseId = (BaseIdIter + v).GetAsUINT32();
+                const uint32 morphIdx = IndexMap.Find( baseId );
 
                 if ( morphIdx == SparseIndexMap::NotFoundValue )
                 {
                     continue;
                 }
 
-                const int32_t m = static_cast<int32_t>( morphIdx );
+                const int32 m = static_cast<int32>( morphIdx );
 
                 // Find consecutive run.
                 auto RunBaseIter = BaseIdIter + v;
                 auto RunMorphIter = MorphIdIter + m;
 
-                int32_t runSize = 0;
+                int32 runSize = 0;
                 for (  ; v + runSize < baseSize && m + runSize < morphSize && RunBaseIter.GetAsUINT32() == (*RunMorphIter)[0];
                         ++runSize, ++RunBaseIter, ++RunMorphIter );
 
@@ -85,13 +84,13 @@ namespace mu
                     const auto dstChannelComps = baseChannelsIters[c].GetComponents();
 
                     // Apply Morph to range found above.
-                    for ( int32_t r = 0; r < runSize; ++r, ++channelBaseIter, ++channelMorphIter )
+                    for ( int32 r = 0; r < runSize; ++r, ++channelBaseIter, ++channelMorphIter )
                     {
                         const vec4<float> value = channelBaseIter.GetAsVec4f() + channelMorphIter.GetAsVec4f()*factor;
 
                         // TODO: Optimize this for the specific components.
                         // Max 4 components
-                        for ( int32_t comp = 0; comp < dstChannelComps && comp < 4; ++comp )
+                        for ( int32 comp = 0; comp < dstChannelComps && comp < 4; ++comp )
                         {
                             ConvertData( comp, channelBaseIter.ptr(), dstChannelFormat, &value, MBF_FLOAT32 );
                         }
@@ -105,9 +104,9 @@ namespace mu
 		MeshPtr pDest = pBase->Clone();
 
 		// Number of vertices to modify
-		const uint32_t minCount = pMin ? pMin->GetVertexBuffers().GetElementCount() : 0;
-		const uint32_t maxCount = pMax ? pMax->GetVertexBuffers().GetElementCount() : 0;
-		const uint32_t baseCount = pBase ? pBase->GetVertexBuffers().GetElementCount() : 0;
+		const uint32 minCount = pMin ? pMin->GetVertexBuffers().GetElementCount() : 0;
+		const uint32 maxCount = pMax ? pMax->GetVertexBuffers().GetElementCount() : 0;
+		const uint32 baseCount = pBase ? pBase->GetVertexBuffers().GetElementCount() : 0;
 		const Mesh* refTarget = minCount > 0 ? pMin : pMax;
 
 		if ( baseCount == 0 || (minCount + maxCount) == 0)
@@ -117,7 +116,7 @@ namespace mu
 
 		if (refTarget)
 		{
-			uint32_t ccount = refTarget->GetVertexBuffers().GetBufferChannelCount(0);
+			uint32 ccount = refTarget->GetVertexBuffers().GetBufferChannelCount(0);
 
 			// Iterator of the vertex ids of the base vertices
 			UntypedMeshBufferIteratorConst itBaseId(pBase->GetVertexBuffers(), MBS_VERTEXINDEX);
@@ -149,13 +148,13 @@ namespace mu
 			
 			if (minCount > 0)
 			{
-				MeshBufferIteratorConst<MBF_UINT32, uint32_t, 1> itMinId(pMin->GetVertexBuffers(), MBS_VERTEXINDEX);
+				MeshBufferIteratorConst<MBF_UINT32, uint32, 1> itMinId(pMin->GetVertexBuffers(), MBS_VERTEXINDEX);
 				ApplyMorph(itBaseId, itBaseChannels, baseCount, itMinId, itMinChannels, minCount, 1.0f - factor);
 			}
 
 			if (maxCount > 0)
 			{
-				MeshBufferIteratorConst<MBF_UINT32, uint32_t, 1> itMaxId(pMax->GetVertexBuffers(), MBS_VERTEXINDEX);
+				MeshBufferIteratorConst<MBF_UINT32, uint32, 1> itMaxId(pMax->GetVertexBuffers(), MBS_VERTEXINDEX);
 				ApplyMorph(itBaseId, itBaseChannels, baseCount, itMaxId, itMaxChannels, maxCount, factor);
 			}
 		}
diff --git a/Engine/Plugins/Experimental/Mutable/Source/MutableRuntime/Public/MuR/Mesh.h b/Engine/Plugins/Experimental/Mutable/Source/MutableRuntime/Public/MuR/Mesh.h
--- a/Engine/Plugins/Experimental/Mutable/Source/MutableRuntime/Public/MuR/Mesh.h
+++ b/Engine/Plugins/Experimental/Mutable/Source/MutableRuntime/Public/MuR/Mesh.h
@@ -520,10 +520,10 @@ namespace mu
 
 		//! Compare the mesh with another one, but ignore internal data like generated vertex
 		//! indices.
-		bool IsSimilar(const Mesh& o) const;
+		bool IsSimilar(const Mesh& o, bool bCompareLayouts) const;
 
 		//! Compare the vertex attributes to check if they match.
-		bool IsSameVertex( uint32 v, const Mesh& other, uint32 otherVertexIndex, float tolerance = 1e-5f) const;
+		bool IsSameVertex( uint32 v, const Mesh& other, uint32 otherVertexIndex, float tolerance = UE_SMALL_NUMBER) const;
 
 
 		//! Make a map from the vertices in this mesh to thefirst matching vertex of the given
diff --git a/Engine/Plugins/Experimental/Mutable/Source/MutableRuntime/Public/MuR/MeshBufferSet.h b/Engine/Plugins/Experimental/Mutable/Source/MutableRuntime/Public/MuR/MeshBufferSet.h
--- a/Engine/Plugins/Experimental/Mutable/Source/MutableRuntime/Public/MuR/MeshBufferSet.h
+++ b/Engine/Plugins/Experimental/Mutable/Source/MutableRuntime/Public/MuR/MeshBufferSet.h
@@ -393,7 +393,7 @@ namespace mu
 		//! Compare the mesh buffer with another one, but ignore internal data like generated
 		//! vertex indices. This version compares the data component-wise, skipping any memory
 		//! not specified in the buffer description.	
-		bool IsSimilarRobust(const FMeshBufferSet& Other) const;
+		bool IsSimilarRobust(const FMeshBufferSet& Other,bool bCompareUVs) const;
 
 		//! Change the buffer descriptions so that all buffer indices start at 0 and are in the
 		//! same order than memory.
diff --git a/Engine/Plugins/Experimental/Mutable/Source/MutableTools/Private/MuT/ASTOpConditional.cpp b/Engine/Plugins/Experimental/Mutable/Source/MutableTools/Private/MuT/ASTOpConditional.cpp
--- a/Engine/Plugins/Experimental/Mutable/Source/MutableTools/Private/MuT/ASTOpConditional.cpp
+++ b/Engine/Plugins/Experimental/Mutable/Source/MutableTools/Private/MuT/ASTOpConditional.cpp
@@ -14,347 +14,349 @@
 #include "MuT/StreamsPrivate.h"
 
 
-using namespace mu;
+namespace mu
+{
 
 
-//-------------------------------------------------------------------------------------------------
-//-------------------------------------------------------------------------------------------------
-//-------------------------------------------------------------------------------------------------
-ASTOpConditional::ASTOpConditional()
-    : condition(this)
-    , yes(this)
-    , no(this)
-{
-}
+	//-------------------------------------------------------------------------------------------------
+	//-------------------------------------------------------------------------------------------------
+	//-------------------------------------------------------------------------------------------------
+	ASTOpConditional::ASTOpConditional()
+		: condition(this)
+		, yes(this)
+		, no(this)
+	{
+	}
 
 
-//-------------------------------------------------------------------------------------------------
-ASTOpConditional::~ASTOpConditional()
-{
-    // Explicit call needed to avoid recursive destruction
-    ASTOp::RemoveChildren();
-}
+	//-------------------------------------------------------------------------------------------------
+	ASTOpConditional::~ASTOpConditional()
+	{
+		// Explicit call needed to avoid recursive destruction
+		ASTOp::RemoveChildren();
+	}
 
 
-//-------------------------------------------------------------------------------------------------
-bool ASTOpConditional::IsEqual(const ASTOp& otherUntyped) const
-{
-    if ( auto other = dynamic_cast<const ASTOpConditional*>(&otherUntyped) )
-    {
-        return type==other->type && condition==other->condition &&
-                yes==other->yes && no==other->no;
-    }
-    return false;
-}
+	//-------------------------------------------------------------------------------------------------
+	bool ASTOpConditional::IsEqual(const ASTOp& otherUntyped) const
+	{
+		if (auto other = dynamic_cast<const ASTOpConditional*>(&otherUntyped))
+		{
+			return type == other->type && condition == other->condition &&
+				yes == other->yes && no == other->no;
+		}
+		return false;
+	}
 
 
-//-------------------------------------------------------------------------------------------------
-mu::Ptr<ASTOp> ASTOpConditional::Clone(MapChildFuncRef mapChild) const
-{
-    Ptr<ASTOpConditional> n = new ASTOpConditional();
-    n->type = type;
-    n->condition = mapChild(condition.child());
-    n->yes = mapChild(yes.child());
-    n->no = mapChild(no.child());
-    return n;
-}
+	//-------------------------------------------------------------------------------------------------
+	mu::Ptr<ASTOp> ASTOpConditional::Clone(MapChildFuncRef mapChild) const
+	{
+		Ptr<ASTOpConditional> n = new ASTOpConditional();
+		n->type = type;
+		n->condition = mapChild(condition.child());
+		n->yes = mapChild(yes.child());
+		n->no = mapChild(no.child());
+		return n;
+	}
 
 
-//-------------------------------------------------------------------------------------------------
-uint64 ASTOpConditional::Hash() const
-{
-	uint64 res = std::hash<void*>()(condition.m_child.get());
-    hash_combine(res, yes.m_child.get() );
-    hash_combine(res, no.m_child.get() );
-    return res;
-}
+	//-------------------------------------------------------------------------------------------------
+	uint64 ASTOpConditional::Hash() const
+	{
+		uint64 res = std::hash<void*>()(condition.m_child.get());
+		hash_combine(res, yes.m_child.get());
+		hash_combine(res, no.m_child.get());
+		return res;
+	}
 
 
-//-------------------------------------------------------------------------------------------------
-void ASTOpConditional::Assert()
-{
-    switch(type)
-    {
-    case OP_TYPE::NU_CONDITIONAL:
-    case OP_TYPE::SC_CONDITIONAL:
-    case OP_TYPE::CO_CONDITIONAL:
-    case OP_TYPE::IM_CONDITIONAL:
-    case OP_TYPE::ME_CONDITIONAL:
-    case OP_TYPE::LA_CONDITIONAL:
-    case OP_TYPE::IN_CONDITIONAL:
-        break;
-    default:
-        // Unexpected type
-        check(false);
-        break;
-    }
-
-    ASTOp::Assert();
-}
-
-
-//-------------------------------------------------------------------------------------------------
-void ASTOpConditional::ForEachChild(const TFunctionRef<void(ASTChild&)> f )
-{
-    f(condition);
-    f(yes);
-    f(no);
-}
+	//-------------------------------------------------------------------------------------------------
+	void ASTOpConditional::Assert()
+	{
+		switch (type)
+		{
+		case OP_TYPE::NU_CONDITIONAL:
+		case OP_TYPE::SC_CONDITIONAL:
+		case OP_TYPE::CO_CONDITIONAL:
+		case OP_TYPE::IM_CONDITIONAL:
+		case OP_TYPE::ME_CONDITIONAL:
+		case OP_TYPE::LA_CONDITIONAL:
+		case OP_TYPE::IN_CONDITIONAL:
+			break;
+		default:
+			// Unexpected type
+			check(false);
+			break;
+		}
 
+		ASTOp::Assert();
+	}
 
-//-------------------------------------------------------------------------------------------------
-void ASTOpConditional::Link( PROGRAM& program, const FLinkerOptions*)
-{
-    // Already linked?
-    if (!linkedAddress)
-    {
-        OP::ConditionalArgs args;
-        memset( &args,0, sizeof(args) );
 
-        if (condition) args.condition = condition->linkedAddress;
-        if (yes) args.yes = yes->linkedAddress;
-        if (no) args.no = no->linkedAddress;
+	//-------------------------------------------------------------------------------------------------
+	void ASTOpConditional::ForEachChild(const TFunctionRef<void(ASTChild&)> f)
+	{
+		f(condition);
+		f(yes);
+		f(no);
+	}
 
-        linkedAddress = (OP::ADDRESS)program.m_opAddress.Num();
-        //program.m_code.push_back(op);
-        program.m_opAddress.Add((uint32_t)program.m_byteCode.Num());
-        AppendCode(program.m_byteCode,type);
-        AppendCode(program.m_byteCode,args);
 
-    }
-}
+	//-------------------------------------------------------------------------------------------------
+	void ASTOpConditional::Link(PROGRAM& program, const FLinkerOptions*)
+	{
+		// Already linked?
+		if (!linkedAddress)
+		{
+			OP::ConditionalArgs args;
+			memset(&args, 0, sizeof(args));
 
+			if (condition) args.condition = condition->linkedAddress;
+			if (yes) args.yes = yes->linkedAddress;
+			if (no) args.no = no->linkedAddress;
 
-//-------------------------------------------------------------------------------------------------
-FImageDesc ASTOpConditional::GetImageDesc( bool returnBestOption, class GetImageDescContext* context )
-{
-    FImageDesc res;
-
-    // Local context in case it is necessary
-    GetImageDescContext localContext;
-    if (!context)
-    {
-        context = &localContext;
-    }
-    else
-    {
-        // Cached result?
-		FImageDesc* PtrValue = context->m_results.Find(this);
-		if (PtrValue)
+			linkedAddress = (OP::ADDRESS)program.m_opAddress.Num();
+			//program.m_code.push_back(op);
+			program.m_opAddress.Add((uint32_t)program.m_byteCode.Num());
+			AppendCode(program.m_byteCode, type);
+			AppendCode(program.m_byteCode, args);
+
+		}
+	}
+
+
+	//-------------------------------------------------------------------------------------------------
+	FImageDesc ASTOpConditional::GetImageDesc(bool returnBestOption, class GetImageDescContext* context)
+	{
+		FImageDesc res;
+
+		// Local context in case it is necessary
+		GetImageDescContext localContext;
+		if (!context)
+		{
+			context = &localContext;
+		}
+		else
+		{
+			// Cached result?
+			FImageDesc* PtrValue = context->m_results.Find(this);
+			if (PtrValue)
+			{
+				return *PtrValue;
+			}
+		}
+
+		if (type == OP_TYPE::IM_CONDITIONAL)
 		{
-			return *PtrValue;
+			// In a conditional we cannot guarantee the size and format.
+			// We check both options, and if they are the same we return that.
+			// Otherwise, we return an empty descriptor that will force re-formatting of the image.
+			// The code optimiser will take care then of moving the format operations down to each
+			// branch and remove the unnecessary ones.
+			FImageDesc noDesc;
+			FImageDesc yesDesc;
+
+			if (no.child())
+			{
+				noDesc = no->GetImageDesc(returnBestOption, context);
+			}
+			if (yes.child())
+			{
+				yesDesc = yes->GetImageDesc(returnBestOption, context);
+			}
+
+			if (yesDesc == noDesc || returnBestOption)
+			{
+				res = yesDesc;
+			}
+			else
+			{
+				res.m_format = (yesDesc.m_format == noDesc.m_format) ? yesDesc.m_format : EImageFormat::IF_NONE;
+				res.m_lods = (yesDesc.m_lods == noDesc.m_lods) ? yesDesc.m_lods : 0;
+				res.m_size = (yesDesc.m_size == noDesc.m_size) ? yesDesc.m_size : FImageSize(0, 0);
+			}
 		}
+		else
+		{
+			checkf(false, TEXT("Instruction not supported"));
+		}
+
+
+		// Cache the result
+		if (context)
+		{
+			context->m_results.Add(this, res);
+		}
+
+		return res;
 	}
 
-    if ( type==OP_TYPE::IM_CONDITIONAL )
-    {
-        // In a conditional we cannot guarantee the size and format.
-        // We check both options, and if they are the same we return that.
-        // Otherwise, we return an empty descriptor that will force re-formatting of the image.
-        // The code optimiser will take care then of moving the format operations down to each
-        // branch and remove the unnecessary ones.
-        FImageDesc noDesc;
-        FImageDesc yesDesc;
-
-        if (no.child())
-        {
-            noDesc = no->GetImageDesc(returnBestOption, context);
-        }
-        if (yes.child())
-        {
-            yesDesc = yes->GetImageDesc(returnBestOption, context);
-        }
-
-        if ( yesDesc == noDesc || returnBestOption )
-        {
-            res = yesDesc;
-        }
-        else
-        {
-            res.m_format = (yesDesc.m_format==noDesc.m_format) ? yesDesc.m_format : EImageFormat::IF_NONE;
-            res.m_lods = (yesDesc.m_lods==noDesc.m_lods) ? yesDesc.m_lods : 0;
-            res.m_size = (yesDesc.m_size==noDesc.m_size) ? yesDesc.m_size : FImageSize( 0, 0 );
-        }
-    }
-    else
-    {
-        checkf(false, TEXT("Instruction not supported") );
-    }
-
-
-    // Cache the result
-    if (context)
-    {
-        context->m_results.Add(this,res);
-    }
-
-    return res;
-}
-
-
-//-------------------------------------------------------------------------------------------------
-void ASTOpConditional::GetLayoutBlockSize( int* pBlockX, int* pBlockY )
-{
-    if ( type==OP_TYPE::IM_CONDITIONAL )
-    {
-        yes->GetLayoutBlockSize( pBlockX, pBlockY );
-
-        if (!*pBlockX)
-        {
-            no->GetLayoutBlockSize( pBlockX, pBlockY );
-        }
-    }
-    else
-    {
-        checkf( false, TEXT("Instruction not supported") );
-    }
-}
-
-
-//-------------------------------------------------------------------------------------------------
-void ASTOpConditional::GetBlockLayoutSize( int blockIndex, int* pBlockX, int* pBlockY,
-                                           BLOCK_LAYOUT_SIZE_CACHE* cache )
-{
-    if ( type==OP_TYPE::LA_CONDITIONAL )
-    {
-        yes->GetBlockLayoutSizeCached( blockIndex, pBlockX, pBlockY, cache );
-
-        if (*pBlockX==0)
-        {
-            no->GetBlockLayoutSizeCached( blockIndex, pBlockX, pBlockY, cache );
-        }
-    }
-    else
-    {
-        checkf( false, TEXT("Instruction not supported") );
-    }
-}
-
-
-//-------------------------------------------------------------------------------------------------
-bool ASTOpConditional::GetNonBlackRect( FImageRect& maskUsage ) const
-{
-    if (type == OP_TYPE::IM_CONDITIONAL )
-    {
-        FImageRect local;
-        bool localValid = false;
-        if (yes)
-        {
-            localValid = yes->GetNonBlackRect(local);
-            if (!localValid)
-            {
-                return false;
-            }
-        }
-
-        if (no)
-        {
-            FImageRect noRect;
-            bool validNo = no->GetNonBlackRect(noRect);
-            if (validNo)
-            {
-                if (localValid)
-                {
-                    local.Bound(noRect);
-                }
-                else
-                {
-                    local = noRect;
-                    localValid = true;
-                }
-            }
-            else
-            {
-                return false;
-            }
-        }
-
-        if (localValid)
-        {
-            maskUsage = local;
-            return true;
-        }
-    }
-
-    return false;
-}
-
-
-//-------------------------------------------------------------------------------------------------
-mu::Ptr<ASTOp> ASTOpConditional::OptimiseSemantic( const MODEL_OPTIMIZATION_OPTIONS& ) const
-{
-    if (!condition)
-    {
-        // If there is no expression, we'll assume true.
-        return yes.child();
-    }
-
-    // If the branches are the same, remove the instruction
-    if ( yes.child() == no.child() )
-    {
-        return yes.child();
-    }
-
-    // Constant condition?
-    if ( condition->GetOpType() == OP_TYPE::BO_CONSTANT )
-    {
-        auto typedCondition = dynamic_cast<const ASTOpConstantBool*>(condition.child().get());
-        if ( typedCondition->value )
-        {
-            return yes.child();
-        }
-
-        return no.child();
-    }
-
-    else
-    {
-        // If the yes branch is a conditional with the same condition
-        if (yes && yes->GetOpType() == type)
-        {
-            auto typedYes = dynamic_cast<const ASTOpConditional*>(yes.child().get());
-            if ( condition.child()==typedYes->condition.child()
-                 ||
-                 *condition.child()==*typedYes->condition.child() )
-            {
-                auto op = mu::Clone<ASTOpConditional>(this);
-                op->yes = typedYes->yes.child();
-                return op;
-            }
-        }
-
-        // If the no branch is a conditional with the same condition
-        else if (no && no->GetOpType() == type)
-        {
-            auto typedNo = dynamic_cast<const ASTOpConditional*>(no.child().get());
-            if ( condition.child()==typedNo->condition.child()
-                 ||
-                 *condition.child()==*typedNo->condition.child() )
-            {
-                auto op = mu::Clone<ASTOpConditional>(this);
-                op->no = typedNo->no.child();
-                return op;
-            }
-        }
-    }
-
-    return nullptr;
-}
-
-
-//-------------------------------------------------------------------------------------------------
-mu::Ptr<ImageSizeExpression> ASTOpConditional::GetImageSizeExpression() const
-{
-    Ptr<ImageSizeExpression> pRes = new ImageSizeExpression;
-    pRes->type = ImageSizeExpression::ISET_CONDITIONAL;
-    pRes->yes = yes->GetImageSizeExpression();
-    pRes->no = no->GetImageSizeExpression();
 
-    if ( *pRes->yes == *pRes->no )
-    {
-        pRes = pRes->yes;
-    }
+	//-------------------------------------------------------------------------------------------------
+	void ASTOpConditional::GetLayoutBlockSize(int* pBlockX, int* pBlockY)
+	{
+		if (type == OP_TYPE::IM_CONDITIONAL)
+		{
+			yes->GetLayoutBlockSize(pBlockX, pBlockY);
+
+			if (!*pBlockX)
+			{
+				no->GetLayoutBlockSize(pBlockX, pBlockY);
+			}
+		}
+		else
+		{
+			checkf(false, TEXT("Instruction not supported"));
+		}
+	}
+
+
+	//-------------------------------------------------------------------------------------------------
+	void ASTOpConditional::GetBlockLayoutSize(int blockIndex, int* pBlockX, int* pBlockY,
+		BLOCK_LAYOUT_SIZE_CACHE* cache)
+	{
+		if (type == OP_TYPE::LA_CONDITIONAL)
+		{
+			yes->GetBlockLayoutSizeCached(blockIndex, pBlockX, pBlockY, cache);
+
+			if (*pBlockX == 0)
+			{
+				no->GetBlockLayoutSizeCached(blockIndex, pBlockX, pBlockY, cache);
+			}
+		}
+		else
+		{
+			checkf(false, TEXT("Instruction not supported"));
+		}
+	}
+
+
+	//-------------------------------------------------------------------------------------------------
+	bool ASTOpConditional::GetNonBlackRect(FImageRect& maskUsage) const
+	{
+		if (type == OP_TYPE::IM_CONDITIONAL)
+		{
+			FImageRect local;
+			bool localValid = false;
+			if (yes)
+			{
+				localValid = yes->GetNonBlackRect(local);
+				if (!localValid)
+				{
+					return false;
+				}
+			}
+
+			if (no)
+			{
+				FImageRect noRect;
+				bool validNo = no->GetNonBlackRect(noRect);
+				if (validNo)
+				{
+					if (localValid)
+					{
+						local.Bound(noRect);
+					}
+					else
+					{
+						local = noRect;
+						localValid = true;
+					}
+				}
+				else
+				{
+					return false;
+				}
+			}
+
+			if (localValid)
+			{
+				maskUsage = local;
+				return true;
+			}
+		}
+
+		return false;
+	}
+
 
-    return pRes;
-}
+	//-------------------------------------------------------------------------------------------------
+	mu::Ptr<ASTOp> ASTOpConditional::OptimiseSemantic(const MODEL_OPTIMIZATION_OPTIONS&) const
+	{
+		if (!condition)
+		{
+			// If there is no expression, we'll assume true.
+			return yes.child();
+		}
+
+		// If the branches are the same, remove the instruction
+		if (yes.child() == no.child())
+		{
+			return yes.child();
+		}
+
+		// Constant condition?
+		if (condition->GetOpType() == OP_TYPE::BO_CONSTANT)
+		{
+			auto typedCondition = dynamic_cast<const ASTOpConstantBool*>(condition.child().get());
+			if (typedCondition->value)
+			{
+				return yes.child();
+			}
+
+			return no.child();
+		}
+
+		else
+		{
+			// If the yes branch is a conditional with the same condition
+			if (yes && yes->GetOpType() == type)
+			{
+				auto typedYes = dynamic_cast<const ASTOpConditional*>(yes.child().get());
+				if (condition.child() == typedYes->condition.child()
+					||
+					*condition.child() == *typedYes->condition.child())
+				{
+					auto op = mu::Clone<ASTOpConditional>(this);
+					op->yes = typedYes->yes.child();
+					return op;
+				}
+			}
+
+			// If the no branch is a conditional with the same condition
+			else if (no && no->GetOpType() == type)
+			{
+				auto typedNo = dynamic_cast<const ASTOpConditional*>(no.child().get());
+				if (condition.child() == typedNo->condition.child()
+					||
+					*condition.child() == *typedNo->condition.child())
+				{
+					auto op = mu::Clone<ASTOpConditional>(this);
+					op->no = typedNo->no.child();
+					return op;
+				}
+			}
+		}
+
+		return nullptr;
+	}
+
+
+	//-------------------------------------------------------------------------------------------------
+	mu::Ptr<ImageSizeExpression> ASTOpConditional::GetImageSizeExpression() const
+	{
+		Ptr<ImageSizeExpression> pRes = new ImageSizeExpression;
+		pRes->type = ImageSizeExpression::ISET_CONDITIONAL;
+		pRes->yes = yes->GetImageSizeExpression();
+		pRes->no = no->GetImageSizeExpression();
+
+		if (*pRes->yes == *pRes->no)
+		{
+			pRes = pRes->yes;
+		}
+
+		return pRes;
+	}
 
+}
\ No newline at end of file
diff --git a/Engine/Plugins/Experimental/Mutable/Source/MutableTools/Private/MuT/ASTOpConstantBool.cpp b/Engine/Plugins/Experimental/Mutable/Source/MutableTools/Private/MuT/ASTOpConstantBool.cpp
--- a/Engine/Plugins/Experimental/Mutable/Source/MutableTools/Private/MuT/ASTOpConstantBool.cpp
+++ b/Engine/Plugins/Experimental/Mutable/Source/MutableTools/Private/MuT/ASTOpConstantBool.cpp
@@ -10,69 +10,70 @@
 #include "MuT/StreamsPrivate.h"
 
 
-using namespace mu;
+namespace mu
+{
 
 
-//-------------------------------------------------------------------------------------------------
-//-------------------------------------------------------------------------------------------------
-//-------------------------------------------------------------------------------------------------
-ASTOpConstantBool::ASTOpConstantBool( bool v )
-{
-    value = v;
-}
+	//-------------------------------------------------------------------------------------------------
+	//-------------------------------------------------------------------------------------------------
+	//-------------------------------------------------------------------------------------------------
+	ASTOpConstantBool::ASTOpConstantBool(bool v)
+	{
+		value = v;
+	}
 
 
-void ASTOpConstantBool::ForEachChild(const TFunctionRef<void(ASTChild&)> )
-{
-}
+	void ASTOpConstantBool::ForEachChild(const TFunctionRef<void(ASTChild&)>)
+	{
+	}
 
 
-bool ASTOpConstantBool::IsEqual(const ASTOp& otherUntyped) const
-{
-    if ( auto other = dynamic_cast<const ASTOpConstantBool*>(&otherUntyped) )
-    {
-        return value==other->value;
-    }
-    return false;
-}
+	bool ASTOpConstantBool::IsEqual(const ASTOp& otherUntyped) const
+	{
+		if (auto other = dynamic_cast<const ASTOpConstantBool*>(&otherUntyped))
+		{
+			return value == other->value;
+		}
+		return false;
+	}
 
 
-uint64 ASTOpConstantBool::Hash() const
-{
-	uint64 res = std::hash<uint64>()(uint64(OP_TYPE::BO_CONSTANT));
-    hash_combine( res, value );
-    return res;
-}
+	uint64 ASTOpConstantBool::Hash() const
+	{
+		uint64 res = std::hash<uint64>()(uint64(OP_TYPE::BO_CONSTANT));
+		hash_combine(res, value);
+		return res;
+	}
 
 
-mu::Ptr<ASTOp> ASTOpConstantBool::Clone(MapChildFuncRef) const
-{
-    Ptr<ASTOpConstantBool> n = new ASTOpConstantBool();
-    n->value = value;
-    return n;
-}
+	mu::Ptr<ASTOp> ASTOpConstantBool::Clone(MapChildFuncRef) const
+	{
+		Ptr<ASTOpConstantBool> n = new ASTOpConstantBool();
+		n->value = value;
+		return n;
+	}
 
 
-void ASTOpConstantBool::Link( PROGRAM& program, const FLinkerOptions*)
-{
-    if (!linkedAddress)
-    {
-        OP::BoolConstantArgs args;
-        memset( &args,0, sizeof(args) );
-        args.value = value;
-
-        linkedAddress = (OP::ADDRESS)program.m_opAddress.Num();
-//        program.m_code.push_back(op);
-        program.m_opAddress.Add((uint32_t)program.m_byteCode.Num());
-        AppendCode(program.m_byteCode,OP_TYPE::BO_CONSTANT);
-        AppendCode(program.m_byteCode,args);
-    }
-}
+	void ASTOpConstantBool::Link(PROGRAM& program, const FLinkerOptions*)
+	{
+		if (!linkedAddress)
+		{
+			OP::BoolConstantArgs args;
+			memset(&args, 0, sizeof(args));
+			args.value = value;
 
+			linkedAddress = (OP::ADDRESS)program.m_opAddress.Num();
+			//        program.m_code.push_back(op);
+			program.m_opAddress.Add((uint32_t)program.m_byteCode.Num());
+			AppendCode(program.m_byteCode, OP_TYPE::BO_CONSTANT);
+			AppendCode(program.m_byteCode, args);
+		}
+	}
 
-ASTOp::BOOL_EVAL_RESULT ASTOpConstantBool::EvaluateBool( ASTOpList&, EVALUATE_BOOL_CACHE* ) const
-{
-    return value? BET_TRUE : BET_FALSE;
-}
 
+	ASTOp::BOOL_EVAL_RESULT ASTOpConstantBool::EvaluateBool(ASTOpList&, EVALUATE_BOOL_CACHE*) const
+	{
+		return value ? BET_TRUE : BET_FALSE;
+	}
 
+}
diff --git a/Engine/Plugins/Experimental/Mutable/Source/MutableTools/Private/MuT/ASTOpConstantResource.cpp b/Engine/Plugins/Experimental/Mutable/Source/MutableTools/Private/MuT/ASTOpConstantResource.cpp
--- a/Engine/Plugins/Experimental/Mutable/Source/MutableTools/Private/MuT/ASTOpConstantResource.cpp
+++ b/Engine/Plugins/Experimental/Mutable/Source/MutableTools/Private/MuT/ASTOpConstantResource.cpp
@@ -21,422 +21,423 @@
 #include <utility>
 
 
-using namespace mu;
+namespace mu
+{
 
+	//-------------------------------------------------------------------------------------------------
+	//-------------------------------------------------------------------------------------------------
+	//-------------------------------------------------------------------------------------------------
+	void ASTOpConstantResource::ForEachChild(const TFunctionRef<void(ASTChild&)>)
+	{
+	}
 
-//-------------------------------------------------------------------------------------------------
-//-------------------------------------------------------------------------------------------------
-//-------------------------------------------------------------------------------------------------
-void ASTOpConstantResource::ForEachChild(const TFunctionRef<void (ASTChild&)> )
-{
-}
 
+	//-------------------------------------------------------------------------------------------------
+	bool ASTOpConstantResource::IsEqual(const ASTOp& otherUntyped) const
+	{
+		if (auto other = dynamic_cast<const ASTOpConstantResource*>(&otherUntyped))
+		{
+			return type == other->type && hash == other->hash &&
+				loadedValue == other->loadedValue && proxy == other->proxy;
+		}
+		return false;
+	}
 
-//-------------------------------------------------------------------------------------------------
-bool ASTOpConstantResource::IsEqual(const ASTOp& otherUntyped) const
-{
-    if ( auto other = dynamic_cast<const ASTOpConstantResource*>(&otherUntyped) )
-    {
-        return type == other->type && hash == other->hash &&
-               loadedValue == other->loadedValue && proxy == other->proxy;
-    }
-    return false;
-}
 
+	//-------------------------------------------------------------------------------------------------
+	mu::Ptr<ASTOp> ASTOpConstantResource::Clone(MapChildFuncRef) const
+	{
+		Ptr<ASTOpConstantResource> n = new ASTOpConstantResource();
+		n->type = type;
+		n->proxy = proxy;
+		n->loadedValue = loadedValue;
+		n->hash = hash;
+		return n;
+	}
 
-//-------------------------------------------------------------------------------------------------
-mu::Ptr<ASTOp> ASTOpConstantResource::Clone(MapChildFuncRef) const
-{
-    Ptr<ASTOpConstantResource> n = new ASTOpConstantResource();
-    n->type = type;
-    n->proxy = proxy;
-    n->loadedValue = loadedValue;
-    n->hash = hash;
-    return n;
-}
 
+	//-------------------------------------------------------------------------------------------------
+	uint64 ASTOpConstantResource::Hash() const
+	{
+		uint64 res = std::hash<uint64>()(uint64(type));
+		hash_combine(res, hash);
+		return res;
+	}
 
-//-------------------------------------------------------------------------------------------------
-uint64 ASTOpConstantResource::Hash() const
-{
-	uint64 res = std::hash<uint64>()(uint64(type));
-    hash_combine( res, hash );
-    return res;
-}
 
+	//-------------------------------------------------------------------------------------------------
+	void ASTOpConstantResource::Link(PROGRAM& program, const FLinkerOptions* Options)
+	{
+		if (!linkedAddress && !bLinkedAndNull)
+		{
+			if (type == OP_TYPE::ME_CONSTANT)
+			{
+				OP::MeshConstantArgs args;
+				memset(&args, 0, sizeof(args));
 
-//-------------------------------------------------------------------------------------------------
-void ASTOpConstantResource::Link( PROGRAM& program, const FLinkerOptions* Options)
-{
-    if (!linkedAddress && !bLinkedAndNull)
-    {
-        if (type==OP_TYPE::ME_CONSTANT)
-        {
-            OP::MeshConstantArgs args;
-            memset( &args,0, sizeof(args) );
-
-            Ptr<Mesh> pTyped = static_cast<const Mesh*>(GetValue().get())->Clone();
-            check(pTyped);
-
-            args.skeleton = -1;
-            if (Ptr<const Skeleton> pSkeleton = pTyped->GetSkeleton())
-            {
-                // See if it can be merged into another skeleton
-//                for (const auto& pOtherSkeleton : program.m_constantSkeletons)
-//                {
-//                    if (true)
-//                    //if (AreSkeletonsRelated(pOtherSkeleton, pSkeleton))
-//                    {
-//                        // \todo : ugly const hack
-//                        ExtendSkeleton( const_cast<Skeleton*>(pOtherSkeleton.get()),
-//                                        pSkeleton.get() );
-//                        pSkeleton = pOtherSkeleton;
-
-//                        Ptr<Mesh> pRemapped = MeshRemapSkeleton( pTyped.get(),
-//                                                                 pOtherSkeleton.get() );
-//                        if (pRemapped)
-//                        {
-//                            pTyped = pRemapped;
-//                        }
-//                        break;
-//                    }
-//                }
-
-                args.skeleton = program.AddConstant( pSkeleton.get() );
-                pTyped->SetSkeleton( nullptr );
-            }
-
-			args.physicsBody = -1;
-			if (Ptr<const PhysicsBody> pPhysicsBody = pTyped->GetPhysicsBody())
+				Ptr<Mesh> pTyped = static_cast<const Mesh*>(GetValue().get())->Clone();
+				check(pTyped);
+
+				args.skeleton = -1;
+				if (Ptr<const Skeleton> pSkeleton = pTyped->GetSkeleton())
+				{
+					// See if it can be merged into another skeleton
+	//                for (const auto& pOtherSkeleton : program.m_constantSkeletons)
+	//                {
+	//                    if (true)
+	//                    //if (AreSkeletonsRelated(pOtherSkeleton, pSkeleton))
+	//                    {
+	//                        // \todo : ugly const hack
+	//                        ExtendSkeleton( const_cast<Skeleton*>(pOtherSkeleton.get()),
+	//                                        pSkeleton.get() );
+	//                        pSkeleton = pOtherSkeleton;
+
+	//                        Ptr<Mesh> pRemapped = MeshRemapSkeleton( pTyped.get(),
+	//                                                                 pOtherSkeleton.get() );
+	//                        if (pRemapped)
+	//                        {
+	//                            pTyped = pRemapped;
+	//                        }
+	//                        break;
+	//                    }
+	//                }
+
+					args.skeleton = program.AddConstant(pSkeleton.get());
+					pTyped->SetSkeleton(nullptr);
+				}
+
+				args.physicsBody = -1;
+				if (Ptr<const PhysicsBody> pPhysicsBody = pTyped->GetPhysicsBody())
+				{
+					args.physicsBody = program.AddConstant(pPhysicsBody.get());
+					pTyped->SetPhysicsBody(nullptr);
+				}
+
+				args.value = program.AddConstant(pTyped.get());
+
+				linkedAddress = (OP::ADDRESS)program.m_opAddress.Num();
+				program.m_opAddress.Add((uint32_t)program.m_byteCode.Num());
+				AppendCode(program.m_byteCode, type);
+				AppendCode(program.m_byteCode, args);
+			}
+			else
 			{
-				args.physicsBody = program.AddConstant( pPhysicsBody.get() );
-                pTyped->SetPhysicsBody( nullptr );
+				OP::ResourceConstantArgs args;
+				memset(&args, 0, sizeof(args));
+
+				bool bValidData = true;
+
+				switch (type)
+				{
+				case OP_TYPE::IM_CONSTANT:
+				{
+					Ptr<const Image> pTyped = static_cast<const Image*>(GetValue().get());
+					check(pTyped);
+
+					if (pTyped->GetSizeX() * pTyped->GetSizeY() == 0)
+					{
+						// It's an empty or degenerated image, return a null operation.
+						bValidData = false;
+					}
+					else
+					{
+						// TODO: Compilation options?
+						int32 MinTextureResidentMipCount = Options->MinTextureResidentMipCount;
+						args.value = program.AddConstant(pTyped, MinTextureResidentMipCount);
+					}
+					break;
+				}
+				case OP_TYPE::LA_CONSTANT:
+				{
+					Ptr<const Layout> pTyped = static_cast<const Layout*>(GetValue().get());
+					check(pTyped);
+					args.value = program.AddConstant(pTyped);
+					break;
+				}
+				default:
+					check(false);
+				}
+
+				if (bValidData)
+				{
+					linkedAddress = (OP::ADDRESS)program.m_opAddress.Num();
+					program.m_opAddress.Add((uint32_t)program.m_byteCode.Num());
+					AppendCode(program.m_byteCode, type);
+					AppendCode(program.m_byteCode, args);
+				}
+				else
+				{
+					// Null op
+					linkedAddress = 0;
+					bLinkedAndNull = true;
+				}
 			}
+		}
+	}
 
-            args.value = program.AddConstant(pTyped.get());
 
-            linkedAddress = (OP::ADDRESS)program.m_opAddress.Num();
-            program.m_opAddress.Add((uint32_t)program.m_byteCode.Num());
-            AppendCode(program.m_byteCode,type);
-            AppendCode(program.m_byteCode,args);
-        }
-        else
-        {
-            OP::ResourceConstantArgs args;
-            memset( &args,0, sizeof(args) );
+	//-------------------------------------------------------------------------------------------------
+	FImageDesc ASTOpConstantResource::GetImageDesc(bool, class GetImageDescContext*)
+	{
+		FImageDesc res;
+
+		if (type == OP_TYPE::IM_CONSTANT)
+		{
+			// TODO: cache to avoid disk loading
+			Ptr<const Image> pConst = static_cast<const Image*>(GetValue().get());
+			res.m_format = pConst->m_format;
+			res.m_lods = pConst->m_lods;
+			res.m_size = pConst->m_size;
+		}
+		else
+		{
+			check(false);
+		}
+
+		return res;
+	}
 
-			bool bValidData = true;
 
-            switch (type)
-            {
-            case OP_TYPE::IM_CONSTANT:
-            {
-                Ptr<const Image> pTyped = static_cast<const Image*>(GetValue().get());
-                check(pTyped);
+	//-------------------------------------------------------------------------------------------------
+	void ASTOpConstantResource::GetBlockLayoutSize(int blockIndex, int* pBlockX, int* pBlockY,
+		BLOCK_LAYOUT_SIZE_CACHE*)
+	{
+		switch (type)
+		{
+		case OP_TYPE::LA_CONSTANT:
+		{
+			Ptr<const Layout> pLayout = static_cast<const Layout*>(GetValue().get());
+			check(pLayout);
 
-				if (pTyped->GetSizeX()*pTyped->GetSizeY()==0)
+			if (pLayout)
+			{
+				int relId = pLayout->FindBlock(blockIndex);
+				if (relId >= 0)
 				{
-					// It's an empty or degenerated image, return a null operation.
-					bValidData = false;
+					*pBlockX = pLayout->m_blocks[relId].m_size[0];
+					*pBlockY = pLayout->m_blocks[relId].m_size[1];
 				}
 				else
 				{
-					// TODO: Compilation options?
-					int32 MinTextureResidentMipCount = Options->MinTextureResidentMipCount;
-					args.value = program.AddConstant(pTyped, MinTextureResidentMipCount);
+					*pBlockX = 0;
+					*pBlockY = 0;
 				}
-                break;
-            }
-            case OP_TYPE::LA_CONSTANT:
-            {
-                Ptr<const Layout> pTyped = static_cast<const Layout*>(GetValue().get());
-                check(pTyped);
-                args.value = program.AddConstant(pTyped);
-                break;
-            }
-            default:
-                check(false);
-            }
-
-			if (bValidData)
+			}
+
+			break;
+		}
+		default:
+			check(false);
+		}
+	}
+
+
+	//-------------------------------------------------------------------------------------------------
+	void ASTOpConstantResource::GetLayoutBlockSize(int* pBlockX, int* pBlockY)
+	{
+		switch (type)
+		{
+
+		case OP_TYPE::IM_CONSTANT:
+		{
+			// We didn't find any layout.
+			*pBlockX = 0;
+			*pBlockY = 0;
+			break;
+		}
+
+		default:
+			checkf(false, TEXT("Instruction not supported"));
+		}
+	}
+
+
+	//-------------------------------------------------------------------------------------------------
+	bool ASTOpConstantResource::GetNonBlackRect(FImageRect& maskUsage) const
+	{
+		if (type == OP_TYPE::IM_CONSTANT)
+		{
+			// TODO: cache
+			Ptr<const Image> pMask = static_cast<const Image*>(GetValue().get());
+			pMask->GetNonBlackRect(maskUsage);
+			return true;
+		}
+
+		return false;
+	}
+
+
+	//-------------------------------------------------------------------------------------------------
+	bool ASTOpConstantResource::IsImagePlainConstant(vec4<float>& colour) const
+	{
+		bool res = false;
+		switch (type)
+		{
+
+		case OP_TYPE::IM_CONSTANT:
+		{
+			Ptr<const Image> pImage = static_cast<const Image*>(GetValue().get());
+			check(pImage->m_size[0] > 0);
+			check(pImage->m_size[1] > 0);
+
+			if (pImage->m_flags & Image::IF_IS_PLAIN_COLOUR_VALID)
 			{
-				linkedAddress = (OP::ADDRESS)program.m_opAddress.Num();
-				program.m_opAddress.Add((uint32_t)program.m_byteCode.Num());
-				AppendCode(program.m_byteCode, type);
-				AppendCode(program.m_byteCode, args);
+				if (pImage->m_flags & Image::IF_IS_PLAIN_COLOUR)
+				{
+					res = true;
+					colour = pImage->Sample(vec2<float>(0, 0));
+				}
+				else
+				{
+					res = false;
+				}
 			}
 			else
 			{
-				// Null op
-				linkedAddress = 0;
-				bLinkedAndNull = true;
+				if (pImage->IsPlainColour(colour))
+				{
+					res = true;
+					pImage->m_flags |= Image::IF_IS_PLAIN_COLOUR;
+				}
+
+				pImage->m_flags |= Image::IF_IS_PLAIN_COLOUR_VALID;
 			}
-        }
-    }
-}
+			break;
+		}
 
+		default:
+			break;
+		}
 
-//-------------------------------------------------------------------------------------------------
-FImageDesc ASTOpConstantResource::GetImageDesc( bool, class GetImageDescContext* )
-{
-    FImageDesc res;
-
-    if (type==OP_TYPE::IM_CONSTANT)
-    {
-        // TODO: cache to avoid disk loading
-        Ptr<const Image> pConst = static_cast<const Image*>(GetValue().get());
-        res.m_format = pConst->m_format;
-		res.m_lods = pConst->m_lods;
-		res.m_size = pConst->m_size;
+		return res;
 	}
-    else
-    {
-        check(false);
-    }
 
-    return res;
-}
 
+	//-------------------------------------------------------------------------------------------------
+	ASTOpConstantResource::~ASTOpConstantResource()
+	{
+	}
 
-//-------------------------------------------------------------------------------------------------
-void ASTOpConstantResource::GetBlockLayoutSize( int blockIndex, int* pBlockX, int* pBlockY,
-                                                BLOCK_LAYOUT_SIZE_CACHE* )
-{
-    switch ( type )
-    {
-    case OP_TYPE::LA_CONSTANT:
-    {
-        Ptr<const Layout> pLayout = static_cast<const Layout*>( GetValue().get() );
-        check(pLayout);
-
-        if (pLayout)
-        {
-            int relId = pLayout->FindBlock( blockIndex );
-            if ( relId>=0 )
-            {
-                *pBlockX = pLayout->m_blocks[relId].m_size[0];
-                *pBlockY = pLayout->m_blocks[relId].m_size[1];
-            }
-            else
-            {
-                *pBlockX = 0;
-                *pBlockY = 0;
-            }
-        }
-
-        break;
-    }
-    default:
-        check(false);
-    }
-}
-
-
-//-------------------------------------------------------------------------------------------------
-void ASTOpConstantResource::GetLayoutBlockSize( int* pBlockX, int* pBlockY )
-{
-    switch ( type )
-    {
 
-    case OP_TYPE::IM_CONSTANT:
-    {
-        // We didn't find any layout.
-        *pBlockX = 0;
-        *pBlockY = 0;
-        break;
-    }
+	//-------------------------------------------------------------------------------------------------
+	uint64 ASTOpConstantResource::GetValueHash() const
+	{
+		return hash;
+	}
 
-    default:
-        checkf( false, TEXT("Instruction not supported") );
-    }
-}
 
+	//-------------------------------------------------------------------------------------------------
+	mu::Ptr<const RefCounted> ASTOpConstantResource::GetValue() const
+	{
+		if (loadedValue)
+		{
+			return loadedValue;
+		}
+		else
+		{
+			switch (type)
+			{
 
-//-------------------------------------------------------------------------------------------------
-bool ASTOpConstantResource::GetNonBlackRect( FImageRect& maskUsage ) const
-{
-    if (type==OP_TYPE::IM_CONSTANT)
-    {
-        // TODO: cache
-        Ptr<const Image> pMask = static_cast<const Image*>(GetValue().get());
-        pMask->GetNonBlackRect(maskUsage);
-        return true;
-    }
+			case OP_TYPE::IM_CONSTANT:
+			{
+				Ptr<ResourceProxy<Image>> typedProxy = dynamic_cast<ResourceProxy<Image>*>(proxy.get());
+				Ptr<const Image> r = typedProxy->Get();
+				return r;
+			}
+
+			default:
+				check(false);
+				break;
+			}
+		}
 
-    return false;
-}
+		return nullptr;
+	}
 
 
-//-------------------------------------------------------------------------------------------------
-bool ASTOpConstantResource::IsImagePlainConstant( vec4<float>& colour ) const
-{
-    bool res = false;
-    switch( type )
-    {
-
-    case OP_TYPE::IM_CONSTANT:
-    {
-        Ptr<const Image> pImage = static_cast<const Image*>(GetValue().get());
-		check(pImage->m_size[0] > 0);
-		check(pImage->m_size[1] > 0);
-
-        if ( pImage->m_flags & Image::IF_IS_PLAIN_COLOUR_VALID )
-        {
-            if ( pImage->m_flags & Image::IF_IS_PLAIN_COLOUR )
-            {
-                res = true;
-                colour = pImage->Sample( vec2<float>(0,0) );
-            }
-            else
-            {
-                res = false;
-            }
-        }
-        else
-        {
-            if ( pImage->IsPlainColour( colour ) )
-            {
-                res = true;
-                pImage->m_flags |= Image::IF_IS_PLAIN_COLOUR;
-            }
-
-            pImage->m_flags |= Image::IF_IS_PLAIN_COLOUR_VALID;
-        }
-        break;
-    }
-
-    default:
-        break;
-    }
-
-    return res;
-}
-
-
-//-------------------------------------------------------------------------------------------------
-ASTOpConstantResource::~ASTOpConstantResource()
-{
-}
+	//-------------------------------------------------------------------------------------------------
+	void ASTOpConstantResource::SetValue(const Ptr<const RefCounted>& v,
+		bool useDiskCache)
+	{
+		switch (type)
+		{
+		case OP_TYPE::IM_CONSTANT:
+		{
+			Ptr<const Image> r = static_cast<const Image*>(v.get());
 
+			OutputMemoryStream stream(r->GetDataSize() + 1024);
+			OutputArchive arch(&stream);
+			Image::Serialise(r.get(), arch);
 
-//-------------------------------------------------------------------------------------------------
-uint64 ASTOpConstantResource::GetValueHash() const
-{
-    return hash;
-}
+			hash = CityHash64(static_cast<const char*>(stream.GetBuffer()), stream.GetBufferSize());
+
+			if (useDiskCache)
+			{
+				proxy = new ResourceProxyTempFile<Image>(r.get());
+			}
+			else
+			{
+				loadedValue = r;
+			}
+			break;
+		}
 
+		case OP_TYPE::ME_CONSTANT:
+		{
+			OutputMemoryStream stream;
+			OutputArchive arch(&stream);
 
-//-------------------------------------------------------------------------------------------------
-mu::Ptr<const RefCounted> ASTOpConstantResource::GetValue() const
-{
-    if (loadedValue)
-    {
-        return loadedValue;
-    }
-    else
-    {
-        switch( type )
-        {
-
-        case OP_TYPE::IM_CONSTANT:
-        {
-            Ptr<ResourceProxy<Image>> typedProxy = dynamic_cast< ResourceProxy<Image>* >( proxy.get() );
-            Ptr<const Image> r = typedProxy->Get();
-            return r;
-        }
-
-        default:
-            check(false);
-            break;
-        }
-    }
-
-    return nullptr;
-}
-
-
-//-------------------------------------------------------------------------------------------------
-void ASTOpConstantResource::SetValue( const Ptr<const RefCounted>& v,
-                                      bool useDiskCache )
-{
-    switch( type )
-    {
-    case OP_TYPE::IM_CONSTANT:
-    {
-        Ptr<const Image> r = static_cast<const Image*>(v.get());
-
-        OutputMemoryStream stream(r->GetDataSize()+1024);
-        OutputArchive arch(&stream);
-        Image::Serialise( r.get(), arch );
-
-		hash = CityHash64(static_cast<const char*>(stream.GetBuffer()), stream.GetBufferSize());
-
-        if (useDiskCache)
-        {
-            proxy = new ResourceProxyTempFile<Image>(r.get());
-        }
-        else
-        {
-            loadedValue = r;
-        }
-        break;
-    }
-
-    case OP_TYPE::ME_CONSTANT:
-    {
-        OutputMemoryStream stream;
-        OutputArchive arch(&stream);
-
-        Ptr<const Mesh> r = static_cast<const Mesh*>(v.get());
-        Mesh::Serialise( r.get(), arch );
-
-		hash = CityHash64(static_cast<const char*>(stream.GetBuffer()), stream.GetBufferSize());
-
-        loadedValue = v;
-        break;
-    }
-
-    case OP_TYPE::LA_CONSTANT:
-    {
-        OutputMemoryStream stream;
-        OutputArchive arch(&stream);
-
-        Ptr<const Layout> r = static_cast<const Layout*>(v.get());
-        Layout::Serialise( r.get(), arch );
-
-		hash = CityHash64(static_cast<const char*>(stream.GetBuffer()), stream.GetBufferSize());
-
-        loadedValue = v;
-        break;
-    }
-
-    default:
-        loadedValue = v;
-        break;
-    }
-}
-
-
-//-------------------------------------------------------------------------------------------------
-mu::Ptr<ImageSizeExpression> ASTOpConstantResource::GetImageSizeExpression() const
-{
-    Ptr<ImageSizeExpression> pRes = new ImageSizeExpression;
-    pRes->type = ImageSizeExpression::ISET_CONSTANT;
-
-    switch( type )
-    {
-    case OP_TYPE::IM_CONSTANT:
-    {
-        Ptr<const Image> pConst = static_cast<const Image*>(GetValue().get());
-        pRes->size = pConst->m_size;
-        break;
-    }
-
-    default:
-        check( false );
-        return nullptr;
-    }
-
-    return pRes;
-}
+			Ptr<const Mesh> r = static_cast<const Mesh*>(v.get());
+			Mesh::Serialise(r.get(), arch);
+
+			hash = CityHash64(static_cast<const char*>(stream.GetBuffer()), stream.GetBufferSize());
+
+			loadedValue = v;
+			break;
+		}
+
+		case OP_TYPE::LA_CONSTANT:
+		{
+			OutputMemoryStream stream;
+			OutputArchive arch(&stream);
+
+			Ptr<const Layout> r = static_cast<const Layout*>(v.get());
+			Layout::Serialise(r.get(), arch);
+
+			hash = CityHash64(static_cast<const char*>(stream.GetBuffer()), stream.GetBufferSize());
+
+			loadedValue = v;
+			break;
+		}
+
+		default:
+			loadedValue = v;
+			break;
+		}
+	}
+
+
+	//-------------------------------------------------------------------------------------------------
+	mu::Ptr<ImageSizeExpression> ASTOpConstantResource::GetImageSizeExpression() const
+	{
+		Ptr<ImageSizeExpression> pRes = new ImageSizeExpression;
+		pRes->type = ImageSizeExpression::ISET_CONSTANT;
+
+		switch (type)
+		{
+		case OP_TYPE::IM_CONSTANT:
+		{
+			Ptr<const Image> pConst = static_cast<const Image*>(GetValue().get());
+			pRes->size = pConst->m_size;
+			break;
+		}
+
+		default:
+			check(false);
+			return nullptr;
+		}
+
+		return pRes;
+	}
 
+}
\ No newline at end of file
diff --git a/Engine/Plugins/Experimental/Mutable/Source/MutableTools/Private/MuT/ASTOpConstantString.cpp b/Engine/Plugins/Experimental/Mutable/Source/MutableTools/Private/MuT/ASTOpConstantString.cpp
--- a/Engine/Plugins/Experimental/Mutable/Source/MutableTools/Private/MuT/ASTOpConstantString.cpp
+++ b/Engine/Plugins/Experimental/Mutable/Source/MutableTools/Private/MuT/ASTOpConstantString.cpp
@@ -10,54 +10,56 @@
 #include <string>
 
 
-using namespace mu;
+namespace mu
+{
 
 
-//-------------------------------------------------------------------------------------------------
-//-------------------------------------------------------------------------------------------------
-//-------------------------------------------------------------------------------------------------
-void ASTOpConstantString::ForEachChild(const TFunctionRef<void(ASTChild&)> )
-{
-}
+	//-------------------------------------------------------------------------------------------------
+	//-------------------------------------------------------------------------------------------------
+	//-------------------------------------------------------------------------------------------------
+	void ASTOpConstantString::ForEachChild(const TFunctionRef<void(ASTChild&)>)
+	{
+	}
 
 
-uint64 ASTOpConstantString::Hash() const
-{
-	uint64 res = std::hash<std::string>()(value.c_str());
-    return res;
-}
+	uint64 ASTOpConstantString::Hash() const
+	{
+		uint64 res = std::hash<std::string>()(value.c_str());
+		return res;
+	}
 
 
-bool ASTOpConstantString::IsEqual(const ASTOp& otherUntyped) const
-{
-    if ( auto other = dynamic_cast<const ASTOpConstantString*>(&otherUntyped) )
-    {
-        return value==other->value;
-    }
-    return false;
-}
+	bool ASTOpConstantString::IsEqual(const ASTOp& otherUntyped) const
+	{
+		if (auto other = dynamic_cast<const ASTOpConstantString*>(&otherUntyped))
+		{
+			return value == other->value;
+		}
+		return false;
+	}
 
 
-mu::Ptr<ASTOp> ASTOpConstantString::Clone(MapChildFuncRef) const
-{
-    Ptr<ASTOpConstantString> n = new ASTOpConstantString();
-    n->value = value;
-    return n;
-}
+	mu::Ptr<ASTOp> ASTOpConstantString::Clone(MapChildFuncRef) const
+	{
+		Ptr<ASTOpConstantString> n = new ASTOpConstantString();
+		n->value = value;
+		return n;
+	}
 
 
-void ASTOpConstantString::Link( PROGRAM& program, const FLinkerOptions* )
-{
-    if (!linkedAddress)
-    {
-        OP::ResourceConstantArgs args;
-        memset( &args,0, sizeof(args) );
-        args.value = program.AddConstant(value);
-
-        linkedAddress = (OP::ADDRESS)program.m_opAddress.Num();
-        program.m_opAddress.Add((uint32_t)program.m_byteCode.Num());
-        AppendCode(program.m_byteCode,OP_TYPE::ST_CONSTANT);
-        AppendCode(program.m_byteCode,args);
-    }
-}
+	void ASTOpConstantString::Link(PROGRAM& program, const FLinkerOptions*)
+	{
+		if (!linkedAddress)
+		{
+			OP::ResourceConstantArgs args;
+			memset(&args, 0, sizeof(args));
+			args.value = program.AddConstant(value);
+
+			linkedAddress = (OP::ADDRESS)program.m_opAddress.Num();
+			program.m_opAddress.Add((uint32_t)program.m_byteCode.Num());
+			AppendCode(program.m_byteCode, OP_TYPE::ST_CONSTANT);
+			AppendCode(program.m_byteCode, args);
+		}
+	}
 
+}
\ No newline at end of file
diff --git a/Engine/Plugins/Experimental/Mutable/Source/MutableTools/Private/MuT/ASTOpImageCompose.cpp b/Engine/Plugins/Experimental/Mutable/Source/MutableTools/Private/MuT/ASTOpImageCompose.cpp
--- a/Engine/Plugins/Experimental/Mutable/Source/MutableTools/Private/MuT/ASTOpImageCompose.cpp
+++ b/Engine/Plugins/Experimental/Mutable/Source/MutableTools/Private/MuT/ASTOpImageCompose.cpp
@@ -20,7 +20,8 @@
 #include <utility>
 
 
-using namespace mu;
+namespace mu
+{
 
 
 //-------------------------------------------------------------------------------------------------
@@ -281,3 +282,5 @@ mu::Ptr<ASTOp> ASTOpImageCompose::OptimiseSemantic(const MODEL_OPTIMIZATION_OPTI
 
 	return at;
 }
+
+}
\ No newline at end of file
diff --git a/Engine/Plugins/Experimental/Mutable/Source/MutableTools/Private/MuT/ASTOpImageMipmap.cpp b/Engine/Plugins/Experimental/Mutable/Source/MutableTools/Private/MuT/ASTOpImageMipmap.cpp
--- a/Engine/Plugins/Experimental/Mutable/Source/MutableTools/Private/MuT/ASTOpImageMipmap.cpp
+++ b/Engine/Plugins/Experimental/Mutable/Source/MutableTools/Private/MuT/ASTOpImageMipmap.cpp
@@ -23,12 +23,9 @@
 #include <utility>
 
 
-using namespace mu;
-
+namespace mu
+{
 
-//-------------------------------------------------------------------------------------------------
-//-------------------------------------------------------------------------------------------------
-//-------------------------------------------------------------------------------------------------
 ASTOpImageMipmap::ASTOpImageMipmap()
     : Source(this)
 {
@@ -458,3 +455,4 @@ mu::Ptr<ASTOp> Sink_ImageMipmapAST::Visit(mu::Ptr<ASTOp> at, const ASTOpImageMip
 	return newAt;
 }
 
+}
\ No newline at end of file
diff --git a/Engine/Plugins/Experimental/Mutable/Source/MutableTools/Private/MuT/ASTOpImageMultiLayer.cpp b/Engine/Plugins/Experimental/Mutable/Source/MutableTools/Private/MuT/ASTOpImageMultiLayer.cpp
--- a/Engine/Plugins/Experimental/Mutable/Source/MutableTools/Private/MuT/ASTOpImageMultiLayer.cpp
+++ b/Engine/Plugins/Experimental/Mutable/Source/MutableTools/Private/MuT/ASTOpImageMultiLayer.cpp
@@ -14,164 +14,163 @@
 #include <utility>
 
 
-using namespace mu;
-
-
-//-------------------------------------------------------------------------------------------------
-//-------------------------------------------------------------------------------------------------
-//-------------------------------------------------------------------------------------------------
-ASTOpImageMultiLayer::ASTOpImageMultiLayer()
-    : base(this)
-    , blend(this)
-    , mask(this)
-    , range(this,nullptr,string(),string())
-	, blendType(EBlendType::BT_BLEND)
+namespace mu
 {
-}
 
 
-//-------------------------------------------------------------------------------------------------
-ASTOpImageMultiLayer::~ASTOpImageMultiLayer()
-{
-    // Explicit call needed to avoid recursive destruction
-    ASTOp::RemoveChildren();
-}
+	ASTOpImageMultiLayer::ASTOpImageMultiLayer()
+		: base(this)
+		, blend(this)
+		, mask(this)
+		, range(this, nullptr, string(), string())
+		, blendType(EBlendType::BT_BLEND)
+	{
+	}
 
 
-//-------------------------------------------------------------------------------------------------
-bool ASTOpImageMultiLayer::IsEqual(const ASTOp& otherUntyped) const
-{
-    if ( auto other = dynamic_cast<const ASTOpImageMultiLayer*>(&otherUntyped) )
-    {
-        return base==other->base &&
-                blend==other->blend &&
-                mask==other->mask &&
-                range==other->range &&
-                blendType==other->blendType;
-    }
-    return false;
-}
+	//-------------------------------------------------------------------------------------------------
+	ASTOpImageMultiLayer::~ASTOpImageMultiLayer()
+	{
+		// Explicit call needed to avoid recursive destruction
+		ASTOp::RemoveChildren();
+	}
 
 
-//-------------------------------------------------------------------------------------------------
-uint64 ASTOpImageMultiLayer::Hash() const
-{
-	uint64 res = std::hash<OP_TYPE>()(OP_TYPE::IM_MULTILAYER);
-    hash_combine( res, base.child().get() );
-    hash_combine( res, blend.child().get() );
-    hash_combine( res, mask.child().get() );
-    return res;
-}
+	//-------------------------------------------------------------------------------------------------
+	bool ASTOpImageMultiLayer::IsEqual(const ASTOp& otherUntyped) const
+	{
+		if (auto other = dynamic_cast<const ASTOpImageMultiLayer*>(&otherUntyped))
+		{
+			return base == other->base &&
+				blend == other->blend &&
+				mask == other->mask &&
+				range == other->range &&
+				blendType == other->blendType;
+		}
+		return false;
+	}
 
 
-//-------------------------------------------------------------------------------------------------
-mu::Ptr<ASTOp> ASTOpImageMultiLayer::Clone(MapChildFuncRef mapChild) const
-{
-    Ptr<ASTOpImageMultiLayer> n = new ASTOpImageMultiLayer();
-    n->base = mapChild(base.child());
-    n->blend = mapChild(blend.child());
-    n->mask = mapChild(mask.child());
-    n->range.rangeName = range.rangeName;
-    n->range.rangeUID = range.rangeUID;
-    n->range.rangeSize = mapChild(range.rangeSize.child());
-    n->blendType = blendType;
-    return n;
-}
+	//-------------------------------------------------------------------------------------------------
+	uint64 ASTOpImageMultiLayer::Hash() const
+	{
+		uint64 res = std::hash<OP_TYPE>()(OP_TYPE::IM_MULTILAYER);
+		hash_combine(res, base.child().get());
+		hash_combine(res, blend.child().get());
+		hash_combine(res, mask.child().get());
+		return res;
+	}
 
 
-//-------------------------------------------------------------------------------------------------
-void ASTOpImageMultiLayer::ForEachChild( const TFunctionRef<void(ASTChild&)> f )
-{
-    f( base );
-    f( blend );
-    f( mask );
-    f( range.rangeSize );
-}
+	//-------------------------------------------------------------------------------------------------
+	mu::Ptr<ASTOp> ASTOpImageMultiLayer::Clone(MapChildFuncRef mapChild) const
+	{
+		Ptr<ASTOpImageMultiLayer> n = new ASTOpImageMultiLayer();
+		n->base = mapChild(base.child());
+		n->blend = mapChild(blend.child());
+		n->mask = mapChild(mask.child());
+		n->range.rangeName = range.rangeName;
+		n->range.rangeUID = range.rangeUID;
+		n->range.rangeSize = mapChild(range.rangeSize.child());
+		n->blendType = blendType;
+		return n;
+	}
 
 
-//-------------------------------------------------------------------------------------------------
-void ASTOpImageMultiLayer::Link( PROGRAM& program, const FLinkerOptions*)
-{
-    // Already linked?
-    if (!linkedAddress)
-    {
-        OP::ImageMultiLayerArgs args;
-        memset( &args,0, sizeof(args) );
-
-        args.blendType = (uint16)blendType;
-
-        if (base) args.base = base->linkedAddress;
-        if (blend) args.blended = blend->linkedAddress;
-        if (mask) args.mask = mask->linkedAddress;
-        if (range.rangeSize)
-        {
-            LinkRange( program, range, args.rangeSize, args.rangeId );
-        }
-
-        linkedAddress = (OP::ADDRESS)program.m_opAddress.Num();
-        program.m_opAddress.Add((uint32_t)program.m_byteCode.Num());
-        AppendCode(program.m_byteCode,OP_TYPE::IM_MULTILAYER);
-        AppendCode(program.m_byteCode,args);
-    }
-}
+	//-------------------------------------------------------------------------------------------------
+	void ASTOpImageMultiLayer::ForEachChild(const TFunctionRef<void(ASTChild&)> f)
+	{
+		f(base);
+		f(blend);
+		f(mask);
+		f(range.rangeSize);
+	}
 
 
-//-------------------------------------------------------------------------------------------------
-FImageDesc ASTOpImageMultiLayer::GetImageDesc( bool returnBestOption, GetImageDescContext* context )
-{
-    FImageDesc res;
-
-    // Local context in case it is necessary
-    GetImageDescContext localContext;
-    if (!context)
-    {
-      context = &localContext;
-    }
-    else
-    {
-        // Cached result?
-		FImageDesc* PtrValue = context->m_results.Find(this);
-		if (PtrValue)
+	//-------------------------------------------------------------------------------------------------
+	void ASTOpImageMultiLayer::Link(PROGRAM& program, const FLinkerOptions*)
+	{
+		// Already linked?
+		if (!linkedAddress)
 		{
-			return *PtrValue;
+			OP::ImageMultiLayerArgs args;
+			memset(&args, 0, sizeof(args));
+
+			args.blendType = (uint16)blendType;
+
+			if (base) args.base = base->linkedAddress;
+			if (blend) args.blended = blend->linkedAddress;
+			if (mask) args.mask = mask->linkedAddress;
+			if (range.rangeSize)
+			{
+				LinkRange(program, range, args.rangeSize, args.rangeId);
+			}
+
+			linkedAddress = (OP::ADDRESS)program.m_opAddress.Num();
+			program.m_opAddress.Add((uint32_t)program.m_byteCode.Num());
+			AppendCode(program.m_byteCode, OP_TYPE::IM_MULTILAYER);
+			AppendCode(program.m_byteCode, args);
 		}
 	}
 
-    // Actual work
-    if (base)
-    {
-        res = base->GetImageDesc( returnBestOption, context );
-    }
 
+	//-------------------------------------------------------------------------------------------------
+	FImageDesc ASTOpImageMultiLayer::GetImageDesc(bool returnBestOption, GetImageDescContext* context)
+	{
+		FImageDesc res;
 
-    // Cache the result
-    if (context)
-    {
-		context->m_results.Add(this, res);
-	}
+		// Local context in case it is necessary
+		GetImageDescContext localContext;
+		if (!context)
+		{
+			context = &localContext;
+		}
+		else
+		{
+			// Cached result?
+			FImageDesc* PtrValue = context->m_results.Find(this);
+			if (PtrValue)
+			{
+				return *PtrValue;
+			}
+		}
 
-    return res;
-}
+		// Actual work
+		if (base)
+		{
+			res = base->GetImageDesc(returnBestOption, context);
+		}
 
 
-//-------------------------------------------------------------------------------------------------
-void ASTOpImageMultiLayer::GetLayoutBlockSize(int* pBlockX, int* pBlockY)
-{
-	if (base)
+		// Cache the result
+		if (context)
+		{
+			context->m_results.Add(this, res);
+		}
+
+		return res;
+	}
+
+
+	//-------------------------------------------------------------------------------------------------
+	void ASTOpImageMultiLayer::GetLayoutBlockSize(int* pBlockX, int* pBlockY)
 	{
-		base->GetLayoutBlockSize(pBlockX, pBlockY);
+		if (base)
+		{
+			base->GetLayoutBlockSize(pBlockX, pBlockY);
+		}
 	}
-}
 
 
-//-------------------------------------------------------------------------------------------------
-mu::Ptr<ImageSizeExpression> ASTOpImageMultiLayer::GetImageSizeExpression() const
-{
-    if ( base )
-    {
-        return base->GetImageSizeExpression();
-    }
+	//-------------------------------------------------------------------------------------------------
+	mu::Ptr<ImageSizeExpression> ASTOpImageMultiLayer::GetImageSizeExpression() const
+	{
+		if (base)
+		{
+			return base->GetImageSizeExpression();
+		}
 
-    return nullptr;
-}
+		return nullptr;
+	}
 
+}
diff --git a/Engine/Plugins/Experimental/Mutable/Source/MutableTools/Private/MuT/ASTOpImageNormalComposite.cpp b/Engine/Plugins/Experimental/Mutable/Source/MutableTools/Private/MuT/ASTOpImageNormalComposite.cpp
--- a/Engine/Plugins/Experimental/Mutable/Source/MutableTools/Private/MuT/ASTOpImageNormalComposite.cpp
+++ b/Engine/Plugins/Experimental/Mutable/Source/MutableTools/Private/MuT/ASTOpImageNormalComposite.cpp
@@ -12,141 +12,144 @@
 #include "MuT/StreamsPrivate.h"
 
 
-using namespace mu;
+namespace mu
+{
 
 
-ASTOpImageNormalComposite::ASTOpImageNormalComposite()
-    : Base(this)
-    , Normal(this)
-    , Mode(ECompositeImageMode::CIM_Disabled)
-    , Power(1.0f)
-{
-}
+	ASTOpImageNormalComposite::ASTOpImageNormalComposite()
+		: Base(this)
+		, Normal(this)
+		, Mode(ECompositeImageMode::CIM_Disabled)
+		, Power(1.0f)
+	{
+	}
 
 
-ASTOpImageNormalComposite::~ASTOpImageNormalComposite()
-{
-    // Explicit call needed to avoid recursive destruction
-    ASTOp::RemoveChildren();
-}
+	ASTOpImageNormalComposite::~ASTOpImageNormalComposite()
+	{
+		// Explicit call needed to avoid recursive destruction
+		ASTOp::RemoveChildren();
+	}
 
 
-bool ASTOpImageNormalComposite::IsEqual(const ASTOp& OtherUntyped) const
-{
-    if ( const ASTOpImageNormalComposite* Other = dynamic_cast<const ASTOpImageNormalComposite*>( &OtherUntyped) )
-    {
-        return Base == Other->Base && Normal == Other->Normal && Power == Other->Power && Mode == Other->Mode;
-    }
+	bool ASTOpImageNormalComposite::IsEqual(const ASTOp& OtherUntyped) const
+	{
+		if (const ASTOpImageNormalComposite* Other = dynamic_cast<const ASTOpImageNormalComposite*>(&OtherUntyped))
+		{
+			return Base == Other->Base && Normal == Other->Normal && Power == Other->Power && Mode == Other->Mode;
+		}
 
-    return false;
-}
+		return false;
+	}
 
 
-uint64 ASTOpImageNormalComposite::Hash() const
-{
-	uint64 Res = std::hash<void*>()(Base.child().get() );
-	hash_combine(Res, Normal.child().get());
+	uint64 ASTOpImageNormalComposite::Hash() const
+	{
+		uint64 Res = std::hash<void*>()(Base.child().get());
+		hash_combine(Res, Normal.child().get());
 
-	return Res;
-}
+		return Res;
+	}
 
 
-mu::Ptr<ASTOp> ASTOpImageNormalComposite::Clone(MapChildFuncRef mapChild) const
-{
-    Ptr<ASTOpImageNormalComposite> N = new ASTOpImageNormalComposite();
-	N->Base = mapChild(Base.child());
-	N->Normal = mapChild(Normal.child());
-    N->Mode = Mode;
-    N->Power = Power;
+	mu::Ptr<ASTOp> ASTOpImageNormalComposite::Clone(MapChildFuncRef mapChild) const
+	{
+		Ptr<ASTOpImageNormalComposite> N = new ASTOpImageNormalComposite();
+		N->Base = mapChild(Base.child());
+		N->Normal = mapChild(Normal.child());
+		N->Mode = Mode;
+		N->Power = Power;
 
-    return N;
-}
+		return N;
+	}
 
 
-void ASTOpImageNormalComposite::ForEachChild(const TFunctionRef<void(ASTChild&)> f )
-{
-	f( Base );
-	f( Normal );
-}
+	void ASTOpImageNormalComposite::ForEachChild(const TFunctionRef<void(ASTChild&)> f)
+	{
+		f(Base);
+		f(Normal);
+	}
 
 
-void ASTOpImageNormalComposite::Link( PROGRAM& program, const FLinkerOptions* )
-{
-    // Already linked?
-    if (!linkedAddress)
-    {
-        OP::ImageNormalCompositeArgs args;
-        FMemory::Memset( &args, 0, sizeof(args) );
-
-		if (Base) 
-        {
-            args.base = Base->linkedAddress;
-        }
-
-        if (Normal) 
-        {
-            args.normal = Normal->linkedAddress;
-        }
-
-		args.power = Power;
-		args.mode = Mode;
-
-        linkedAddress = (OP::ADDRESS)program.m_opAddress.Num();
-        program.m_opAddress.Add((uint32_t)program.m_byteCode.Num());
-        AppendCode(program.m_byteCode,OP_TYPE::IM_NORMALCOMPOSITE);
-        AppendCode(program.m_byteCode,args);
-    }
+	void ASTOpImageNormalComposite::Link(PROGRAM& program, const FLinkerOptions*)
+	{
+		// Already linked?
+		if (!linkedAddress)
+		{
+			OP::ImageNormalCompositeArgs args;
+			FMemory::Memset(&args, 0, sizeof(args));
+
+			if (Base)
+			{
+				args.base = Base->linkedAddress;
+			}
+
+			if (Normal)
+			{
+				args.normal = Normal->linkedAddress;
+			}
+
+			args.power = Power;
+			args.mode = Mode;
+
+			linkedAddress = (OP::ADDRESS)program.m_opAddress.Num();
+			program.m_opAddress.Add((uint32_t)program.m_byteCode.Num());
+			AppendCode(program.m_byteCode, OP_TYPE::IM_NORMALCOMPOSITE);
+			AppendCode(program.m_byteCode, args);
+		}
 
-}
+	}
 
 
-FImageDesc ASTOpImageNormalComposite::GetImageDesc( bool returnBestOption, GetImageDescContext* context )
-{
-    FImageDesc res;
-
-    // Local context in case it is necessary
-    GetImageDescContext localContext;
-    if (!context)
-    {
-      context = &localContext;
-    }
-    else
-    {
-        // Cached result?
-		FImageDesc* PtrValue = context->m_results.Find(this);
-		if (PtrValue)
+	FImageDesc ASTOpImageNormalComposite::GetImageDesc(bool returnBestOption, GetImageDescContext* context)
+	{
+		FImageDesc res;
+
+		// Local context in case it is necessary
+		GetImageDescContext localContext;
+		if (!context)
 		{
-			return *PtrValue;
+			context = &localContext;
+		}
+		else
+		{
+			// Cached result?
+			FImageDesc* PtrValue = context->m_results.Find(this);
+			if (PtrValue)
+			{
+				return *PtrValue;
+			}
 		}
-	}
 
-    if (Base)
-    {
-        res = Base->GetImageDesc( returnBestOption, context );
-    }
+		if (Base)
+		{
+			res = Base->GetImageDesc(returnBestOption, context);
+		}
 
-    if (context)
-    {
-        context->m_results.Add(this, res);
-    }
+		if (context)
+		{
+			context->m_results.Add(this, res);
+		}
 
-    return res;
-}
+		return res;
+	}
 
-void ASTOpImageNormalComposite::GetLayoutBlockSize(int* pBlockX, int* pBlockY)
-{
-    if (Base)
-    {
-        Base->GetLayoutBlockSize(pBlockX, pBlockY);
-    }
-}
+	void ASTOpImageNormalComposite::GetLayoutBlockSize(int* pBlockX, int* pBlockY)
+	{
+		if (Base)
+		{
+			Base->GetLayoutBlockSize(pBlockX, pBlockY);
+		}
+	}
 
-mu::Ptr<ImageSizeExpression> ASTOpImageNormalComposite::GetImageSizeExpression() const 
-{
-    if (Base)
-    {
-        return Base->GetImageSizeExpression();
-    }
+	mu::Ptr<ImageSizeExpression> ASTOpImageNormalComposite::GetImageSizeExpression() const
+	{
+		if (Base)
+		{
+			return Base->GetImageSizeExpression();
+		}
+
+		return nullptr;
+	}
 
-    return nullptr;
 }
diff --git a/Engine/Plugins/Experimental/Mutable/Source/MutableTools/Private/MuT/ASTOpImagePatch.cpp b/Engine/Plugins/Experimental/Mutable/Source/MutableTools/Private/MuT/ASTOpImagePatch.cpp
--- a/Engine/Plugins/Experimental/Mutable/Source/MutableTools/Private/MuT/ASTOpImagePatch.cpp
+++ b/Engine/Plugins/Experimental/Mutable/Source/MutableTools/Private/MuT/ASTOpImagePatch.cpp
@@ -10,138 +10,137 @@
 #include "MuT/StreamsPrivate.h"
 
 
-using namespace mu;
+namespace mu
+{
 
+	ASTOpImagePatch::ASTOpImagePatch()
+		: base(this)
+		, patch(this)
+	{
+	}
 
-//-------------------------------------------------------------------------------------------------
-//-------------------------------------------------------------------------------------------------
-//-------------------------------------------------------------------------------------------------
-ASTOpImagePatch::ASTOpImagePatch()
-    : base(this)
-    , patch(this)
-{
-}
 
+	//-------------------------------------------------------------------------------------------------
+	ASTOpImagePatch::~ASTOpImagePatch()
+	{
+		// Explicit call needed to avoid recursive destruction
+		ASTOp::RemoveChildren();
+	}
 
-//-------------------------------------------------------------------------------------------------
-ASTOpImagePatch::~ASTOpImagePatch()
-{
-    // Explicit call needed to avoid recursive destruction
-    ASTOp::RemoveChildren();
-}
 
+	//-------------------------------------------------------------------------------------------------
+	bool ASTOpImagePatch::IsEqual(const ASTOp& otherUntyped) const
+	{
+		if (auto other = dynamic_cast<const ASTOpImagePatch*>(&otherUntyped))
+		{
+			return base == other->base &&
+				patch == other->patch &&
+				location == other->location;
+		}
+		return false;
+	}
 
-//-------------------------------------------------------------------------------------------------
-bool ASTOpImagePatch::IsEqual(const ASTOp& otherUntyped) const
-{
-    if ( auto other = dynamic_cast<const ASTOpImagePatch*>(&otherUntyped) )
-    {
-        return base==other->base &&
-                patch==other->patch &&
-                location==other->location;
-    }
-    return false;
-}
 
+	//-------------------------------------------------------------------------------------------------
+	uint64 ASTOpImagePatch::Hash() const
+	{
+		uint64 res = std::hash<OP_TYPE>()(OP_TYPE::IM_PATCH);
+		hash_combine(res, base.child().get());
+		hash_combine(res, patch.child().get());
+		return res;
+	}
 
-//-------------------------------------------------------------------------------------------------
-uint64 ASTOpImagePatch::Hash() const
-{
-	uint64 res = std::hash<OP_TYPE>()(OP_TYPE::IM_PATCH);
-    hash_combine( res, base.child().get() );
-    hash_combine( res, patch.child().get() );
-    return res;
-}
 
+	//-------------------------------------------------------------------------------------------------
+	mu::Ptr<ASTOp> ASTOpImagePatch::Clone(MapChildFuncRef mapChild) const
+	{
+		Ptr<ASTOpImagePatch> n = new ASTOpImagePatch();
+		n->base = mapChild(base.child());
+		n->patch = mapChild(patch.child());
+		n->location = location;
+		return n;
+	}
 
-//-------------------------------------------------------------------------------------------------
-mu::Ptr<ASTOp> ASTOpImagePatch::Clone(MapChildFuncRef mapChild) const
-{
-    Ptr<ASTOpImagePatch> n = new ASTOpImagePatch();
-    n->base = mapChild(base.child());
-    n->patch = mapChild(patch.child());
-    n->location = location;
-    return n;
-}
 
+	//-------------------------------------------------------------------------------------------------
+	void ASTOpImagePatch::ForEachChild(const TFunctionRef<void(ASTChild&)> f)
+	{
+		f(base);
+		f(patch);
+	}
 
-//-------------------------------------------------------------------------------------------------
-void ASTOpImagePatch::ForEachChild(const TFunctionRef<void(ASTChild&)> f )
-{
-    f( base );
-    f( patch );
-}
 
+	//-------------------------------------------------------------------------------------------------
+	void ASTOpImagePatch::Link(PROGRAM& program, const FLinkerOptions*)
+	{
+		// Already linked?
+		if (!linkedAddress)
+		{
+			OP::ImagePatchArgs args;
+			memset(&args, 0, sizeof(args));
+
+			if (base) args.base = base->linkedAddress;
+			if (patch) args.patch = patch->linkedAddress;
+			args.minX = location[0];
+			args.minY = location[1];
+
+			linkedAddress = (OP::ADDRESS)program.m_opAddress.Num();
+			program.m_opAddress.Add((uint32_t)program.m_byteCode.Num());
+			AppendCode(program.m_byteCode, OP_TYPE::IM_PATCH);
+			AppendCode(program.m_byteCode, args);
+		}
 
-//-------------------------------------------------------------------------------------------------
-void ASTOpImagePatch::Link( PROGRAM& program, const FLinkerOptions*)
-{
-    // Already linked?
-    if (!linkedAddress)
-    {
-        OP::ImagePatchArgs args;
-        memset( &args,0, sizeof(args) );
-
-        if (base) args.base = base->linkedAddress;
-        if (patch) args.patch = patch->linkedAddress;
-        args.minX = location[0];
-        args.minY = location[1];
-
-        linkedAddress = (OP::ADDRESS)program.m_opAddress.Num();
-        program.m_opAddress.Add((uint32_t)program.m_byteCode.Num());
-        AppendCode(program.m_byteCode,OP_TYPE::IM_PATCH);
-        AppendCode(program.m_byteCode,args);
-    }
+	}
 
-}
 
+	//-------------------------------------------------------------------------------------------------
+	FImageDesc ASTOpImagePatch::GetImageDesc(bool returnBestOption, GetImageDescContext* context)
+	{
+		FImageDesc res;
 
-//-------------------------------------------------------------------------------------------------
-FImageDesc ASTOpImagePatch::GetImageDesc( bool returnBestOption, GetImageDescContext* context )
-{
-    FImageDesc res;
-
-    // Local context in case it is necessary
-    GetImageDescContext localContext;
-    if (!context)
-    {
-      context = &localContext;
-    }
-    else
-    {
-        // Cached result?
-		FImageDesc* PtrValue = context->m_results.Find(this);
-		if (PtrValue)
+		// Local context in case it is necessary
+		GetImageDescContext localContext;
+		if (!context)
 		{
-			return *PtrValue;
+			context = &localContext;
+		}
+		else
+		{
+			// Cached result?
+			FImageDesc* PtrValue = context->m_results.Find(this);
+			if (PtrValue)
+			{
+				return *PtrValue;
+			}
 		}
-	}
 
-    // Actual work
-    if (base)
-    {
-        res = base->GetImageDesc( returnBestOption, context );
-    }
+		// Actual work
+		if (base)
+		{
+			res = base->GetImageDesc(returnBestOption, context);
+		}
 
 
-    // Cache the result
-    if (context)
-    {
-		context->m_results.Add(this, res);
+		// Cache the result
+		if (context)
+		{
+			context->m_results.Add(this, res);
+		}
+
+		return res;
 	}
 
-    return res;
-}
 
+	//-------------------------------------------------------------------------------------------------
+	mu::Ptr<ImageSizeExpression> ASTOpImagePatch::GetImageSizeExpression() const
+	{
+		if (base)
+		{
+			return base->GetImageSizeExpression();
+		}
+
+		return nullptr;
+	}
 
-//-------------------------------------------------------------------------------------------------
-mu::Ptr<ImageSizeExpression> ASTOpImagePatch::GetImageSizeExpression() const
-{
-    if ( base )
-    {
-        return base->GetImageSizeExpression();
-    }
 
-    return nullptr;
 }
-
diff --git a/Engine/Plugins/Experimental/Mutable/Source/MutableTools/Private/MuT/ASTOpImagePixelFormat.cpp b/Engine/Plugins/Experimental/Mutable/Source/MutableTools/Private/MuT/ASTOpImagePixelFormat.cpp
--- a/Engine/Plugins/Experimental/Mutable/Source/MutableTools/Private/MuT/ASTOpImagePixelFormat.cpp
+++ b/Engine/Plugins/Experimental/Mutable/Source/MutableTools/Private/MuT/ASTOpImagePixelFormat.cpp
@@ -21,11 +21,10 @@
 #include <memory>
 #include <utility>
 
-namespace mu { template <class SCALAR> class vec4; }
-
-
-using namespace mu;
+namespace mu
+{
 
+	template <class SCALAR> class vec4;
 
 //-------------------------------------------------------------------------------------------------
 //-------------------------------------------------------------------------------------------------
@@ -530,3 +529,5 @@ mu::Ptr<ASTOp> Sink_ImagePixelFormatAST::Visit(mu::Ptr<ASTOp> at, const ASTOpIma
 	return newAt;
 }
 
+
+}
