[mutable] Fix morphs not having any effect.

diff --git a/Engine/Plugins/Experimental/Mutable/Source/MutableTools/Private/MuT/ASTOpImageTransform.cpp b/Engine/Plugins/Experimental/Mutable/Source/MutableTools/Private/MuT/ASTOpImageTransform.cpp
--- a/Engine/Plugins/Experimental/Mutable/Source/MutableTools/Private/MuT/ASTOpImageTransform.cpp
+++ b/Engine/Plugins/Experimental/Mutable/Source/MutableTools/Private/MuT/ASTOpImageTransform.cpp
@@ -14,165 +14,164 @@
 #include <utility>
 
 
-using namespace mu;
-
-
-//-------------------------------------------------------------------------------------------------
-//-------------------------------------------------------------------------------------------------
-//-------------------------------------------------------------------------------------------------
-ASTOpImageTransform::ASTOpImageTransform()
-    : base(this)
-    , offsetX(this)
-    , offsetY(this)
-    , scaleX(this)
-    , scaleY(this)
-    , rotation(this)
+namespace mu
 {
-}
 
+	//-------------------------------------------------------------------------------------------------
+	ASTOpImageTransform::ASTOpImageTransform()
+		: base(this)
+		, offsetX(this)
+		, offsetY(this)
+		, scaleX(this)
+		, scaleY(this)
+		, rotation(this)
+	{
+	}
 
-//-------------------------------------------------------------------------------------------------
-ASTOpImageTransform::~ASTOpImageTransform()
-{
-    // Explicit call needed to avoid recursive destruction
-    ASTOp::RemoveChildren();
-}
 
+	//-------------------------------------------------------------------------------------------------
+	ASTOpImageTransform::~ASTOpImageTransform()
+	{
+		// Explicit call needed to avoid recursive destruction
+		ASTOp::RemoveChildren();
+	}
 
-//-------------------------------------------------------------------------------------------------
-bool ASTOpImageTransform::IsEqual(const ASTOp& otherUntyped) const
-{
-    if ( auto other = dynamic_cast<const ASTOpImageTransform*>(&otherUntyped) )
-    {
-        return base == other->base &&
-			offsetX == other->offsetX &&
-            offsetY == other->offsetY &&
-            scaleX == other->scaleX &&
-            scaleY == other->scaleY &&
-            rotation == other->rotation;
-    }
-    return false;
-}
 
+	//-------------------------------------------------------------------------------------------------
+	bool ASTOpImageTransform::IsEqual(const ASTOp& otherUntyped) const
+	{
+		if (auto other = dynamic_cast<const ASTOpImageTransform*>(&otherUntyped))
+		{
+			return base == other->base &&
+				offsetX == other->offsetX &&
+				offsetY == other->offsetY &&
+				scaleX == other->scaleX &&
+				scaleY == other->scaleY &&
+				rotation == other->rotation;
+		}
+		return false;
+	}
 
-//-------------------------------------------------------------------------------------------------
-uint64 ASTOpImageTransform::Hash() const
-{
-	uint64 res = std::hash<OP_TYPE>()(OP_TYPE::IM_MULTILAYER);
-    hash_combine( res, base.child().get() );
-    hash_combine( res, offsetX.child().get() );
-    hash_combine( res, offsetY.child().get() );
-    hash_combine( res, scaleX.child().get() );
-    hash_combine( res, scaleY.child().get() );
-    hash_combine( res, rotation.child().get() );
-    return res;
-}
 
+	//-------------------------------------------------------------------------------------------------
+	uint64 ASTOpImageTransform::Hash() const
+	{
+		uint64 res = std::hash<OP_TYPE>()(OP_TYPE::IM_MULTILAYER);
+		hash_combine(res, base.child().get());
+		hash_combine(res, offsetX.child().get());
+		hash_combine(res, offsetY.child().get());
+		hash_combine(res, scaleX.child().get());
+		hash_combine(res, scaleY.child().get());
+		hash_combine(res, rotation.child().get());
+		return res;
+	}
 
-//-------------------------------------------------------------------------------------------------
-mu::Ptr<ASTOp> ASTOpImageTransform::Clone(MapChildFuncRef mapChild) const
-{
-    Ptr<ASTOpImageTransform> n = new ASTOpImageTransform();
-    n->base = mapChild(base.child());
-	n->offsetX = mapChild(offsetX.child());
-	n->offsetY = mapChild(offsetY.child());
-	n->scaleX = mapChild(scaleX.child());
-	n->scaleY = mapChild(scaleY.child());
-	n->rotation = mapChild(rotation.child());
-    return n;
-}
 
+	//-------------------------------------------------------------------------------------------------
+	mu::Ptr<ASTOp> ASTOpImageTransform::Clone(MapChildFuncRef mapChild) const
+	{
+		Ptr<ASTOpImageTransform> n = new ASTOpImageTransform();
+		n->base = mapChild(base.child());
+		n->offsetX = mapChild(offsetX.child());
+		n->offsetY = mapChild(offsetY.child());
+		n->scaleX = mapChild(scaleX.child());
+		n->scaleY = mapChild(scaleY.child());
+		n->rotation = mapChild(rotation.child());
+		return n;
+	}
 
-//-------------------------------------------------------------------------------------------------
-void ASTOpImageTransform::ForEachChild( const TFunctionRef<void(ASTChild&)> f )
-{
-    f( base );
-    f( offsetX );
-    f( offsetY );
-    f( scaleX );
-    f( scaleY );
-    f( rotation );
-}
 
+	//-------------------------------------------------------------------------------------------------
+	void ASTOpImageTransform::ForEachChild(const TFunctionRef<void(ASTChild&)> f)
+	{
+		f(base);
+		f(offsetX);
+		f(offsetY);
+		f(scaleX);
+		f(scaleY);
+		f(rotation);
+	}
 
-//-------------------------------------------------------------------------------------------------
-void ASTOpImageTransform::Link( PROGRAM& program, const FLinkerOptions*)
-{
-    // Already linked?
-    if (!linkedAddress)
-    {
-		OP::ImageTransformArgs Args;
-	
-		Args.base = base ? base->linkedAddress : 0;
-		Args.offsetX = offsetX ? offsetX->linkedAddress : 0;
-		Args.offsetY = offsetY ? offsetY->linkedAddress : 0;
-		Args.scaleX = scaleX ? scaleX->linkedAddress : 0;
-		Args.scaleY = scaleY ? scaleY->linkedAddress : 0;
-		Args.rotation = rotation ? rotation->linkedAddress : 0;
-
-        linkedAddress = (OP::ADDRESS)program.m_opAddress.Num();
-        program.m_opAddress.Add((uint32_t)program.m_byteCode.Num());
-        AppendCode(program.m_byteCode, OP_TYPE::IM_TRANSFORM);
-		AppendCode(program.m_byteCode, Args);
-    }
-}
 
-//-------------------------------------------------------------------------------------------------
-FImageDesc ASTOpImageTransform::GetImageDesc( bool returnBestOption, GetImageDescContext* context )
-{
-    FImageDesc res;
-
-    // Local context in case it is necessary
-    GetImageDescContext localContext;
-    if (!context)
-    {
-      context = &localContext;
-    }
-    else
-    {
-        // Cached result?
-		FImageDesc* PtrValue = context->m_results.Find(this);
-		if (PtrValue)
+	//-------------------------------------------------------------------------------------------------
+	void ASTOpImageTransform::Link(PROGRAM& program, const FLinkerOptions*)
+	{
+		// Already linked?
+		if (!linkedAddress)
 		{
-			return *PtrValue;
+			OP::ImageTransformArgs Args;
+
+			Args.base = base ? base->linkedAddress : 0;
+			Args.offsetX = offsetX ? offsetX->linkedAddress : 0;
+			Args.offsetY = offsetY ? offsetY->linkedAddress : 0;
+			Args.scaleX = scaleX ? scaleX->linkedAddress : 0;
+			Args.scaleY = scaleY ? scaleY->linkedAddress : 0;
+			Args.rotation = rotation ? rotation->linkedAddress : 0;
+
+			linkedAddress = (OP::ADDRESS)program.m_opAddress.Num();
+			program.m_opAddress.Add((uint32_t)program.m_byteCode.Num());
+			AppendCode(program.m_byteCode, OP_TYPE::IM_TRANSFORM);
+			AppendCode(program.m_byteCode, Args);
 		}
 	}
 
-    // Actual work
-    if (base)
-    {
-        res = base->GetImageDesc( returnBestOption, context );
-    }
+	//-------------------------------------------------------------------------------------------------
+	FImageDesc ASTOpImageTransform::GetImageDesc(bool returnBestOption, GetImageDescContext* context)
+	{
+		FImageDesc res;
 
+		// Local context in case it is necessary
+		GetImageDescContext localContext;
+		if (!context)
+		{
+			context = &localContext;
+		}
+		else
+		{
+			// Cached result?
+			FImageDesc* PtrValue = context->m_results.Find(this);
+			if (PtrValue)
+			{
+				return *PtrValue;
+			}
+		}
 
-    // Cache the result
-    if (context)
-    {
-		context->m_results.Add(this, res);
-	}
+		// Actual work
+		if (base)
+		{
+			res = base->GetImageDesc(returnBestOption, context);
+		}
 
-    return res;
-}
 
+		// Cache the result
+		if (context)
+		{
+			context->m_results.Add(this, res);
+		}
 
-//-------------------------------------------------------------------------------------------------
-void ASTOpImageTransform::GetLayoutBlockSize(int* pBlockX, int* pBlockY)
-{
-	if (base)
+		return res;
+	}
+
+
+	//-------------------------------------------------------------------------------------------------
+	void ASTOpImageTransform::GetLayoutBlockSize(int* pBlockX, int* pBlockY)
 	{
-		base->GetLayoutBlockSize(pBlockX, pBlockY);
+		if (base)
+		{
+			base->GetLayoutBlockSize(pBlockX, pBlockY);
+		}
 	}
-}
 
 
-//-------------------------------------------------------------------------------------------------
-mu::Ptr<ImageSizeExpression> ASTOpImageTransform::GetImageSizeExpression() const
-{
-    if ( base )
-    {
-        return base->GetImageSizeExpression();
-    }
+	//-------------------------------------------------------------------------------------------------
+	mu::Ptr<ImageSizeExpression> ASTOpImageTransform::GetImageSizeExpression() const
+	{
+		if (base)
+		{
+			return base->GetImageSizeExpression();
+		}
 
-    return nullptr;
-}
+		return nullptr;
+	}
 
+}
diff --git a/Engine/Plugins/Experimental/Mutable/Source/MutableTools/Private/MuT/ASTOpInstanceAdd.cpp b/Engine/Plugins/Experimental/Mutable/Source/MutableTools/Private/MuT/ASTOpInstanceAdd.cpp
--- a/Engine/Plugins/Experimental/Mutable/Source/MutableTools/Private/MuT/ASTOpInstanceAdd.cpp
+++ b/Engine/Plugins/Experimental/Mutable/Source/MutableTools/Private/MuT/ASTOpInstanceAdd.cpp
@@ -15,156 +15,157 @@
 #include <utility>
 
 
-using namespace mu;
-
-
-//-------------------------------------------------------------------------------------------------
-//-------------------------------------------------------------------------------------------------
-//-------------------------------------------------------------------------------------------------
-ASTOpInstanceAdd::ASTOpInstanceAdd()
-    : instance(this)
-    , value(this)
-{
-}
-
-
-//-------------------------------------------------------------------------------------------------
-ASTOpInstanceAdd::~ASTOpInstanceAdd()
-{
-    // Explicit call needed to avoid recursive destruction
-    ASTOp::RemoveChildren();
-}
-
-
-//-------------------------------------------------------------------------------------------------
-bool ASTOpInstanceAdd::IsEqual(const ASTOp& otherUntyped) const
-{
-    if ( auto other = dynamic_cast<const ASTOpInstanceAdd*>(&otherUntyped) )
-    {
-        return type==other->type &&
-                instance==other->instance &&
-                value==other->value &&
-                id==other->id &&
-                externalId==other->externalId &&
-                name==other->name;
-    }
-    return false;
-}
-
-
-//-------------------------------------------------------------------------------------------------
-mu::Ptr<ASTOp> ASTOpInstanceAdd::Clone(MapChildFuncRef mapChild) const
+namespace mu
 {
-    Ptr<ASTOpInstanceAdd> n = new ASTOpInstanceAdd();
-    n->type = type;
-    n->instance = mapChild(instance.child());
-    n->value = mapChild(value.child());
-    n->id = id;
-    n->externalId = externalId;
-    n->name = name;
-    return n;
-}
-
-
-//-------------------------------------------------------------------------------------------------
-uint64 ASTOpInstanceAdd::Hash() const
-{
-	uint64 res = std::hash<size_t>()(size_t(type));
-    hash_combine( res, instance.child().get() );
-    hash_combine( res, value.child().get() );
-    return res;
-}
-
-
-//-------------------------------------------------------------------------------------------------
-void ASTOpInstanceAdd::Assert()
-{
-    switch(type)
-    {
-    case OP_TYPE::IN_ADDMESH:
-    case OP_TYPE::IN_ADDIMAGE:
-    case OP_TYPE::IN_ADDVECTOR:
-    case OP_TYPE::IN_ADDSCALAR:
-    case OP_TYPE::IN_ADDSTRING:
-    case OP_TYPE::IN_ADDSURFACE:
-    case OP_TYPE::IN_ADDCOMPONENT:
-    case OP_TYPE::IN_ADDLOD:
-        break;
-    default:
-        // Unexpected type
-        check(false);
-        break;
-    }
-
-    ASTOp::Assert();
-}
 
 
-//-------------------------------------------------------------------------------------------------
-void ASTOpInstanceAdd::ForEachChild(const TFunctionRef<void(ASTChild&)> f )
-{
-    f( instance );
-    f( value );
-}
-
-
-//-------------------------------------------------------------------------------------------------
-void ASTOpInstanceAdd::Link( PROGRAM& program, const FLinkerOptions* )
-{
-    // Already linked?
-    if (!linkedAddress)
-    {
-        OP::InstanceAddArgs args;
-        memset( &args,0, sizeof(args) );
-        args.id = id;
-        args.externalId = externalId;
-
-        args.name = program.AddConstant(name);
-
-        if (instance) args.instance = instance->linkedAddress;
-        if (value) args.value = value->linkedAddress;
-
-        if ( type==OP_TYPE::IN_ADDIMAGE ||
-             type==OP_TYPE::IN_ADDMESH  )
-        {
-            // Find out relevant parameters. \todo: this may be optimised by reusing partial
-            // values in a LINK_CONTEXT or similar
-            SubtreeRelevantParametersVisitorAST visitor;
-            visitor.Run( value.child() );
-
-            TArray<uint16> params;
-            for(const string& paramName: visitor.m_params)
-            {
-                for( int32 i=0; i<program.m_parameters.Num(); ++i )
-                {
-                    const auto& param = program.m_parameters[i];
-                    if (param.m_name==paramName)
-                    {
-                        params.Add( uint16(i) );
-                        break;
-                    }
-                }
-            }
-
-            params.Sort();
-
-            auto it = program.m_parameterLists.Find( params );
-
-            if (it!=INDEX_NONE)
-            {
-                args.relevantParametersListIndex = it;
-            }
-            else
-            {
-                args.relevantParametersListIndex = uint32_t(program.m_parameterLists.Num());
-                program.m_parameterLists.Add( params );
-            }
-        }
-
-        linkedAddress = (OP::ADDRESS)program.m_opAddress.Num();
-        program.m_opAddress.Add((uint32_t)program.m_byteCode.Num());
-        AppendCode(program.m_byteCode,type);
-        AppendCode(program.m_byteCode,args);
-    }
+	//-------------------------------------------------------------------------------------------------
+	ASTOpInstanceAdd::ASTOpInstanceAdd()
+		: instance(this)
+		, value(this)
+	{
+	}
+
+
+	//-------------------------------------------------------------------------------------------------
+	ASTOpInstanceAdd::~ASTOpInstanceAdd()
+	{
+		// Explicit call needed to avoid recursive destruction
+		ASTOp::RemoveChildren();
+	}
+
+
+	//-------------------------------------------------------------------------------------------------
+	bool ASTOpInstanceAdd::IsEqual(const ASTOp& otherUntyped) const
+	{
+		if (auto other = dynamic_cast<const ASTOpInstanceAdd*>(&otherUntyped))
+		{
+			return type == other->type &&
+				instance == other->instance &&
+				value == other->value &&
+				id == other->id &&
+				externalId == other->externalId &&
+				name == other->name;
+		}
+		return false;
+	}
+
+
+	//-------------------------------------------------------------------------------------------------
+	mu::Ptr<ASTOp> ASTOpInstanceAdd::Clone(MapChildFuncRef mapChild) const
+	{
+		Ptr<ASTOpInstanceAdd> n = new ASTOpInstanceAdd();
+		n->type = type;
+		n->instance = mapChild(instance.child());
+		n->value = mapChild(value.child());
+		n->id = id;
+		n->externalId = externalId;
+		n->name = name;
+		return n;
+	}
+
+
+	//-------------------------------------------------------------------------------------------------
+	uint64 ASTOpInstanceAdd::Hash() const
+	{
+		uint64 res = std::hash<size_t>()(size_t(type));
+		hash_combine(res, instance.child().get());
+		hash_combine(res, value.child().get());
+		return res;
+	}
+
+
+	//-------------------------------------------------------------------------------------------------
+	void ASTOpInstanceAdd::Assert()
+	{
+		switch (type)
+		{
+		case OP_TYPE::IN_ADDMESH:
+		case OP_TYPE::IN_ADDIMAGE:
+		case OP_TYPE::IN_ADDVECTOR:
+		case OP_TYPE::IN_ADDSCALAR:
+		case OP_TYPE::IN_ADDSTRING:
+		case OP_TYPE::IN_ADDSURFACE:
+		case OP_TYPE::IN_ADDCOMPONENT:
+		case OP_TYPE::IN_ADDLOD:
+			break;
+		default:
+			// Unexpected type
+			check(false);
+			break;
+		}
+
+		ASTOp::Assert();
+	}
+
+
+	//-------------------------------------------------------------------------------------------------
+	void ASTOpInstanceAdd::ForEachChild(const TFunctionRef<void(ASTChild&)> f)
+	{
+		f(instance);
+		f(value);
+	}
+
+
+	//-------------------------------------------------------------------------------------------------
+	void ASTOpInstanceAdd::Link(PROGRAM& program, const FLinkerOptions*)
+	{
+		// Already linked?
+		if (!linkedAddress)
+		{
+			OP::InstanceAddArgs args;
+			memset(&args, 0, sizeof(args));
+			args.id = id;
+			args.externalId = externalId;
+
+			args.name = program.AddConstant(name);
+
+			if (instance) args.instance = instance->linkedAddress;
+			if (value) args.value = value->linkedAddress;
+
+			if (type == OP_TYPE::IN_ADDIMAGE ||
+				type == OP_TYPE::IN_ADDMESH)
+			{
+				// Find out relevant parameters. \todo: this may be optimised by reusing partial
+				// values in a LINK_CONTEXT or similar
+				SubtreeRelevantParametersVisitorAST visitor;
+				visitor.Run(value.child());
+
+				TArray<uint16> params;
+				for (const string& paramName : visitor.m_params)
+				{
+					for (int32 i = 0; i < program.m_parameters.Num(); ++i)
+					{
+						const auto& param = program.m_parameters[i];
+						if (param.m_name == paramName)
+						{
+							params.Add(uint16(i));
+							break;
+						}
+					}
+				}
+
+				params.Sort();
+
+				auto it = program.m_parameterLists.Find(params);
+
+				if (it != INDEX_NONE)
+				{
+					args.relevantParametersListIndex = it;
+				}
+				else
+				{
+					args.relevantParametersListIndex = uint32_t(program.m_parameterLists.Num());
+					program.m_parameterLists.Add(params);
+				}
+			}
+
+			linkedAddress = (OP::ADDRESS)program.m_opAddress.Num();
+			program.m_opAddress.Add((uint32_t)program.m_byteCode.Num());
+			AppendCode(program.m_byteCode, type);
+			AppendCode(program.m_byteCode, args);
+		}
+
+	}
 
 }
diff --git a/Engine/Plugins/Experimental/Mutable/Source/MutableTools/Private/MuT/ASTOpMeshApplyPose.cpp b/Engine/Plugins/Experimental/Mutable/Source/MutableTools/Private/MuT/ASTOpMeshApplyPose.cpp
--- a/Engine/Plugins/Experimental/Mutable/Source/MutableTools/Private/MuT/ASTOpMeshApplyPose.cpp
+++ b/Engine/Plugins/Experimental/Mutable/Source/MutableTools/Private/MuT/ASTOpMeshApplyPose.cpp
@@ -13,79 +13,78 @@
 #include <utility>
 
 
+namespace mu
+{
 
-using namespace mu;
 
+	//-------------------------------------------------------------------------------------------------
+	ASTOpMeshApplyPose::ASTOpMeshApplyPose()
+		: base(this)
+		, pose(this)
+	{
+	}
 
-//-------------------------------------------------------------------------------------------------
-//-------------------------------------------------------------------------------------------------
-//-------------------------------------------------------------------------------------------------
-ASTOpMeshApplyPose::ASTOpMeshApplyPose()
-    : base(this)
-    , pose(this)
-{
-}
 
+	ASTOpMeshApplyPose::~ASTOpMeshApplyPose()
+	{
+		// Explicit call needed to avoid recursive destruction
+		ASTOp::RemoveChildren();
+	}
 
-ASTOpMeshApplyPose::~ASTOpMeshApplyPose()
-{
-    // Explicit call needed to avoid recursive destruction
-    ASTOp::RemoveChildren();
-}
 
+	bool ASTOpMeshApplyPose::IsEqual(const ASTOp& otherUntyped) const
+	{
+		if (auto other = dynamic_cast<const ASTOpMeshApplyPose*>(&otherUntyped))
+		{
+			return base == other->base &&
+				pose == other->pose;
+		}
+		return false;
+	}
 
-bool ASTOpMeshApplyPose::IsEqual(const ASTOp& otherUntyped) const
-{
-    if ( auto other = dynamic_cast<const ASTOpMeshApplyPose*>(&otherUntyped) )
-    {
-        return base==other->base &&
-                pose==other->pose;
-    }
-    return false;
-}
 
+	uint64 ASTOpMeshApplyPose::Hash() const
+	{
+		uint64 res = std::hash<void*>()(base.child().get());
+		hash_combine(res, pose.child().get());
+		return res;
+	}
 
-uint64 ASTOpMeshApplyPose::Hash() const
-{
-	uint64 res = std::hash<void*>()(base.child().get() );
-    hash_combine( res, pose.child().get() );
-    return res;
-}
 
-
-mu::Ptr<ASTOp> ASTOpMeshApplyPose::Clone(MapChildFuncRef mapChild) const
-{
-    Ptr<ASTOpMeshApplyPose> n = new ASTOpMeshApplyPose();
-    n->base = mapChild(base.child());
-    n->pose = mapChild(pose.child());
-    return n;
-}
+	mu::Ptr<ASTOp> ASTOpMeshApplyPose::Clone(MapChildFuncRef mapChild) const
+	{
+		Ptr<ASTOpMeshApplyPose> n = new ASTOpMeshApplyPose();
+		n->base = mapChild(base.child());
+		n->pose = mapChild(pose.child());
+		return n;
+	}
 
 
-void ASTOpMeshApplyPose::ForEachChild(const TFunctionRef<void(ASTChild&)> f )
-{
-    f( base );
-    f( pose );
-}
+	void ASTOpMeshApplyPose::ForEachChild(const TFunctionRef<void(ASTChild&)> f)
+	{
+		f(base);
+		f(pose);
+	}
 
 
-void ASTOpMeshApplyPose::Link( PROGRAM& program, const FLinkerOptions* )
-{
-    // Already linked?
-    if (!linkedAddress)
-    {
-        OP::MeshApplyPoseArgs args;
-        memset( &args,0, sizeof(args) );
+	void ASTOpMeshApplyPose::Link(PROGRAM& program, const FLinkerOptions*)
+	{
+		// Already linked?
+		if (!linkedAddress)
+		{
+			OP::MeshApplyPoseArgs args;
+			memset(&args, 0, sizeof(args));
 
-        if (base) args.base = base->linkedAddress;
-        if (pose) args.pose = pose->linkedAddress;
+			if (base) args.base = base->linkedAddress;
+			if (pose) args.pose = pose->linkedAddress;
 
-        linkedAddress = (OP::ADDRESS)program.m_opAddress.Num();
-        //program.m_code.push_back(op);
-        program.m_opAddress.Add((uint32_t)program.m_byteCode.Num());
-        AppendCode(program.m_byteCode,OP_TYPE::ME_APPLYPOSE);
-        AppendCode(program.m_byteCode,args);
-    }
+			linkedAddress = (OP::ADDRESS)program.m_opAddress.Num();
+			//program.m_code.push_back(op);
+			program.m_opAddress.Add((uint32_t)program.m_byteCode.Num());
+			AppendCode(program.m_byteCode, OP_TYPE::ME_APPLYPOSE);
+			AppendCode(program.m_byteCode, args);
+		}
 
-}
+	}
 
+}
\ No newline at end of file
diff --git a/Engine/Plugins/Experimental/Mutable/Source/MutableTools/Private/MuT/ASTOpMeshApplyShape.cpp b/Engine/Plugins/Experimental/Mutable/Source/MutableTools/Private/MuT/ASTOpMeshApplyShape.cpp
--- a/Engine/Plugins/Experimental/Mutable/Source/MutableTools/Private/MuT/ASTOpMeshApplyShape.cpp
+++ b/Engine/Plugins/Experimental/Mutable/Source/MutableTools/Private/MuT/ASTOpMeshApplyShape.cpp
@@ -12,100 +12,103 @@
 #include <memory>
 #include <utility>
 
-using namespace mu;
 
-
-ASTOpMeshApplyShape::ASTOpMeshApplyShape()
-    : Mesh(this)
-    , Shape(this)
-{
-}
-
-
-ASTOpMeshApplyShape::~ASTOpMeshApplyShape()
+namespace mu
 {
-    // Explicit call needed to avoid recursive destruction
-    ASTOp::RemoveChildren();
-}
 
 
-bool ASTOpMeshApplyShape::IsEqual(const ASTOp& otherUntyped) const
-{
-    if ( auto other = dynamic_cast<const ASTOpMeshApplyShape*>(&otherUntyped) )
-    {
-    	const bool bSameFlags =
-    		m_reshapePhysicsVolumes == other->m_reshapePhysicsVolumes &&
-    		m_reshapeSkeleton == other->m_reshapeSkeleton &&
-			m_reshapeVertices == other->m_reshapeVertices;
-    	
-        return Mesh == other->Mesh && Shape == other->Shape && bSameFlags;
-    }
-    return false;
-}
+	ASTOpMeshApplyShape::ASTOpMeshApplyShape()
+		: Mesh(this)
+		, Shape(this)
+	{
+	}
 
 
-uint64 ASTOpMeshApplyShape::Hash() const
-{
-	uint64 res = std::hash<void*>()(Mesh.child().get() );
-	hash_combine(res, Shape.child().get());
-	hash_combine(res, m_reshapeSkeleton);
-	hash_combine(res, m_reshapePhysicsVolumes);
-	hash_combine(res, m_reshapeVertices);
-	return res;
-}
+	ASTOpMeshApplyShape::~ASTOpMeshApplyShape()
+	{
+		// Explicit call needed to avoid recursive destruction
+		ASTOp::RemoveChildren();
+	}
 
 
-mu::Ptr<ASTOp> ASTOpMeshApplyShape::Clone(MapChildFuncRef mapChild) const
-{
-    Ptr<ASTOpMeshApplyShape> n = new ASTOpMeshApplyShape();
-	n->Mesh = mapChild(Mesh.child());
-	n->Shape = mapChild(Shape.child());
-	n->m_reshapeSkeleton = m_reshapeSkeleton;
-	n->m_reshapePhysicsVolumes = m_reshapePhysicsVolumes;
-	n->m_reshapeVertices = m_reshapeVertices;
-	return n;
-}
+	bool ASTOpMeshApplyShape::IsEqual(const ASTOp& otherUntyped) const
+	{
+		if (auto other = dynamic_cast<const ASTOpMeshApplyShape*>(&otherUntyped))
+		{
+			const bool bSameFlags =
+				m_reshapePhysicsVolumes == other->m_reshapePhysicsVolumes &&
+				m_reshapeSkeleton == other->m_reshapeSkeleton &&
+				m_reshapeVertices == other->m_reshapeVertices;
 
+			return Mesh == other->Mesh && Shape == other->Shape && bSameFlags;
+		}
+		return false;
+	}
+
+
+	uint64 ASTOpMeshApplyShape::Hash() const
+	{
+		uint64 res = std::hash<void*>()(Mesh.child().get());
+		hash_combine(res, Shape.child().get());
+		hash_combine(res, m_reshapeSkeleton);
+		hash_combine(res, m_reshapePhysicsVolumes);
+		hash_combine(res, m_reshapeVertices);
+		return res;
+	}
+
+
+	mu::Ptr<ASTOp> ASTOpMeshApplyShape::Clone(MapChildFuncRef mapChild) const
+	{
+		Ptr<ASTOpMeshApplyShape> n = new ASTOpMeshApplyShape();
+		n->Mesh = mapChild(Mesh.child());
+		n->Shape = mapChild(Shape.child());
+		n->m_reshapeSkeleton = m_reshapeSkeleton;
+		n->m_reshapePhysicsVolumes = m_reshapePhysicsVolumes;
+		n->m_reshapeVertices = m_reshapeVertices;
+		return n;
+	}
+
+
+	void ASTOpMeshApplyShape::ForEachChild(const TFunctionRef<void(ASTChild&)> f)
+	{
+		f(Mesh);
+		f(Shape);
+	}
+
+
+	void ASTOpMeshApplyShape::Link(PROGRAM& program, const FLinkerOptions*)
+	{
+		// Already linked?
+		if (!linkedAddress)
+		{
+			OP::MeshApplyShapeArgs args;
+			memset(&args, 0, sizeof(args));
 
-void ASTOpMeshApplyShape::ForEachChild(const TFunctionRef<void(ASTChild&)> f )
-{
-	f(Mesh);
-	f(Shape);
-}
+			if (m_reshapeSkeleton)
+			{
+				args.flags |= uint32(OP::EMeshBindShapeFlags::ReshapeSkeleton);
+			}
 
+			if (m_reshapePhysicsVolumes)
+			{
+				args.flags |= uint32(OP::EMeshBindShapeFlags::ReshapePhysicsVolumes);
+			}
 
-void ASTOpMeshApplyShape::Link( PROGRAM& program, const FLinkerOptions* )
-{
-    // Already linked?
-    if (!linkedAddress)
-    {
-        OP::MeshApplyShapeArgs args;
-        memset( &args,0, sizeof(args) );
+			if (m_reshapeVertices)
+			{
+				args.flags |= uint32(OP::EMeshBindShapeFlags::ReshapeVertices);
+			}
 
-		if (m_reshapeSkeleton)
-		{
-			args.flags |= uint32(OP::EMeshBindShapeFlags::ReshapeSkeleton);
-		}
 
-    	if (m_reshapePhysicsVolumes)
-    	{
-			args.flags |= uint32(OP::EMeshBindShapeFlags::ReshapePhysicsVolumes);
-    	}
+			if (Mesh) args.mesh = Mesh->linkedAddress;
+			if (Shape) args.shape = Shape->linkedAddress;
 
-		if (m_reshapeVertices)
-		{
-			args.flags |= uint32(OP::EMeshBindShapeFlags::ReshapeVertices);
+			linkedAddress = (OP::ADDRESS)program.m_opAddress.Num();
+			program.m_opAddress.Add((uint32_t)program.m_byteCode.Num());
+			AppendCode(program.m_byteCode, OP_TYPE::ME_APPLYSHAPE);
+			AppendCode(program.m_byteCode, args);
 		}
 
-    	
-		if (Mesh) args.mesh = Mesh->linkedAddress;
-		if (Shape) args.shape = Shape->linkedAddress;
-
-        linkedAddress = (OP::ADDRESS)program.m_opAddress.Num();
-        program.m_opAddress.Add((uint32_t)program.m_byteCode.Num());
-        AppendCode(program.m_byteCode,OP_TYPE::ME_APPLYSHAPE);
-        AppendCode(program.m_byteCode,args);
-    }
+	}
 
 }
-
diff --git a/Engine/Plugins/Experimental/Mutable/Source/MutableTools/Private/MuT/ASTOpMeshBindShape.cpp b/Engine/Plugins/Experimental/Mutable/Source/MutableTools/Private/MuT/ASTOpMeshBindShape.cpp
--- a/Engine/Plugins/Experimental/Mutable/Source/MutableTools/Private/MuT/ASTOpMeshBindShape.cpp
+++ b/Engine/Plugins/Experimental/Mutable/Source/MutableTools/Private/MuT/ASTOpMeshBindShape.cpp
@@ -13,329 +13,333 @@
 #include <memory>
 #include <utility>
 
-using namespace mu;
 
-
-ASTOpMeshBindShape::ASTOpMeshBindShape()
-    : Mesh(this)
-    , Shape(this)
-{
-}
-
-
-ASTOpMeshBindShape::~ASTOpMeshBindShape()
-{
-    // Explicit call needed to avoid recursive destruction
-    ASTOp::RemoveChildren();
-}
-
-
-bool ASTOpMeshBindShape::IsEqual(const ASTOp& otherUntyped) const
+namespace mu
 {
-	if ( auto other = dynamic_cast<const ASTOpMeshBindShape*>(&otherUntyped) )
-    {
-		const bool bSameFlags =
-			m_reshapeSkeleton == other->m_reshapeSkeleton &&
-			m_discardInvalidBindings == other->m_discardInvalidBindings &&
-			m_enableRigidParts == other->m_enableRigidParts &&
-			m_deformAllBones == other->m_deformAllBones &&
-			m_deformAllPhysics == other->m_deformAllPhysics &&
-    		m_reshapePhysicsVolumes == other->m_reshapePhysicsVolumes &&
-			m_reshapeVertices == other->m_reshapeVertices;
-    	
-		return bSameFlags && 
-				Mesh == other->Mesh && 
-				Shape == other->Shape && 
-				m_bonesToDeform == other->m_bonesToDeform &&
-				m_physicsToDeform == other->m_physicsToDeform && 
-				m_bindingMethod == other->m_bindingMethod;
-	}
-	
-    return false;
-}
 
 
-uint64 ASTOpMeshBindShape::Hash() const
-{
-	uint64 res = std::hash<void*>()(Mesh.child().get() );
-	hash_combine(res, Shape.child().get());
-	hash_combine(res, m_reshapeSkeleton);
-	hash_combine(res, m_discardInvalidBindings);
-	hash_combine(res, m_enableRigidParts);
-	hash_combine(res, m_deformAllBones);
-	hash_combine(res, m_deformAllPhysics);
-	hash_combine(res, m_reshapePhysicsVolumes);
-	hash_combine(res, m_reshapeVertices);
-	hash_combine(res, m_bindingMethod);
-
-	for ( const string& S : m_bonesToDeform)
+	ASTOpMeshBindShape::ASTOpMeshBindShape()
+		: Mesh(this)
+		, Shape(this)
 	{
-		hash_combine(res, S);
 	}
 
-	for ( const string& S : m_physicsToDeform)
+
+	ASTOpMeshBindShape::~ASTOpMeshBindShape()
 	{
-		hash_combine(res, S);
+		// Explicit call needed to avoid recursive destruction
+		ASTOp::RemoveChildren();
 	}
 
-	return res;
-}
 
-
-mu::Ptr<ASTOp> ASTOpMeshBindShape::Clone(MapChildFuncRef mapChild) const
-{
-    Ptr<ASTOpMeshBindShape> n = new ASTOpMeshBindShape();
-	n->Mesh = mapChild(Mesh.child());
-	n->Shape = mapChild(Shape.child());
-	n->m_reshapeSkeleton = m_reshapeSkeleton;
-	n->m_discardInvalidBindings = m_discardInvalidBindings;
-	n->m_enableRigidParts = m_enableRigidParts;
-	n->m_deformAllBones = m_deformAllBones;
-	n->m_deformAllPhysics = m_deformAllPhysics;
-	n->m_reshapePhysicsVolumes = m_reshapePhysicsVolumes;
-	n->m_reshapeVertices = m_reshapeVertices;
-	n->m_bonesToDeform = m_bonesToDeform;
-	n->m_physicsToDeform = m_physicsToDeform;	
-	n->m_bindingMethod = m_bindingMethod;
-	return n;
-}
-
-
-void ASTOpMeshBindShape::ForEachChild(const TFunctionRef<void(ASTChild&)> f )
-{
-	f(Mesh);
-	f(Shape);
-}
-
-
-void ASTOpMeshBindShape::Link( PROGRAM& program, const FLinkerOptions*)
-{
-    // Already linked?
-    if (!linkedAddress)
-    {
-        OP::MeshBindShapeArgs args;
-        FMemory::Memzero( &args, sizeof(args) );
-
-		if (m_reshapeSkeleton)
+	bool ASTOpMeshBindShape::IsEqual(const ASTOp& otherUntyped) const
+	{
+		if (auto other = dynamic_cast<const ASTOpMeshBindShape*>(&otherUntyped))
 		{
-			args.flags |= uint32(OP::EMeshBindShapeFlags::ReshapeSkeleton);
+			const bool bSameFlags =
+				m_reshapeSkeleton == other->m_reshapeSkeleton &&
+				m_discardInvalidBindings == other->m_discardInvalidBindings &&
+				m_enableRigidParts == other->m_enableRigidParts &&
+				m_deformAllBones == other->m_deformAllBones &&
+				m_deformAllPhysics == other->m_deformAllPhysics &&
+				m_reshapePhysicsVolumes == other->m_reshapePhysicsVolumes &&
+				m_reshapeVertices == other->m_reshapeVertices;
+
+			return bSameFlags &&
+				Mesh == other->Mesh &&
+				Shape == other->Shape &&
+				m_bonesToDeform == other->m_bonesToDeform &&
+				m_physicsToDeform == other->m_physicsToDeform &&
+				m_bindingMethod == other->m_bindingMethod;
 		}
 
-    	if (m_discardInvalidBindings)
-    	{
-    		args.flags |= uint32(OP::EMeshBindShapeFlags::DiscardInvalidBindings);
-    	}
-    
-    	if (m_enableRigidParts)
-    	{
-    		args.flags |= uint32(OP::EMeshBindShapeFlags::EnableRigidParts);
-    	}
-		
-		if (m_deformAllBones)
-    	{
-    		args.flags |= uint32(OP::EMeshBindShapeFlags::DeformAllBones);
-    	}
-
-		if (m_deformAllPhysics)
-    	{
-    		args.flags |= uint32(OP::EMeshBindShapeFlags::DeformAllPhysics);
-    	}
-
-    	if (m_reshapePhysicsVolumes)
-    	{
-    		args.flags |= uint32(OP::EMeshBindShapeFlags::ReshapePhysicsVolumes);
-    	}
-  	
-		if (m_reshapeVertices)
-    	{
-    		args.flags |= uint32(OP::EMeshBindShapeFlags::ReshapeVertices);
-    	}  
- 	
-		args.bindingMethod = m_bindingMethod;    
-		
-		if (Mesh) 
-		{
-			args.mesh = Mesh->linkedAddress;
-		}
+		return false;
+	}
 
-		if (Shape)
-		{
-			args.shape = Shape->linkedAddress;
-		}
-		
-		linkedAddress = (OP::ADDRESS)program.m_opAddress.Num();
-		program.m_opAddress.Add((uint32_t)program.m_byteCode.Num());
-		AppendCode(program.m_byteCode, OP_TYPE::ME_BINDSHAPE);
-		AppendCode(program.m_byteCode, args);
 
-		AppendCode(program.m_byteCode, (int32)m_bonesToDeform.Num());
+	uint64 ASTOpMeshBindShape::Hash() const
+	{
+		uint64 res = std::hash<void*>()(Mesh.child().get());
+		hash_combine(res, Shape.child().get());
+		hash_combine(res, m_reshapeSkeleton);
+		hash_combine(res, m_discardInvalidBindings);
+		hash_combine(res, m_enableRigidParts);
+		hash_combine(res, m_deformAllBones);
+		hash_combine(res, m_deformAllPhysics);
+		hash_combine(res, m_reshapePhysicsVolumes);
+		hash_combine(res, m_reshapeVertices);
+		hash_combine(res, m_bindingMethod);
+
 		for (const string& S : m_bonesToDeform)
 		{
-			const uint32 index = program.AddConstant(S);
-			AppendCode(program.m_byteCode, index);
+			hash_combine(res, S);
 		}
 
-		AppendCode(program.m_byteCode, (int32)m_physicsToDeform.Num());
 		for (const string& S : m_physicsToDeform)
 		{
-			const uint32 index = program.AddConstant(S);
-			AppendCode(program.m_byteCode, index);
+			hash_combine(res, S);
 		}
-    }
-}
 
+		return res;
+	}
 
-//-------------------------------------------------------------------------------------------------
-mu::Ptr<ASTOp> ASTOpMeshBindShape::OptimiseSink(const MODEL_OPTIMIZATION_OPTIONS&, OPTIMIZE_SINK_CONTEXT&) const
-{
-	Ptr<ASTOp> NewOp;
 
-	Ptr<ASTOp> MeshAt = Mesh.child();
-	if (!MeshAt)
-	{ 
-		return nullptr;
+	mu::Ptr<ASTOp> ASTOpMeshBindShape::Clone(MapChildFuncRef mapChild) const
+	{
+		Ptr<ASTOpMeshBindShape> n = new ASTOpMeshBindShape();
+		n->Mesh = mapChild(Mesh.child());
+		n->Shape = mapChild(Shape.child());
+		n->m_reshapeSkeleton = m_reshapeSkeleton;
+		n->m_discardInvalidBindings = m_discardInvalidBindings;
+		n->m_enableRigidParts = m_enableRigidParts;
+		n->m_deformAllBones = m_deformAllBones;
+		n->m_deformAllPhysics = m_deformAllPhysics;
+		n->m_reshapePhysicsVolumes = m_reshapePhysicsVolumes;
+		n->m_reshapeVertices = m_reshapeVertices;
+		n->m_bonesToDeform = m_bonesToDeform;
+		n->m_physicsToDeform = m_physicsToDeform;
+		n->m_bindingMethod = m_bindingMethod;
+		return n;
 	}
 
-	Ptr<ASTOp> ShapeAt = Shape.child();
-	if (!ShapeAt)
+
+	void ASTOpMeshBindShape::ForEachChild(const TFunctionRef<void(ASTChild&)> f)
 	{
-		return nullptr;
+		f(Mesh);
+		f(Shape);
 	}
 
-	OP_TYPE MeshType = MeshAt->GetOpType();
-	OP_TYPE ShapeType = ShapeAt->GetOpType();
 
-	// See if both mesh and shape have an operation that can be optimized in a combined way
-	if (MeshType==ShapeType)
+	void ASTOpMeshBindShape::Link(PROGRAM& program, const FLinkerOptions*)
 	{
-		switch (MeshType)
+		// Already linked?
+		if (!linkedAddress)
 		{
+			OP::MeshBindShapeArgs args;
+			FMemory::Memzero(&args, sizeof(args));
 
-		case OP_TYPE::ME_SWITCH:
-		{
-			// Id the switch variable and structure is the same
-			const ASTOpSwitch* MeshSwitch = reinterpret_cast<const ASTOpSwitch*>(MeshAt.get());
-			const ASTOpSwitch* ShapeSwitch = reinterpret_cast<const ASTOpSwitch*>(ShapeAt.get());
-			bool bIsSimilarSwitch = MeshSwitch->IsCompatibleWith(ShapeSwitch);
-			if (!bIsSimilarSwitch)
+			if (m_reshapeSkeleton)
 			{
-				break;
+				args.flags |= uint32(OP::EMeshBindShapeFlags::ReshapeSkeleton);
+			}
+
+			if (m_discardInvalidBindings)
+			{
+				args.flags |= uint32(OP::EMeshBindShapeFlags::DiscardInvalidBindings);
 			}
 
-			// Move the operation down all the paths
-			Ptr<ASTOpSwitch> NewSwitch = mu::Clone<ASTOpSwitch>(MeshAt);
+			if (m_enableRigidParts)
+			{
+				args.flags |= uint32(OP::EMeshBindShapeFlags::EnableRigidParts);
+			}
 
-			if (NewSwitch->def)
+			if (m_deformAllBones)
 			{
-				Ptr<ASTOpMeshBindShape> NewBind = mu::Clone<ASTOpMeshBindShape>(this);
-				NewBind->Mesh = MeshSwitch->def.child();
-				NewBind->Shape = ShapeSwitch->def.child();
-				NewSwitch->def = NewBind;
+				args.flags |= uint32(OP::EMeshBindShapeFlags::DeformAllBones);
 			}
 
-			// We need to copy the options because we change them
-			for (int32 v = 0; v < NewSwitch->cases.Num(); ++v)
+			if (m_deformAllPhysics)
 			{
-				if (NewSwitch->cases[v].branch)
-				{
-					Ptr<ASTOpMeshBindShape> NewBind = mu::Clone<ASTOpMeshBindShape>(this);
-					NewBind->Mesh = MeshSwitch->cases[v].branch.child();
-					NewBind->Shape = ShapeSwitch->FindBranch(MeshSwitch->cases[v].condition);
-					NewSwitch->cases[v].branch = NewBind;
-				}
+				args.flags |= uint32(OP::EMeshBindShapeFlags::DeformAllPhysics);
 			}
 
-			NewOp = NewSwitch;
-			break;
-		}
+			if (m_reshapePhysicsVolumes)
+			{
+				args.flags |= uint32(OP::EMeshBindShapeFlags::ReshapePhysicsVolumes);
+			}
+
+			if (m_reshapeVertices)
+			{
+				args.flags |= uint32(OP::EMeshBindShapeFlags::ReshapeVertices);
+			}
 
-		default:
-			break;
+			args.bindingMethod = m_bindingMethod;
+
+			if (Mesh)
+			{
+				args.mesh = Mesh->linkedAddress;
+			}
 
+			if (Shape)
+			{
+				args.shape = Shape->linkedAddress;
+			}
+
+			linkedAddress = (OP::ADDRESS)program.m_opAddress.Num();
+			program.m_opAddress.Add((uint32_t)program.m_byteCode.Num());
+			AppendCode(program.m_byteCode, OP_TYPE::ME_BINDSHAPE);
+			AppendCode(program.m_byteCode, args);
+
+			AppendCode(program.m_byteCode, (int32)m_bonesToDeform.Num());
+			for (const string& S : m_bonesToDeform)
+			{
+				const uint32 index = program.AddConstant(S);
+				AppendCode(program.m_byteCode, index);
+			}
+
+			AppendCode(program.m_byteCode, (int32)m_physicsToDeform.Num());
+			for (const string& S : m_physicsToDeform)
+			{
+				const uint32 index = program.AddConstant(S);
+				AppendCode(program.m_byteCode, index);
+			}
 		}
 	}
 
 
-	// If not already optimized
-	if (!NewOp)
+	//-------------------------------------------------------------------------------------------------
+	mu::Ptr<ASTOp> ASTOpMeshBindShape::OptimiseSink(const MODEL_OPTIMIZATION_OPTIONS&, OPTIMIZE_SINK_CONTEXT&) const
 	{
-		// Optimize only the mesh parameter
-		switch (MeshType)
+		Ptr<ASTOp> NewOp;
+
+		Ptr<ASTOp> MeshAt = Mesh.child();
+		if (!MeshAt)
 		{
+			return nullptr;
+		}
 
-		case OP_TYPE::ME_SWITCH:
+		Ptr<ASTOp> ShapeAt = Shape.child();
+		if (!ShapeAt)
 		{
-			// Move the operation down all the paths
-			Ptr<ASTOpSwitch> NewSwitch = mu::Clone<ASTOpSwitch>(MeshAt);
+			return nullptr;
+		}
 
-			if (NewSwitch->def)
+		OP_TYPE MeshType = MeshAt->GetOpType();
+		OP_TYPE ShapeType = ShapeAt->GetOpType();
+
+		// See if both mesh and shape have an operation that can be optimized in a combined way
+		if (MeshType == ShapeType)
+		{
+			switch (MeshType)
 			{
-				Ptr<ASTOpMeshBindShape> NewBind = mu::Clone<ASTOpMeshBindShape>(this);
-				NewBind->Mesh = NewSwitch->def.child();
-				NewSwitch->def = NewBind;
-			}
 
-			// We need to copy the options because we change them
-			for (int32 v = 0; v < NewSwitch->cases.Num(); ++v)
+			case OP_TYPE::ME_SWITCH:
 			{
-				if (NewSwitch->cases[v].branch)
+				// Id the switch variable and structure is the same
+				const ASTOpSwitch* MeshSwitch = reinterpret_cast<const ASTOpSwitch*>(MeshAt.get());
+				const ASTOpSwitch* ShapeSwitch = reinterpret_cast<const ASTOpSwitch*>(ShapeAt.get());
+				bool bIsSimilarSwitch = MeshSwitch->IsCompatibleWith(ShapeSwitch);
+				if (!bIsSimilarSwitch)
+				{
+					break;
+				}
+
+				// Move the operation down all the paths
+				Ptr<ASTOpSwitch> NewSwitch = mu::Clone<ASTOpSwitch>(MeshAt);
+
+				if (NewSwitch->def)
 				{
 					Ptr<ASTOpMeshBindShape> NewBind = mu::Clone<ASTOpMeshBindShape>(this);
-					NewBind->Mesh = NewSwitch->cases[v].branch.child();
-					NewSwitch->cases[v].branch = NewBind;
+					NewBind->Mesh = MeshSwitch->def.child();
+					NewBind->Shape = ShapeSwitch->def.child();
+					NewSwitch->def = NewBind;
 				}
-			}
 
-			NewOp = NewSwitch;
-			break;
-		}
+				// We need to copy the options because we change them
+				for (int32 v = 0; v < NewSwitch->cases.Num(); ++v)
+				{
+					if (NewSwitch->cases[v].branch)
+					{
+						Ptr<ASTOpMeshBindShape> NewBind = mu::Clone<ASTOpMeshBindShape>(this);
+						NewBind->Mesh = MeshSwitch->cases[v].branch.child();
+						NewBind->Shape = ShapeSwitch->FindBranch(MeshSwitch->cases[v].condition);
+						NewSwitch->cases[v].branch = NewBind;
+					}
+				}
 
-		default:
-			break;
+				NewOp = NewSwitch;
+				break;
+			}
+
+			default:
+				break;
 
+			}
 		}
-	}
 
-	// If not already optimized
-	if (!NewOp)
-	{
-		// Optimize only the shape parameter
-		switch (ShapeType)
-		{
 
-		case OP_TYPE::ME_SWITCH:
+		// If not already optimized
+		if (!NewOp)
 		{
-			// Move the operation down all the paths
-			Ptr<ASTOpSwitch> NewSwitch = mu::Clone<ASTOpSwitch>(ShapeAt);
+			// Optimize only the mesh parameter
+			switch (MeshType)
+			{
 
-			if (NewSwitch->def)
+			case OP_TYPE::ME_SWITCH:
 			{
-				Ptr<ASTOpMeshBindShape> NewBind = mu::Clone<ASTOpMeshBindShape>(this);
-				NewBind->Shape = NewSwitch->def.child();
-				NewSwitch->def = NewBind;
+				// Move the operation down all the paths
+				Ptr<ASTOpSwitch> NewSwitch = mu::Clone<ASTOpSwitch>(MeshAt);
+
+				if (NewSwitch->def)
+				{
+					Ptr<ASTOpMeshBindShape> NewBind = mu::Clone<ASTOpMeshBindShape>(this);
+					NewBind->Mesh = NewSwitch->def.child();
+					NewSwitch->def = NewBind;
+				}
+
+				// We need to copy the options because we change them
+				for (int32 v = 0; v < NewSwitch->cases.Num(); ++v)
+				{
+					if (NewSwitch->cases[v].branch)
+					{
+						Ptr<ASTOpMeshBindShape> NewBind = mu::Clone<ASTOpMeshBindShape>(this);
+						NewBind->Mesh = NewSwitch->cases[v].branch.child();
+						NewSwitch->cases[v].branch = NewBind;
+					}
+				}
+
+				NewOp = NewSwitch;
+				break;
+			}
+
+			default:
+				break;
+
 			}
+		}
+
+		// If not already optimized
+		if (!NewOp)
+		{
+			// Optimize only the shape parameter
+			switch (ShapeType)
+			{
 
-			// We need to copy the options because we change them
-			for (int32 v = 0; v < NewSwitch->cases.Num(); ++v)
+			case OP_TYPE::ME_SWITCH:
 			{
-				if (NewSwitch->cases[v].branch)
+				// Move the operation down all the paths
+				Ptr<ASTOpSwitch> NewSwitch = mu::Clone<ASTOpSwitch>(ShapeAt);
+
+				if (NewSwitch->def)
 				{
 					Ptr<ASTOpMeshBindShape> NewBind = mu::Clone<ASTOpMeshBindShape>(this);
-					NewBind->Shape = NewSwitch->cases[v].branch.child();
-					NewSwitch->cases[v].branch = NewBind;
+					NewBind->Shape = NewSwitch->def.child();
+					NewSwitch->def = NewBind;
+				}
+
+				// We need to copy the options because we change them
+				for (int32 v = 0; v < NewSwitch->cases.Num(); ++v)
+				{
+					if (NewSwitch->cases[v].branch)
+					{
+						Ptr<ASTOpMeshBindShape> NewBind = mu::Clone<ASTOpMeshBindShape>(this);
+						NewBind->Shape = NewSwitch->cases[v].branch.child();
+						NewSwitch->cases[v].branch = NewBind;
+					}
 				}
+
+				NewOp = NewSwitch;
+				break;
 			}
 
-			NewOp = NewSwitch;
-			break;
+			default:
+				break;
+
+			}
 		}
 
-		default:
-			break;
 
-		}
+		return NewOp;
 	}
 
-
-	return NewOp;
 }
diff --git a/Engine/Plugins/Experimental/Mutable/Source/MutableTools/Private/MuT/ASTOpMeshClipDeform.cpp b/Engine/Plugins/Experimental/Mutable/Source/MutableTools/Private/MuT/ASTOpMeshClipDeform.cpp
--- a/Engine/Plugins/Experimental/Mutable/Source/MutableTools/Private/MuT/ASTOpMeshClipDeform.cpp
+++ b/Engine/Plugins/Experimental/Mutable/Source/MutableTools/Private/MuT/ASTOpMeshClipDeform.cpp
@@ -12,81 +12,84 @@
 #include <memory>
 #include <utility>
 
-using namespace mu;
 
-
-ASTOpMeshClipDeform::ASTOpMeshClipDeform()
-    : Mesh(this)
-    , ClipShape(this)
+namespace mu
 {
-}
 
 
-ASTOpMeshClipDeform::~ASTOpMeshClipDeform()
-{
-    // Explicit call needed to avoid recursive destruction
-    ASTOp::RemoveChildren();
-}
+	ASTOpMeshClipDeform::ASTOpMeshClipDeform()
+		: Mesh(this)
+		, ClipShape(this)
+	{
+	}
 
 
-bool ASTOpMeshClipDeform::IsEqual(const ASTOp& OtherUntyped) const
-{
-    if ( const ASTOpMeshClipDeform* Other = dynamic_cast<const ASTOpMeshClipDeform*>(&OtherUntyped) )
-    {
-        return Mesh == Other->Mesh && ClipShape == Other->ClipShape;
-    }
+	ASTOpMeshClipDeform::~ASTOpMeshClipDeform()
+	{
+		// Explicit call needed to avoid recursive destruction
+		ASTOp::RemoveChildren();
+	}
 
-    return false;
-}
 
+	bool ASTOpMeshClipDeform::IsEqual(const ASTOp& OtherUntyped) const
+	{
+		if (const ASTOpMeshClipDeform* Other = dynamic_cast<const ASTOpMeshClipDeform*>(&OtherUntyped))
+		{
+			return Mesh == Other->Mesh && ClipShape == Other->ClipShape;
+		}
 
-uint64 ASTOpMeshClipDeform::Hash() const
-{
-	uint64 res = std::hash<void*>()(Mesh.child().get());
-	hash_combine(res, ClipShape.child().get());
-	return res;
-}
+		return false;
+	}
 
 
-mu::Ptr<ASTOp> ASTOpMeshClipDeform::Clone(MapChildFuncRef mapChild) const
-{
-    Ptr<ASTOpMeshClipDeform> n = new ASTOpMeshClipDeform();
-	n->Mesh = mapChild(Mesh.child());
-	n->ClipShape = mapChild(ClipShape.child());
-	return n;
-}
+	uint64 ASTOpMeshClipDeform::Hash() const
+	{
+		uint64 res = std::hash<void*>()(Mesh.child().get());
+		hash_combine(res, ClipShape.child().get());
+		return res;
+	}
 
 
-void ASTOpMeshClipDeform::ForEachChild(const TFunctionRef<void(ASTChild&)> f )
-{
-	f(Mesh);
-	f(ClipShape);
-}
+	mu::Ptr<ASTOp> ASTOpMeshClipDeform::Clone(MapChildFuncRef mapChild) const
+	{
+		Ptr<ASTOpMeshClipDeform> n = new ASTOpMeshClipDeform();
+		n->Mesh = mapChild(Mesh.child());
+		n->ClipShape = mapChild(ClipShape.child());
+		return n;
+	}
 
 
-void ASTOpMeshClipDeform::Link( PROGRAM& program, const FLinkerOptions*)
-{
-    // Already linked?
-    if (!linkedAddress)
-    {
-        OP::MeshClipDeformArgs args;
-        memset( &args,0, sizeof(args) );	
+	void ASTOpMeshClipDeform::ForEachChild(const TFunctionRef<void(ASTChild&)> f)
+	{
+		f(Mesh);
+		f(ClipShape);
+	}
 
-		if (Mesh) 
-		{
-			args.mesh = Mesh->linkedAddress;
-		}
 
-		if (ClipShape)
+	void ASTOpMeshClipDeform::Link(PROGRAM& program, const FLinkerOptions*)
+	{
+		// Already linked?
+		if (!linkedAddress)
 		{
-			 args.clipShape = ClipShape->linkedAddress;
-		}	
+			OP::MeshClipDeformArgs args;
+			memset(&args, 0, sizeof(args));
+
+			if (Mesh)
+			{
+				args.mesh = Mesh->linkedAddress;
+			}
+
+			if (ClipShape)
+			{
+				args.clipShape = ClipShape->linkedAddress;
+			}
+
+			linkedAddress = (OP::ADDRESS)program.m_opAddress.Num();
+			program.m_opAddress.Add((uint32_t)program.m_byteCode.Num());
+			AppendCode(program.m_byteCode, OP_TYPE::ME_CLIPDEFORM);
+			AppendCode(program.m_byteCode, args);
+		}
 
-        linkedAddress = (OP::ADDRESS)program.m_opAddress.Num();
-        program.m_opAddress.Add((uint32_t)program.m_byteCode.Num());
-        AppendCode(program.m_byteCode,OP_TYPE::ME_CLIPDEFORM);
-        AppendCode(program.m_byteCode,args);
-    }
+	}
 
 }
-
diff --git a/Engine/Plugins/Experimental/Mutable/Source/MutableTools/Private/MuT/ASTOpMeshClipMorphPlane.cpp b/Engine/Plugins/Experimental/Mutable/Source/MutableTools/Private/MuT/ASTOpMeshClipMorphPlane.cpp
--- a/Engine/Plugins/Experimental/Mutable/Source/MutableTools/Private/MuT/ASTOpMeshClipMorphPlane.cpp
+++ b/Engine/Plugins/Experimental/Mutable/Source/MutableTools/Private/MuT/ASTOpMeshClipMorphPlane.cpp
@@ -12,103 +12,105 @@
 #include <utility>
 
 
-using namespace mu;
-
-
-//-------------------------------------------------------------------------------------------------
-//-------------------------------------------------------------------------------------------------
-//-------------------------------------------------------------------------------------------------
-ASTOpMeshClipMorphPlane::ASTOpMeshClipMorphPlane()
-    : source(this)
+namespace mu
 {
-}
 
 
-ASTOpMeshClipMorphPlane::~ASTOpMeshClipMorphPlane()
-{
-    // Explicit call needed to avoid recursive destruction
-    ASTOp::RemoveChildren();
-}
-
+	//-------------------------------------------------------------------------------------------------
+	//-------------------------------------------------------------------------------------------------
+	//-------------------------------------------------------------------------------------------------
+	ASTOpMeshClipMorphPlane::ASTOpMeshClipMorphPlane()
+		: source(this)
+	{
+	}
+
+
+	ASTOpMeshClipMorphPlane::~ASTOpMeshClipMorphPlane()
+	{
+		// Explicit call needed to avoid recursive destruction
+		ASTOp::RemoveChildren();
+	}
+
+
+	bool ASTOpMeshClipMorphPlane::IsEqual(const ASTOp& otherUntyped) const
+	{
+		if (auto other = dynamic_cast<const ASTOpMeshClipMorphPlane*>(&otherUntyped))
+		{
+			return source == other->source &&
+				morphShape == other->morphShape &&
+				selectionShape == other->selectionShape &&
+				vertexSelectionBone == other->vertexSelectionBone &&
+				vertexSelectionType == other->vertexSelectionType &&
+				dist == other->dist &&
+				factor == other->factor;
+		}
+		return false;
+	}
+
+
+	uint64 ASTOpMeshClipMorphPlane::Hash() const
+	{
+		uint64 res = std::hash<void*>()(source.child().get());
+		hash_combine(res, factor);
+		hash_combine(res, dist);
+		return res;
+	}
+
+
+	mu::Ptr<ASTOp> ASTOpMeshClipMorphPlane::Clone(MapChildFuncRef mapChild) const
+	{
+		Ptr<ASTOpMeshClipMorphPlane> n = new ASTOpMeshClipMorphPlane();
+		n->source = mapChild(source.child());
+		n->morphShape = morphShape;
+		n->selectionShape = selectionShape;
+		n->vertexSelectionBone = vertexSelectionBone;
+		n->vertexSelectionType = vertexSelectionType;
+		n->dist = dist;
+		n->factor = factor;
+		return n;
+	}
+
+
+	void ASTOpMeshClipMorphPlane::ForEachChild(const TFunctionRef<void(ASTChild&)> f)
+	{
+		f(source);
+	}
+
+
+	void ASTOpMeshClipMorphPlane::Link(PROGRAM& program, const FLinkerOptions*)
+	{
+		// Already linked?
+		if (!linkedAddress)
+		{
+			OP::MeshClipMorphPlaneArgs args;
+			memset(&args, 0, sizeof(args));
+
+			if (source) args.source = source->linkedAddress;
+
+			args.morphShape = (OP::ADDRESS)program.m_constantShapes.Num();
+			program.m_constantShapes.Add(morphShape);
+
+			args.vertexSelectionType = (uint8_t)vertexSelectionType;
+			if (vertexSelectionType == OP::MeshClipMorphPlaneArgs::VS_BONE_HIERARCHY)
+			{
+				args.vertexSelectionShapeOrBone = program.AddConstant(vertexSelectionBone);
+			}
+			else if (vertexSelectionType == OP::MeshClipMorphPlaneArgs::VS_SHAPE)
+			{
+				args.vertexSelectionShapeOrBone = program.AddConstant(selectionShape);
+			}
+
+			args.dist = dist;
+			args.factor = factor;
+			args.maxBoneRadius = vertexSelectionBoneMaxRadius;
+
+			linkedAddress = (OP::ADDRESS)program.m_opAddress.Num();
+			//program.m_code.push_back(op);
+			program.m_opAddress.Add((uint32_t)program.m_byteCode.Num());
+			AppendCode(program.m_byteCode, OP_TYPE::ME_CLIPMORPHPLANE);
+			AppendCode(program.m_byteCode, args);
+		}
+
+	}
 
-bool ASTOpMeshClipMorphPlane::IsEqual(const ASTOp& otherUntyped) const
-{
-    if ( auto other = dynamic_cast<const ASTOpMeshClipMorphPlane*>(&otherUntyped) )
-    {
-        return source==other->source &&
-                morphShape==other->morphShape &&
-                selectionShape==other->selectionShape &&
-                vertexSelectionBone==other->vertexSelectionBone &&
-                vertexSelectionType==other->vertexSelectionType &&
-                dist==other->dist &&
-                factor==other->factor;
-    }
-    return false;
 }
-
-
-uint64 ASTOpMeshClipMorphPlane::Hash() const
-{
-	uint64 res = std::hash<void*>()(source.child().get());
-    hash_combine( res, factor );
-    hash_combine( res, dist );
-    return res;
-}
-
-
-mu::Ptr<ASTOp> ASTOpMeshClipMorphPlane::Clone(MapChildFuncRef mapChild) const
-{
-    Ptr<ASTOpMeshClipMorphPlane> n = new ASTOpMeshClipMorphPlane();
-    n->source = mapChild(source.child());
-    n->morphShape = morphShape;
-    n->selectionShape = selectionShape;
-    n->vertexSelectionBone = vertexSelectionBone;
-    n->vertexSelectionType = vertexSelectionType;
-    n->dist = dist;
-    n->factor = factor;
-    return n;
-}
-
-
-void ASTOpMeshClipMorphPlane::ForEachChild(const TFunctionRef<void(ASTChild&)> f )
-{
-    f( source );
-}
-
-
-void ASTOpMeshClipMorphPlane::Link( PROGRAM& program, const FLinkerOptions* )
-{
-    // Already linked?
-    if (!linkedAddress)
-    {
-        OP::MeshClipMorphPlaneArgs args;
-        memset( &args,0, sizeof(args) );
-
-        if (source) args.source = source->linkedAddress;
-
-        args.morphShape = (OP::ADDRESS)program.m_constantShapes.Num();
-        program.m_constantShapes.Add(morphShape);
-
-        args.vertexSelectionType = (uint8_t)vertexSelectionType;
-        if (vertexSelectionType==OP::MeshClipMorphPlaneArgs::VS_BONE_HIERARCHY)
-        {
-            args.vertexSelectionShapeOrBone = program.AddConstant(vertexSelectionBone);
-        }
-        else if (vertexSelectionType==OP::MeshClipMorphPlaneArgs::VS_SHAPE)
-        {
-            args.vertexSelectionShapeOrBone = program.AddConstant(selectionShape);
-        }
-
-        args.dist = dist;
-        args.factor = factor;
-		args.maxBoneRadius = vertexSelectionBoneMaxRadius;
-
-        linkedAddress = (OP::ADDRESS)program.m_opAddress.Num();
-        //program.m_code.push_back(op);
-        program.m_opAddress.Add((uint32_t)program.m_byteCode.Num());
-        AppendCode(program.m_byteCode,OP_TYPE::ME_CLIPMORPHPLANE);
-        AppendCode(program.m_byteCode,args);
-    }
-
-}
-
diff --git a/Engine/Plugins/Experimental/Mutable/Source/MutableTools/Private/MuT/ASTOpMeshExtractLayoutBlocks.cpp b/Engine/Plugins/Experimental/Mutable/Source/MutableTools/Private/MuT/ASTOpMeshExtractLayoutBlocks.cpp
--- a/Engine/Plugins/Experimental/Mutable/Source/MutableTools/Private/MuT/ASTOpMeshExtractLayoutBlocks.cpp
+++ b/Engine/Plugins/Experimental/Mutable/Source/MutableTools/Private/MuT/ASTOpMeshExtractLayoutBlocks.cpp
@@ -10,85 +10,84 @@
 #include <memory>
 #include <utility>
 
-
-using namespace mu;
-
-
-//---------------------------------------------------------------------------------------------
-//---------------------------------------------------------------------------------------------
-//---------------------------------------------------------------------------------------------
-ASTOpMeshExtractLayoutBlocks::ASTOpMeshExtractLayoutBlocks()
-    : source(this)
-{
-}
-
-
-ASTOpMeshExtractLayoutBlocks::~ASTOpMeshExtractLayoutBlocks()
-{
-    // Explicit call needed to avoid recursive destruction
-    ASTOp::RemoveChildren();
-}
-
-
-bool ASTOpMeshExtractLayoutBlocks::IsEqual(const ASTOp& otherUntyped) const
+namespace mu
 {
-    if ( auto other = dynamic_cast<const ASTOpMeshExtractLayoutBlocks*>(&otherUntyped) )
-    {
-        return source==other->source && layout==other->layout && blocks==other->blocks;
-    }
-    return false;
-}
 
+	//---------------------------------------------------------------------------------------------
+	//---------------------------------------------------------------------------------------------
+	//---------------------------------------------------------------------------------------------
+	ASTOpMeshExtractLayoutBlocks::ASTOpMeshExtractLayoutBlocks()
+		: source(this)
+	{
+	}
+
+
+	ASTOpMeshExtractLayoutBlocks::~ASTOpMeshExtractLayoutBlocks()
+	{
+		// Explicit call needed to avoid recursive destruction
+		ASTOp::RemoveChildren();
+	}
+
+
+	bool ASTOpMeshExtractLayoutBlocks::IsEqual(const ASTOp& otherUntyped) const
+	{
+		if (auto other = dynamic_cast<const ASTOpMeshExtractLayoutBlocks*>(&otherUntyped))
+		{
+			return source == other->source && layout == other->layout && blocks == other->blocks;
+		}
+		return false;
+	}
+
+
+	mu::Ptr<ASTOp> ASTOpMeshExtractLayoutBlocks::Clone(MapChildFuncRef mapChild) const
+	{
+		Ptr<ASTOpMeshExtractLayoutBlocks> n = new ASTOpMeshExtractLayoutBlocks();
+		n->source = mapChild(source.child());
+		n->layout = layout;
+		n->blocks = blocks;
+		return n;
+	}
+
+
+	void ASTOpMeshExtractLayoutBlocks::Assert()
+	{
+		check(blocks.Num() < std::numeric_limits<uint16>::max());
+		ASTOp::Assert();
+	}
+
+
+	void ASTOpMeshExtractLayoutBlocks::ForEachChild(const TFunctionRef<void(ASTChild&)> f)
+	{
+		f(source);
+	}
+
+
+	uint64 ASTOpMeshExtractLayoutBlocks::Hash() const
+	{
+		uint64 res = std::hash<size_t>()(size_t(source.child().get()));
+		return res;
+	}
+
+
+	void ASTOpMeshExtractLayoutBlocks::Link(PROGRAM& program, const FLinkerOptions*)
+	{
+		// Already linked?
+		if (!linkedAddress)
+		{
+			linkedAddress = (OP::ADDRESS)program.m_opAddress.Num();
+
+			program.m_opAddress.Add((uint32)program.m_byteCode.Num());
+			AppendCode(program.m_byteCode, OP_TYPE::ME_EXTRACTLAYOUTBLOCK);
+			OP::ADDRESS sourceAt = source ? source->linkedAddress : 0;
+			AppendCode(program.m_byteCode, sourceAt);
+			AppendCode(program.m_byteCode, (uint16)layout);
+			AppendCode(program.m_byteCode, (uint16)blocks.Num());
+
+			for (auto b : blocks)
+			{
+				AppendCode(program.m_byteCode, (uint32)b);
+			}
+		}
+	}
 
-mu::Ptr<ASTOp> ASTOpMeshExtractLayoutBlocks::Clone(MapChildFuncRef mapChild) const
-{
-    Ptr<ASTOpMeshExtractLayoutBlocks> n = new ASTOpMeshExtractLayoutBlocks();
-    n->source = mapChild(source.child());
-    n->layout = layout;
-    n->blocks = blocks;
-    return n;
 }
-
-
-void ASTOpMeshExtractLayoutBlocks::Assert()
-{
-    check( blocks.Num()<std::numeric_limits<uint16>::max() );
-    ASTOp::Assert();
-}
-
-
-void ASTOpMeshExtractLayoutBlocks::ForEachChild(const TFunctionRef<void(ASTChild&)> f )
-{
-    f(source);
-}
-
-
-uint64 ASTOpMeshExtractLayoutBlocks::Hash() const
-{
-	uint64 res = std::hash<size_t>()( size_t(source.child().get() ) );
-    return res;
-}
-
-
-void ASTOpMeshExtractLayoutBlocks::Link( PROGRAM& program, const FLinkerOptions* )
-{
-    // Already linked?
-    if (!linkedAddress)
-    {
-        linkedAddress = (OP::ADDRESS)program.m_opAddress.Num();
-
-        program.m_opAddress.Add((uint32)program.m_byteCode.Num());
-        AppendCode(program.m_byteCode, OP_TYPE::ME_EXTRACTLAYOUTBLOCK);
-        OP::ADDRESS sourceAt = source ? source->linkedAddress : 0;
-        AppendCode(program.m_byteCode, sourceAt );
-        AppendCode(program.m_byteCode, (uint16)layout );
-        AppendCode(program.m_byteCode, (uint16)blocks.Num() );
-
-        for ( auto b: blocks )
-        {
-            AppendCode(program.m_byteCode, (uint32)b );
-        }
-    }
-}
-
-
diff --git a/Engine/Plugins/Experimental/Mutable/Source/MutableTools/Private/MuT/ASTOpMeshFormat.cpp b/Engine/Plugins/Experimental/Mutable/Source/MutableTools/Private/MuT/ASTOpMeshFormat.cpp
--- a/Engine/Plugins/Experimental/Mutable/Source/MutableTools/Private/MuT/ASTOpMeshFormat.cpp
+++ b/Engine/Plugins/Experimental/Mutable/Source/MutableTools/Private/MuT/ASTOpMeshFormat.cpp
@@ -23,11 +23,11 @@
 #include <utility>
 
 
-using namespace mu;
+
+namespace mu
+{
 
 
-//-------------------------------------------------------------------------------------------------
-//-------------------------------------------------------------------------------------------------
 //-------------------------------------------------------------------------------------------------
 ASTOpMeshFormat::ASTOpMeshFormat()
     : Source(this)
@@ -410,3 +410,5 @@ mu::Ptr<ASTOp> Sink_MeshFormatAST::Visit(const mu::Ptr<ASTOp>& at, const ASTOpMe
 
 	return newAt;
 }
+
+}
diff --git a/Engine/Plugins/Experimental/Mutable/Source/MutableTools/Private/MuT/ASTOpMeshMorphReshape.cpp b/Engine/Plugins/Experimental/Mutable/Source/MutableTools/Private/MuT/ASTOpMeshMorphReshape.cpp
--- a/Engine/Plugins/Experimental/Mutable/Source/MutableTools/Private/MuT/ASTOpMeshMorphReshape.cpp
+++ b/Engine/Plugins/Experimental/Mutable/Source/MutableTools/Private/MuT/ASTOpMeshMorphReshape.cpp
@@ -12,81 +12,85 @@
 #include <memory>
 #include <utility>
 
-using namespace mu;
 
-
-ASTOpMeshMorphReshape::ASTOpMeshMorphReshape()
-    : Morph(this)
-    , Reshape(this)
+namespace mu
 {
-}
 
 
-ASTOpMeshMorphReshape::~ASTOpMeshMorphReshape()
-{
-    // Explicit call needed to avoid recursive destruction
-    ASTOp::RemoveChildren();
-}
+	ASTOpMeshMorphReshape::ASTOpMeshMorphReshape()
+		: Morph(this)
+		, Reshape(this)
+	{
+	}
 
 
-bool ASTOpMeshMorphReshape::IsEqual(const ASTOp& otherUntyped) const
-{
-	if ( auto other = dynamic_cast<const ASTOpMeshMorphReshape*>(&otherUntyped) )
+	ASTOpMeshMorphReshape::~ASTOpMeshMorphReshape()
 	{
-		return Morph == other->Morph && Reshape == other->Reshape;
+		// Explicit call needed to avoid recursive destruction
+		ASTOp::RemoveChildren();
 	}
-	
-    return false;
-}
 
 
-uint64 ASTOpMeshMorphReshape::Hash() const
-{
-	uint64 res = std::hash<void*>()(Morph.child().get() );
-	hash_combine(res, Reshape.child().get());
+	bool ASTOpMeshMorphReshape::IsEqual(const ASTOp& otherUntyped) const
+	{
+		if (auto other = dynamic_cast<const ASTOpMeshMorphReshape*>(&otherUntyped))
+		{
+			return Morph == other->Morph && Reshape == other->Reshape;
+		}
 
-	return res;
-}
+		return false;
+	}
 
 
-mu::Ptr<ASTOp> ASTOpMeshMorphReshape::Clone(MapChildFuncRef mapChild) const
-{
-    Ptr<ASTOpMeshMorphReshape> n = new ASTOpMeshMorphReshape();
-	n->Morph = mapChild(Morph.child());
-	n->Reshape = mapChild(Reshape.child());
-	return n;
-}
+	uint64 ASTOpMeshMorphReshape::Hash() const
+	{
+		uint64 res = std::hash<void*>()(Morph.child().get());
+		hash_combine(res, Reshape.child().get());
 
+		return res;
+	}
 
-void ASTOpMeshMorphReshape::ForEachChild(const TFunctionRef<void(ASTChild&)> f )
-{
-	f(Morph);
-	f(Reshape);
-}
 
+	mu::Ptr<ASTOp> ASTOpMeshMorphReshape::Clone(MapChildFuncRef mapChild) const
+	{
+		Ptr<ASTOpMeshMorphReshape> n = new ASTOpMeshMorphReshape();
+		n->Morph = mapChild(Morph.child());
+		n->Reshape = mapChild(Reshape.child());
+		return n;
+	}
 
-void ASTOpMeshMorphReshape::Link( PROGRAM& program, const FLinkerOptions*)
-{
-    // Already linked?
-    if (!linkedAddress)
-    {
-        OP::MeshMorphReshapeArgs Args;
-        memset( &Args, 0, sizeof(Args) );
-    	
-		if (Morph) 
-		{
-			Args.Morph = Morph->linkedAddress;
-		}
 
-		if (Reshape) 
+	void ASTOpMeshMorphReshape::ForEachChild(const TFunctionRef<void(ASTChild&)> f)
+	{
+		f(Morph);
+		f(Reshape);
+	}
+
+
+	void ASTOpMeshMorphReshape::Link(PROGRAM& program, const FLinkerOptions*)
+	{
+		// Already linked?
+		if (!linkedAddress)
 		{
-			Args.Reshape = Reshape->linkedAddress;
+			OP::MeshMorphReshapeArgs Args;
+			memset(&Args, 0, sizeof(Args));
+
+			if (Morph)
+			{
+				Args.Morph = Morph->linkedAddress;
+			}
+
+			if (Reshape)
+			{
+				Args.Reshape = Reshape->linkedAddress;
+			}
+
+			linkedAddress = (OP::ADDRESS)program.m_opAddress.Num();
+			program.m_opAddress.Add((uint32_t)program.m_byteCode.Num());
+			AppendCode(program.m_byteCode, OP_TYPE::ME_MORPHRESHAPE);
+			AppendCode(program.m_byteCode, Args);
 		}
 
-        linkedAddress = (OP::ADDRESS)program.m_opAddress.Num();
-        program.m_opAddress.Add((uint32_t)program.m_byteCode.Num());
-        AppendCode(program.m_byteCode, OP_TYPE::ME_MORPHRESHAPE);
-        AppendCode(program.m_byteCode, Args);
-    }
+	}
 
 }
diff --git a/Engine/Plugins/Experimental/Mutable/Source/MutableTools/Private/MuT/ASTOpMeshRemapIndices.cpp b/Engine/Plugins/Experimental/Mutable/Source/MutableTools/Private/MuT/ASTOpMeshRemapIndices.cpp
--- a/Engine/Plugins/Experimental/Mutable/Source/MutableTools/Private/MuT/ASTOpMeshRemapIndices.cpp
+++ b/Engine/Plugins/Experimental/Mutable/Source/MutableTools/Private/MuT/ASTOpMeshRemapIndices.cpp
@@ -12,81 +12,84 @@
 #include <memory>
 #include <utility>
 
-using namespace mu;
 
-
-//-------------------------------------------------------------------------------------------------
-//-------------------------------------------------------------------------------------------------
-//-------------------------------------------------------------------------------------------------
-ASTOpMeshRemapIndices::ASTOpMeshRemapIndices()
-    : source(this)
-    , reference(this)
-{
-}
-
-
-//-------------------------------------------------------------------------------------------------
-ASTOpMeshRemapIndices::~ASTOpMeshRemapIndices()
+namespace mu
 {
-    // Explicit call needed to avoid recursive destruction
-    ASTOp::RemoveChildren();
-}
 
 
-//-------------------------------------------------------------------------------------------------
-bool ASTOpMeshRemapIndices::IsEqual(const ASTOp& otherUntyped) const
-{
-    if ( auto other = dynamic_cast<const ASTOpMeshRemapIndices*>(&otherUntyped) )
-    {
-        return source==other->source && reference==other->reference;
-    }
-    return false;
-}
-
+	//-------------------------------------------------------------------------------------------------
+	//-------------------------------------------------------------------------------------------------
+	//-------------------------------------------------------------------------------------------------
+	ASTOpMeshRemapIndices::ASTOpMeshRemapIndices()
+		: source(this)
+		, reference(this)
+	{
+	}
+
+
+	//-------------------------------------------------------------------------------------------------
+	ASTOpMeshRemapIndices::~ASTOpMeshRemapIndices()
+	{
+		// Explicit call needed to avoid recursive destruction
+		ASTOp::RemoveChildren();
+	}
+
+
+	//-------------------------------------------------------------------------------------------------
+	bool ASTOpMeshRemapIndices::IsEqual(const ASTOp& otherUntyped) const
+	{
+		if (auto other = dynamic_cast<const ASTOpMeshRemapIndices*>(&otherUntyped))
+		{
+			return source == other->source && reference == other->reference;
+		}
+		return false;
+	}
+
+
+	//-------------------------------------------------------------------------------------------------
+	uint64 ASTOpMeshRemapIndices::Hash() const
+	{
+		uint64 res = std::hash<void*>()(source.child().get());
+		hash_combine(res, reference.child());
+		return res;
+	}
+
+
+	//-------------------------------------------------------------------------------------------------
+	mu::Ptr<ASTOp> ASTOpMeshRemapIndices::Clone(MapChildFuncRef mapChild) const
+	{
+		Ptr<ASTOpMeshRemapIndices> n = new ASTOpMeshRemapIndices();
+		n->source = mapChild(source.child());
+		n->reference = mapChild(reference.child());
+		return n;
+	}
+
+
+	//-------------------------------------------------------------------------------------------------
+	void ASTOpMeshRemapIndices::ForEachChild(const TFunctionRef<void(ASTChild&)> f)
+	{
+		f(source);
+		f(reference);
+	}
+
+
+	//-------------------------------------------------------------------------------------------------
+	void ASTOpMeshRemapIndices::Link(PROGRAM& program, const FLinkerOptions*)
+	{
+		// Already linked?
+		if (!linkedAddress)
+		{
+			OP::MeshRemapIndicesArgs args;
+			memset(&args, 0, sizeof(args));
+
+			if (source) args.source = source->linkedAddress;
+			if (reference) args.reference = reference->linkedAddress;
+
+			linkedAddress = (OP::ADDRESS)program.m_opAddress.Num();
+			program.m_opAddress.Add((uint32_t)program.m_byteCode.Num());
+			AppendCode(program.m_byteCode, OP_TYPE::ME_REMAPINDICES);
+			AppendCode(program.m_byteCode, args);
+		}
+	}
 
-//-------------------------------------------------------------------------------------------------
-uint64 ASTOpMeshRemapIndices::Hash() const
-{
-	uint64 res = std::hash<void*>()(source.child().get());
-    hash_combine( res, reference.child() );
-    return res;
 }
-
-
-//-------------------------------------------------------------------------------------------------
-mu::Ptr<ASTOp> ASTOpMeshRemapIndices::Clone(MapChildFuncRef mapChild) const
-{
-    Ptr<ASTOpMeshRemapIndices> n = new ASTOpMeshRemapIndices();
-    n->source = mapChild(source.child());
-    n->reference = mapChild(reference.child());
-    return n;
-}
-
-
-//-------------------------------------------------------------------------------------------------
-void ASTOpMeshRemapIndices::ForEachChild(const TFunctionRef<void(ASTChild&)> f )
-{
-    f( source );
-    f( reference );
-}
-
-
-//-------------------------------------------------------------------------------------------------
-void ASTOpMeshRemapIndices::Link( PROGRAM& program, const FLinkerOptions* )
-{
-    // Already linked?
-    if (!linkedAddress)
-    {
-        OP::MeshRemapIndicesArgs args;
-        memset( &args,0, sizeof(args) );
-
-        if (source) args.source = source->linkedAddress;
-        if (reference) args.reference = reference->linkedAddress;
-
-        linkedAddress = (OP::ADDRESS)program.m_opAddress.Num();
-        program.m_opAddress.Add((uint32_t)program.m_byteCode.Num());
-        AppendCode(program.m_byteCode,OP_TYPE::ME_REMAPINDICES);
-        AppendCode(program.m_byteCode,args);
-    }
-}
-
diff --git a/Engine/Plugins/Experimental/Mutable/Source/MutableTools/Private/MuT/ASTOpMeshRemapIndices.h b/Engine/Plugins/Experimental/Mutable/Source/MutableTools/Private/MuT/ASTOpMeshRemapIndices.h
--- a/Engine/Plugins/Experimental/Mutable/Source/MutableTools/Private/MuT/ASTOpMeshRemapIndices.h
+++ b/Engine/Plugins/Experimental/Mutable/Source/MutableTools/Private/MuT/ASTOpMeshRemapIndices.h
@@ -14,7 +14,7 @@ struct PROGRAM;
 
 
 	//---------------------------------------------------------------------------------------------
-	//! 
+	//! \TODO: Deprecated?
 	//---------------------------------------------------------------------------------------------
 	class ASTOpMeshRemapIndices : public ASTOp
 	{
diff --git a/Engine/Plugins/Experimental/Mutable/Source/MutableTools/Private/MuT/ASTOpMeshRemoveMask.cpp b/Engine/Plugins/Experimental/Mutable/Source/MutableTools/Private/MuT/ASTOpMeshRemoveMask.cpp
--- a/Engine/Plugins/Experimental/Mutable/Source/MutableTools/Private/MuT/ASTOpMeshRemoveMask.cpp
+++ b/Engine/Plugins/Experimental/Mutable/Source/MutableTools/Private/MuT/ASTOpMeshRemoveMask.cpp
@@ -11,179 +11,178 @@
 #include <memory>
 
 
-using namespace mu;
-
-
-//---------------------------------------------------------------------------------------------
-//---------------------------------------------------------------------------------------------
-//---------------------------------------------------------------------------------------------
-ASTOpMeshRemoveMask::ASTOpMeshRemoveMask()
-    : source(this)
-{
-}
-
-
-//---------------------------------------------------------------------------------------------
-ASTOpMeshRemoveMask::~ASTOpMeshRemoveMask()
-{
-    // Explicit call needed to avoid recursive destruction
-    ASTOp::RemoveChildren();
-}
-
-
-//---------------------------------------------------------------------------------------------
-void ASTOpMeshRemoveMask::AddRemove( const Ptr<ASTOp>& condition, const Ptr<ASTOp>& mask  )
-{
-	removes.Add({ ASTChild(this,condition), ASTChild(this,mask) });
-}
-
-
-//---------------------------------------------------------------------------------------------
-bool ASTOpMeshRemoveMask::IsEqual(const ASTOp& otherUntyped) const
+namespace mu
 {
-    if ( auto other = dynamic_cast<const ASTOpMeshRemoveMask*>(&otherUntyped) )
-    {
-        return source==other->source && removes==other->removes;
-    }
-    return false;
-}
-
 
-//---------------------------------------------------------------------------------------------
-mu::Ptr<ASTOp> ASTOpMeshRemoveMask::Clone(MapChildFuncRef mapChild) const
-{
-    Ptr<ASTOpMeshRemoveMask> n = new ASTOpMeshRemoveMask();
-    n->source = mapChild(source.child());
-    for(const auto& r:removes)
-    {
-		n->removes.Add({ ASTChild(n,mapChild(r.Key.child())), ASTChild(n,mapChild(r.Value.child())) });
-    }
-    return n;
-}
-
-
-//---------------------------------------------------------------------------------------------
-void ASTOpMeshRemoveMask::Assert()
-{
-    ASTOp::Assert();
-}
 
-
-//---------------------------------------------------------------------------------------------
-void ASTOpMeshRemoveMask::ForEachChild(const TFunctionRef<void(ASTChild&)> f )
-{
-    f(source);
-    for(auto& r:removes)
-    {
-        f(r.Key);
-        f(r.Value);
-    }
-}
-
-
-//---------------------------------------------------------------------------------------------
-uint64 ASTOpMeshRemoveMask::Hash() const
-{
-	uint64 res = std::hash<ASTOp*>()( source.child().get() );
-    for(const auto& r:removes)
-    {
-        hash_combine( res, r.Key.child().get() );
-        hash_combine( res, r.Value.child().get() );
-    }
-    return res;
-}
-
-
-//---------------------------------------------------------------------------------------------
-void ASTOpMeshRemoveMask::Link( PROGRAM& program, const FLinkerOptions*)
-{
-    // Already linked?
-    if (!linkedAddress)
-    {
-        linkedAddress = (OP::ADDRESS)program.m_opAddress.Num();
-
-        program.m_opAddress.Add((uint32_t)program.m_byteCode.Num());
-        AppendCode(program.m_byteCode, OP_TYPE::ME_REMOVEMASK);
-        OP::ADDRESS sourceAt = source ? source->linkedAddress : 0;
-        AppendCode(program.m_byteCode, sourceAt );
-        AppendCode(program.m_byteCode, (uint16)removes.Num() );
-        for ( const auto& b: removes )
-        {
-            OP::ADDRESS condition = b.Key ? b.Key->linkedAddress : 0;
-            AppendCode(program.m_byteCode, condition );
-
-            OP::ADDRESS remove = b.Value ? b.Value->linkedAddress : 0;
-            AppendCode(program.m_byteCode, remove );
-        }
-    }
-}
-
-
-//-------------------------------------------------------------------------------------------------
-//-------------------------------------------------------------------------------------------------
-//-------------------------------------------------------------------------------------------------
-namespace
-{
-    class Sink_MeshRemoveMaskAST
-    {
-    public:
-
-        // \TODO This is recursive and may cause stack overflows in big models.
-
-		mu::Ptr<ASTOp> Apply( const ASTOpMeshRemoveMask* root )
-        {
-            m_root = root;
-            m_oldToNew.Empty();
-
-            m_initialSource = root->source.child();
-			mu::Ptr<ASTOp> newSource = Visit( m_initialSource );
-
-            // If there is any change, it is the new root.
-            if (newSource!=m_initialSource)
-            {
-                return newSource;
-            }
-
-            return nullptr;
-        }
-
-    protected:
-
-        const ASTOpMeshRemoveMask* m_root;
-		mu::Ptr<ASTOp> m_initialSource;
-        TMap<mu::Ptr<ASTOp>, mu::Ptr<ASTOp>> m_oldToNew;
-        TArray<mu::Ptr<ASTOp>> m_newOps;
-
-		mu::Ptr<ASTOp> Visit( const mu::Ptr<ASTOp>& at )
-        {
-            if (!at) return nullptr;
-
-            // Newly created?
-            if (m_newOps.Contains( at ))
-            {
-                return at;
-            }
-
-            // Already visited?
-            auto cacheIt = m_oldToNew.Find(at);
-            if (cacheIt)
-            {
-                return *cacheIt;
-            }
-
-			mu::Ptr<ASTOp> newAt = at;
-            switch ( at->GetOpType() )
-            {
-
-            case OP_TYPE::ME_MORPH2:
-            {
-				mu::Ptr<ASTOpFixed> newOp = mu::Clone<ASTOpFixed>(at);
-                newOp->SetChild( newOp->op.args.MeshMorph2.base, Visit( newOp->children[newOp->op.args.MeshMorph2.base].child() ) );
-                newAt = newOp;
-                break;
-            }
-
-                // disabled to avoid code explosion (or bug?) TODO
+	//---------------------------------------------------------------------------------------------
+	ASTOpMeshRemoveMask::ASTOpMeshRemoveMask()
+		: source(this)
+	{
+	}
+
+
+	//---------------------------------------------------------------------------------------------
+	ASTOpMeshRemoveMask::~ASTOpMeshRemoveMask()
+	{
+		// Explicit call needed to avoid recursive destruction
+		ASTOp::RemoveChildren();
+	}
+
+
+	//---------------------------------------------------------------------------------------------
+	void ASTOpMeshRemoveMask::AddRemove(const Ptr<ASTOp>& condition, const Ptr<ASTOp>& mask)
+	{
+		removes.Add({ ASTChild(this,condition), ASTChild(this,mask) });
+	}
+
+
+	//---------------------------------------------------------------------------------------------
+	bool ASTOpMeshRemoveMask::IsEqual(const ASTOp& otherUntyped) const
+	{
+		if (auto other = dynamic_cast<const ASTOpMeshRemoveMask*>(&otherUntyped))
+		{
+			return source == other->source && removes == other->removes;
+		}
+		return false;
+	}
+
+
+	//---------------------------------------------------------------------------------------------
+	mu::Ptr<ASTOp> ASTOpMeshRemoveMask::Clone(MapChildFuncRef mapChild) const
+	{
+		Ptr<ASTOpMeshRemoveMask> n = new ASTOpMeshRemoveMask();
+		n->source = mapChild(source.child());
+		for (const auto& r : removes)
+		{
+			n->removes.Add({ ASTChild(n,mapChild(r.Key.child())), ASTChild(n,mapChild(r.Value.child())) });
+		}
+		return n;
+	}
+
+
+	//---------------------------------------------------------------------------------------------
+	void ASTOpMeshRemoveMask::Assert()
+	{
+		ASTOp::Assert();
+	}
+
+
+	//---------------------------------------------------------------------------------------------
+	void ASTOpMeshRemoveMask::ForEachChild(const TFunctionRef<void(ASTChild&)> f)
+	{
+		f(source);
+		for (auto& r : removes)
+		{
+			f(r.Key);
+			f(r.Value);
+		}
+	}
+
+
+	//---------------------------------------------------------------------------------------------
+	uint64 ASTOpMeshRemoveMask::Hash() const
+	{
+		uint64 res = std::hash<ASTOp*>()(source.child().get());
+		for (const auto& r : removes)
+		{
+			hash_combine(res, r.Key.child().get());
+			hash_combine(res, r.Value.child().get());
+		}
+		return res;
+	}
+
+
+	//---------------------------------------------------------------------------------------------
+	void ASTOpMeshRemoveMask::Link(PROGRAM& program, const FLinkerOptions*)
+	{
+		// Already linked?
+		if (!linkedAddress)
+		{
+			linkedAddress = (OP::ADDRESS)program.m_opAddress.Num();
+
+			program.m_opAddress.Add((uint32_t)program.m_byteCode.Num());
+			AppendCode(program.m_byteCode, OP_TYPE::ME_REMOVEMASK);
+			OP::ADDRESS sourceAt = source ? source->linkedAddress : 0;
+			AppendCode(program.m_byteCode, sourceAt);
+			AppendCode(program.m_byteCode, (uint16)removes.Num());
+			for (const auto& b : removes)
+			{
+				OP::ADDRESS condition = b.Key ? b.Key->linkedAddress : 0;
+				AppendCode(program.m_byteCode, condition);
+
+				OP::ADDRESS remove = b.Value ? b.Value->linkedAddress : 0;
+				AppendCode(program.m_byteCode, remove);
+			}
+		}
+	}
+
+
+	//-------------------------------------------------------------------------------------------------
+	//-------------------------------------------------------------------------------------------------
+	//-------------------------------------------------------------------------------------------------
+	namespace
+	{
+		class Sink_MeshRemoveMaskAST
+		{
+		public:
+
+			// \TODO This is recursive and may cause stack overflows in big models.
+
+			mu::Ptr<ASTOp> Apply(const ASTOpMeshRemoveMask* root)
+			{
+				m_root = root;
+				m_oldToNew.Empty();
+
+				m_initialSource = root->source.child();
+				mu::Ptr<ASTOp> newSource = Visit(m_initialSource);
+
+				// If there is any change, it is the new root.
+				if (newSource != m_initialSource)
+				{
+					return newSource;
+				}
+
+				return nullptr;
+			}
+
+		protected:
+
+			const ASTOpMeshRemoveMask* m_root;
+			mu::Ptr<ASTOp> m_initialSource;
+			TMap<mu::Ptr<ASTOp>, mu::Ptr<ASTOp>> m_oldToNew;
+			TArray<mu::Ptr<ASTOp>> m_newOps;
+
+			mu::Ptr<ASTOp> Visit(const mu::Ptr<ASTOp>& at)
+			{
+				if (!at) return nullptr;
+
+				// Newly created?
+				if (m_newOps.Contains(at))
+				{
+					return at;
+				}
+
+				// Already visited?
+				auto cacheIt = m_oldToNew.Find(at);
+				if (cacheIt)
+				{
+					return *cacheIt;
+				}
+
+				mu::Ptr<ASTOp> newAt = at;
+				switch (at->GetOpType())
+				{
+
+				case OP_TYPE::ME_MORPH2:
+				{
+					mu::Ptr<ASTOpFixed> newOp = mu::Clone<ASTOpFixed>(at);
+					newOp->SetChild(newOp->op.args.MeshMorph2.base, Visit(newOp->children[newOp->op.args.MeshMorph2.base].child()));
+					newAt = newOp;
+					break;
+				}
+
+				// disabled to avoid code explosion (or bug?) TODO
 //            case OP_TYPE::ME_CONDITIONAL:
 //            {
 //                Ptr<ASTOpConditional> newOp = mu::Clone<ASTOpConditional>(at);
@@ -205,34 +204,35 @@ namespace
 //                break;
 //            }
 
-            default:
-            {
-                //
-                if (at!=m_initialSource)
-                {
-                    auto newOp = mu::Clone<ASTOpMeshRemoveMask>(m_root);
-                    newOp->source = at;
-                    newAt = newOp;
-                }
-                break;
-            }
-
-            }
-
-            m_oldToNew.Add(at, newAt);
-
-            return newAt;
-        }
-    };
-}
+				default:
+				{
+					//
+					if (at != m_initialSource)
+					{
+						auto newOp = mu::Clone<ASTOpMeshRemoveMask>(m_root);
+						newOp->source = at;
+						newAt = newOp;
+					}
+					break;
+				}
 
+				}
 
-//-------------------------------------------------------------------------------------------------
-mu::Ptr<ASTOp> ASTOpMeshRemoveMask::OptimiseSink(const MODEL_OPTIMIZATION_OPTIONS&, OPTIMIZE_SINK_CONTEXT&) const
-{
-    Sink_MeshRemoveMaskAST sinker;
-    auto at = sinker.Apply(this);
+				m_oldToNew.Add(at, newAt);
 
-    return at;
-}
+				return newAt;
+			}
+		};
+	}
+
+
+	//-------------------------------------------------------------------------------------------------
+	mu::Ptr<ASTOp> ASTOpMeshRemoveMask::OptimiseSink(const MODEL_OPTIMIZATION_OPTIONS&, OPTIMIZE_SINK_CONTEXT&) const
+	{
+		Sink_MeshRemoveMaskAST sinker;
+		auto at = sinker.Apply(this);
 
+		return at;
+	}
+
+}
diff --git a/Engine/Plugins/Experimental/Mutable/Source/MutableTools/Private/MuT/ASTOpMeshTransform.cpp b/Engine/Plugins/Experimental/Mutable/Source/MutableTools/Private/MuT/ASTOpMeshTransform.cpp
--- a/Engine/Plugins/Experimental/Mutable/Source/MutableTools/Private/MuT/ASTOpMeshTransform.cpp
+++ b/Engine/Plugins/Experimental/Mutable/Source/MutableTools/Private/MuT/ASTOpMeshTransform.cpp
@@ -13,76 +13,75 @@
 #include <utility>
 
 
-using namespace mu;
+namespace mu
+{
 
+	ASTOpMeshTransform::ASTOpMeshTransform()
+		: source(this)
+	{
+	}
 
-//-------------------------------------------------------------------------------------------------
-//-------------------------------------------------------------------------------------------------
-//-------------------------------------------------------------------------------------------------
-ASTOpMeshTransform::ASTOpMeshTransform()
-    : source(this)
-{
-}
 
+	ASTOpMeshTransform::~ASTOpMeshTransform()
+	{
+		// Explicit call needed to avoid recursive destruction
+		ASTOp::RemoveChildren();
+	}
 
-ASTOpMeshTransform::~ASTOpMeshTransform()
-{
-    // Explicit call needed to avoid recursive destruction
-    ASTOp::RemoveChildren();
-}
 
+	bool ASTOpMeshTransform::IsEqual(const ASTOp& otherUntyped) const
+	{
+		if (auto other = dynamic_cast<const ASTOpMeshTransform*>(&otherUntyped))
+		{
+			return source == other->source &&
+				matrix == other->matrix;
+		}
+		return false;
+	}
 
-bool ASTOpMeshTransform::IsEqual(const ASTOp& otherUntyped) const
-{
-    if ( auto other = dynamic_cast<const ASTOpMeshTransform*>(&otherUntyped) )
-    {
-        return source==other->source &&
-                matrix==other->matrix;
-    }
-    return false;
-}
 
+	uint64 ASTOpMeshTransform::Hash() const
+	{
+		uint64 res = std::hash<OP_TYPE>()(OP_TYPE::ME_TRANSFORM);
+		hash_combine(res, source.child().get());
+		return res;
+	}
 
-uint64 ASTOpMeshTransform::Hash() const
-{
-	uint64 res = std::hash<OP_TYPE>()(OP_TYPE::ME_TRANSFORM);
-    hash_combine( res, source.child().get() );
-    return res;
-}
 
+	mu::Ptr<ASTOp> ASTOpMeshTransform::Clone(MapChildFuncRef mapChild) const
+	{
+		Ptr<ASTOpMeshTransform> n = new ASTOpMeshTransform();
+		n->matrix = matrix;
+		n->source = mapChild(source.child());
+		return n;
+	}
 
-mu::Ptr<ASTOp> ASTOpMeshTransform::Clone(MapChildFuncRef mapChild) const
-{
-    Ptr<ASTOpMeshTransform> n = new ASTOpMeshTransform();
-    n->matrix = matrix;
-    n->source = mapChild(source.child());
-    return n;
-}
 
+	void ASTOpMeshTransform::ForEachChild(const TFunctionRef<void(ASTChild&)> f)
+	{
+		f(source);
+	}
 
-void ASTOpMeshTransform::ForEachChild(const TFunctionRef<void(ASTChild&)> f )
-{
-    f( source );
-}
 
+	void ASTOpMeshTransform::Link(PROGRAM& program, const FLinkerOptions*)
+	{
+		// Already linked?
+		if (!linkedAddress)
+		{
+			OP::MeshTransformArgs args;
+			memset(&args, 0, sizeof(args));
 
-void ASTOpMeshTransform::Link( PROGRAM& program, const FLinkerOptions*)
-{
-    // Already linked?
-    if (!linkedAddress)
-    {
-        OP::MeshTransformArgs args;
-        memset( &args,0, sizeof(args) );
-
-        if (source) args.source = source->linkedAddress;
-
-        args.matrix = program.AddConstant(matrix);
-
-        linkedAddress = (OP::ADDRESS)program.m_opAddress.Num();
-        //program.m_code.push_back(op);
-        program.m_opAddress.Add((uint32_t)program.m_byteCode.Num());
-        AppendCode(program.m_byteCode,OP_TYPE::ME_TRANSFORM);
-        AppendCode(program.m_byteCode,args);
-    }
+			if (source) args.source = source->linkedAddress;
+
+			args.matrix = program.AddConstant(matrix);
+
+			linkedAddress = (OP::ADDRESS)program.m_opAddress.Num();
+			//program.m_code.push_back(op);
+			program.m_opAddress.Add((uint32_t)program.m_byteCode.Num());
+			AppendCode(program.m_byteCode, OP_TYPE::ME_TRANSFORM);
+			AppendCode(program.m_byteCode, args);
+		}
+
+	}
 
 }
diff --git a/Engine/Plugins/Experimental/Mutable/Source/MutableTools/Private/MuT/ASTOpParameter.cpp b/Engine/Plugins/Experimental/Mutable/Source/MutableTools/Private/MuT/ASTOpParameter.cpp
--- a/Engine/Plugins/Experimental/Mutable/Source/MutableTools/Private/MuT/ASTOpParameter.cpp
+++ b/Engine/Plugins/Experimental/Mutable/Source/MutableTools/Private/MuT/ASTOpParameter.cpp
@@ -14,151 +14,150 @@
 #include <utility>
 
 
-using namespace mu;
-
-
-//-------------------------------------------------------------------------------------------------
-//-------------------------------------------------------------------------------------------------
-//-------------------------------------------------------------------------------------------------
-ASTOpParameter::~ASTOpParameter()
-{
-    // Explicit call needed to avoid recursive destruction
-    ASTOp::RemoveChildren();
-}
-
-
-void ASTOpParameter::ForEachChild(const TFunctionRef<void(ASTChild&)> f)
-{
-    for (auto &c: ranges)
-    {
-        f(c.rangeSize);
-    }
-    for (auto &c: additionalImages)
-    {
-        f(c);
-    }
-}
-
-
-bool ASTOpParameter::IsEqual(const ASTOp& otherUntyped) const
-{
-    if ( auto other = dynamic_cast<const ASTOpParameter*>(&otherUntyped) )
-    {
-        return type==other->type &&
-                parameter==other->parameter &&
-                ranges==other->ranges &&
-                additionalImages==other->additionalImages;
-    }
-    return false;
-}
-
-
-mu::Ptr<ASTOp> ASTOpParameter::Clone(MapChildFuncRef mapChild) const
-{
-    Ptr<ASTOpParameter> n = new ASTOpParameter();
-    n->type = type;
-    n->parameter = parameter;
-    for(const auto& c:ranges)
-    {
-        n->ranges.Emplace( n.get(), mapChild(c.rangeSize.child()), c.rangeName, c.rangeUID );
-    }
-    for(const auto& c:additionalImages)
-    {
-        n->additionalImages.Emplace(n,mapChild(c.child()));
-    }
-    return n;
-}
-
-
-uint64 ASTOpParameter::Hash() const
+namespace mu
 {
-	uint64 res = std::hash<uint64>()(uint64(type));
-    hash_combine( res, parameter.m_type );
-    hash_combine( res, parameter.m_name );
-    return res;
-}
-
-
-void ASTOpParameter::Assert()
-{
-    ASTOp::Assert();
-}
-
-
-void ASTOpParameter::Link( PROGRAM& program, const FLinkerOptions* )
-{
-    // Already linked?
-    if (!linkedAddress)
-    {
-        OP::ParameterArgs args;
-        memset( &args,0, sizeof(args) );
-
-        args.variable = (OP::ADDRESS)program.m_parameters.Num();
-        program.m_parameters.Add(parameter);
-
-        for( const auto& d: ranges)
-        {
-            OP::ADDRESS sizeAt = 0;
-            uint16 rangeId = 0;
-            LinkRange( program, d, sizeAt, rangeId );
-            program.m_parameters.Last().m_ranges.Add(rangeId);
-        }
-
-        for( const auto& d: additionalImages)
-        {
-            OP::ADDRESS descAt = d ? d->linkedAddress : 0;
-            program.m_parameters.Last().m_descImages.Add(descAt);
-        }
-
-        linkedAddress = (OP::ADDRESS)program.m_opAddress.Num();
-        //program.m_code.push_back(op);
-
-        program.m_opAddress.Add((uint32_t)program.m_byteCode.Num());
-        AppendCode(program.m_byteCode,type);
-        AppendCode(program.m_byteCode,args);
-    }
-}
 
 
-int ASTOpParameter::EvaluateInt( ASTOpList& facts, bool &unknown ) const
-{
-    unknown = true;
-
-    // Check the facts, in case we have the value for our parameter.
-    for( const auto& f: facts )
-    {
-        if (f->GetOpType()==OP_TYPE::BO_EQUAL_INT_CONST)
-        {
-            const auto typedFact = dynamic_cast<const ASTOpFixed*>(f.get());
-            auto value = typedFact->children[typedFact->op.args.BoolEqualScalarConst.value].child();
-            if (value.get()==this)
-            {
-                unknown = false;
-                return typedFact->op.args.BoolEqualScalarConst.constant;
-            }
-            else
-            {
-                // We could try something more if it was an expression and it had the parameter
-                // somewhere in it.
-            }
-        }
-    }
-
-    return 0;
-}
-
+	ASTOpParameter::~ASTOpParameter()
+	{
+		// Explicit call needed to avoid recursive destruction
+		ASTOp::RemoveChildren();
+	}
+
+
+	void ASTOpParameter::ForEachChild(const TFunctionRef<void(ASTChild&)> f)
+	{
+		for (auto& c : ranges)
+		{
+			f(c.rangeSize);
+		}
+		for (auto& c : additionalImages)
+		{
+			f(c);
+		}
+	}
+
+
+	bool ASTOpParameter::IsEqual(const ASTOp& otherUntyped) const
+	{
+		if (auto other = dynamic_cast<const ASTOpParameter*>(&otherUntyped))
+		{
+			return type == other->type &&
+				parameter == other->parameter &&
+				ranges == other->ranges &&
+				additionalImages == other->additionalImages;
+		}
+		return false;
+	}
+
+
+	mu::Ptr<ASTOp> ASTOpParameter::Clone(MapChildFuncRef mapChild) const
+	{
+		Ptr<ASTOpParameter> n = new ASTOpParameter();
+		n->type = type;
+		n->parameter = parameter;
+		for (const auto& c : ranges)
+		{
+			n->ranges.Emplace(n.get(), mapChild(c.rangeSize.child()), c.rangeName, c.rangeUID);
+		}
+		for (const auto& c : additionalImages)
+		{
+			n->additionalImages.Emplace(n, mapChild(c.child()));
+		}
+		return n;
+	}
+
+
+	uint64 ASTOpParameter::Hash() const
+	{
+		uint64 res = std::hash<uint64>()(uint64(type));
+		hash_combine(res, parameter.m_type);
+		hash_combine(res, parameter.m_name);
+		return res;
+	}
+
+
+	void ASTOpParameter::Assert()
+	{
+		ASTOp::Assert();
+	}
+
+
+	void ASTOpParameter::Link(PROGRAM& program, const FLinkerOptions*)
+	{
+		// Already linked?
+		if (!linkedAddress)
+		{
+			OP::ParameterArgs args;
+			memset(&args, 0, sizeof(args));
+
+			args.variable = (OP::ADDRESS)program.m_parameters.Num();
+			program.m_parameters.Add(parameter);
+
+			for (const auto& d : ranges)
+			{
+				OP::ADDRESS sizeAt = 0;
+				uint16 rangeId = 0;
+				LinkRange(program, d, sizeAt, rangeId);
+				program.m_parameters.Last().m_ranges.Add(rangeId);
+			}
+
+			for (const auto& d : additionalImages)
+			{
+				OP::ADDRESS descAt = d ? d->linkedAddress : 0;
+				program.m_parameters.Last().m_descImages.Add(descAt);
+			}
+
+			linkedAddress = (OP::ADDRESS)program.m_opAddress.Num();
+			//program.m_code.push_back(op);
+
+			program.m_opAddress.Add((uint32_t)program.m_byteCode.Num());
+			AppendCode(program.m_byteCode, type);
+			AppendCode(program.m_byteCode, args);
+		}
+	}
+
+
+	int ASTOpParameter::EvaluateInt(ASTOpList& facts, bool& unknown) const
+	{
+		unknown = true;
+
+		// Check the facts, in case we have the value for our parameter.
+		for (const auto& f : facts)
+		{
+			if (f->GetOpType() == OP_TYPE::BO_EQUAL_INT_CONST)
+			{
+				const auto typedFact = dynamic_cast<const ASTOpFixed*>(f.get());
+				auto value = typedFact->children[typedFact->op.args.BoolEqualScalarConst.value].child();
+				if (value.get() == this)
+				{
+					unknown = false;
+					return typedFact->op.args.BoolEqualScalarConst.constant;
+				}
+				else
+				{
+					// We could try something more if it was an expression and it had the parameter
+					// somewhere in it.
+				}
+			}
+		}
+
+		return 0;
+	}
+
+
+	ASTOp::BOOL_EVAL_RESULT ASTOpParameter::EvaluateBool(ASTOpList& /*facts*/,
+		EVALUATE_BOOL_CACHE*) const
+	{
+		check(type == OP_TYPE::BO_PARAMETER);
+		return BET_UNKNOWN;
+	}
+
+
+	FImageDesc ASTOpParameter::GetImageDesc(bool, GetImageDescContext*)
+	{
+		check(type == OP_TYPE::IM_PARAMETER);
+		return FImageDesc();
+	}
 
-ASTOp::BOOL_EVAL_RESULT ASTOpParameter::EvaluateBool( ASTOpList& /*facts*/,
-                                                      EVALUATE_BOOL_CACHE* ) const
-{
-    check(type==OP_TYPE::BO_PARAMETER);
-    return BET_UNKNOWN;
 }
-
-
-FImageDesc ASTOpParameter::GetImageDesc(bool, GetImageDescContext*)
-{
-	check(type == OP_TYPE::IM_PARAMETER);
-	return FImageDesc();
-}
-
diff --git a/Engine/Plugins/Experimental/Mutable/Source/MutableTools/Private/MuT/ASTOpScalarCurve.cpp b/Engine/Plugins/Experimental/Mutable/Source/MutableTools/Private/MuT/ASTOpScalarCurve.cpp
--- a/Engine/Plugins/Experimental/Mutable/Source/MutableTools/Private/MuT/ASTOpScalarCurve.cpp
+++ b/Engine/Plugins/Experimental/Mutable/Source/MutableTools/Private/MuT/ASTOpScalarCurve.cpp
@@ -13,74 +13,73 @@
 #include <utility>
 
 
-using namespace mu;
+namespace mu
+{
 
 
-//-------------------------------------------------------------------------------------------------
-//-------------------------------------------------------------------------------------------------
-//-------------------------------------------------------------------------------------------------
-ASTOpScalarCurve::ASTOpScalarCurve()
-    : time(this)
-{
-}
+	ASTOpScalarCurve::ASTOpScalarCurve()
+		: time(this)
+	{
+	}
 
 
-ASTOpScalarCurve::~ASTOpScalarCurve()
-{
-    // Explicit call needed to avoid recursive destruction
-    ASTOp::RemoveChildren();
-}
+	ASTOpScalarCurve::~ASTOpScalarCurve()
+	{
+		// Explicit call needed to avoid recursive destruction
+		ASTOp::RemoveChildren();
+	}
 
 
-void ASTOpScalarCurve::ForEachChild(const TFunctionRef<void(ASTChild&)> f)
-{
-    f( time );
-}
+	void ASTOpScalarCurve::ForEachChild(const TFunctionRef<void(ASTChild&)> f)
+	{
+		f(time);
+	}
 
 
-uint64 ASTOpScalarCurve::Hash() const
-{
-	uint64 res = std::hash<uint64>()(size_t(OP_TYPE::SC_CURVE));
-    hash_combine( res, curve.keyFrames.Num() );
-    return res;
-}
+	uint64 ASTOpScalarCurve::Hash() const
+	{
+		uint64 res = std::hash<uint64>()(size_t(OP_TYPE::SC_CURVE));
+		hash_combine(res, curve.keyFrames.Num());
+		return res;
+	}
 
 
-bool ASTOpScalarCurve::IsEqual(const ASTOp& otherUntyped) const
-{
-    if ( auto other = dynamic_cast<const ASTOpScalarCurve*>(&otherUntyped) )
-    {
-        return time==other->time && curve==other->curve;
-    }
-    return false;
-}
+	bool ASTOpScalarCurve::IsEqual(const ASTOp& otherUntyped) const
+	{
+		if (auto other = dynamic_cast<const ASTOpScalarCurve*>(&otherUntyped))
+		{
+			return time == other->time && curve == other->curve;
+		}
+		return false;
+	}
 
 
-mu::Ptr<ASTOp> ASTOpScalarCurve::Clone(MapChildFuncRef mapChild) const
-{
-    Ptr<ASTOpScalarCurve> n = new ASTOpScalarCurve();
-    n->curve = curve;
-    n->time = mapChild(time.child());
-    return n;
-}
+	mu::Ptr<ASTOp> ASTOpScalarCurve::Clone(MapChildFuncRef mapChild) const
+	{
+		Ptr<ASTOpScalarCurve> n = new ASTOpScalarCurve();
+		n->curve = curve;
+		n->time = mapChild(time.child());
+		return n;
+	}
 
 
-void ASTOpScalarCurve::Link( PROGRAM& program, const FLinkerOptions* )
-{
-    if (!linkedAddress)
-    {
-        OP::ScalarCurveArgs args;
-        memset( &args,0, sizeof(args) );
-        args.time = time->linkedAddress;
-        args.curve = program.AddConstant(curve);
-
-        linkedAddress = (OP::ADDRESS)program.m_opAddress.Num();
-        //program.m_code.push_back(op);
-        program.m_opAddress.Add((uint32_t)program.m_byteCode.Num());
-        AppendCode(program.m_byteCode,OP_TYPE::SC_CURVE);
-        AppendCode(program.m_byteCode,args);
-    }
+	void ASTOpScalarCurve::Link(PROGRAM& program, const FLinkerOptions*)
+	{
+		if (!linkedAddress)
+		{
+			OP::ScalarCurveArgs args;
+			memset(&args, 0, sizeof(args));
+			args.time = time ? time->linkedAddress : 0;
+			args.curve = program.AddConstant(curve);
 
-}
+			linkedAddress = (OP::ADDRESS)program.m_opAddress.Num();
+			//program.m_code.push_back(op);
+			program.m_opAddress.Add((uint32_t)program.m_byteCode.Num());
+			AppendCode(program.m_byteCode, OP_TYPE::SC_CURVE);
+			AppendCode(program.m_byteCode, args);
+		}
 
+	}
+
+}
 
