[mutable] Fix morphs not having any effect.

diff --git a/Engine/Plugins/Experimental/Mutable/Source/MutableTools/Private/MuT/ASTOpSwitch.cpp b/Engine/Plugins/Experimental/Mutable/Source/MutableTools/Private/MuT/ASTOpSwitch.cpp
--- a/Engine/Plugins/Experimental/Mutable/Source/MutableTools/Private/MuT/ASTOpSwitch.cpp
+++ b/Engine/Plugins/Experimental/Mutable/Source/MutableTools/Private/MuT/ASTOpSwitch.cpp
@@ -15,497 +15,498 @@
 #include <utility>
 
 
-using namespace mu;
+namespace mu
+{
 
 
-//---------------------------------------------------------------------------------------------
-//---------------------------------------------------------------------------------------------
-//---------------------------------------------------------------------------------------------
-ASTOpSwitch::ASTOpSwitch()
-    : variable(this)
-    , def(this)
-{
-}
+	//---------------------------------------------------------------------------------------------
+	ASTOpSwitch::ASTOpSwitch()
+		: variable(this)
+		, def(this)
+	{
+	}
 
 
-//-------------------------------------------------------------------------------------------------
-ASTOpSwitch::~ASTOpSwitch()
-{
-    // Explicit call needed to avoid recursive destruction
-    ASTOp::RemoveChildren();
-}
+	//-------------------------------------------------------------------------------------------------
+	ASTOpSwitch::~ASTOpSwitch()
+	{
+		// Explicit call needed to avoid recursive destruction
+		ASTOp::RemoveChildren();
+	}
 
 
-//-------------------------------------------------------------------------------------------------
-bool ASTOpSwitch::IsEqual(const ASTOp& otherUntyped) const
-{
-    if ( auto other = dynamic_cast<const ASTOpSwitch*>(&otherUntyped) )
-    {
-        return type==other->type && variable==other->variable &&
-                cases==other->cases && def==other->def;
-    }
-    return false;
-}
+	//-------------------------------------------------------------------------------------------------
+	bool ASTOpSwitch::IsEqual(const ASTOp& otherUntyped) const
+	{
+		if (auto other = dynamic_cast<const ASTOpSwitch*>(&otherUntyped))
+		{
+			return type == other->type && variable == other->variable &&
+				cases == other->cases && def == other->def;
+		}
+		return false;
+	}
 
 
-//-------------------------------------------------------------------------------------------------
-mu::Ptr<ASTOp> ASTOpSwitch::Clone(MapChildFuncRef mapChild) const
-{
-    Ptr<ASTOpSwitch> n = new ASTOpSwitch();
-    n->type = type;
-    n->variable = mapChild(variable.child());
-    n->def = mapChild(def.child());
-    for(const auto& c:cases)
-    {
-        n->cases.Emplace(c.condition, n, mapChild(c.branch.child()));
-    }
-    return n;
-}
+	//-------------------------------------------------------------------------------------------------
+	mu::Ptr<ASTOp> ASTOpSwitch::Clone(MapChildFuncRef mapChild) const
+	{
+		Ptr<ASTOpSwitch> n = new ASTOpSwitch();
+		n->type = type;
+		n->variable = mapChild(variable.child());
+		n->def = mapChild(def.child());
+		for (const auto& c : cases)
+		{
+			n->cases.Emplace(c.condition, n, mapChild(c.branch.child()));
+		}
+		return n;
+	}
 
 
-//-------------------------------------------------------------------------------------------------
-void ASTOpSwitch::Assert()
-{
-    switch(type)
-    {
-    case OP_TYPE::NU_SWITCH:
-    case OP_TYPE::SC_SWITCH:
-    case OP_TYPE::CO_SWITCH:
-    case OP_TYPE::IM_SWITCH:
-    case OP_TYPE::ME_SWITCH:
-    case OP_TYPE::LA_SWITCH:
-    case OP_TYPE::IN_SWITCH:
-        break;
-    default:
-        // Unexpected type
-        check(false);
-        break;
-    }
-
-    ASTOp::Assert();
-}
+	//-------------------------------------------------------------------------------------------------
+	void ASTOpSwitch::Assert()
+	{
+		switch (type)
+		{
+		case OP_TYPE::NU_SWITCH:
+		case OP_TYPE::SC_SWITCH:
+		case OP_TYPE::CO_SWITCH:
+		case OP_TYPE::IM_SWITCH:
+		case OP_TYPE::ME_SWITCH:
+		case OP_TYPE::LA_SWITCH:
+		case OP_TYPE::IN_SWITCH:
+			break;
+		default:
+			// Unexpected type
+			check(false);
+			break;
+		}
 
+		ASTOp::Assert();
+	}
 
-//-------------------------------------------------------------------------------------------------
-uint64 ASTOpSwitch::Hash() const
-{
-	uint64 res = std::hash<uint64>()(uint64(type));
-    for (const auto& c: cases)
-    {
-        hash_combine( res, c.condition );
-        hash_combine( res, c.branch.child().get() );
-    }
-    return res;
-}
 
+	//-------------------------------------------------------------------------------------------------
+	uint64 ASTOpSwitch::Hash() const
+	{
+		uint64 res = std::hash<uint64>()(uint64(type));
+		for (const auto& c : cases)
+		{
+			hash_combine(res, c.condition);
+			hash_combine(res, c.branch.child().get());
+		}
+		return res;
+	}
 
-//-------------------------------------------------------------------------------------------------
-mu::Ptr<ASTOp> ASTOpSwitch::GetFirstValidValue()
-{
-    for ( int32 i=0; i<cases.Num(); ++i )
-    {
-        if ( cases[i].branch )
-        {
-            return cases[i].branch.child();
-        }
-    }
-    return nullptr;
-}
 
+	//-------------------------------------------------------------------------------------------------
+	mu::Ptr<ASTOp> ASTOpSwitch::GetFirstValidValue()
+	{
+		for (int32 i = 0; i < cases.Num(); ++i)
+		{
+			if (cases[i].branch)
+			{
+				return cases[i].branch.child();
+			}
+		}
+		return nullptr;
+	}
 
-//-------------------------------------------------------------------------------------------------
-bool ASTOpSwitch::IsCompatibleWith( const ASTOpSwitch* other ) const
-{
-    if (!other) return false;
-    if (variable.child()!=other->variable.child()) return false;
-    if (cases.Num()!=other->cases.Num()) return false;
-    for (const auto& c: cases)
-    {
-        bool found = false;
-        for (const auto& o: other->cases)
-        {
-            if (c.condition==o.condition)
-            {
-                found=true;
-                break;
-            }
-        }
-        if (!found)
-        {
-            return false;
-        }
-    }
-
-    return true;
-}
 
+	//-------------------------------------------------------------------------------------------------
+	bool ASTOpSwitch::IsCompatibleWith(const ASTOpSwitch* other) const
+	{
+		if (!other) return false;
+		if (variable.child() != other->variable.child()) return false;
+		if (cases.Num() != other->cases.Num()) return false;
+		for (const auto& c : cases)
+		{
+			bool found = false;
+			for (const auto& o : other->cases)
+			{
+				if (c.condition == o.condition)
+				{
+					found = true;
+					break;
+				}
+			}
+			if (!found)
+			{
+				return false;
+			}
+		}
 
-//-------------------------------------------------------------------------------------------------
-mu::Ptr<ASTOp> ASTOpSwitch::FindBranch( int32_t condition ) const
-{
-    for (const auto& c: cases)
-    {
-        if (c.condition==condition)
-        {
-            return c.branch.child();
-        }
-    }
-
-    return def.child();
-}
+		return true;
+	}
 
 
-//-------------------------------------------------------------------------------------------------
-void ASTOpSwitch::ForEachChild(const TFunctionRef<void(ASTChild&)> f )
-{
-    f(variable);
-    f(def);
-    for ( auto& cas: cases)
-    {
-        f(cas.branch);
-    }
-}
+	//-------------------------------------------------------------------------------------------------
+	mu::Ptr<ASTOp> ASTOpSwitch::FindBranch(int32_t condition) const
+	{
+		for (const auto& c : cases)
+		{
+			if (c.condition == condition)
+			{
+				return c.branch.child();
+			}
+		}
 
+		return def.child();
+	}
 
-//-------------------------------------------------------------------------------------------------
-void ASTOpSwitch::Link( PROGRAM& program, const FLinkerOptions* )
-{
-    // Already linked?
-    if (!linkedAddress)
-    {        
-        linkedAddress = (OP::ADDRESS)program.m_opAddress.Num();
-        program.m_opAddress.Add((uint32_t)program.m_byteCode.Num());
- 
-        OP::ADDRESS VarAddress = variable ? variable->linkedAddress : 0;
-        OP::ADDRESS DefAddress = def ? def->linkedAddress : 0;
- 
-        AppendCode( program.m_byteCode, type );
-        AppendCode( program.m_byteCode, VarAddress );
-        AppendCode( program.m_byteCode, DefAddress );
-        AppendCode( program.m_byteCode, (uint32_t)cases.Num() );
-
-        for ( const CASE& Case : cases )
-        {
-            OP::ADDRESS CaseBranchAddress = Case.branch ? Case.branch->linkedAddress : 0;  
-            AppendCode( program.m_byteCode, Case.condition );
-            AppendCode( program.m_byteCode, CaseBranchAddress );
-        }
-    }
-}
 
-//-------------------------------------------------------------------------------------------------
-FImageDesc ASTOpSwitch::GetImageDesc( bool returnBestOption, class GetImageDescContext* context )
-{
-    FImageDesc res;
-
-    // Local context in case it is necessary
-    GetImageDescContext localContext;
-    if (!context)
-    {
-        context = &localContext;
-    }
-    else
-    {
-        // Cached result?
-		FImageDesc* PtrValue = context->m_results.Find(this);
-		if (PtrValue)
+	//-------------------------------------------------------------------------------------------------
+	void ASTOpSwitch::ForEachChild(const TFunctionRef<void(ASTChild&)> f)
+	{
+		f(variable);
+		f(def);
+		for (auto& cas : cases)
 		{
-			return *PtrValue;
+			f(cas.branch);
 		}
 	}
 
-    // In a switch we cannot guarantee the size and format.
-    // We check all the options, and if they are the same we return that.
-    // Otherwise, we return a descriptor with empty fields in the conflicting ones, size or format.
-    // In some places this will force re-formatting of the image.
-    // The code optimiser will take care then of moving the format operations down to each
-    // branch and remove the unnecessary ones.
-    FImageDesc candidate;
-    bool sameSize = true;
-    bool sameFormat = true;
-    bool sameLods = true;
-    bool first = true;
-
-    if ( def )
-    {
-        FImageDesc childDesc = def->GetImageDesc( returnBestOption, context );
-        candidate = childDesc;
-        first = false;
-    }
-
-    for ( int i=0; i<cases.Num(); ++i )
-    {
-        if ( cases[i].branch )
-        {
-            FImageDesc childDesc = cases[i].branch->GetImageDesc(  returnBestOption, context );
-            if ( first )
-            {
-                candidate = childDesc;
-                first = false;
-            }
-            else
-            {
-                sameSize = (candidate.m_size == childDesc.m_size);
-                sameFormat = (candidate.m_format == childDesc.m_format);
-                sameLods = (candidate.m_lods == childDesc.m_lods);
-
-                if (returnBestOption)
-                {
-                    candidate.m_format =
-                        GetMostGenericFormat( candidate.m_format, childDesc.m_format );
-                }
-            }
-          }
-      }
-
-
-    res = candidate;
-
-    if (!sameFormat && !returnBestOption)
-    {
-        res.m_format = EImageFormat::IF_NONE;
-    }
-
-    if (!sameSize && !returnBestOption)
-    {
-        res.m_size = FImageSize(0, 0);
-    }
-
-    if (!sameLods && !returnBestOption)
-    {
-        res.m_lods = 0;
-    }
-
-    // Cache the result
-    if (context)
-    {
-		context->m_results.Add(this, res);
+
+	//-------------------------------------------------------------------------------------------------
+	void ASTOpSwitch::Link(PROGRAM& program, const FLinkerOptions*)
+	{
+		// Already linked?
+		if (!linkedAddress)
+		{
+			linkedAddress = (OP::ADDRESS)program.m_opAddress.Num();
+			program.m_opAddress.Add((uint32_t)program.m_byteCode.Num());
+
+			OP::ADDRESS VarAddress = variable ? variable->linkedAddress : 0;
+			OP::ADDRESS DefAddress = def ? def->linkedAddress : 0;
+
+			AppendCode(program.m_byteCode, type);
+			AppendCode(program.m_byteCode, VarAddress);
+			AppendCode(program.m_byteCode, DefAddress);
+			AppendCode(program.m_byteCode, (uint32_t)cases.Num());
+
+			for (const CASE& Case : cases)
+			{
+				OP::ADDRESS CaseBranchAddress = Case.branch ? Case.branch->linkedAddress : 0;
+				AppendCode(program.m_byteCode, Case.condition);
+				AppendCode(program.m_byteCode, CaseBranchAddress);
+			}
+		}
 	}
 
-    return res;
-}
+	//-------------------------------------------------------------------------------------------------
+	FImageDesc ASTOpSwitch::GetImageDesc(bool returnBestOption, class GetImageDescContext* context)
+	{
+		FImageDesc res;
 
+		// Local context in case it is necessary
+		GetImageDescContext localContext;
+		if (!context)
+		{
+			context = &localContext;
+		}
+		else
+		{
+			// Cached result?
+			FImageDesc* PtrValue = context->m_results.Find(this);
+			if (PtrValue)
+			{
+				return *PtrValue;
+			}
+		}
 
-//-------------------------------------------------------------------------------------------------
-void ASTOpSwitch::GetBlockLayoutSize( int blockIndex, int* pBlockX, int* pBlockY,
-                                      BLOCK_LAYOUT_SIZE_CACHE* cache )
-{
-    switch ( type )
-    {
-    case OP_TYPE::LA_SWITCH:
-    {
-        Ptr<ASTOp> child = GetFirstValidValue();
-        if (!child)
-        {
-            child = def.child();
-        }
-
-        if (child)
-        {
-            child->GetBlockLayoutSizeCached( blockIndex, pBlockX, pBlockY, cache );
-        }
-        else
-        {
-            *pBlockX = 0;
-            *pBlockY = 0;
-        }
-        break;
-    }
-
-    default:
-        check(false);
-        break;
-    }
-}
+		// In a switch we cannot guarantee the size and format.
+		// We check all the options, and if they are the same we return that.
+		// Otherwise, we return a descriptor with empty fields in the conflicting ones, size or format.
+		// In some places this will force re-formatting of the image.
+		// The code optimiser will take care then of moving the format operations down to each
+		// branch and remove the unnecessary ones.
+		FImageDesc candidate;
+		bool sameSize = true;
+		bool sameFormat = true;
+		bool sameLods = true;
+		bool first = true;
+
+		if (def)
+		{
+			FImageDesc childDesc = def->GetImageDesc(returnBestOption, context);
+			candidate = childDesc;
+			first = false;
+		}
 
+		for (int i = 0; i < cases.Num(); ++i)
+		{
+			if (cases[i].branch)
+			{
+				FImageDesc childDesc = cases[i].branch->GetImageDesc(returnBestOption, context);
+				if (first)
+				{
+					candidate = childDesc;
+					first = false;
+				}
+				else
+				{
+					sameSize = (candidate.m_size == childDesc.m_size);
+					sameFormat = (candidate.m_format == childDesc.m_format);
+					sameLods = (candidate.m_lods == childDesc.m_lods);
+
+					if (returnBestOption)
+					{
+						candidate.m_format =
+							GetMostGenericFormat(candidate.m_format, childDesc.m_format);
+					}
+				}
+			}
+		}
 
-//-------------------------------------------------------------------------------------------------
-void ASTOpSwitch::GetLayoutBlockSize( int* pBlockX, int* pBlockY )
-{
-    switch ( type )
-    {
-    case OP_TYPE::IM_SWITCH:
-    {
-        Ptr<ASTOp> child = GetFirstValidValue();
-        if (!child)
-        {
-            child = def.child();
-        }
-
-        if (child)
-        {
-            child->GetLayoutBlockSize(pBlockX, pBlockY);
-        }
-        else
-        {
-            checkf(false, TEXT("Image switch had no options.") );
-        }
-        break;
-    }
-
-    default:
-        checkf( false, TEXT("Instruction not supported") );
-    }
-}
 
+		res = candidate;
 
-//-------------------------------------------------------------------------------------------------
-bool ASTOpSwitch::GetNonBlackRect( FImageRect& maskUsage ) const
-{
-    if (type == OP_TYPE::IM_SWITCH )
-    {
-        FImageRect local;
-        bool localValid = false;
-        if (def)
-        {
-            localValid = def->GetNonBlackRect(local);
-            if (!localValid)
-            {
-                return false;
-            }
-        }
-
-        for (const auto& c: cases)
-        {
-            if (c.branch)
-            {
-                FImageRect branchRect;
-                bool validBranch = c.branch->GetNonBlackRect(branchRect);
-                if (validBranch)
-                {
-                    if (localValid)
-                    {
-                        local.Bound(branchRect);
-                    }
-                    else
-                    {
-                        local = branchRect;
-                        localValid = true;
-                    }
-                }
-                else
-                {
-                    return false;
-                }
-            }
-        }
-
-        if (localValid)
-        {
-            maskUsage = local;
-            return true;
-        }
-    }
-
-    return false;
-}
+		if (!sameFormat && !returnBestOption)
+		{
+			res.m_format = EImageFormat::IF_NONE;
+		}
 
+		if (!sameSize && !returnBestOption)
+		{
+			res.m_size = FImageSize(0, 0);
+		}
 
-//-------------------------------------------------------------------------------------------------
-bool ASTOpSwitch::IsImagePlainConstant( vec4<float>& ) const
-{
-    // We could check if every option is plain and exactly the same colour, but probably it is
-    // not worth.
-    return false;
-}
+		if (!sameLods && !returnBestOption)
+		{
+			res.m_lods = 0;
+		}
+
+		// Cache the result
+		if (context)
+		{
+			context->m_results.Add(this, res);
+		}
 
+		return res;
+	}
 
-//-------------------------------------------------------------------------------------------------
-mu::Ptr<ASTOp> ASTOpSwitch::OptimiseSemantic( const MODEL_OPTIMIZATION_OPTIONS& ) const
-{
-    // Constant condition?
-    if  ( variable->GetOpType() == OP_TYPE::NU_CONSTANT )
-    {
-        Ptr<ASTOp> branch = def.child();
-
-        auto typedCondition = dynamic_cast<const ASTOpFixed*>(variable.child().get());
-        for ( int32 o=0; o<cases.Num(); ++o )
-        {
-            if ( cases[o].branch &&
-                typedCondition->op.args.IntConstant.value
-                 ==
-                 (int)cases[o].condition)
-            {
-                branch = cases[o].branch.child();
-                break;
-            }
-        }
-
-        return branch;
-    }
-
-    //else
-    //{
-        // If all the possible branches are the same, remove the instruction
-        // TODO: If the variables is a parameter, checking that all the possible
-        // values of this parameter are handled.
-
-//            Disable this: I think it doesn't deal correctly with def. If all branches are the same but the
-//              value is none of the options, "def" should be selected.
-//            OP::ADDRESS switchAt = at;
-//            bool same = true;
-//            OP::ADDRESS value = 0;
-//            //vector< UINT16 > options;
-
-//            const auto& options = program.m_constantSwitches[program.m_code[switchAt].args.Switch.options].m_options;
-//            for ( size_t o=0; o<options.size(); ++o )
-//            {
-//                if ( options[o].at )
-//                {
-//                    if (!value)
-//                    {
-//                        value = options[o].at;
-//                    }
-//                    else
-//                    {
-//                        same &= ( value == options[o].at );
-//                    }
-
-//                    //options.push_back
-//                    //		( program.m_code[switchAt].args.Switch.conditions[o] );
-//                }
-//            }
-
-
-//            if (same)
-//            {
-//                m_modified = true;
-
-//                at = value;
-//            }
-
-    //}
-
-    return nullptr;
-}
 
+	//-------------------------------------------------------------------------------------------------
+	void ASTOpSwitch::GetBlockLayoutSize(int blockIndex, int* pBlockX, int* pBlockY,
+		BLOCK_LAYOUT_SIZE_CACHE* cache)
+	{
+		switch (type)
+		{
+		case OP_TYPE::LA_SWITCH:
+		{
+			Ptr<ASTOp> child = GetFirstValidValue();
+			if (!child)
+			{
+				child = def.child();
+			}
+
+			if (child)
+			{
+				child->GetBlockLayoutSizeCached(blockIndex, pBlockX, pBlockY, cache);
+			}
+			else
+			{
+				*pBlockX = 0;
+				*pBlockY = 0;
+			}
+			break;
+		}
 
-//-------------------------------------------------------------------------------------------------
-mu::Ptr<ImageSizeExpression> ASTOpSwitch::GetImageSizeExpression() const
-{
-    Ptr<ImageSizeExpression> pRes = new ImageSizeExpression;
-
-    bool first = true;
-    for (const auto& c : cases )
-    {
-        if (c.branch)
-        {
-            if (first)
-            {
-                pRes = c.branch->GetImageSizeExpression();
-            }
-            else
-            {
-                Ptr<ImageSizeExpression> pOther = c.branch->GetImageSizeExpression();
-                if ( !(*pOther==*pRes) )
-                {
-                    pRes->type = ImageSizeExpression::ISET_UNKNOWN;
-                    break;
-                }
-            }
-        }
-    }
-
-    return pRes;
-}
+		default:
+			check(false);
+			break;
+		}
+	}
 
+
+	//-------------------------------------------------------------------------------------------------
+	void ASTOpSwitch::GetLayoutBlockSize(int* pBlockX, int* pBlockY)
+	{
+		switch (type)
+		{
+		case OP_TYPE::IM_SWITCH:
+		{
+			Ptr<ASTOp> child = GetFirstValidValue();
+			if (!child)
+			{
+				child = def.child();
+			}
+
+			if (child)
+			{
+				child->GetLayoutBlockSize(pBlockX, pBlockY);
+			}
+			else
+			{
+				checkf(false, TEXT("Image switch had no options."));
+			}
+			break;
+		}
+
+		default:
+			checkf(false, TEXT("Instruction not supported"));
+		}
+	}
+
+
+	//-------------------------------------------------------------------------------------------------
+	bool ASTOpSwitch::GetNonBlackRect(FImageRect& maskUsage) const
+	{
+		if (type == OP_TYPE::IM_SWITCH)
+		{
+			FImageRect local;
+			bool localValid = false;
+			if (def)
+			{
+				localValid = def->GetNonBlackRect(local);
+				if (!localValid)
+				{
+					return false;
+				}
+			}
+
+			for (const auto& c : cases)
+			{
+				if (c.branch)
+				{
+					FImageRect branchRect;
+					bool validBranch = c.branch->GetNonBlackRect(branchRect);
+					if (validBranch)
+					{
+						if (localValid)
+						{
+							local.Bound(branchRect);
+						}
+						else
+						{
+							local = branchRect;
+							localValid = true;
+						}
+					}
+					else
+					{
+						return false;
+					}
+				}
+			}
+
+			if (localValid)
+			{
+				maskUsage = local;
+				return true;
+			}
+		}
+
+		return false;
+	}
+
+
+	//-------------------------------------------------------------------------------------------------
+	bool ASTOpSwitch::IsImagePlainConstant(vec4<float>&) const
+	{
+		// We could check if every option is plain and exactly the same colour, but probably it is
+		// not worth.
+		return false;
+	}
+
+
+	//-------------------------------------------------------------------------------------------------
+	mu::Ptr<ASTOp> ASTOpSwitch::OptimiseSemantic(const MODEL_OPTIMIZATION_OPTIONS&) const
+	{
+		// Constant condition?
+		if (variable->GetOpType() == OP_TYPE::NU_CONSTANT)
+		{
+			Ptr<ASTOp> branch = def.child();
+
+			auto typedCondition = dynamic_cast<const ASTOpFixed*>(variable.child().get());
+			for (int32 o = 0; o < cases.Num(); ++o)
+			{
+				if (cases[o].branch &&
+					typedCondition->op.args.IntConstant.value
+					==
+					(int)cases[o].condition)
+				{
+					branch = cases[o].branch.child();
+					break;
+				}
+			}
+
+			return branch;
+		}
+
+		//else
+		//{
+			// If all the possible branches are the same, remove the instruction
+			// TODO: If the variables is a parameter, checking that all the possible
+			// values of this parameter are handled.
+
+	//            Disable this: I think it doesn't deal correctly with def. If all branches are the same but the
+	//              value is none of the options, "def" should be selected.
+	//            OP::ADDRESS switchAt = at;
+	//            bool same = true;
+	//            OP::ADDRESS value = 0;
+	//            //vector< UINT16 > options;
+
+	//            const auto& options = program.m_constantSwitches[program.m_code[switchAt].args.Switch.options].m_options;
+	//            for ( size_t o=0; o<options.size(); ++o )
+	//            {
+	//                if ( options[o].at )
+	//                {
+	//                    if (!value)
+	//                    {
+	//                        value = options[o].at;
+	//                    }
+	//                    else
+	//                    {
+	//                        same &= ( value == options[o].at );
+	//                    }
+
+	//                    //options.push_back
+	//                    //		( program.m_code[switchAt].args.Switch.conditions[o] );
+	//                }
+	//            }
+
+
+	//            if (same)
+	//            {
+	//                m_modified = true;
+
+	//                at = value;
+	//            }
+
+		//}
+
+		return nullptr;
+	}
+
+
+	//-------------------------------------------------------------------------------------------------
+	mu::Ptr<ImageSizeExpression> ASTOpSwitch::GetImageSizeExpression() const
+	{
+		Ptr<ImageSizeExpression> pRes = new ImageSizeExpression;
+
+		bool first = true;
+		for (const auto& c : cases)
+		{
+			if (c.branch)
+			{
+				if (first)
+				{
+					pRes = c.branch->GetImageSizeExpression();
+				}
+				else
+				{
+					Ptr<ImageSizeExpression> pOther = c.branch->GetImageSizeExpression();
+					if (!(*pOther == *pRes))
+					{
+						pRes->type = ImageSizeExpression::ISET_UNKNOWN;
+						break;
+					}
+				}
+			}
+		}
+
+		return pRes;
+	}
+
+
+}
diff --git a/Engine/Plugins/Experimental/Mutable/Source/MutableTools/Private/MuT/CodeGenerator.cpp b/Engine/Plugins/Experimental/Mutable/Source/MutableTools/Private/MuT/CodeGenerator.cpp
--- a/Engine/Plugins/Experimental/Mutable/Source/MutableTools/Private/MuT/CodeGenerator.cpp
+++ b/Engine/Plugins/Experimental/Mutable/Source/MutableTools/Private/MuT/CodeGenerator.cpp
@@ -275,7 +275,6 @@ namespace mu
         m_generatedTables.clear();
         m_firstPass = FirstPassGenerator();
         m_currentBottomUpState = BOTTOM_UP_STATE();
-        m_overrideLayoutsStack.Empty();
         m_imageState.Empty();
         m_currentParents.Empty();
         m_currentObject.Empty();
@@ -336,8 +335,14 @@ namespace mu
 
 		if (auto MeshNode = dynamic_cast<const NodeMesh*>(pNode.get()))
 		{
-			MESH_GENERATION_RESULT MeshResult;
-			GenerateMesh(MeshResult, MeshNode);
+			// This should only happen with mesh root nodes
+			FMeshGenerationOptions MeshOptions;
+			MeshOptions.bUniqueVertexIDs = true;
+			MeshOptions.bLayouts = true;
+			MeshOptions.State = m_currentStateIndex;
+
+			FMeshGenerationResult MeshResult;
+			GenerateMesh(MeshOptions, MeshResult, MeshNode);
 			return MeshResult.meshOp;
 		}
 
@@ -862,7 +867,7 @@ namespace mu
         }
 
 
-        MESH_GENERATION_RESULT meshResults;
+        FMeshGenerationResult meshResults;
 
 
 		// We don't add the mesh here, since it will be added directly at the top of the
@@ -915,7 +920,13 @@ namespace mu
 			m_activeTags.Add(node.m_tags);
 
             // Generate the mesh
-            GenerateMesh( meshResults, pMesh );
+			FMeshGenerationOptions MeshOptions;
+			MeshOptions.bUniqueVertexIDs = true;
+			MeshOptions.bLayouts = true;
+			MeshOptions.State = m_currentStateIndex;
+			MeshOptions.ActiveTags = m_activeTags.Last();
+
+            GenerateMesh(MeshOptions, meshResults, pMesh );
             lastMeshOp = meshResults.meshOp;
             meshResults.extraMeshLayouts.SetNum( edits.Num() );
 
@@ -938,8 +949,13 @@ namespace mu
 						// Store the data necessary to apply modifiers for the pre-normal operations stage.
 						m_activeTags.Add(e.node->m_tags);
 						
-						MESH_GENERATION_RESULT addResults;
-                        GenerateMesh( addResults, pAdd );
+						FMeshGenerationOptions MergedMeshOptions;
+						MergedMeshOptions.bUniqueVertexIDs = true;
+						MergedMeshOptions.bLayouts = true;
+						MergedMeshOptions.State = m_currentStateIndex;
+						MergedMeshOptions.ActiveTags = m_activeTags.Last();
+						FMeshGenerationResult addResults;
+                        GenerateMesh(MergedMeshOptions, addResults, pAdd);
 
 						m_activeTags.Pop();
 
@@ -950,7 +966,7 @@ namespace mu
 						lastMeshOp = ApplyMeshModifiers(lastMeshOp, e.node->m_tags,
 							bModifiersForBeforeOperations, node.m_errorContext);
 
-                        MESH_GENERATION_RESULT::EXTRA_LAYOUTS data;
+                        FMeshGenerationResult::EXTRA_LAYOUTS data;
                         data.layouts = addResults.layouts;
                         data.condition = e.condition;
                         data.meshFragment = addResults.meshOp;
@@ -992,8 +1008,17 @@ namespace mu
                 {
                     if ( NodeMeshPtr pRemove = e.node->m_pMesh->GetRemove() )
                     {
-                        MESH_GENERATION_RESULT removeResults;
-                        GenerateMesh( removeResults, pRemove );
+                        FMeshGenerationResult removeResults;
+						FMeshGenerationOptions RemoveMeshOptions;
+						RemoveMeshOptions.bUniqueVertexIDs = false;
+						RemoveMeshOptions.bLayouts = false;
+						RemoveMeshOptions.State = m_currentStateIndex;
+						if (!m_activeTags.IsEmpty())
+						{
+							RemoveMeshOptions.ActiveTags = m_activeTags.Last();
+						}
+
+                        GenerateMesh(RemoveMeshOptions, removeResults, pRemove );
 
                         Ptr<ASTOpFixed> maskOp = new ASTOpFixed();
                         maskOp->op.type = OP_TYPE::ME_MASKDIFF;
@@ -1034,21 +1059,27 @@ namespace mu
             {
                 if ( NodeMeshPtr pMorph = e.node->m_pMorph )
                 {
+					// Not needed because it has been generated before already.
 					// Base mesh
-					MESH_GENERATION_RESULT baseMesh;
-					GenerateMesh(baseMesh, pMesh);
+					//FMeshGenerationResult baseMesh;
+					//GenerateMesh(baseMesh, pMesh);
 
                     // Target mesh
-                    MESH_GENERATION_RESULT morphResult;
-                    GenerateMesh( morphResult, pMorph );
+					FMeshGenerationOptions MorphTargetMeshOptions;
+					MorphTargetMeshOptions.bUniqueVertexIDs = false;
+					MorphTargetMeshOptions.bLayouts = false;
+					MorphTargetMeshOptions.State = m_currentStateIndex;
+
+                    FMeshGenerationResult morphResult;
+                    GenerateMesh(MorphTargetMeshOptions, morphResult, pMorph );
 
 					// BaseMorph generation through mesh diff
 					Ptr<ASTOpFixed> diffBase;
 					{
 						Ptr<ASTOpFixed> op = new ASTOpFixed();
 						op->op.type = OP_TYPE::ME_DIFFERENCE;
-						op->SetChild(op->op.args.MeshDifference.base, baseMesh.meshOp);
-						op->SetChild(op->op.args.MeshDifference.target, baseMesh.meshOp);
+						op->SetChild(op->op.args.MeshDifference.base, meshResults.meshOp);
+						op->SetChild(op->op.args.MeshDifference.target, meshResults.meshOp);
 
 						// Morphing tex coords here is not supported:
 						// Generating the homogoneous UVs is difficult since we don't have the base
@@ -1094,8 +1125,6 @@ namespace mu
 						}
 
 						// Base		
-						MESH_GENERATION_RESULT test;
-						GenerateMesh(test, pMesh);
 						op->SetChild(op->op.args.MeshMorph2.base, lastMeshOp);
 
 						// Targets
@@ -2258,8 +2287,13 @@ namespace mu
 				op->source = preModifiersMesh;
 
 				// Parameters
-				MESH_GENERATION_RESULT clipResult;
-				GenerateMesh(clipResult, TypedClipNode->m_clipMesh);
+				FMeshGenerationOptions ClipOptions;
+				ClipOptions.bUniqueVertexIDs = false;
+				ClipOptions.bLayouts = false;
+				ClipOptions.State = m_currentStateIndex;
+
+				FMeshGenerationResult clipResult;
+				GenerateMesh(ClipOptions, clipResult, TypedClipNode->m_clipMesh);
 				op->clip = clipResult.meshOp;
 
 				if (!op->clip)
@@ -2380,8 +2414,13 @@ namespace mu
 				Ptr<ASTOpMeshBindShape>  BindOp = new ASTOpMeshBindShape();
 				Ptr<ASTOpMeshClipDeform> ClipOp = new ASTOpMeshClipDeform();
 
-				MESH_GENERATION_RESULT ClipShapeResult;
-				GenerateMesh(ClipShapeResult, TypedClipNode->ClipMesh);
+				FMeshGenerationOptions ClipOptions;
+				ClipOptions.bUniqueVertexIDs = false;
+				ClipOptions.bLayouts = false;
+				ClipOptions.State = m_currentStateIndex;
+
+				FMeshGenerationResult ClipShapeResult;
+				GenerateMesh(ClipOptions, ClipShapeResult, TypedClipNode->ClipMesh);
 				ClipOp->ClipShape = ClipShapeResult.meshOp;
 				
 				BindOp->Mesh = lastMeshOp;
@@ -2423,25 +2462,5 @@ namespace mu
 
 		return lastMeshOp;
 	}
-	//---------------------------------------------------------------------------------------------
-//	void CodeGenerator::GetSurfacesWithTag(
-//		const string& tag,
-//		vector<FirstPassGenerator::SURFACE>& surfaces)
-//	{
-//		MUTABLE_CPUPROFILER_SCOPE(GetSurfacesWithTag);
-
-//		for (const auto& m : m_firstPass.surfaces)
-//		{
-//			// Matching tags?
-//            bool found = std::find(m.node->GetPrivate()->m_tags.begin(), m.node->GetPrivate()->m_tags.end(),tag)
-//				!=
-//                m.node->GetPrivate()->m_tags.end();
-
-//			if (found)
-//			{
-//				surfaces.push_back(m);
-//			}
-//		}
-//	}
 
 }
diff --git a/Engine/Plugins/Experimental/Mutable/Source/MutableTools/Private/MuT/CodeGenerator.h b/Engine/Plugins/Experimental/Mutable/Source/MutableTools/Private/MuT/CodeGenerator.h
--- a/Engine/Plugins/Experimental/Mutable/Source/MutableTools/Private/MuT/CodeGenerator.h
+++ b/Engine/Plugins/Experimental/Mutable/Source/MutableTools/Private/MuT/CodeGenerator.h
@@ -113,6 +113,8 @@ namespace mu
 	struct OBJECT_STATE;
 	struct PROGRAM;
 
+
+
     //---------------------------------------------------------------------------------------------
     //! Code generator
     //---------------------------------------------------------------------------------------------
@@ -232,20 +234,15 @@ namespace mu
         //! palceholders for missing images.
         ImagePtr m_missingImage[size_t(EImageFormat::IF_COUNT)];
 
-        //! If this has something, while generating meshes, the layouts will be ignored, because
-        //! they are supposed to match some other set of layouts. If the vector is empty, layouts
-        //! are generated normally.
-		TArray< TArray<Ptr<const Layout>> > m_overrideLayoutsStack;
-
         //! Map of layouts found in the code already generated. The map is from the source layout
         //! pointer to the cloned layout. The cloned layout will have absolute block ids assigned.
         TMap<Ptr<const Layout>,Ptr<const Layout>> m_addedLayouts;
 
         //! First free index for a layout block
-        int32_t m_absoluteLayoutIndex = 0;
+        int32 m_absoluteLayoutIndex = 0;
 
         //! First free index to be used to identify mesh vertices.
-        uint32_t m_freeVertexIndex = 0;
+        uint32 m_freeVertexIndex = 0;
 
         //! When generating images, here we have the entire source image size and the rect of the
         //! image that we are generating.
@@ -253,7 +250,7 @@ namespace mu
         {
             vec2<int> m_imageSize;
             box< vec2<int> > m_imageRect;
-            int32_t m_layoutBlock;
+            int32 m_layoutBlock;
             LayoutPtrConst m_pLayout;
         };
 		TArray<IMAGE_STATE> m_imageState;
@@ -375,10 +372,6 @@ namespace mu
 			if (!m_activeTags.IsEmpty())
 			{
 				key.activeTags = m_activeTags.Last();
-				if (!m_overrideLayoutsStack.IsEmpty())
-				{
-					key.overrideLayouts = m_overrideLayoutsStack.Last();
-				}
 			}
 			return key;
 		}
@@ -447,27 +440,115 @@ namespace mu
 
         //-----------------------------------------------------------------------------------------
         // Meshes
-        typedef TMap<VISITED_MAP_KEY,MESH_GENERATION_RESULT> GeneratedMeshMap;
+
+		/** Options that affect the generation of meshes. It is like list of what required data we want
+		* while parsing down the mesh node graph.
+		*/
+		struct FMeshGenerationOptions
+		{
+			/** TODO: Review and document. */
+			int State = 0;
+			TArray<mu::string> ActiveTags;
+
+			/** Whatever mesh we reach at the leaves of the graph will need to have unique ids for its vertices.
+			* This is used to track mesh removal indices, morph data in other nodes, clothing data, etc.
+			*/
+			bool bUniqueVertexIDs = false;
+
+			/** The meshes at the leaves will need their own layout block data. */
+			bool bLayouts = false;
+
+			/** If this has something the layouts in constant meshes will be ignored, because
+			* they are supposed to match some other set of layouts. If the vector is empty, layouts
+			* are generated normally.
+			*/
+			TArray<Ptr<const Layout>> OverrideLayouts;			
+
+			friend FORCEINLINE uint32 GetTypeHash(const FMeshGenerationOptions& InKey)
+			{
+				uint32 KeyHash = 0;
+				KeyHash = HashCombine(KeyHash, ::GetTypeHash(InKey.bUniqueVertexIDs));
+				KeyHash = HashCombine(KeyHash, ::GetTypeHash(InKey.bLayouts));
+				KeyHash = HashCombine(KeyHash, ::GetTypeHash(InKey.OverrideLayouts.Num()));
+				KeyHash = HashCombine(KeyHash, ::GetTypeHash(InKey.State));
+				KeyHash = HashCombine(KeyHash, ::GetTypeHash(InKey.ActiveTags.Num()));
+				return KeyHash;
+			}
+
+			FORCEINLINE bool operator==(const FMeshGenerationOptions& Other) const
+			{
+				return State==Other.State 
+					&& bUniqueVertexIDs==Other.bUniqueVertexIDs && bLayouts==Other.bLayouts
+					&& ActiveTags==Other.ActiveTags
+					&& OverrideLayouts ==Other.OverrideLayouts;
+			}
+
+		};
+
+		//! Store the results of the code generation of a mesh.
+		struct FMeshGenerationResult
+		{
+			//! Mesh after all code tree is applied
+			Ptr<ASTOp> meshOp;
+
+			//! Original base mesh before removes, morphs, etc.
+			Ptr<ASTOp> baseMeshOp;
+
+			/** Source node layouts to use with these extra mesh. They don't have block ids. */
+			TArray<Ptr<const Layout>> layouts;
+
+			TArray<Ptr<ASTOp>> layoutOps;
+
+			struct EXTRA_LAYOUTS
+			{
+				/** Source node layouts to use with these extra mesh. They don't have block ids. */
+				TArray<Ptr<const Layout>> layouts;
+				Ptr<ASTOp> condition;
+				Ptr<ASTOp> meshFragment;
+			};
+			TArray< EXTRA_LAYOUTS > extraMeshLayouts;
+		};
+		
+		struct FGeneratedMeshCacheKey
+		{
+			NodePtrConst Node;
+			FMeshGenerationOptions Options;
+
+			friend FORCEINLINE uint32 GetTypeHash(const FGeneratedMeshCacheKey& InKey)
+			{
+				uint32 KeyHash = 0;
+				KeyHash = HashCombine(KeyHash, ::GetTypeHash(InKey.Node.get()));
+				KeyHash = HashCombine(KeyHash, GetTypeHash(InKey.Options));
+				return KeyHash;
+			}
+
+			FORCEINLINE bool operator==(const FGeneratedMeshCacheKey& Other) const
+			{
+				return Node == Other.Node && Options == Other.Options;
+			}
+		};
+
+        typedef TMap<FGeneratedMeshCacheKey,FMeshGenerationResult> GeneratedMeshMap;
         GeneratedMeshMap m_generatedMeshes;
 
-        void GenerateMesh( MESH_GENERATION_RESULT& result, const NodeMeshPtrConst& node);
-        void GenerateMesh_Constant( MESH_GENERATION_RESULT&, const NodeMeshConstant* );
-        void GenerateMesh_Format( MESH_GENERATION_RESULT&, const NodeMeshFormat* );
-        void GenerateMesh_Morph( MESH_GENERATION_RESULT&, const NodeMeshMorph* );
-        void GenerateMesh_MakeMorph( MESH_GENERATION_RESULT&, const NodeMeshMakeMorph* );
-        void GenerateMesh_Fragment( MESH_GENERATION_RESULT&, const NodeMeshFragment* );
-        void GenerateMesh_Interpolate( MESH_GENERATION_RESULT&, const NodeMeshInterpolate* );
-        void GenerateMesh_Switch( MESH_GENERATION_RESULT&, const NodeMeshSwitch* );
-        void GenerateMesh_Subtract( MESH_GENERATION_RESULT&, const NodeMeshSubtract* );
-        void GenerateMesh_Transform( MESH_GENERATION_RESULT&, const NodeMeshTransform* );
-        void GenerateMesh_ClipMorphPlane( MESH_GENERATION_RESULT&, const NodeMeshClipMorphPlane* );
-        void GenerateMesh_ClipWithMesh( MESH_GENERATION_RESULT&, const NodeMeshClipWithMesh* );
-        void GenerateMesh_ApplyPose( MESH_GENERATION_RESULT&, const NodeMeshApplyPose* );
-        void GenerateMesh_Variation( MESH_GENERATION_RESULT&, const NodeMeshVariation* );
-		void GenerateMesh_Table(MESH_GENERATION_RESULT&, const NodeMeshTable*);
-		void GenerateMesh_GeometryOperation(MESH_GENERATION_RESULT&, const NodeMeshGeometryOperation*);
-		void GenerateMesh_Reshape(MESH_GENERATION_RESULT&, const NodeMeshReshape*);
-		void GenerateMesh_ClipDeform(MESH_GENERATION_RESULT&, const NodeMeshClipDeform*);
+        void GenerateMesh(const FMeshGenerationOptions&, FMeshGenerationResult& result, const NodeMeshPtrConst& node);
+        void GenerateMesh_Constant(const FMeshGenerationOptions&, FMeshGenerationResult&, const NodeMeshConstant* );
+        void GenerateMesh_Format(const FMeshGenerationOptions&, FMeshGenerationResult&, const NodeMeshFormat* );
+        void GenerateMesh_Morph(const FMeshGenerationOptions&, FMeshGenerationResult&, const NodeMeshMorph* );
+        void GenerateMesh_MakeMorph(const FMeshGenerationOptions&, FMeshGenerationResult&, const NodeMeshMakeMorph* );
+        void GenerateMesh_Fragment(const FMeshGenerationOptions&, FMeshGenerationResult&, const NodeMeshFragment* );
+        void GenerateMesh_Interpolate(const FMeshGenerationOptions&, FMeshGenerationResult&, const NodeMeshInterpolate* );
+        void GenerateMesh_Switch(const FMeshGenerationOptions&, FMeshGenerationResult&, const NodeMeshSwitch* );
+        void GenerateMesh_Subtract(const FMeshGenerationOptions&, FMeshGenerationResult&, const NodeMeshSubtract* );
+        void GenerateMesh_Transform(const FMeshGenerationOptions&, FMeshGenerationResult&, const NodeMeshTransform* );
+        void GenerateMesh_ClipMorphPlane(const FMeshGenerationOptions&, FMeshGenerationResult&, const NodeMeshClipMorphPlane* );
+        void GenerateMesh_ClipWithMesh(const FMeshGenerationOptions&, FMeshGenerationResult&, const NodeMeshClipWithMesh* );
+        void GenerateMesh_ApplyPose(const FMeshGenerationOptions&, FMeshGenerationResult&, const NodeMeshApplyPose* );
+        void GenerateMesh_Variation(const FMeshGenerationOptions&, FMeshGenerationResult&, const NodeMeshVariation* );
+		void GenerateMesh_Table(const FMeshGenerationOptions&, FMeshGenerationResult&, const NodeMeshTable*);
+		void GenerateMesh_GeometryOperation(const FMeshGenerationOptions&, FMeshGenerationResult&, const NodeMeshGeometryOperation*);
+		void GenerateMesh_Reshape(const FMeshGenerationOptions&, FMeshGenerationResult&, const NodeMeshReshape*);
+		void GenerateMesh_ClipDeform(const FMeshGenerationOptions&, FMeshGenerationResult&, const NodeMeshClipDeform*);
 
 		//!
 		Ptr<const Layout> AddLayout(Ptr<const Layout> SourceLayout);
diff --git a/Engine/Plugins/Experimental/Mutable/Source/MutableTools/Private/MuT/CodeGenerator_Bool.cpp b/Engine/Plugins/Experimental/Mutable/Source/MutableTools/Private/MuT/CodeGenerator_Bool.cpp
--- a/Engine/Plugins/Experimental/Mutable/Source/MutableTools/Private/MuT/CodeGenerator_Bool.cpp
+++ b/Engine/Plugins/Experimental/Mutable/Source/MutableTools/Private/MuT/CodeGenerator_Bool.cpp
@@ -22,145 +22,148 @@
 #include <utility>
 
 
-using namespace mu;
+namespace mu
+{
 
 
-//-------------------------------------------------------------------------------------------------
-void CodeGenerator::GenerateBool( BOOL_GENERATION_RESULT& result, const NodeBoolPtrConst& untyped )
-{
-    if (!untyped)
-    {
-        result = BOOL_GENERATION_RESULT();
-        return;
-    }
-
-    // See if it was already generated
-	VISITED_MAP_KEY key = GetCurrentCacheKey(untyped);
-	GeneratedBoolsMap::ValueType* it = m_generatedBools.Find( key );
-    if ( it )
-    {
-        result = *it;
-        return;
-    }
-
-    // Generate for each different type of node
-    if ( auto Constant = dynamic_cast<const NodeBoolConstant*>(untyped.get()) )
-    {
-        GenerateBool_Constant( result, Constant );
-    }
-    else if ( auto Param = dynamic_cast<const NodeBoolParameter*>(untyped.get()) )
-    {
-        GenerateBool_Parameter( result, Param );
-    }
-	else if (auto Switch = dynamic_cast<const NodeBoolIsNull*>(untyped.get()))
-	{
-		GenerateBool_IsNull(result, Switch);
-	}
-	else if (auto Sample = dynamic_cast<const NodeBoolNot*>(untyped.get()))
+	//-------------------------------------------------------------------------------------------------
+	void CodeGenerator::GenerateBool(BOOL_GENERATION_RESULT& result, const NodeBoolPtrConst& untyped)
 	{
-		GenerateBool_Not(result, Sample);
-	}
-	else if (auto From = dynamic_cast<const NodeBoolAnd*>(untyped.get()))
-	{
-		GenerateBool_And(result, From);
-	}	
-	else
-    {
-        check(false);
-    }
-
-    // Cache the result
-    m_generatedBools.Add( key, result );
-}
+		if (!untyped)
+		{
+			result = BOOL_GENERATION_RESULT();
+			return;
+		}
 
+		// See if it was already generated
+		VISITED_MAP_KEY key = GetCurrentCacheKey(untyped);
+		GeneratedBoolsMap::ValueType* it = m_generatedBools.Find(key);
+		if (it)
+		{
+			result = *it;
+			return;
+		}
 
-//-------------------------------------------------------------------------------------------------
-void CodeGenerator::GenerateBool_Constant(BOOL_GENERATION_RESULT& result, const Ptr<const NodeBoolConstant>& Typed )
-{
-    const NodeBoolConstant::Private& node = *Typed->GetPrivate();
+		// Generate for each different type of node
+		if (auto Constant = dynamic_cast<const NodeBoolConstant*>(untyped.get()))
+		{
+			GenerateBool_Constant(result, Constant);
+		}
+		else if (auto Param = dynamic_cast<const NodeBoolParameter*>(untyped.get()))
+		{
+			GenerateBool_Parameter(result, Param);
+		}
+		else if (auto Switch = dynamic_cast<const NodeBoolIsNull*>(untyped.get()))
+		{
+			GenerateBool_IsNull(result, Switch);
+		}
+		else if (auto Sample = dynamic_cast<const NodeBoolNot*>(untyped.get()))
+		{
+			GenerateBool_Not(result, Sample);
+		}
+		else if (auto From = dynamic_cast<const NodeBoolAnd*>(untyped.get()))
+		{
+			GenerateBool_And(result, From);
+		}
+		else
+		{
+			check(false);
+		}
+
+		// Cache the result
+		m_generatedBools.Add(key, result);
+	}
 
-	Ptr<ASTOpConstantBool> op = new ASTOpConstantBool();
-	op->value = node.m_value;
 
-    result.op = op;
-}
+	//-------------------------------------------------------------------------------------------------
+	void CodeGenerator::GenerateBool_Constant(BOOL_GENERATION_RESULT& result, const Ptr<const NodeBoolConstant>& Typed)
+	{
+		const NodeBoolConstant::Private& node = *Typed->GetPrivate();
 
+		Ptr<ASTOpConstantBool> op = new ASTOpConstantBool();
+		op->value = node.m_value;
 
-//-------------------------------------------------------------------------------------------------
-void CodeGenerator::GenerateBool_Parameter( BOOL_GENERATION_RESULT& result, const Ptr<const NodeBoolParameter>& Typed)
-{
-	const NodeBoolParameter::Private& node = *Typed->GetPrivate();
+		result.op = op;
+	}
 
-	Ptr<ASTOpParameter> op;
 
-	auto it = m_nodeVariables.find(node.m_pNode);
-	if (it == m_nodeVariables.end())
+	//-------------------------------------------------------------------------------------------------
+	void CodeGenerator::GenerateBool_Parameter(BOOL_GENERATION_RESULT& result, const Ptr<const NodeBoolParameter>& Typed)
 	{
-		PARAMETER_DESC param;
-		param.m_name = node.m_name;
-		param.m_uid = node.m_uid;
-		param.m_type = PARAMETER_TYPE::T_BOOL;
-		param.m_defaultValue.m_bool = node.m_defaultValue;
-
-		op = new ASTOpParameter();
-		op->type = OP_TYPE::BO_PARAMETER;
-		op->parameter = param;
-
-		// Generate the code for the ranges
-		for (int32 a = 0; a < node.m_ranges.Num(); ++a)
+		const NodeBoolParameter::Private& node = *Typed->GetPrivate();
+
+		Ptr<ASTOpParameter> op;
+
+		auto it = m_nodeVariables.find(node.m_pNode);
+		if (it == m_nodeVariables.end())
 		{
-			RANGE_GENERATION_RESULT rangeResult;
-			GenerateRange(rangeResult, node.m_ranges[a]);
-			op->ranges.Emplace(op.get(), rangeResult.sizeOp, rangeResult.rangeName, rangeResult.rangeUID);
+			PARAMETER_DESC param;
+			param.m_name = node.m_name;
+			param.m_uid = node.m_uid;
+			param.m_type = PARAMETER_TYPE::T_BOOL;
+			param.m_defaultValue.m_bool = node.m_defaultValue;
+
+			op = new ASTOpParameter();
+			op->type = OP_TYPE::BO_PARAMETER;
+			op->parameter = param;
+
+			// Generate the code for the ranges
+			for (int32 a = 0; a < node.m_ranges.Num(); ++a)
+			{
+				RANGE_GENERATION_RESULT rangeResult;
+				GenerateRange(rangeResult, node.m_ranges[a]);
+				op->ranges.Emplace(op.get(), rangeResult.sizeOp, rangeResult.rangeName, rangeResult.rangeUID);
+			}
+
+			m_nodeVariables[node.m_pNode] = op;
+		}
+		else
+		{
+			op = it->second;
 		}
 
-		m_nodeVariables[node.m_pNode] = op;
-	}
-	else
-	{
-		op = it->second;
+		result.op = op;
 	}
 
-    result.op = op;
-}
 
+	//-------------------------------------------------------------------------------------------------
+	void CodeGenerator::GenerateBool_IsNull(BOOL_GENERATION_RESULT& result, const Ptr<const NodeBoolIsNull>& Typed)
+	{
+		const NodeBoolIsNull::Private& node = *Typed->GetPrivate();
 
-//-------------------------------------------------------------------------------------------------
-void CodeGenerator::GenerateBool_IsNull(BOOL_GENERATION_RESULT& result, const Ptr<const NodeBoolIsNull>& Typed)
-{
-	const NodeBoolIsNull::Private& node = *Typed->GetPrivate();
+		Ptr<ASTOpConstantBool> op = new ASTOpConstantBool();
+		Ptr<ASTOp> source = Generate(node.m_pSource);
+		op->value = !source;
 
-	Ptr<ASTOpConstantBool> op = new ASTOpConstantBool();
-	Ptr<ASTOp> source = Generate(node.m_pSource);
-	op->value = !source;
+		result.op = op;
+	}
 
-	result.op = op;
-}
 
+	//-------------------------------------------------------------------------------------------------
+	void CodeGenerator::GenerateBool_Not(BOOL_GENERATION_RESULT& result, const Ptr<const NodeBoolNot>& Typed)
+	{
+		const NodeBoolNot::Private& node = *Typed->GetPrivate();
 
-//-------------------------------------------------------------------------------------------------
-void CodeGenerator::GenerateBool_Not(BOOL_GENERATION_RESULT& result, const Ptr<const NodeBoolNot>& Typed)
-{
-	const NodeBoolNot::Private& node = *Typed->GetPrivate();
+		Ptr<ASTOpFixed> op = new ASTOpFixed();
+		op->op.type = OP_TYPE::BO_NOT;
+		op->SetChild(op->op.args.BoolNot.source, Generate(node.m_pSource));
+
+		result.op = op;
+	}
 
-	Ptr<ASTOpFixed> op = new ASTOpFixed();
-	op->op.type = OP_TYPE::BO_NOT;
-	op->SetChild(op->op.args.BoolNot.source, Generate(node.m_pSource));
 
-	result.op = op;
-}
+	//-------------------------------------------------------------------------------------------------
+	void CodeGenerator::GenerateBool_And(BOOL_GENERATION_RESULT& result, const Ptr<const NodeBoolAnd>& Typed)
+	{
+		const NodeBoolAnd::Private& node = *Typed->GetPrivate();
 
+		Ptr<ASTOpFixed> op = new ASTOpFixed();
+		op->op.type = OP_TYPE::BO_AND;
+		op->SetChild(op->op.args.BoolBinary.a, Generate(node.m_pA));
+		op->SetChild(op->op.args.BoolBinary.b, Generate(node.m_pB));
 
-//-------------------------------------------------------------------------------------------------
-void CodeGenerator::GenerateBool_And(BOOL_GENERATION_RESULT& result, const Ptr<const NodeBoolAnd>& Typed)
-{
-	const NodeBoolAnd::Private& node = *Typed->GetPrivate();
+		result.op = op;
+	}
 
-	Ptr<ASTOpFixed> op = new ASTOpFixed();
-	op->op.type = OP_TYPE::BO_AND;
-	op->SetChild(op->op.args.BoolBinary.a, Generate(node.m_pA));
-	op->SetChild(op->op.args.BoolBinary.b, Generate(node.m_pB));
 
-	result.op = op;
 }
-
diff --git a/Engine/Plugins/Experimental/Mutable/Source/MutableTools/Private/MuT/CodeGenerator_Color.cpp b/Engine/Plugins/Experimental/Mutable/Source/MutableTools/Private/MuT/CodeGenerator_Color.cpp
--- a/Engine/Plugins/Experimental/Mutable/Source/MutableTools/Private/MuT/CodeGenerator_Color.cpp
+++ b/Engine/Plugins/Experimental/Mutable/Source/MutableTools/Private/MuT/CodeGenerator_Color.cpp
@@ -53,462 +53,465 @@
 #include <utility>
 
 
-using namespace mu;
+namespace mu
+{
 
 
-//-------------------------------------------------------------------------------------------------
-void CodeGenerator::GenerateColor( COLOR_GENERATION_RESULT& result, const NodeColourPtrConst& untyped )
-{
-    if (!untyped)
-    {
-        result = COLOR_GENERATION_RESULT();
-        return;
-    }
-
-    // See if it was already generated
-	VISITED_MAP_KEY key = GetCurrentCacheKey(untyped);
-	GeneratedColorsMap::ValueType* it = m_generatedColors.Find( key );
-    if ( it )
-    {
-        result = *it;
-        return;
-    }
-
-    // Generate for each different type of node
-    if ( auto Constant = dynamic_cast<const NodeColourConstant*>(untyped.get()) )
-    {
-        GenerateColor_Constant( result, Constant );
-    }
-    else if ( auto Param = dynamic_cast<const NodeColourParameter*>(untyped.get()) )
-    {
-        GenerateColor_Parameter( result, Param );
-    }
-	else if (auto Switch = dynamic_cast<const NodeColourSwitch*>(untyped.get()))
-	{
-		GenerateColor_Switch(result, Switch);
-	}
-	else if (auto Sample = dynamic_cast<const NodeColourSampleImage*>(untyped.get()))
+	//-------------------------------------------------------------------------------------------------
+	void CodeGenerator::GenerateColor(COLOR_GENERATION_RESULT& result, const NodeColourPtrConst& untyped)
 	{
-		GenerateColor_SampleImage(result, Sample);
-	}
-	else if (auto From = dynamic_cast<const NodeColourFromScalars*>(untyped.get()))
-	{
-		GenerateColor_FromScalars(result, From);
-	}
-	else if (auto Arithmetic = dynamic_cast<const NodeColourArithmeticOperation*>(untyped.get()))
-	{
-		GenerateColor_Arithmetic(result, Arithmetic);
-	}
-	else if (auto Variation = dynamic_cast<const NodeColourVariation*>(untyped.get()))
-	{
-		GenerateColor_Variation(result, Variation);
-	}
-	else if (auto Table = dynamic_cast<const NodeColourTable*>(untyped.get()))
-	{
-		GenerateColor_Table(result, Table);
-	}
-	else
-    {
-        check(false);
-    }
-
-    // Cache the result
-    m_generatedColors.Add( key, result );
-}
+		if (!untyped)
+		{
+			result = COLOR_GENERATION_RESULT();
+			return;
+		}
 
+		// See if it was already generated
+		VISITED_MAP_KEY key = GetCurrentCacheKey(untyped);
+		GeneratedColorsMap::ValueType* it = m_generatedColors.Find(key);
+		if (it)
+		{
+			result = *it;
+			return;
+		}
 
-//-------------------------------------------------------------------------------------------------
-void CodeGenerator::GenerateColor_Constant(COLOR_GENERATION_RESULT& result, const Ptr<const NodeColourConstant>& Typed )
-{
-    const NodeColourConstant::Private& node = *Typed->GetPrivate();
+		// Generate for each different type of node
+		if (auto Constant = dynamic_cast<const NodeColourConstant*>(untyped.get()))
+		{
+			GenerateColor_Constant(result, Constant);
+		}
+		else if (auto Param = dynamic_cast<const NodeColourParameter*>(untyped.get()))
+		{
+			GenerateColor_Parameter(result, Param);
+		}
+		else if (auto Switch = dynamic_cast<const NodeColourSwitch*>(untyped.get()))
+		{
+			GenerateColor_Switch(result, Switch);
+		}
+		else if (auto Sample = dynamic_cast<const NodeColourSampleImage*>(untyped.get()))
+		{
+			GenerateColor_SampleImage(result, Sample);
+		}
+		else if (auto From = dynamic_cast<const NodeColourFromScalars*>(untyped.get()))
+		{
+			GenerateColor_FromScalars(result, From);
+		}
+		else if (auto Arithmetic = dynamic_cast<const NodeColourArithmeticOperation*>(untyped.get()))
+		{
+			GenerateColor_Arithmetic(result, Arithmetic);
+		}
+		else if (auto Variation = dynamic_cast<const NodeColourVariation*>(untyped.get()))
+		{
+			GenerateColor_Variation(result, Variation);
+		}
+		else if (auto Table = dynamic_cast<const NodeColourTable*>(untyped.get()))
+		{
+			GenerateColor_Table(result, Table);
+		}
+		else
+		{
+			check(false);
+		}
 
-	Ptr<ASTOpFixed> op = new ASTOpFixed();
-	op->op.type = OP_TYPE::CO_CONSTANT;
-	op->op.args.ColourConstant.value[0] = node.m_value[0];
-	op->op.args.ColourConstant.value[1] = node.m_value[1];
-	op->op.args.ColourConstant.value[2] = node.m_value[2];
+		// Cache the result
+		m_generatedColors.Add(key, result);
+	}
 
-    result.op = op;
-}
 
+	//-------------------------------------------------------------------------------------------------
+	void CodeGenerator::GenerateColor_Constant(COLOR_GENERATION_RESULT& result, const Ptr<const NodeColourConstant>& Typed)
+	{
+		const NodeColourConstant::Private& node = *Typed->GetPrivate();
 
-//-------------------------------------------------------------------------------------------------
-void CodeGenerator::GenerateColor_Parameter( COLOR_GENERATION_RESULT& result, const Ptr<const NodeColourParameter>& Typed)
-{
-    const NodeColourParameter::Private& node = *Typed->GetPrivate();
+		Ptr<ASTOpFixed> op = new ASTOpFixed();
+		op->op.type = OP_TYPE::CO_CONSTANT;
+		op->op.args.ColourConstant.value[0] = node.m_value[0];
+		op->op.args.ColourConstant.value[1] = node.m_value[1];
+		op->op.args.ColourConstant.value[2] = node.m_value[2];
 
-    Ptr<ASTOpParameter> op;
+		result.op = op;
+	}
 
-	auto it = m_nodeVariables.find(node.m_pNode);
 
-	if (it == m_nodeVariables.end())
+	//-------------------------------------------------------------------------------------------------
+	void CodeGenerator::GenerateColor_Parameter(COLOR_GENERATION_RESULT& result, const Ptr<const NodeColourParameter>& Typed)
 	{
-		PARAMETER_DESC param;
-		param.m_name = node.m_name;
-		param.m_uid = node.m_uid;
-		param.m_type = PARAMETER_TYPE::T_COLOUR;
-		param.m_defaultValue.m_colour[0] = node.m_defaultValue[0];
-		param.m_defaultValue.m_colour[1] = node.m_defaultValue[1];
-		param.m_defaultValue.m_colour[2] = node.m_defaultValue[2];
+		const NodeColourParameter::Private& node = *Typed->GetPrivate();
+
+		Ptr<ASTOpParameter> op;
 
-		op = new ASTOpParameter();
-		op->type = OP_TYPE::CO_PARAMETER;
-		op->parameter = param;
+		auto it = m_nodeVariables.find(node.m_pNode);
 
-		// Generate the code for the ranges
-		for (int32 a = 0; a < node.m_ranges.Num(); ++a)
+		if (it == m_nodeVariables.end())
+		{
+			PARAMETER_DESC param;
+			param.m_name = node.m_name;
+			param.m_uid = node.m_uid;
+			param.m_type = PARAMETER_TYPE::T_COLOUR;
+			param.m_defaultValue.m_colour[0] = node.m_defaultValue[0];
+			param.m_defaultValue.m_colour[1] = node.m_defaultValue[1];
+			param.m_defaultValue.m_colour[2] = node.m_defaultValue[2];
+
+			op = new ASTOpParameter();
+			op->type = OP_TYPE::CO_PARAMETER;
+			op->parameter = param;
+
+			// Generate the code for the ranges
+			for (int32 a = 0; a < node.m_ranges.Num(); ++a)
+			{
+				RANGE_GENERATION_RESULT rangeResult;
+				GenerateRange(rangeResult, node.m_ranges[a]);
+				op->ranges.Emplace(op.get(), rangeResult.sizeOp, rangeResult.rangeName, rangeResult.rangeUID);
+			}
+
+			m_nodeVariables[node.m_pNode] = op;
+		}
+		else
 		{
-			RANGE_GENERATION_RESULT rangeResult;
-			GenerateRange(rangeResult, node.m_ranges[a]);
-			op->ranges.Emplace(op.get(), rangeResult.sizeOp, rangeResult.rangeName, rangeResult.rangeUID);
+			op = it->second;
 		}
 
-		m_nodeVariables[node.m_pNode] = op;
+		result.op = op;
 	}
-	else
+
+
+	//-------------------------------------------------------------------------------------------------
+	void CodeGenerator::GenerateColor_Switch(COLOR_GENERATION_RESULT& result, const Ptr<const NodeColourSwitch>& Typed)
 	{
-		op = it->second;
-	}
+		const NodeColourSwitch::Private& node = *Typed->GetPrivate();
 
-    result.op = op;
-}
+		MUTABLE_CPUPROFILER_SCOPE(NodeColourSwitch);
 
+		if (node.m_options.Num() == 0)
+		{
+			// No options in the switch!
+			Ptr<ASTOp> missingOp = GenerateMissingColourCode("Switch option",
+				node.m_errorContext);
+			result.op = missingOp;
+			return;
+		}
 
-//-------------------------------------------------------------------------------------------------
-void CodeGenerator::GenerateColor_Switch(COLOR_GENERATION_RESULT& result, const Ptr<const NodeColourSwitch>& Typed)
-{
-	const NodeColourSwitch::Private& node = *Typed->GetPrivate();
+		Ptr<ASTOpSwitch> op = new ASTOpSwitch();
+		op->type = OP_TYPE::CO_SWITCH;
+
+		// Variable value
+		if (node.m_pParameter)
+		{
+			op->variable = Generate(node.m_pParameter.get());
+		}
+		else
+		{
+			// This argument is required
+			op->variable = GenerateMissingScalarCode("Switch variable", 0.0f, node.m_errorContext);
+		}
 
-	MUTABLE_CPUPROFILER_SCOPE(NodeColourSwitch);
+		// Options
+		for (std::size_t t = 0; t < node.m_options.Num(); ++t)
+		{
+			Ptr<ASTOp> branch;
+			if (node.m_options[t])
+			{
+				branch = Generate(node.m_options[t].get());
+			}
+			else
+			{
+				// This argument is required
+				branch = GenerateMissingColourCode("Switch option",
+					node.m_errorContext);
+			}
+			op->cases.Emplace((int16)t, op, branch);
+		}
 
-	if (node.m_options.Num() == 0)
-	{
-		// No options in the switch!
-		Ptr<ASTOp> missingOp = GenerateMissingColourCode("Switch option",
-			node.m_errorContext);
-		result.op = missingOp;
-		return;
+		result.op = op;
 	}
 
-	Ptr<ASTOpSwitch> op = new ASTOpSwitch();
-	op->type = OP_TYPE::CO_SWITCH;
 
-	// Variable value
-	if (node.m_pParameter)
-	{
-		op->variable = Generate(node.m_pParameter.get());
-	}
-	else
+	//-------------------------------------------------------------------------------------------------
+	void CodeGenerator::GenerateColor_Variation(COLOR_GENERATION_RESULT& result, const Ptr<const NodeColourVariation>& Typed)
 	{
-		// This argument is required
-		op->variable = GenerateMissingScalarCode("Switch variable", 0.0f, node.m_errorContext);
-	}
+		const NodeColourVariation::Private& node = *Typed->GetPrivate();
 
-	// Options
-	for (std::size_t t = 0; t < node.m_options.Num(); ++t)
-	{
-		Ptr<ASTOp> branch;
-		if (node.m_options[t])
+		Ptr<ASTOp> currentOp;
+
+		// Default case
+		if (node.m_defaultColour)
 		{
-			branch = Generate(node.m_options[t].get());
+			FMeshGenerationResult branchResults;
+			currentOp = Generate(node.m_defaultColour);
 		}
 		else
 		{
 			// This argument is required
-			branch = GenerateMissingColourCode("Switch option",
-				node.m_errorContext);
+			currentOp = GenerateMissingColourCode("Variation default", node.m_errorContext);
 		}
-		op->cases.Emplace((int16)t, op, branch);
-	}
 
-	result.op = op;
-}
+		// Process variations in reverse order, since conditionals are built bottom-up.
+		for (int t = node.m_variations.Num() - 1; t >= 0; --t)
+		{
+			int tagIndex = -1;
+			const string& tag = node.m_variations[t].m_tag;
+			for (int i = 0; i < m_firstPass.m_tags.Num(); ++i)
+			{
+				if (m_firstPass.m_tags[i].tag == tag)
+				{
+					tagIndex = i;
+				}
+			}
 
+			if (tagIndex < 0)
+			{
+				char buf[256];
+				mutable_snprintf(buf, 256, "Unknown tag found in image variation [%s].",
+					tag.c_str());
 
-//-------------------------------------------------------------------------------------------------
-void CodeGenerator::GenerateColor_Variation(COLOR_GENERATION_RESULT& result, const Ptr<const NodeColourVariation>& Typed)
-{
-	const NodeColourVariation::Private& node = *Typed->GetPrivate();
+				m_pErrorLog->GetPrivate()->Add(buf, ELMT_WARNING, node.m_errorContext);
+				continue;
+			}
 
-	Ptr<ASTOp> currentOp;
+			Ptr<ASTOp> variationOp;
+			if (node.m_variations[t].m_colour)
+			{
+				variationOp = Generate(node.m_variations[t].m_colour);
+			}
+			else
+			{
+				// This argument is required
+				variationOp = GenerateMissingColourCode("Variation option", node.m_errorContext);
+			}
 
-	// Default case
-	if (node.m_defaultColour)
-	{
-		MESH_GENERATION_RESULT branchResults;
-		currentOp = Generate(node.m_defaultColour);
-	}
-	else
-	{
-		// This argument is required
-		currentOp = GenerateMissingColourCode("Variation default", node.m_errorContext);
+
+			Ptr<ASTOpConditional> conditional = new ASTOpConditional;
+			conditional->type = OP_TYPE::CO_CONDITIONAL;
+			conditional->no = currentOp;
+			conditional->yes = variationOp;
+			conditional->condition = m_firstPass.m_tags[tagIndex].genericCondition;
+
+			currentOp = conditional;
+		}
+
+		result.op = currentOp;
 	}
 
-	// Process variations in reverse order, since conditionals are built bottom-up.
-	for (int t = node.m_variations.Num() - 1; t >= 0; --t)
+
+	//-------------------------------------------------------------------------------------------------
+	void CodeGenerator::GenerateColor_SampleImage(COLOR_GENERATION_RESULT& result, const Ptr<const NodeColourSampleImage>& Typed)
 	{
-		int tagIndex = -1;
-		const string& tag = node.m_variations[t].m_tag;
-		for (int i = 0; i < m_firstPass.m_tags.Num(); ++i)
-		{
-			if (m_firstPass.m_tags[i].tag == tag)
-			{
-				tagIndex = i;
-			}
-		}
+		const NodeColourSampleImage::Private& node = *Typed->GetPrivate();
 
-		if (tagIndex < 0)
-		{
-			char buf[256];
-			mutable_snprintf(buf, 256, "Unknown tag found in image variation [%s].",
-				tag.c_str());
+		// Generate the code
+		Ptr<ASTOpFixed> op = new ASTOpFixed();
+		op->op.type = OP_TYPE::CO_SAMPLEIMAGE;
 
-			m_pErrorLog->GetPrivate()->Add(buf, ELMT_WARNING, node.m_errorContext);
-			continue;
-		}
 
-		Ptr<ASTOp> variationOp;
-		if (node.m_variations[t].m_colour)
+		// Source image
+		Ptr<ASTOp> base;
+		if (Node* pSource = node.m_pImage.get())
 		{
-			variationOp = Generate(node.m_variations[t].m_colour);
+			// We take whatever size will be produced
+			FImageDesc desc = CalculateImageDesc(*pSource->GetBasePrivate());
+			IMAGE_STATE newState;
+			newState.m_imageSize = desc.m_size;
+			newState.m_imageRect.min[0] = 0;
+			newState.m_imageRect.min[1] = 0;
+			newState.m_imageRect.size = desc.m_size;
+			m_imageState.Add(newState);
+
+			// Generate
+			base = Generate(pSource);
+
+			// Restore rect
+			m_imageState.Pop();
 		}
 		else
 		{
 			// This argument is required
-			variationOp = GenerateMissingColourCode("Variation option", node.m_errorContext);
+			base = GenerateMissingImageCode("Sample image", EImageFormat::IF_RGB_UBYTE, node.m_errorContext);
 		}
+		base = GenerateImageFormat(base, EImageFormat::IF_RGB_UBYTE);
+		op->SetChild(op->op.args.ColourSampleImage.image, base);
 
 
-		Ptr<ASTOpConditional> conditional = new ASTOpConditional;
-		conditional->type = OP_TYPE::CO_CONDITIONAL;
-		conditional->no = currentOp;
-		conditional->yes = variationOp;
-		conditional->condition = m_firstPass.m_tags[tagIndex].genericCondition;
-
-		currentOp = conditional;
-	}
-
-	result.op = currentOp;
-}
-
+		// X
+		if (Node* pX = node.m_pX.get())
+		{
+			op->SetChild(op->op.args.ColourSampleImage.x, Generate(pX));
+		}
+		else
+		{
+			// Set a constant 0.5 value
+			NodeScalarConstantPtr pNode = new NodeScalarConstant();
+			pNode->SetValue(0.5f);
+			op->SetChild(op->op.args.ColourSampleImage.x, Generate(pNode));
+		}
 
-//-------------------------------------------------------------------------------------------------
-void CodeGenerator::GenerateColor_SampleImage(COLOR_GENERATION_RESULT& result, const Ptr<const NodeColourSampleImage>& Typed)
-{
-	const NodeColourSampleImage::Private& node = *Typed->GetPrivate();
 
-	// Generate the code
-	Ptr<ASTOpFixed> op = new ASTOpFixed();
-	op->op.type = OP_TYPE::CO_SAMPLEIMAGE;
+		// Y
+		if (Node* pY = node.m_pY.get())
+		{
+			op->SetChild(op->op.args.ColourSampleImage.y, Generate(pY));
+		}
+		else
+		{
+			// Set a constant 0.5 value
+			NodeScalarConstantPtr pNode = new NodeScalarConstant();
+			pNode->SetValue(0.5f);
+			op->SetChild(op->op.args.ColourSampleImage.y, Generate(pNode));
+		}
 
+		// TODO
+		op->op.args.ColourSampleImage.filter = 0;
 
-	// Source image
-	Ptr<ASTOp> base;
-	if (Node* pSource = node.m_pImage.get())
-	{
-		// We take whatever size will be produced
-		FImageDesc desc = CalculateImageDesc(*pSource->GetBasePrivate());
-		IMAGE_STATE newState;
-		newState.m_imageSize = desc.m_size;
-		newState.m_imageRect.min[0] = 0;
-		newState.m_imageRect.min[1] = 0;
-		newState.m_imageRect.size = desc.m_size;
-		m_imageState.Add(newState);
-
-		// Generate
-		base = Generate(pSource);
-
-		// Restore rect
-		m_imageState.Pop();
+		result.op = op;
 	}
-	else
-	{
-		// This argument is required
-		base = GenerateMissingImageCode("Sample image", EImageFormat::IF_RGB_UBYTE, node.m_errorContext);
-	}
-	base = GenerateImageFormat(base, EImageFormat::IF_RGB_UBYTE);
-	op->SetChild(op->op.args.ColourSampleImage.image, base);
 
 
-	// X
-	if (Node* pX = node.m_pX.get())
-	{
-		op->SetChild(op->op.args.ColourSampleImage.x, Generate(pX));
-	}
-	else
+	//-------------------------------------------------------------------------------------------------
+	void CodeGenerator::GenerateColor_FromScalars(COLOR_GENERATION_RESULT& result, const Ptr<const NodeColourFromScalars>& Typed)
 	{
-		// Set a constant 0.5 value
-		NodeScalarConstantPtr pNode = new NodeScalarConstant();
-		pNode->SetValue(0.5f);
-		op->SetChild(op->op.args.ColourSampleImage.x, Generate(pNode));
-	}
+		const NodeColourFromScalars::Private& node = *Typed->GetPrivate();
 
+		Ptr<ASTOpFixed> op = new ASTOpFixed();
+		op->op.type = OP_TYPE::CO_FROMSCALARS;
 
-	// Y
-	if (Node* pY = node.m_pY.get())
-	{
-		op->SetChild(op->op.args.ColourSampleImage.y, Generate(pY));
-	}
-	else
-	{
-		// Set a constant 0.5 value
-		NodeScalarConstantPtr pNode = new NodeScalarConstant();
-		pNode->SetValue(0.5f);
-		op->SetChild(op->op.args.ColourSampleImage.y, Generate(pNode));
-	}
-
-	// TODO
-	op->op.args.ColourSampleImage.filter = 0;
-
-	result.op = op;
-}
-
+		// X
+		if (Node* pX = node.m_pX.get())
+		{
+			op->SetChild(op->op.args.ColourFromScalars.x, Generate(pX));
+		}
+		else
+		{
+			NodeScalarConstantPtr pNode = new NodeScalarConstant();
+			pNode->SetValue(1.0f);
+			op->SetChild(op->op.args.ColourFromScalars.x, Generate(pNode));
+		}
 
-//-------------------------------------------------------------------------------------------------
-void CodeGenerator::GenerateColor_FromScalars(COLOR_GENERATION_RESULT& result, const Ptr<const NodeColourFromScalars>& Typed)
-{
-	const NodeColourFromScalars::Private& node = *Typed->GetPrivate();
+		// Y
+		if (Node* pY = node.m_pY.get())
+		{
+			op->SetChild(op->op.args.ColourFromScalars.y, Generate(pY));
+		}
+		else
+		{
+			NodeScalarConstantPtr pNode = new NodeScalarConstant();
+			pNode->SetValue(1.0f);
+			op->SetChild(op->op.args.ColourFromScalars.y, Generate(pNode));
+		}
 
-	Ptr<ASTOpFixed> op = new ASTOpFixed();
-	op->op.type = OP_TYPE::CO_FROMSCALARS;
+		// Z
+		if (Node* pZ = node.m_pZ.get())
+		{
+			op->SetChild(op->op.args.ColourFromScalars.z, Generate(pZ));
+		}
+		else
+		{
+			NodeScalarConstantPtr pNode = new NodeScalarConstant();
+			pNode->SetValue(1.0f);
+			op->SetChild(op->op.args.ColourFromScalars.z, Generate(pNode));
+		}
 
-	// X
-	if (Node* pX = node.m_pX.get())
-	{
-		op->SetChild(op->op.args.ColourFromScalars.x, Generate(pX));
-	}
-	else
-	{
-		NodeScalarConstantPtr pNode = new NodeScalarConstant();
-		pNode->SetValue(1.0f);
-		op->SetChild(op->op.args.ColourFromScalars.x, Generate(pNode));
-	}
+		// W
+		if (Node* pW = node.m_pW.get())
+		{
+			op->SetChild(op->op.args.ColourFromScalars.w, Generate(pW));
+		}
+		else
+		{
+			NodeScalarConstantPtr pNode = new NodeScalarConstant();
+			pNode->SetValue(1.0f);
+			op->SetChild(op->op.args.ColourFromScalars.w, Generate(pNode));
+		}
 
-	// Y
-	if (Node* pY = node.m_pY.get())
-	{
-		op->SetChild(op->op.args.ColourFromScalars.y, Generate(pY));
-	}
-	else
-	{
-		NodeScalarConstantPtr pNode = new NodeScalarConstant();
-		pNode->SetValue(1.0f);
-		op->SetChild(op->op.args.ColourFromScalars.y, Generate(pNode));
+		result.op = op;
 	}
 
-	// Z
-	if (Node* pZ = node.m_pZ.get())
-	{
-		op->SetChild(op->op.args.ColourFromScalars.z, Generate(pZ));
-	}
-	else
-	{
-		NodeScalarConstantPtr pNode = new NodeScalarConstant();
-		pNode->SetValue(1.0f);
-		op->SetChild(op->op.args.ColourFromScalars.z, Generate(pNode));
-	}
 
-	// W
-	if (Node* pW = node.m_pW.get())
-	{
-		op->SetChild(op->op.args.ColourFromScalars.w, Generate(pW));
-	}
-	else
+	//-------------------------------------------------------------------------------------------------
+	void CodeGenerator::GenerateColor_Arithmetic(COLOR_GENERATION_RESULT& result, const Ptr<const NodeColourArithmeticOperation>& Typed)
 	{
-		NodeScalarConstantPtr pNode = new NodeScalarConstant();
-		pNode->SetValue(1.0f);
-		op->SetChild(op->op.args.ColourFromScalars.w, Generate(pNode));
-	}
+		const NodeColourArithmeticOperation::Private& node = *Typed->GetPrivate();
 
-	result.op = op;
-}
+		Ptr<ASTOpFixed> op = new ASTOpFixed();
+		op->op.type = OP_TYPE::CO_ARITHMETIC;
 
+		switch (node.m_operation)
+		{
+		case NodeColourArithmeticOperation::AO_ADD: op->op.args.ColourArithmetic.operation = OP::ArithmeticArgs::ADD; break;
+		case NodeColourArithmeticOperation::AO_SUBTRACT: op->op.args.ColourArithmetic.operation = OP::ArithmeticArgs::SUBTRACT; break;
+		case NodeColourArithmeticOperation::AO_MULTIPLY: op->op.args.ColourArithmetic.operation = OP::ArithmeticArgs::MULTIPLY; break;
+		case NodeColourArithmeticOperation::AO_DIVIDE: op->op.args.ColourArithmetic.operation = OP::ArithmeticArgs::DIVIDE; break;
+		default:
+			checkf(false, TEXT("Unknown arithmetic operation."));
+			op->op.args.ColourArithmetic.operation = OP::ArithmeticArgs::NONE;
+			break;
+		}
 
-//-------------------------------------------------------------------------------------------------
-void CodeGenerator::GenerateColor_Arithmetic(COLOR_GENERATION_RESULT& result, const Ptr<const NodeColourArithmeticOperation>& Typed)
-{
-	const NodeColourArithmeticOperation::Private& node = *Typed->GetPrivate();
+		// A
+		if (Node* pA = node.m_pA.get())
+		{
+			op->SetChild(op->op.args.ColourArithmetic.a, Generate(pA));
+		}
+		else
+		{
+			op->SetChild(op->op.args.ColourArithmetic.a,
+				CodeGenerator::GenerateMissingColourCode
+				(
+					"ColourArithmetic A",
+					node.m_errorContext
+				));
+		}
 
-	Ptr<ASTOpFixed> op = new ASTOpFixed();
-	op->op.type = OP_TYPE::CO_ARITHMETIC;
+		// B
+		if (Node* pB = node.m_pB.get())
+		{
+			op->SetChild(op->op.args.ColourArithmetic.b, Generate(pB));
+		}
+		else
+		{
+			op->SetChild(op->op.args.ColourArithmetic.b,
+				CodeGenerator::GenerateMissingColourCode
+				(
+					"ColourArithmetic B",
+					node.m_errorContext
+				));
+		}
 
-	switch (node.m_operation)
-	{
-	case NodeColourArithmeticOperation::AO_ADD: op->op.args.ColourArithmetic.operation = OP::ArithmeticArgs::ADD; break;
-	case NodeColourArithmeticOperation::AO_SUBTRACT: op->op.args.ColourArithmetic.operation = OP::ArithmeticArgs::SUBTRACT; break;
-	case NodeColourArithmeticOperation::AO_MULTIPLY: op->op.args.ColourArithmetic.operation = OP::ArithmeticArgs::MULTIPLY; break;
-	case NodeColourArithmeticOperation::AO_DIVIDE: op->op.args.ColourArithmetic.operation = OP::ArithmeticArgs::DIVIDE; break;
-	default:
-		checkf(false, TEXT("Unknown arithmetic operation."));
-		op->op.args.ColourArithmetic.operation = OP::ArithmeticArgs::NONE;
-		break;
+		result.op = op;
 	}
 
-	// A
-	if (Node* pA = node.m_pA.get())
-	{
-		op->SetChild(op->op.args.ColourArithmetic.a, Generate(pA));
-	}
-	else
-	{
-		op->SetChild(op->op.args.ColourArithmetic.a,
-			CodeGenerator::GenerateMissingColourCode
-			(
-				"ColourArithmetic A",
-				node.m_errorContext
-			));
-	}
 
-	// B
-	if (Node* pB = node.m_pB.get())
-	{
-		op->SetChild(op->op.args.ColourArithmetic.b, Generate(pB));
-	}
-	else
+	//-------------------------------------------------------------------------------------------------
+	void CodeGenerator::GenerateColor_Table(COLOR_GENERATION_RESULT& result, const Ptr<const NodeColourTable>& Typed)
 	{
-		op->SetChild(op->op.args.ColourArithmetic.b,
-			CodeGenerator::GenerateMissingColourCode
-			(
-				"ColourArithmetic B",
-				node.m_errorContext
-			));
-	}
-
-	result.op = op;
-}
+		const NodeColourTable::Private& node = *Typed->GetPrivate();
 
+		result.op = GenerateTableSwitch<NodeColourTable::Private, TCT_COLOUR, OP_TYPE::CO_SWITCH>(node,
+			[this](const NodeColourTable::Private& node, int colIndex, int row, ErrorLog* pErrorLog)
+			{
+				NodeColourConstantPtr pCell = new NodeColourConstant();
+				vec3<float> colour = node.m_pTable->GetPrivate()->m_rows[row].m_values[colIndex].m_colour;
+				pCell->SetValue(colour[0], colour[1], colour[2]);
+				return Generate(pCell);
+			});
+	}
 
-//-------------------------------------------------------------------------------------------------
-void CodeGenerator::GenerateColor_Table(COLOR_GENERATION_RESULT& result, const Ptr<const NodeColourTable>& Typed)
-{
-	const NodeColourTable::Private& node = *Typed->GetPrivate();
 
-	result.op = GenerateTableSwitch<NodeColourTable::Private, TCT_COLOUR, OP_TYPE::CO_SWITCH>(node,
-		[this](const NodeColourTable::Private& node, int colIndex, int row, ErrorLog* pErrorLog)
-		{
-			NodeColourConstantPtr pCell = new NodeColourConstant();
-			vec3<float> colour = node.m_pTable->GetPrivate()->m_rows[row].m_values[colIndex].m_colour;
-			pCell->SetValue(colour[0], colour[1], colour[2]);
-			return Generate(pCell);
-		});
-}
+	//---------------------------------------------------------------------------------------------
+	mu::Ptr<ASTOp> CodeGenerator::GenerateMissingColourCode(const char* strWhere, const void* errorContext)
+	{
+		// Log a warning
+		char buf[256];
+		mutable_snprintf(buf, 256, "Required connection not found: %s", strWhere);
+		m_pErrorLog->GetPrivate()->Add(buf, ELMT_ERROR, errorContext);
 
+		// Create a constant colour node
+		NodeColourConstantPtr pNode = new NodeColourConstant();
+		pNode->SetValue(1, 1, 0);
 
-//---------------------------------------------------------------------------------------------
-mu::Ptr<ASTOp> CodeGenerator::GenerateMissingColourCode(const char* strWhere, const void* errorContext)
-{
-	// Log a warning
-	char buf[256];
-	mutable_snprintf( buf, 256, "Required connection not found: %s", strWhere );
-	m_pErrorLog->GetPrivate()->Add(buf, ELMT_ERROR, errorContext);
+		COLOR_GENERATION_RESULT Result;
+		GenerateColor(Result, pNode);
 
-	// Create a constant colour node
-	NodeColourConstantPtr pNode = new NodeColourConstant();
-	pNode->SetValue(1, 1, 0);
+		return Result.op;
+	}
 
-	COLOR_GENERATION_RESULT Result;
-	GenerateColor(Result,pNode);
 
-	return Result.op;
 }
-
diff --git a/Engine/Plugins/Experimental/Mutable/Source/MutableTools/Private/MuT/CodeGenerator_FirstPass.h b/Engine/Plugins/Experimental/Mutable/Source/MutableTools/Private/MuT/CodeGenerator_FirstPass.h
--- a/Engine/Plugins/Experimental/Mutable/Source/MutableTools/Private/MuT/CodeGenerator_FirstPass.h
+++ b/Engine/Plugins/Experimental/Mutable/Source/MutableTools/Private/MuT/CodeGenerator_FirstPass.h
@@ -34,33 +34,6 @@ namespace mu
 	class Layout;
 	struct OBJECT_STATE;
 
-    //! Store the results of the code generation of a mesh.
-    struct MESH_GENERATION_RESULT
-    {
-        //! Mesh after all code tree is applied
-        Ptr<ASTOp> meshOp;
-
-        //! Original base mesh before removes, morphs, etc.
-        Ptr<ASTOp> baseMeshOp;
-
-		/** Source node layouts to use with these extra mesh. They don't have block ids. */
-		TArray<Ptr<const Layout>> layouts;
-
-		TArray<Ptr<ASTOp>> layoutOps;
-
-
-        struct EXTRA_LAYOUTS
-        {
-			/** Source node layouts to use with these extra mesh. They don't have block ids. */
-			TArray<Ptr<const Layout>> layouts;
-            Ptr<ASTOp> condition;
-            Ptr<ASTOp> meshFragment;
-        };
-		TArray< EXTRA_LAYOUTS > extraMeshLayouts;
-
-    };
-
-
 	//---------------------------------------------------------------------------------------------
 	//! First pass of the code generation process.
 	//! It collects data about the object hierarchy, the conditions for each object and the global
diff --git a/Engine/Plugins/Experimental/Mutable/Source/MutableTools/Private/MuT/CodeGenerator_Image.cpp b/Engine/Plugins/Experimental/Mutable/Source/MutableTools/Private/MuT/CodeGenerator_Image.cpp
--- a/Engine/Plugins/Experimental/Mutable/Source/MutableTools/Private/MuT/CodeGenerator_Image.cpp
+++ b/Engine/Plugins/Experimental/Mutable/Source/MutableTools/Private/MuT/CodeGenerator_Image.cpp
@@ -202,8 +202,6 @@ namespace mu
     {
 		const NodeImageConstant::Private& node = *InNode->GetPrivate();
 		
-		//MUTABLE_CPUPROFILER_SCOPE(NodeImageConstant);
-
         Ptr<ASTOpConstantResource> op = new ASTOpConstantResource();
         op->type = OP_TYPE::IM_CONSTANT;
 
@@ -232,10 +230,9 @@ namespace mu
 
         // Order of the operations is important: multiply first to avoid losing precision.
         // It will not overflow since image sizes are limited to 16 bit
-        cropRect.min = ( m_imageState.Last().m_imageRect.min * imageSize )
-                / m_imageState.Last().m_imageSize;
-        cropRect.size = ( m_imageState.Last().m_imageRect.size * imageSize )
-                / m_imageState.Last().m_imageSize;
+		vec2<int> RectDivisor = vec2<int>::max(vec2<int>(1,1),m_imageState.Last().m_imageSize);
+        cropRect.min = ( m_imageState.Last().m_imageRect.min * imageSize ) / RectDivisor;
+        cropRect.size = ( m_imageState.Last().m_imageRect.size * imageSize ) / RectDivisor;
 
         //check( cropRect.size[0]>0 && cropRect.size[1]>0 );
         cropRect.size[0] = FMath::Max( cropRect.size[0], 1 );
@@ -244,9 +241,7 @@ namespace mu
         //
         if ( pImage->GetSizeX()!=cropRect.size[0] || pImage->GetSizeY()!=cropRect.size[1] )
         {
-            ImagePtrConst pCropped =
-                ImageCrop( m_compilerOptions->m_imageCompressionQuality,
-                           pImage.get(), cropRect );
+            ImagePtrConst pCropped = ImageCrop( m_compilerOptions->m_imageCompressionQuality, pImage.get(), cropRect );
             op->SetValue( pCropped, m_compilerOptions->m_optimisationOptions.m_useDiskCache );
         }
         else
@@ -1440,9 +1435,16 @@ namespace mu
         // Mesh
         if ( node.m_pMesh )
         {
-            MESH_GENERATION_RESULT MeshResult;
-			check(m_overrideLayoutsStack.IsEmpty());
-            GenerateMesh(MeshResult, node.m_pMesh );
+            FMeshGenerationResult MeshResult;
+			FMeshGenerationOptions MeshOptions;
+			MeshOptions.State = m_currentStateIndex;
+			if (m_activeTags.Num())
+			{
+				MeshOptions.ActiveTags = m_activeTags.Last();
+			}
+			MeshOptions.bLayouts = true;			// We need the layout that we will use to render
+			MeshOptions.bUniqueVertexIDs = false;	// We don't need the IDs at this point.
+            GenerateMesh( MeshOptions, MeshResult, node.m_pMesh );
 
             pop->SetChild( pop->op.args.MeshProject.mesh, MeshResult.meshOp );
 
