[mutable] Fix morphs not having any effect.

diff --git a/Engine/Plugins/Experimental/Mutable/Source/MutableTools/Private/MuT/CodeGenerator_Mesh.cpp b/Engine/Plugins/Experimental/Mutable/Source/MutableTools/Private/MuT/CodeGenerator_Mesh.cpp
--- a/Engine/Plugins/Experimental/Mutable/Source/MutableTools/Private/MuT/CodeGenerator_Mesh.cpp
+++ b/Engine/Plugins/Experimental/Mutable/Source/MutableTools/Private/MuT/CodeGenerator_Mesh.cpp
@@ -27,7 +27,7 @@
 #include "MuT/ASTOpMeshFormat.h"
 #include "MuT/ASTOpMeshGeometryOperation.h"
 #include "MuT/ASTOpMeshMorphReshape.h"
-#include "MuT/ASTOpMeshRemapIndices.h"
+//#include "MuT/ASTOpMeshRemapIndices.h"
 #include "MuT/ASTOpMeshTransform.h"
 #include "MuT/ASTOpSwitch.h"
 #include "MuT/CodeGenerator.h"
@@ -282,11 +282,11 @@ class Node;
 
 
     //---------------------------------------------------------------------------------------------
-    void CodeGenerator::GenerateMesh( MESH_GENERATION_RESULT& result, const NodeMeshPtrConst& Untyped )
+    void CodeGenerator::GenerateMesh( const FMeshGenerationOptions& InOptions, FMeshGenerationResult& OutResult, const NodeMeshPtrConst& InUntypedNode)
     {
-        if (!Untyped)
+        if (!InUntypedNode)
         {
-            result = MESH_GENERATION_RESULT();
+            OutResult = FMeshGenerationResult();
             return;
         }
 
@@ -294,48 +294,54 @@ class Node;
 		m_currentBottomUpState.m_address = nullptr;
 
         // See if it was already generated
-		VISITED_MAP_KEY key = GetCurrentCacheKey(Untyped);
-        GeneratedMeshMap::ValueType* it = m_generatedMeshes.Find( key );
+		FGeneratedMeshCacheKey Key;
+		Key.Node = InUntypedNode;
+		Key.Options = InOptions;
+        GeneratedMeshMap::ValueType* it = m_generatedMeshes.Find(Key);
         if ( it )
         {
-            result = *it;
+			OutResult = *it;
             return;
         }
 
-		const NodeMesh* Node = Untyped.get();
+		const NodeMesh* Node = InUntypedNode.get();
 
         // Generate for each different type of node
-		switch (Untyped->GetMeshNodeType())
+		switch (Node->GetMeshNodeType())
 		{
-		case NodeMesh::EType::Constant: GenerateMesh_Constant(result, static_cast<const NodeMeshConstant*>(Node)); break;
-		case NodeMesh::EType::Format: GenerateMesh_Format(result, static_cast<const NodeMeshFormat*>(Node)); break;
-		case NodeMesh::EType::Morph: GenerateMesh_Morph(result, static_cast<const NodeMeshMorph*>(Node)); break;
-		case NodeMesh::EType::MakeMorph: GenerateMesh_MakeMorph(result, static_cast<const NodeMeshMakeMorph*>(Node)); break;
-		case NodeMesh::EType::Fragment: GenerateMesh_Fragment(result, static_cast<const NodeMeshFragment*>(Node)); break;
-		case NodeMesh::EType::Interpolate: GenerateMesh_Interpolate(result, static_cast<const NodeMeshInterpolate*>(Node)); break;
-		case NodeMesh::EType::Switch: GenerateMesh_Switch(result, static_cast<const NodeMeshSwitch*>(Node)); break;
-		case NodeMesh::EType::Subtract: GenerateMesh_Subtract(result, static_cast<const NodeMeshSubtract*>(Node)); break;
-		case NodeMesh::EType::Transform: GenerateMesh_Transform(result, static_cast<const NodeMeshTransform*>(Node)); break;
-		case NodeMesh::EType::ClipMorphPlane: GenerateMesh_ClipMorphPlane(result, static_cast<const NodeMeshClipMorphPlane*>(Node)); break;
-		case NodeMesh::EType::ClipWithMesh: GenerateMesh_ClipWithMesh(result, static_cast<const NodeMeshClipWithMesh*>(Node)); break;
-		case NodeMesh::EType::ApplyPose: GenerateMesh_ApplyPose(result, static_cast<const NodeMeshApplyPose*>(Node)); break;
-		case NodeMesh::EType::Variation: GenerateMesh_Variation(result, static_cast<const NodeMeshVariation*>(Node)); break;
-		case NodeMesh::EType::Table: GenerateMesh_Table(result, static_cast<const NodeMeshTable*>(Node)); break;
-		case NodeMesh::EType::GeometryOperation: GenerateMesh_GeometryOperation(result, static_cast<const NodeMeshGeometryOperation*>(Node)); break;
-		case NodeMesh::EType::Reshape: GenerateMesh_Reshape(result, static_cast<const NodeMeshReshape*>(Node)); break;
-		case NodeMesh::EType::ClipDeform: GenerateMesh_ClipDeform(result, static_cast<const NodeMeshClipDeform*>(Node)); break;
+		case NodeMesh::EType::Constant: GenerateMesh_Constant(InOptions, OutResult, static_cast<const NodeMeshConstant*>(Node)); break;
+		case NodeMesh::EType::Format: GenerateMesh_Format(InOptions, OutResult, static_cast<const NodeMeshFormat*>(Node)); break;
+		case NodeMesh::EType::Morph: GenerateMesh_Morph(InOptions, OutResult, static_cast<const NodeMeshMorph*>(Node)); break;
+		case NodeMesh::EType::MakeMorph: GenerateMesh_MakeMorph(InOptions, OutResult, static_cast<const NodeMeshMakeMorph*>(Node)); break;
+		case NodeMesh::EType::Fragment: GenerateMesh_Fragment(InOptions, OutResult, static_cast<const NodeMeshFragment*>(Node)); break;
+		case NodeMesh::EType::Interpolate: GenerateMesh_Interpolate(InOptions, OutResult, static_cast<const NodeMeshInterpolate*>(Node)); break;
+		case NodeMesh::EType::Switch: GenerateMesh_Switch(InOptions, OutResult, static_cast<const NodeMeshSwitch*>(Node)); break;
+		case NodeMesh::EType::Subtract: GenerateMesh_Subtract(InOptions, OutResult, static_cast<const NodeMeshSubtract*>(Node)); break;
+		case NodeMesh::EType::Transform: GenerateMesh_Transform(InOptions, OutResult, static_cast<const NodeMeshTransform*>(Node)); break;
+		case NodeMesh::EType::ClipMorphPlane: GenerateMesh_ClipMorphPlane(InOptions, OutResult, static_cast<const NodeMeshClipMorphPlane*>(Node)); break;
+		case NodeMesh::EType::ClipWithMesh: GenerateMesh_ClipWithMesh(InOptions, OutResult, static_cast<const NodeMeshClipWithMesh*>(Node)); break;
+		case NodeMesh::EType::ApplyPose: GenerateMesh_ApplyPose(InOptions, OutResult, static_cast<const NodeMeshApplyPose*>(Node)); break;
+		case NodeMesh::EType::Variation: GenerateMesh_Variation(InOptions, OutResult, static_cast<const NodeMeshVariation*>(Node)); break;
+		case NodeMesh::EType::Table: GenerateMesh_Table(InOptions, OutResult, static_cast<const NodeMeshTable*>(Node)); break;
+		case NodeMesh::EType::GeometryOperation: GenerateMesh_GeometryOperation(InOptions, OutResult, static_cast<const NodeMeshGeometryOperation*>(Node)); break;
+		case NodeMesh::EType::Reshape: GenerateMesh_Reshape(InOptions, OutResult, static_cast<const NodeMeshReshape*>(Node)); break;
+		case NodeMesh::EType::ClipDeform: GenerateMesh_ClipDeform(InOptions, OutResult, static_cast<const NodeMeshClipDeform*>(Node)); break;
 		case NodeMesh::EType::None: check(false);
 		}
 
         // Cache the result
-        m_generatedMeshes.Add( key, result );
+        m_generatedMeshes.Add( Key, OutResult);
     }
 
 
     //---------------------------------------------------------------------------------------------
-    void CodeGenerator::GenerateMesh_Morph( MESH_GENERATION_RESULT& result, const NodeMeshMorph* morph )
+    void CodeGenerator::GenerateMesh_Morph(
+		const FMeshGenerationOptions& InOptions, 
+		FMeshGenerationResult& OutResult, 
+		const NodeMeshMorph* InMorphNode 
+	)
     {
-        NodeMeshMorph::Private& node = *morph->GetPrivate();
+        NodeMeshMorph::Private& node = *InMorphNode->GetPrivate();
 
         Ptr<ASTOpFixed> OpMorph = new ASTOpFixed();
         OpMorph->op.type = OP_TYPE::ME_MORPH2;
@@ -353,11 +359,13 @@ class Node;
         }
 
         // Base
-        MESH_GENERATION_RESULT baseResult;
+        FMeshGenerationResult BaseResult;
         if ( node.m_pBase )
         {
-            GenerateMesh( baseResult, node.m_pBase );
-            OpMorph->SetChild( OpMorph->op.args.MeshMorph2.base, baseResult.meshOp );
+			FMeshGenerationOptions BaseOptions = InOptions;
+			BaseOptions.bUniqueVertexIDs = true;
+            GenerateMesh(BaseOptions,BaseResult, node.m_pBase );
+            OpMorph->SetChild( OpMorph->op.args.MeshMorph2.base, BaseResult.meshOp );
         }
         else
         {
@@ -369,17 +377,12 @@ class Node;
         // TODO: Support more than MUTABLE_OP_MAX_MORPH2_TARGETS targets
         if ( node.m_morphs.Num()>MUTABLE_OP_MAX_MORPH2_TARGETS )
         {
-            char temp[256];
-            mutable_snprintf( temp, 256,
-                              "A morph node has more targets [%d] than currently supported [%d].",
-                              node.m_morphs.Num(),
-                              MUTABLE_OP_MAX_MORPH2_TARGETS );
-            m_pErrorLog->GetPrivate()->Add( temp, ELMT_WARNING, node.m_errorContext );
+            m_pErrorLog->GetPrivate()->Add( 
+				FString::Printf(TEXT("A morph node has more targets [%d] than currently supported [%d]."), node.m_morphs.Num(), MUTABLE_OP_MAX_MORPH2_TARGETS), 
+				ELMT_WARNING, 
+				node.m_errorContext);
         }
 
-        m_overrideLayoutsStack.Add( baseResult.layouts );
-		m_activeTags.Add({});
-
         int count = 0;
         for ( std::size_t t=0
             ; t<node.m_morphs.Num() && t<MUTABLE_OP_MAX_MORPH2_TARGETS
@@ -387,20 +390,25 @@ class Node;
         {
             if ( auto pA = node.m_morphs[t].get() )
             {
-                MESH_GENERATION_RESULT targetResult;
-                GenerateMesh( targetResult, pA );
+                FMeshGenerationResult TargetResult;
+				FMeshGenerationOptions TargetOptions = InOptions;
+				TargetOptions.bUniqueVertexIDs = false;
+				TargetOptions.bLayouts = false;
+				TargetOptions.OverrideLayouts.Empty();
+				TargetOptions.ActiveTags.Empty();
+                GenerateMesh(TargetOptions, TargetResult, pA );
 
                 // TODO: Make sure that the target is a mesh with the morph format
-                Ptr<ASTOp> target = targetResult.meshOp;
+                Ptr<ASTOp> target = TargetResult.meshOp;
 
                 // If the vertex indices are supposed to be relative in the targets, adjust them
-                if (node.m_vertexIndicesAreRelative)
-                {
-                    Ptr<ASTOpMeshRemapIndices> remapIndices = new ASTOpMeshRemapIndices;
-                    remapIndices->source = target;
-                    remapIndices->reference = baseResult.baseMeshOp;
-                    target = remapIndices;
-                }
+                //if (node.m_vertexIndicesAreRelative)
+                //{
+                //    Ptr<ASTOpMeshRemapIndices> remapIndices = new ASTOpMeshRemapIndices;
+                //    remapIndices->source = target;
+                //    remapIndices->reference = baseResult.baseMeshOp;
+                //    target = remapIndices;
+                //}
 
                 OpMorph->SetChild( OpMorph->op.args.MeshMorph2.targets[count], target);
                 count++;
@@ -426,8 +434,8 @@ class Node;
 			OpBind->m_deformAllPhysics = node.m_deformAllPhysics;
 			OpBind->m_bindingMethod = static_cast<uint32>(EShapeBindingMethod::ReshapeClosestProject);
             
-			OpBind->Mesh = baseResult.meshOp;
-            OpBind->Shape = baseResult.meshOp;
+			OpBind->Mesh = BaseResult.meshOp;
+            OpBind->Shape = BaseResult.meshOp;
            
 			OpApply->m_reshapeVertices = OpBind->m_reshapeVertices;
 		    OpApply->m_reshapeSkeleton = OpBind->m_reshapeSkeleton;
@@ -441,28 +449,24 @@ class Node;
             OpMorphReshape->Reshape = OpApply;
         }
 
-        m_overrideLayoutsStack.Pop();
-		m_activeTags.Pop();
-
-		if (OpMorphReshape)
+ 		if (OpMorphReshape)
 		{
-			result.meshOp = OpMorphReshape;
+			OutResult.meshOp = OpMorphReshape;
 		}
 		else
 		{
-			result.meshOp = OpMorph;
+			OutResult.meshOp = OpMorph;
 		}
-        result.baseMeshOp = baseResult.baseMeshOp;
 
-        result.layouts = baseResult.layouts;
+        OutResult.baseMeshOp = BaseResult.baseMeshOp;
+        OutResult.layouts = BaseResult.layouts;
     } 
 
 
     //---------------------------------------------------------------------------------------------
-    void CodeGenerator::GenerateMesh_MakeMorph( MESH_GENERATION_RESULT& result,
-		const NodeMeshMakeMorph* morph )
+    void CodeGenerator::GenerateMesh_MakeMorph(const FMeshGenerationOptions& InOptions, FMeshGenerationResult& OutResult, const NodeMeshMakeMorph* InMakeMorphNode )
     {
-        NodeMeshMakeMorph::Private& node = *morph->GetPrivate();
+        NodeMeshMakeMorph::Private& node = *InMakeMorphNode->GetPrivate();
 
         Ptr<ASTOpFixed> op = new ASTOpFixed();
         op->op.type = OP_TYPE::ME_DIFFERENCE;
@@ -471,12 +475,15 @@ class Node;
         op->op.args.MeshDifference.ignoreTextureCoords = 1;
 
         // Base
-        MESH_GENERATION_RESULT baseResult;
+        FMeshGenerationResult BaseResult;
         if ( node.m_pBase )
         {
-            GenerateMesh( baseResult, node.m_pBase );
+			FMeshGenerationOptions BaseOptions = InOptions;
+			BaseOptions.bUniqueVertexIDs = true;
+			BaseOptions.bLayouts = false;
+			GenerateMesh(BaseOptions, BaseResult, node.m_pBase );
 
-            op->SetChild( op->op.args.MeshDifference.base, baseResult.meshOp );
+            op->SetChild( op->op.args.MeshDifference.base, BaseResult.meshOp );
         }
         else
         {
@@ -486,14 +493,17 @@ class Node;
         }
 
         // Target
-        m_overrideLayoutsStack.Add( baseResult.layouts );
-		m_activeTags.Add({});
 		if ( node.m_pTarget )
         {
-            MESH_GENERATION_RESULT targetResult;
-            GenerateMesh( targetResult, node.m_pTarget );
-
-            op->SetChild( op->op.args.MeshDifference.target, targetResult.meshOp );
+			FMeshGenerationOptions TargetOptions = InOptions;
+			TargetOptions.bUniqueVertexIDs = false;
+			TargetOptions.bLayouts = false;
+			TargetOptions.OverrideLayouts.Empty();
+			TargetOptions.ActiveTags.Empty();
+			FMeshGenerationResult TargetResult;
+            GenerateMesh( TargetOptions, TargetResult, node.m_pTarget );
+
+            op->SetChild( op->op.args.MeshDifference.target, TargetResult.meshOp );
         }
         else
         {
@@ -501,35 +511,38 @@ class Node;
             m_pErrorLog->GetPrivate()->Add( "Mesh make morph target node is not set.",
                                             ELMT_ERROR, node.m_errorContext );
         }
-        m_overrideLayoutsStack.Pop();
-		m_activeTags.Pop();
 
-        result.meshOp = op;
-        result.baseMeshOp = baseResult.baseMeshOp;
-        result.layouts = baseResult.layouts;
+        OutResult.meshOp = op;
+        OutResult.baseMeshOp = BaseResult.baseMeshOp;
+        OutResult.layouts = BaseResult.layouts;
     }
 
     //---------------------------------------------------------------------------------------------
-    void CodeGenerator::GenerateMesh_Fragment( MESH_GENERATION_RESULT& result,
+    void CodeGenerator::GenerateMesh_Fragment(const FMeshGenerationOptions& InOptions, FMeshGenerationResult& OutResult,
 		const NodeMeshFragment* fragment )
     {
         NodeMeshFragment::Private& node = *fragment->GetPrivate();
 
-        MESH_GENERATION_RESULT baseResult;
+        FMeshGenerationResult BaseResult;
         if ( node.m_pMesh )
         {
-            GenerateMesh( baseResult, node.m_pMesh );
+			FMeshGenerationOptions BaseOptions = InOptions;
+			if (node.m_fragmentType == NodeMeshFragment::FT_LAYOUT_BLOCKS)
+			{
+				BaseOptions.bLayouts = true;
+			}
+            GenerateMesh( BaseOptions, BaseResult, node.m_pMesh );
 
             if ( node.m_fragmentType==NodeMeshFragment::FT_LAYOUT_BLOCKS )
             {
                 Ptr<ASTOpMeshExtractLayoutBlocks> op = new ASTOpMeshExtractLayoutBlocks();
-                result.meshOp = op;
+                OutResult.meshOp = op;
 
-                op->source = baseResult.meshOp;
+                op->source = BaseResult.meshOp;
 
-                if ( baseResult.layouts.Num()>node.m_layoutOrGroup )
+                if (BaseResult.layouts.Num()>node.m_layoutOrGroup )
                 {
-                    LayoutPtrConst pSourceLayout = baseResult.layouts[ node.m_layoutOrGroup ];
+                    LayoutPtrConst pSourceLayout = BaseResult.layouts[ node.m_layoutOrGroup ];
                     const Layout* pLayout = m_addedLayouts[ pSourceLayout.get() ].get();
                     op->layout = (uint16)node.m_layoutOrGroup;
 
@@ -557,12 +570,13 @@ class Node;
 
             else if ( node.m_fragmentType==NodeMeshFragment::FT_FACE_GROUP )
             {
+				// \TODO: Deprecated?
                 Ptr<ASTOpFixed> op = new ASTOpFixed();
-                result.meshOp = op;
+                OutResult.meshOp = op;
 
                 op->op.type = OP_TYPE::ME_EXTRACTFACEGROUP;
 
-                op->SetChild( op->op.args.MeshExtractFaceGroup.source, baseResult.meshOp );
+                op->SetChild( op->op.args.MeshExtractFaceGroup.source, BaseResult.meshOp );
                 op->op.args.MeshExtractFaceGroup.group = node.m_layoutOrGroup;
             }
 
@@ -574,13 +588,13 @@ class Node;
                                             ELMT_ERROR, node.m_errorContext );
         }
 
-        result.baseMeshOp = baseResult.baseMeshOp;
-        result.layouts = baseResult.layouts;
+        OutResult.baseMeshOp = BaseResult.baseMeshOp;
+        OutResult.layouts = BaseResult.layouts;
     }
 
 
     //---------------------------------------------------------------------------------------------
-    void CodeGenerator::GenerateMesh_Interpolate( MESH_GENERATION_RESULT& result,
+    void CodeGenerator::GenerateMesh_Interpolate(const FMeshGenerationOptions& InOptions, FMeshGenerationResult& OutResult,
 		const NodeMeshInterpolate* interpolate )
     {
         NodeMeshInterpolate::Private& node = *interpolate->GetPrivate();
@@ -588,7 +602,7 @@ class Node;
         // Generate the code
         Ptr<ASTOpFixed> op = new ASTOpFixed();
         op->op.type = OP_TYPE::ME_INTERPOLATE;
-        result.meshOp = op;
+        OutResult.meshOp = op;
 
         // Factor
         if ( Node* pFactor = node.m_pFactor.get() )
@@ -611,34 +625,29 @@ class Node;
         {
             if ( NodeMesh* pA = node.m_targets[t].get() )
             {
-                if ( count>0 )
-                {
-                    m_overrideLayoutsStack.Add( result.layouts );
-                }
+				FMeshGenerationOptions TargetOptions = InOptions;
+				TargetOptions.bUniqueVertexIDs = true;
+				TargetOptions.OverrideLayouts.Empty();
 
-                MESH_GENERATION_RESULT targetResult;
-                GenerateMesh( targetResult, pA );
+                FMeshGenerationResult TargetResult;
+                GenerateMesh( TargetOptions, TargetResult, pA );
 
-                if ( count>0 )
-                {
-                    m_overrideLayoutsStack.Pop();
-                }
 
                 // The first target is the base
                 if (count==0)
                 {
-                    base = targetResult.meshOp;
-                    op->SetChild( op->op.args.MeshInterpolate.base, targetResult.meshOp );
+                    base = TargetResult.meshOp;
+                    op->SetChild( op->op.args.MeshInterpolate.base, TargetResult.meshOp );
 
-                    result.baseMeshOp = targetResult.baseMeshOp;
-                    result.layouts = targetResult.layouts;
+                    OutResult.baseMeshOp = TargetResult.baseMeshOp;
+                    OutResult.layouts = TargetResult.layouts;
                 }
                 else
                 {
                     Ptr<ASTOpFixed> dop = new ASTOpFixed();
                     dop->op.type = OP_TYPE::ME_DIFFERENCE;
                     dop->SetChild( dop->op.args.MeshDifference.base, base );
-                    dop->SetChild( dop->op.args.MeshDifference.target, targetResult.meshOp );
+                    dop->SetChild( dop->op.args.MeshDifference.target, TargetResult.meshOp );
 
                     // \todo Texcoords are broken?
                     dop->op.args.MeshDifference.ignoreTextureCoords = 1;
@@ -685,7 +694,7 @@ class Node;
 
 
     //---------------------------------------------------------------------------------------------
-    void CodeGenerator::GenerateMesh_Switch( MESH_GENERATION_RESULT& result, const NodeMeshSwitch* sw )
+    void CodeGenerator::GenerateMesh_Switch(const FMeshGenerationOptions& InOptions, FMeshGenerationResult& OutResult, const NodeMeshSwitch* sw )
     {
         NodeMeshSwitch::Private& node = *sw->GetPrivate();
 
@@ -693,7 +702,7 @@ class Node;
         {
             // No options in the switch!
             // TODO
-            result = MESH_GENERATION_RESULT();
+            OutResult = FMeshGenerationResult();
 			return;
         }
 
@@ -714,44 +723,41 @@ class Node;
         // Options
         for ( std::size_t t=0; t< node.m_options.Num(); ++t )
         {
+			FMeshGenerationOptions TargetOptions = InOptions;
+
             if (t!=0)
             {
-                m_overrideLayoutsStack.Add( result.layouts );
+				TargetOptions.OverrideLayouts = OutResult.layouts;
             }
 
             if ( node.m_options[t] )
             {
-                MESH_GENERATION_RESULT branchResults;
-                GenerateMesh( branchResults, node.m_options[t] );
+                FMeshGenerationResult BranchResults;
+                GenerateMesh(TargetOptions, BranchResults, node.m_options[t] );
 
-                auto branch = branchResults.meshOp;
+                auto branch = BranchResults.meshOp;
                 op->cases.Emplace((int16)t,op,branch);
 
                 if (t==0)
                 {
-                    result = branchResults;
+                    OutResult = BranchResults;
                 }
             }
-
-            if (t!=0)
-            {
-                m_overrideLayoutsStack.Pop();
-            }
         }
 
-        result.meshOp = op;
+        OutResult.meshOp = op;
     }
 
 
 	//---------------------------------------------------------------------------------------------
-	void CodeGenerator::GenerateMesh_Table(MESH_GENERATION_RESULT& result, const NodeMeshTable* TableNode)
+	void CodeGenerator::GenerateMesh_Table(const FMeshGenerationOptions& InOptions, FMeshGenerationResult& OutResult, const NodeMeshTable* TableNode)
 	{
 		//
-		MESH_GENERATION_RESULT NewResult = result;
+		FMeshGenerationResult NewResult = OutResult;
 		int t = 0;
 
 		Ptr<ASTOp> Op = GenerateTableSwitch<NodeMeshTable::Private, TCT_MESH, OP_TYPE::ME_SWITCH>(*TableNode->GetPrivate(), 
-			[this, &NewResult, &t] (const NodeMeshTable::Private& node, int colIndex, int row, ErrorLog* pErrorLog)
+			[this, &NewResult, &t, &InOptions] (const NodeMeshTable::Private& node, int colIndex, int row, ErrorLog* pErrorLog)
 			{
 				NodeMeshConstantPtr pCell = new NodeMeshConstant();
 				MeshPtr pMesh = node.m_pTable->GetPrivate()->m_rows[row].m_values[colIndex].m_pMesh;
@@ -774,40 +780,38 @@ class Node;
 					pCell->SetLayout(i, node.m_layouts[i]);
 				}
 
+				FMeshGenerationOptions TargetOptions = InOptions;
+
 				if (t != 0)
 				{
-					m_overrideLayoutsStack.Add(NewResult.layouts);
+					TargetOptions.OverrideLayouts = NewResult.layouts;
 				}
 
-				MESH_GENERATION_RESULT branchResults;
-				GenerateMesh(branchResults, pCell);
+				FMeshGenerationResult BranchResults;
+				GenerateMesh(TargetOptions, BranchResults, pCell);
 
 				if (t == 0)
 				{
-					NewResult = branchResults;
-				}
-				else
-				{
-					m_overrideLayoutsStack.Pop();
-				}
+					NewResult = BranchResults;
+				}				
 
 				++t;
-				return branchResults.meshOp;
+				return BranchResults.meshOp;
 			});
 
 		NewResult.meshOp = Op;
 
-		result = NewResult;
+		OutResult = NewResult;
 	}
 
 
     //---------------------------------------------------------------------------------------------
-    void CodeGenerator::GenerateMesh_Variation( MESH_GENERATION_RESULT& result,
+    void CodeGenerator::GenerateMesh_Variation(const FMeshGenerationOptions& InOptions, FMeshGenerationResult& OutResult,
 		const NodeMeshVariation* va )
     {
         NodeMeshVariation::Private& node = *va->GetPrivate();
 
-        MESH_GENERATION_RESULT currentResult;
+        FMeshGenerationResult currentResult;
         Ptr<ASTOp> currentMeshOp;
 
         bool firstOptionProcessed = false;
@@ -815,10 +819,12 @@ class Node;
         // Default case
         if ( node.m_defaultMesh )
         {
-            MESH_GENERATION_RESULT branchResults;
-			GenerateMesh( branchResults, node.m_defaultMesh );
-            currentMeshOp = branchResults.meshOp;
-            currentResult = branchResults;
+            FMeshGenerationResult BranchResults;
+			FMeshGenerationOptions DefaultOptions = InOptions;
+
+			GenerateMesh(DefaultOptions, BranchResults, node.m_defaultMesh );
+            currentMeshOp = BranchResults.meshOp;
+            currentResult = BranchResults;
             firstOptionProcessed = true;
         }
 
@@ -837,40 +843,36 @@ class Node;
 
             if ( tagIndex < 0 )
             {
-                char buf[256];
-                mutable_snprintf( buf, 256, "Unknown tag found in mesh variation [%s].",
-                                  tag.c_str() );
-
-                m_pErrorLog->GetPrivate()->Add( buf, ELMT_WARNING, node.m_errorContext );
+                m_pErrorLog->GetPrivate()->Add( 
+					FString::Printf(TEXT("Unknown tag found in mesh variation [%s]."), tag.c_str()), 
+					ELMT_WARNING, 
+					node.m_errorContext
+				);
                 continue;
             }
 
             Ptr<ASTOp> variationMeshOp;
             if ( node.m_variations[t].m_mesh )
             {
+				FMeshGenerationOptions VariationOptions = InOptions;
+
                 if (firstOptionProcessed)
                 {
-                    m_overrideLayoutsStack.Add( currentResult.layouts );
+					VariationOptions.OverrideLayouts = currentResult.layouts;
                 }
          
-                MESH_GENERATION_RESULT branchResults;
-				GenerateMesh( branchResults, node.m_variations[t].m_mesh );
-
-                variationMeshOp = branchResults.meshOp;
+                FMeshGenerationResult BranchResults;
+				GenerateMesh(VariationOptions, BranchResults, node.m_variations[t].m_mesh );
 
-				if (firstOptionProcessed )
-				{
-					m_overrideLayoutsStack.Pop();
-				}
+                variationMeshOp = BranchResults.meshOp;
 
                 if ( !firstOptionProcessed )
                 {
                     firstOptionProcessed = true;                   
-                    currentResult = branchResults;
+                    currentResult = BranchResults;
                 }
             }
 
-
             Ptr<ASTOpConditional> conditional = new ASTOpConditional;
             conditional->type = OP_TYPE::ME_CONDITIONAL;
             conditional->no = currentMeshOp;
@@ -880,86 +882,78 @@ class Node;
             currentMeshOp = conditional;
         }
 
-        result = currentResult;
-        result.meshOp = currentMeshOp;
+        OutResult = currentResult;
+        OutResult.meshOp = currentMeshOp;
     }
 
 
     //---------------------------------------------------------------------------------------------
-    void CodeGenerator::GenerateMesh_Constant( MESH_GENERATION_RESULT& result,
-		const NodeMeshConstant* constant )
+    void CodeGenerator::GenerateMesh_Constant(const FMeshGenerationOptions& InOptions, FMeshGenerationResult& OutResult, const NodeMeshConstant* constant )
     {
         NodeMeshConstant::Private& node = *constant->GetPrivate();
 
         Ptr<ASTOpConstantResource> op = new ASTOpConstantResource();
         op->type = OP_TYPE::ME_CONSTANT;
-		result.baseMeshOp = op;
-		result.meshOp = op;
+		OutResult.baseMeshOp = op;
+		OutResult.meshOp = op;
 
-		bool bIsOverridingLayouts = !m_overrideLayoutsStack.IsEmpty();
+		bool bIsOverridingLayouts = !InOptions.OverrideLayouts.IsEmpty();
 
         MeshPtr pMesh = node.m_pValue.get();
-        if (pMesh)
-        {
-            // Clone the mesh
-            MeshPtr pCloned = pMesh->Clone();
-            pCloned->EnsureSurfaceData();
-            if ( !bIsOverridingLayouts )
-            {
-                // This means that we are processing a base mesh
+		if (pMesh)
+		{
+			// Clone the mesh
+			MeshPtr pCloned = pMesh->Clone();
+			pCloned->EnsureSurfaceData();
 
-                // We will redefine the layouts
-				result.layouts.Empty();
+			// We will redefine the layouts
+			OutResult.layouts.Empty();
 
-                // Apply whatever transform is necessary for every layout
-                for ( std::size_t LayoutIndex=0; LayoutIndex<node.m_layouts.Num(); ++LayoutIndex)
-                {
-					NodeLayoutPtr pLayoutNode = node.m_layouts[LayoutIndex];
-					// TODO: In a cleanup of the design of the layouts, we should remove this cast.
-					const NodeLayoutBlocks* TypedNode = dynamic_cast<NodeLayoutBlocks*>(pLayoutNode.get());
-					if (TypedNode)
+			if (InOptions.bLayouts)
+			{
+				if (!bIsOverridingLayouts)
+				{
+					// Apply whatever transform is necessary for every layout
+					for (int32 LayoutIndex = 0; LayoutIndex < node.m_layouts.Num(); ++LayoutIndex)
 					{
-						if (!pCloned)
+						NodeLayoutPtr pLayoutNode = node.m_layouts[LayoutIndex];
+						// TODO: In a cleanup of the design of the layouts, we should remove this cast.
+						const NodeLayoutBlocks* TypedNode = dynamic_cast<NodeLayoutBlocks*>(pLayoutNode.get());
+						if (TypedNode)
 						{
-							m_pErrorLog->GetPrivate()->Add("Generating a layout node without a parent mesh.",
-								ELMT_ERROR, TypedNode->GetPrivate()->m_errorContext);
-							return;
-						}
+							if (!pCloned)
+							{
+								m_pErrorLog->GetPrivate()->Add("Generating a layout node without a parent mesh.",
+									ELMT_ERROR, TypedNode->GetPrivate()->m_errorContext);
+								return;
+							}
 
-						LayoutPtr pSourceLayout = TypedNode->GetPrivate()->m_pLayout;
-						result.layouts.Add(pSourceLayout);
+							LayoutPtr pSourceLayout = TypedNode->GetPrivate()->m_pLayout;
+							OutResult.layouts.Add(pSourceLayout);
 
-						PrepareForLayout(pSourceLayout,
-							pCloned, LayoutIndex,
-							TypedNode->GetPrivate()->m_errorContext);
+							PrepareForLayout(pSourceLayout,
+								pCloned, LayoutIndex,
+								TypedNode->GetPrivate()->m_errorContext);
 
+						}
 					}
-                }
-            }
-            else
-            {
-                // We need to apply the transform of the source layouts
-                for ( int32 LayoutIndex=0; LayoutIndex<m_overrideLayoutsStack.Last().Num(); ++LayoutIndex)
-                {
-                    PrepareForLayout( m_overrideLayoutsStack.Last()[LayoutIndex],
-                                      pCloned, LayoutIndex,
-                                      node.m_errorContext  );
-
-                }
+				}
+				else
+				{
+					// We need to apply the transform of the layouts used to override
+					for (int32 LayoutIndex = 0; LayoutIndex < InOptions.OverrideLayouts.Num(); ++LayoutIndex)
+					{
+						PrepareForLayout(InOptions.OverrideLayouts[LayoutIndex],
+							pCloned, LayoutIndex,
+							node.m_errorContext);
 
-				result.layouts = m_overrideLayoutsStack.Last();
+					}
 
-                // Disabled: lots of misleading messages (when not using layouts?).
-                // Warn about unnecessary layouts
-//                if ( node.m_layouts.size() )
-//                {
-//                    m_pErrorLog->GetPrivate()->Add( "The layouts in this mesh are not required.",
-//                                                    ELMT_WARNING, node.m_errorContext );
-//                }
-            }
+					OutResult.layouts = InOptions.OverrideLayouts;
+				}
+			}
 
-            // See if we already have a mesh identical to this one, except for the internal data
-            // like vertex indices.
+            // See if we already have a mesh identical to this one, except for the internal data like vertex indices.
 			// For now we cannot handle the case were we are overriding layouts and reusing meshes at the same time.
             bool bIsDuplicated = false;
             MeshPtrConst Candidate;
@@ -967,39 +961,49 @@ class Node;
             {
 				Candidate = m_constantMeshes[i];
 
-                if (Candidate->IsSimilar(*pCloned) )
+                if (Candidate->IsSimilar(*pCloned,InOptions.bLayouts) )
                 {
+					// If it is similar, and we need unique vertex IDs, check that it also has them. This was skipped in the IsSimilar.
+					if (InOptions.bUniqueVertexIDs)
+					{
+						int32 FoundBuffer = -1;
+						int32 FoundChannel = -1;
+						Candidate->GetVertexBuffers().FindChannel(MBS_VERTEXINDEX, 0, &FoundBuffer, &FoundChannel);
+						bool bHasUniqueVertexIDs = FoundBuffer >= 0 && FoundChannel >= 0;
+						if (!bHasUniqueVertexIDs)
+						{
+							break;
+						}
+					}
+
 					bIsDuplicated = true;
 
-					// We want to reuse the layouts from the similar mesh.					
-					// Remap layouts that we have created in the previosu step from the source mesh to the ones created for the mesh we will reuse instead.
-					for (int l = 0; l < Candidate->GetLayoutCount(); ++l)
+					if (InOptions.bLayouts)
 					{
-						check(Candidate->GetLayoutCount() == pCloned->GetLayoutCount());
+						// We want to reuse the layouts from the similar mesh.					
+						// Remap layouts that we have created in the previous step from the source mesh to the ones created for the mesh we will reuse instead.
+						for (int l = 0; l < Candidate->GetLayoutCount(); ++l)
+						{
+							check(Candidate->GetLayoutCount() == pCloned->GetLayoutCount());
 
-						const Layout* ClonedMeshLayout = pCloned->GetLayout(l);
+							const Layout* ClonedMeshLayout = pCloned->GetLayout(l);
 
-						Ptr<const Layout> DestLayoutKey = nullptr;
-						for (const auto& e : m_addedLayouts)
-						{
-							if (e.Value == ClonedMeshLayout)
+							Ptr<const Layout> DestLayoutKey = nullptr;
+							for (const auto& e : m_addedLayouts)
 							{
-								DestLayoutKey = e.Key;
+								if (e.Value == ClonedMeshLayout)
+								{
+									DestLayoutKey = e.Key;
+									break;
+								}
 							}
-						}
-						check(DestLayoutKey);
+							check(DestLayoutKey);
 
-						const Layout* GeneratedMeshLayout = Candidate->GetLayout(l);
-						if (bIsOverridingLayouts)
-						{
-							const Layout* OverridingLayout = m_overrideLayoutsStack.Last()[l].get();
-							check(OverridingLayout->m_blocks.IsEmpty() || OverridingLayout->m_blocks[0].m_id == -1);
-							GeneratedMeshLayout = m_addedLayouts[OverridingLayout].get();
+							const Layout* GeneratedMeshLayout = Candidate->GetLayout(l);
+							check(DestLayoutKey->m_blocks.IsEmpty() || DestLayoutKey->m_blocks[0].m_id == -1);
+							check(GeneratedMeshLayout->m_blocks.IsEmpty() || GeneratedMeshLayout->m_blocks[0].m_id != -1);
+							m_addedLayouts.Add(DestLayoutKey, GeneratedMeshLayout);
 						}
-
-						check(DestLayoutKey->m_blocks.IsEmpty() || DestLayoutKey->m_blocks[0].m_id == -1);
-						check(GeneratedMeshLayout->m_blocks.IsEmpty() || GeneratedMeshLayout->m_blocks[0].m_id != -1);
-						m_addedLayouts.Add(DestLayoutKey, GeneratedMeshLayout);
 					}
                 }
             }
@@ -1010,45 +1014,50 @@ class Node;
             }
             else
             {
-                // Enumerate the vertices uniquely unless they already have indices
-                int buf = -1;
-                int chan = -1;
-                pCloned->GetVertexBuffers().FindChannel( MBS_VERTEXINDEX,0,&buf,&chan );
-                bool hasVertexIndices = (buf>=0 && chan>=0);
-                if ( !hasVertexIndices )
-                {
-                    int newBuffer = pCloned->GetVertexBuffers().GetBufferCount();
-                    pCloned->GetVertexBuffers().SetBufferCount( newBuffer+1 );
-                    MESH_BUFFER_SEMANTIC semantic = MBS_VERTEXINDEX;
-                    int semanticIndex = 0;
-                    MESH_BUFFER_FORMAT format = MBF_UINT32;
-                    int components = 1;
-                    int offset = 0;
-                    pCloned->GetVertexBuffers().SetBuffer
-                            (
-                                newBuffer,
-                                sizeof(uint32_t),
-                                1,
-                                &semantic, &semanticIndex,
-                                &format, &components,
-                                &offset
-                            );
-                    uint32_t* pIdData = (uint32_t*)pCloned->GetVertexBuffers().GetBufferData( newBuffer );
-                    for ( int i=0; i<pMesh->GetVertexCount(); ++i )
-                    {
-                        (*pIdData++) = m_freeVertexIndex++;
-						check(m_freeVertexIndex < std::numeric_limits<uint32_t>::max());
-                    }
-                }
+				if (InOptions.bUniqueVertexIDs)
+				{
+					// Enumerate the vertices uniquely unless they already have indices
+					int buf = -1;
+					int chan = -1;
+					pCloned->GetVertexBuffers().FindChannel(MBS_VERTEXINDEX, 0, &buf, &chan);
+					bool hasVertexIndices = (buf >= 0 && chan >= 0);
+					if (!hasVertexIndices)
+					{
+						int newBuffer = pCloned->GetVertexBuffers().GetBufferCount();
+						pCloned->GetVertexBuffers().SetBufferCount(newBuffer + 1);
+						MESH_BUFFER_SEMANTIC semantic = MBS_VERTEXINDEX;
+						int semanticIndex = 0;
+						MESH_BUFFER_FORMAT format = MBF_UINT32;
+						int components = 1;
+						int offset = 0;
+						pCloned->GetVertexBuffers().SetBuffer
+						(
+							newBuffer,
+							sizeof(uint32),
+							1,
+							&semantic, &semanticIndex,
+							&format, &components,
+							&offset
+						);
+						uint32* pIdData = (uint32*)pCloned->GetVertexBuffers().GetBufferData(newBuffer);
+						for (int i = 0; i < pMesh->GetVertexCount(); ++i)
+						{
+							check(m_freeVertexIndex < std::numeric_limits<uint32>::max());
 
-                // Add the constant data
-                m_constantMeshes.Add(pCloned);
-                op->SetValue( pCloned.get(), m_compilerOptions->m_optimisationOptions.m_useDiskCache );
-            }
+							(*pIdData++) = m_freeVertexIndex++;
+							check(m_freeVertexIndex < std::numeric_limits<uint32>::max());
+						}
+					}
+				}
+
+				// Add the constant data
+				m_constantMeshes.Add(pCloned);
+				op->SetValue(pCloned.get(), m_compilerOptions->m_optimisationOptions.m_useDiskCache);
+			}
         }
         else
         {
-			result.layouts.Empty();
+			OutResult.layouts.Empty();
 
             // This data is required
             MeshPtr pTempMesh = new Mesh();
@@ -1056,43 +1065,33 @@ class Node;
             m_constantMeshes.Add(pTempMesh);
 
             // Log an error message
-            m_pErrorLog->GetPrivate()->Add( "Constant mesh not set.",
-                                            ELMT_WARNING, node.m_errorContext );
+            m_pErrorLog->GetPrivate()->Add( "Constant mesh not set.", ELMT_WARNING, node.m_errorContext );
         }
 
 		// Apply the modifier for the pre-normal operations stage.
 		BOTTOM_UP_STATE temp = m_currentBottomUpState;
-		if (!m_activeTags.IsEmpty())
-		{
-			// Clearing layout stack to avoid unwanted information
-			TArray< TArray<Ptr<const Layout>> > m_overrideLayoutsStackCopy = m_overrideLayoutsStack;
-			m_overrideLayoutsStack.Empty();
 
-			bool bModifiersForBeforeOperations = true;
-			result.meshOp = ApplyMeshModifiers(op, m_activeTags.Last(),
-				bModifiersForBeforeOperations, node.m_errorContext);
+		bool bModifiersForBeforeOperations = true;
+		OutResult.meshOp = ApplyMeshModifiers(op, InOptions.ActiveTags,
+			bModifiersForBeforeOperations, node.m_errorContext);
 
-			// Retrieving stack information
-			m_overrideLayoutsStack = m_overrideLayoutsStackCopy;
-		}
 		m_currentBottomUpState = temp;
 
     }
 
 
-
-
     //---------------------------------------------------------------------------------------------
-    void CodeGenerator::GenerateMesh_Format( MESH_GENERATION_RESULT& result,
+    void CodeGenerator::GenerateMesh_Format(const FMeshGenerationOptions& InOptions, FMeshGenerationResult& OutResult,
 		const NodeMeshFormat* format )
     {
         NodeMeshFormat::Private& node = *format->GetPrivate();
 
         if ( node.m_pSource )
         {
-            MESH_GENERATION_RESULT baseResult;
-			GenerateMesh(baseResult, node.m_pSource);
+			FMeshGenerationOptions Options = InOptions;
 
+			FMeshGenerationResult baseResult;
+			GenerateMesh(Options,baseResult, node.m_pSource);
             Ptr<ASTOpMeshFormat> op = new ASTOpMeshFormat();
             op->Source = baseResult.meshOp;
             op->Buffers = 0;
@@ -1124,20 +1123,20 @@ class Node;
 
             m_constantMeshes.Add(pFormatMesh);
 
-            result.meshOp = op;
-            result.baseMeshOp = baseResult.baseMeshOp;
-            result.layouts = baseResult.layouts;
+            OutResult.meshOp = op;
+            OutResult.baseMeshOp = baseResult.baseMeshOp;
+            OutResult.layouts = baseResult.layouts;
         }
         else
         {
             // Put something there
-            GenerateMesh( result, new NodeMeshConstant() );
+            GenerateMesh(InOptions, OutResult, new NodeMeshConstant() );
         }
     }
 
 
     //---------------------------------------------------------------------------------------------
-    void CodeGenerator::GenerateMesh_Subtract( MESH_GENERATION_RESULT& result,
+    void CodeGenerator::GenerateMesh_Subtract(const FMeshGenerationOptions& InOptions, FMeshGenerationResult& OutResult,
 		const NodeMeshSubtract* subs )
     {
 		// This node is deprecated.
@@ -1146,7 +1145,7 @@ class Node;
 
 
     //---------------------------------------------------------------------------------------------
-    void CodeGenerator::GenerateMesh_Transform( MESH_GENERATION_RESULT& result,
+    void CodeGenerator::GenerateMesh_Transform(const FMeshGenerationOptions& InOptions, FMeshGenerationResult& OutResult,
                                                 const NodeMeshTransform* trans )
     {
         const auto& node = *trans->GetPrivate();
@@ -1156,8 +1155,8 @@ class Node;
         // Base
         if (node.m_pSource)
         {
-            GenerateMesh( result, node.m_pSource);
-            op->source = result.meshOp;
+            GenerateMesh(InOptions, OutResult, node.m_pSource);
+            op->source = OutResult.meshOp;
         }
         else
         {
@@ -1168,12 +1167,12 @@ class Node;
 
         op->matrix = node.m_transform;
 
-        result.meshOp = op;
+        OutResult.meshOp = op;
     }
 
 
     //---------------------------------------------------------------------------------------------
-    void CodeGenerator::GenerateMesh_ClipMorphPlane( MESH_GENERATION_RESULT& result,
+    void CodeGenerator::GenerateMesh_ClipMorphPlane(const FMeshGenerationOptions& InOptions, FMeshGenerationResult& OutResult,
                                                      const NodeMeshClipMorphPlane* clip )
     {
         const auto& node = *clip->GetPrivate();
@@ -1183,8 +1182,8 @@ class Node;
         // Base
         if (node.m_pSource)
         {
-            GenerateMesh( result, node.m_pSource);
-            op->source = result.meshOp;
+            GenerateMesh(InOptions, OutResult, node.m_pSource);
+            op->source = OutResult.meshOp;
         }
         else
         {
@@ -1239,12 +1238,12 @@ class Node;
         op->dist = node.m_dist;
         op->factor = node.m_factor;
 
-        result.meshOp = op;
+        OutResult.meshOp = op;
     }
 
 
     //---------------------------------------------------------------------------------------------
-    void CodeGenerator::GenerateMesh_ClipWithMesh( MESH_GENERATION_RESULT& result,
+    void CodeGenerator::GenerateMesh_ClipWithMesh(const FMeshGenerationOptions& InOptions, FMeshGenerationResult& OutResult,
                                                    const NodeMeshClipWithMesh* clip )
     {
         const auto& node = *clip->GetPrivate();
@@ -1255,8 +1254,8 @@ class Node;
         // Base
         if (node.m_pSource)
         {
-            GenerateMesh( result, node.m_pSource );
-            op->SetChild( op->op.args.MeshClipWithMesh.source, result.meshOp );
+            GenerateMesh( InOptions, OutResult, node.m_pSource );
+            op->SetChild( op->op.args.MeshClipWithMesh.source, OutResult.meshOp );
         }
         else
         {
@@ -1268,13 +1267,15 @@ class Node;
         // Clipping mesh
         if (node.m_pClipMesh)
         {
-			m_activeTags.Add({});
-
-            MESH_GENERATION_RESULT clipResult;
-            GenerateMesh( clipResult, node.m_pClipMesh);
+			FMeshGenerationOptions ClipOptions = InOptions;
+			ClipOptions.bUniqueVertexIDs = false;
+			ClipOptions.bLayouts = false;
+			ClipOptions.OverrideLayouts.Empty();
+			ClipOptions.ActiveTags.Empty();
+
+            FMeshGenerationResult clipResult;
+            GenerateMesh(ClipOptions, clipResult, node.m_pClipMesh);
             op->SetChild( op->op.args.MeshClipWithMesh.clipMesh, clipResult.meshOp );
-
-			m_activeTags.Pop();
 		}
         else
         {
@@ -1283,11 +1284,11 @@ class Node;
                 ELMT_ERROR, node.m_errorContext);
         }
 
-        result.meshOp = op;
+        OutResult.meshOp = op;
     }
 
 	//---------------------------------------------------------------------------------------------
-	void CodeGenerator::GenerateMesh_ClipDeform(MESH_GENERATION_RESULT& Result, const NodeMeshClipDeform* ClipDeform)
+	void CodeGenerator::GenerateMesh_ClipDeform(const FMeshGenerationOptions& InOptions, FMeshGenerationResult& Result, const NodeMeshClipDeform* ClipDeform)
 	{
 		const auto& Node = *ClipDeform->GetPrivate();
 
@@ -1297,7 +1298,7 @@ class Node;
 		// Base Mesh
 		if (Node.m_pBaseMesh)
 		{
-			GenerateMesh(Result, Node.m_pBaseMesh);
+			GenerateMesh(InOptions, Result, Node.m_pBaseMesh);
 			OpBind->Mesh = Result.meshOp;
 		}
 		else
@@ -1310,8 +1311,14 @@ class Node;
 		// Base Shape
 		if (Node.m_pClipShape)
 		{
-			MESH_GENERATION_RESULT baseResult;
-			GenerateMesh(baseResult, Node.m_pClipShape);
+			FMeshGenerationOptions ClipOptions = InOptions;
+			ClipOptions.bUniqueVertexIDs = false;
+			ClipOptions.bLayouts = false;
+			ClipOptions.OverrideLayouts.Empty();
+			ClipOptions.ActiveTags.Empty();
+
+			FMeshGenerationResult baseResult;
+			GenerateMesh(ClipOptions, baseResult, Node.m_pClipShape);
 			OpBind->Shape = baseResult.meshOp;
 			OpClipDeform->ClipShape = baseResult.meshOp;
 		}
@@ -1323,7 +1330,7 @@ class Node;
 	}
 
     //---------------------------------------------------------------------------------------------
-    void CodeGenerator::GenerateMesh_ApplyPose( MESH_GENERATION_RESULT& result,
+    void CodeGenerator::GenerateMesh_ApplyPose(const FMeshGenerationOptions& InOptions, FMeshGenerationResult& OutResult,
                                                 const NodeMeshApplyPose* pose )
     {
         const auto& node = *pose->GetPrivate();
@@ -1333,8 +1340,8 @@ class Node;
         // Base
         if (node.m_pBase)
         {
-            GenerateMesh( result, node.m_pBase );
-            op->base = result.meshOp;
+            GenerateMesh(InOptions, OutResult, node.m_pBase );
+            op->base = OutResult.meshOp;
         }
         else
         {
@@ -1346,16 +1353,15 @@ class Node;
         // Pose mesh
         if (node.m_pPose)
         {
-            // We don't need layouts for the pose mesh
-            m_overrideLayoutsStack.Add( TArray<LayoutPtrConst>() );
-			m_activeTags.Add({});
-
-            MESH_GENERATION_RESULT poseResult;
-            GenerateMesh( poseResult, node.m_pPose );
+			FMeshGenerationOptions PoseOptions = InOptions;
+			PoseOptions.bUniqueVertexIDs = false;
+			PoseOptions.bLayouts = false;
+			PoseOptions.OverrideLayouts.Empty();
+			PoseOptions.ActiveTags.Empty();
+
+            FMeshGenerationResult poseResult;
+            GenerateMesh(PoseOptions, poseResult, node.m_pPose );
             op->pose = poseResult.meshOp;
-
-            m_overrideLayoutsStack.Pop();
-			m_activeTags.Pop();
 		}
         else
         {
@@ -1364,12 +1370,12 @@ class Node;
                 ELMT_ERROR, node.m_errorContext);
         }
 
-        result.meshOp = op;
+        OutResult.meshOp = op;
     }
 
 
 	//---------------------------------------------------------------------------------------------
-	void CodeGenerator::GenerateMesh_GeometryOperation(MESH_GENERATION_RESULT& result, const NodeMeshGeometryOperation* geom)
+	void CodeGenerator::GenerateMesh_GeometryOperation(const FMeshGenerationOptions& InOptions, FMeshGenerationResult& OutResult, const NodeMeshGeometryOperation* geom)
 	{
 		const auto& node = *geom->GetPrivate();
 
@@ -1378,8 +1384,8 @@ class Node;
 		// Mesh A
 		if (node.m_pMeshA)
 		{
-			GenerateMesh(result, node.m_pMeshA);
-			op->meshA = result.meshOp;
+			GenerateMesh(InOptions, OutResult, node.m_pMeshA);
+			op->meshA = OutResult.meshOp;
 		}
 		else
 		{
@@ -1391,20 +1397,26 @@ class Node;
 		// Mesh B
 		if (node.m_pMeshB)
 		{
-			MESH_GENERATION_RESULT bResult;
-			GenerateMesh(bResult, node.m_pMeshB);
+			FMeshGenerationOptions OtherOptions = InOptions;
+			OtherOptions.bUniqueVertexIDs = false;
+			OtherOptions.bLayouts = false;
+			OtherOptions.OverrideLayouts.Empty();
+			OtherOptions.ActiveTags.Empty();
+
+			FMeshGenerationResult bResult;
+			GenerateMesh(OtherOptions, bResult, node.m_pMeshB);
 			op->meshB = bResult.meshOp;
 		}
 
 		op->scalarA = Generate(node.m_pScalarA);
 		op->scalarB = Generate(node.m_pScalarB);
 
-		result.meshOp = op;
+		OutResult.meshOp = op;
 	}
 
 
 	//---------------------------------------------------------------------------------------------
-	void CodeGenerator::GenerateMesh_Reshape(MESH_GENERATION_RESULT& result, const NodeMeshReshape* reshape)
+	void CodeGenerator::GenerateMesh_Reshape(const FMeshGenerationOptions& InOptions, FMeshGenerationResult& OutResult, const NodeMeshReshape* reshape)
 	{
 		const NodeMeshReshape::Private& node = *reshape->GetPrivate();
 
@@ -1428,8 +1440,8 @@ class Node;
 		// Base Mesh
 		if (node.m_pBaseMesh)
 		{
-			GenerateMesh(result, node.m_pBaseMesh);
-			opBind->Mesh = result.meshOp;
+			GenerateMesh(InOptions, OutResult, node.m_pBaseMesh);
+			opBind->Mesh = OutResult.meshOp;
 		}
 		else
 		{
@@ -1438,14 +1450,17 @@ class Node;
 		}
 
 		// Base and target shapes shouldn't have layouts or modifiers.
-		m_overrideLayoutsStack.Add({});
-		m_activeTags.Add({});
+		FMeshGenerationOptions ShapeOptions = InOptions;
+		ShapeOptions.bUniqueVertexIDs = false;
+		ShapeOptions.bLayouts = false;
+		ShapeOptions.OverrideLayouts.Empty();
+		ShapeOptions.ActiveTags.Empty();
 
 		// Base Shape
 		if (node.m_pBaseShape)
 		{
-			MESH_GENERATION_RESULT baseResult;
-			GenerateMesh(baseResult, node.m_pBaseShape);
+			FMeshGenerationResult baseResult;
+			GenerateMesh(ShapeOptions, baseResult, node.m_pBaseShape);
 			opBind->Shape = baseResult.meshOp;
 		}
 
@@ -1454,15 +1469,12 @@ class Node;
 		// Target Shape
 		if (node.m_pTargetShape)
 		{
-			MESH_GENERATION_RESULT targetResult;
-			GenerateMesh(targetResult, node.m_pTargetShape);
+			FMeshGenerationResult targetResult;
+			GenerateMesh(ShapeOptions, targetResult, node.m_pTargetShape);
 			opApply->Shape = targetResult.meshOp;
 		}
 
-		m_overrideLayoutsStack.Pop();
-		m_activeTags.Pop();
-
-		result.meshOp = opApply;
+		OutResult.meshOp = opApply;
 	}
 
 }
diff --git a/Engine/Plugins/Experimental/Mutable/Source/MutableTools/Private/MuT/CodeGenerator_Projector.cpp b/Engine/Plugins/Experimental/Mutable/Source/MutableTools/Private/MuT/CodeGenerator_Projector.cpp
--- a/Engine/Plugins/Experimental/Mutable/Source/MutableTools/Private/MuT/CodeGenerator_Projector.cpp
+++ b/Engine/Plugins/Experimental/Mutable/Source/MutableTools/Private/MuT/CodeGenerator_Projector.cpp
@@ -25,157 +25,158 @@
 #include <utility>
 
 
-using namespace mu;
-
-
-//-------------------------------------------------------------------------------------------------
-PROJECTOR ProjectorFromNode( const NodeProjectorConstant::Private& node )
-{
-    PROJECTOR p;
-    p.type = node.m_type;
-    p.position[0] = node.m_position[0];
-    p.position[1] = node.m_position[1];
-    p.position[2] = node.m_position[2];
-    p.direction[0] = node.m_direction[0];
-    p.direction[1] = node.m_direction[1];
-    p.direction[2] = node.m_direction[2];
-    p.up[0] = node.m_up[0];
-    p.up[1] = node.m_up[1];
-    p.up[2] = node.m_up[2];
-    p.scale[0] = node.m_scale[0];
-    p.scale[1] = node.m_scale[1];
-    p.scale[2] = node.m_scale[2];
-    p.projectionAngle = node.m_projectionAngle;
-
-    return p;
-}
-
-
-//-------------------------------------------------------------------------------------------------
-void CodeGenerator::GenerateProjector( PROJECTOR_GENERATION_RESULT& result, const NodeProjectorPtrConst& untyped )
-{
-    if (!untyped)
-    {
-        result = PROJECTOR_GENERATION_RESULT();
-        return;
-    }
-
-    // See if it was already generated
-	VISITED_MAP_KEY key = GetCurrentCacheKey(untyped);
-	GeneratedProjectorsMap::ValueType* it = m_generatedProjectors.Find( key );
-    if ( it )
-    {
-        result = *it;
-        return;
-    }
-
-    // Generate for each different type of node
-    if ( auto Constant = dynamic_cast<const NodeProjectorConstant*>(untyped.get()) )
-    {
-        GenerateProjector_Constant( result, Constant );
-    }
-    else if ( auto Param = dynamic_cast<const NodeProjectorParameter*>(untyped.get()) )
-    {
-        GenerateProjector_Parameter( result, Param );
-    }
-    else
-    {
-        check(false);
-    }
-
-
-    // Cache the result
-    m_generatedProjectors.Add( key, result );
-}
-
-
-//-------------------------------------------------------------------------------------------------
-void CodeGenerator::GenerateProjector_Constant( PROJECTOR_GENERATION_RESULT& result,
-                                                const Ptr<const NodeProjectorConstant>& constant )
-{
-    const NodeProjectorConstant::Private& node = *constant->GetPrivate();
-
-    Ptr<ASTOpConstantProjector> op = new ASTOpConstantProjector();
-    op->value = ProjectorFromNode( node );
-
-    result.op = op;
-    result.type = op->value.type;
-}
-
-
-//-------------------------------------------------------------------------------------------------
-void CodeGenerator::GenerateProjector_Parameter( PROJECTOR_GENERATION_RESULT& result,
-                                                const Ptr<const NodeProjectorParameter>& paramn )
-{
-    const NodeProjectorParameter::Private& node = *paramn->GetPrivate();
-
-    Ptr<ASTOpParameter> op;
-
-    auto it = m_nodeVariables.find(node.m_pNode);
-    if ( it == m_nodeVariables.end() )
-    {
-        PARAMETER_DESC param;
-        param.m_name = node.m_name;
-        param.m_uid = node.m_uid;
-        param.m_type = PARAMETER_TYPE::T_PROJECTOR;
-
-        PROJECTOR p = ProjectorFromNode( node );
-        param.m_defaultValue.m_projector = p;
-
-        op = new ASTOpParameter();
-        op->type = OP_TYPE::PR_PARAMETER;
-        op->parameter = param;
-
-        // Generate the code for the ranges
-        for ( int32 a=0; a<node.m_ranges.Num(); ++a )
-        {
-            RANGE_GENERATION_RESULT rangeResult;
-            GenerateRange( rangeResult, node.m_ranges[a] );
-            op->ranges.Emplace( op.get(), rangeResult.sizeOp, rangeResult.rangeName, rangeResult.rangeUID );
-        }
-
-        m_nodeVariables[node.m_pNode] = op;
-    }
-    else
-    {
-        op = it->second;
-    }
-
-    result.type = op->parameter.m_defaultValue.m_projector.type;
-    result.op = op;
-}
-
-
-//-------------------------------------------------------------------------------------------------
-void CodeGenerator::GenerateMissingProjectorCode( PROJECTOR_GENERATION_RESULT& result,
-                                                  const void* errorContext )
+namespace mu
 {
-    // Log an error message
-    m_pErrorLog->GetPrivate()->Add( "Required projector connection not found.",
-                                    ELMT_ERROR, errorContext );
 
-    PROJECTOR p;
-    p.type = PROJECTOR_TYPE::PLANAR;
+	//-------------------------------------------------------------------------------------------------
+	PROJECTOR ProjectorFromNode(const NodeProjectorConstant::Private& node)
+	{
+		PROJECTOR p;
+		p.type = node.m_type;
+		p.position[0] = node.m_position[0];
+		p.position[1] = node.m_position[1];
+		p.position[2] = node.m_position[2];
+		p.direction[0] = node.m_direction[0];
+		p.direction[1] = node.m_direction[1];
+		p.direction[2] = node.m_direction[2];
+		p.up[0] = node.m_up[0];
+		p.up[1] = node.m_up[1];
+		p.up[2] = node.m_up[2];
+		p.scale[0] = node.m_scale[0];
+		p.scale[1] = node.m_scale[1];
+		p.scale[2] = node.m_scale[2];
+		p.projectionAngle = node.m_projectionAngle;
+
+		return p;
+	}
+
+
+	//-------------------------------------------------------------------------------------------------
+	void CodeGenerator::GenerateProjector(PROJECTOR_GENERATION_RESULT& result, const NodeProjectorPtrConst& untyped)
+	{
+		if (!untyped)
+		{
+			result = PROJECTOR_GENERATION_RESULT();
+			return;
+		}
+
+		// See if it was already generated
+		VISITED_MAP_KEY key = GetCurrentCacheKey(untyped);
+		GeneratedProjectorsMap::ValueType* it = m_generatedProjectors.Find(key);
+		if (it)
+		{
+			result = *it;
+			return;
+		}
+
+		// Generate for each different type of node
+		if (auto Constant = dynamic_cast<const NodeProjectorConstant*>(untyped.get()))
+		{
+			GenerateProjector_Constant(result, Constant);
+		}
+		else if (auto Param = dynamic_cast<const NodeProjectorParameter*>(untyped.get()))
+		{
+			GenerateProjector_Parameter(result, Param);
+		}
+		else
+		{
+			check(false);
+		}
+
+
+		// Cache the result
+		m_generatedProjectors.Add(key, result);
+	}
+
+
+	//-------------------------------------------------------------------------------------------------
+	void CodeGenerator::GenerateProjector_Constant(PROJECTOR_GENERATION_RESULT& result,
+		const Ptr<const NodeProjectorConstant>& constant)
+	{
+		const NodeProjectorConstant::Private& node = *constant->GetPrivate();
+
+		Ptr<ASTOpConstantProjector> op = new ASTOpConstantProjector();
+		op->value = ProjectorFromNode(node);
+
+		result.op = op;
+		result.type = op->value.type;
+	}
+
+
+	//-------------------------------------------------------------------------------------------------
+	void CodeGenerator::GenerateProjector_Parameter(PROJECTOR_GENERATION_RESULT& result,
+		const Ptr<const NodeProjectorParameter>& paramn)
+	{
+		const NodeProjectorParameter::Private& node = *paramn->GetPrivate();
+
+		Ptr<ASTOpParameter> op;
+
+		auto it = m_nodeVariables.find(node.m_pNode);
+		if (it == m_nodeVariables.end())
+		{
+			PARAMETER_DESC param;
+			param.m_name = node.m_name;
+			param.m_uid = node.m_uid;
+			param.m_type = PARAMETER_TYPE::T_PROJECTOR;
+
+			PROJECTOR p = ProjectorFromNode(node);
+			param.m_defaultValue.m_projector = p;
+
+			op = new ASTOpParameter();
+			op->type = OP_TYPE::PR_PARAMETER;
+			op->parameter = param;
+
+			// Generate the code for the ranges
+			for (int32 a = 0; a < node.m_ranges.Num(); ++a)
+			{
+				RANGE_GENERATION_RESULT rangeResult;
+				GenerateRange(rangeResult, node.m_ranges[a]);
+				op->ranges.Emplace(op.get(), rangeResult.sizeOp, rangeResult.rangeName, rangeResult.rangeUID);
+			}
+
+			m_nodeVariables[node.m_pNode] = op;
+		}
+		else
+		{
+			op = it->second;
+		}
+
+		result.type = op->parameter.m_defaultValue.m_projector.type;
+		result.op = op;
+	}
+
+
+	//-------------------------------------------------------------------------------------------------
+	void CodeGenerator::GenerateMissingProjectorCode(PROJECTOR_GENERATION_RESULT& result,
+		const void* errorContext)
+	{
+		// Log an error message
+		m_pErrorLog->GetPrivate()->Add("Required projector connection not found.",
+			ELMT_ERROR, errorContext);
+
+		PROJECTOR p;
+		p.type = PROJECTOR_TYPE::PLANAR;
+
+		p.direction[0] = 1;
+		p.direction[1] = 0;
+		p.direction[2] = 0;
+
+		p.up[0] = 0;
+		p.up[1] = 0;
+		p.up[2] = 1;
+
+		p.position[0] = 0.0f;
+		p.position[1] = 0.0f;
+		p.position[2] = 0.0f;
+
+		p.scale[0] = 1.0f;
+		p.scale[1] = 1.0f;
+
+		Ptr<ASTOpConstantProjector> op = new ASTOpConstantProjector();
+		op->value = p;
+
+		result.op = op;
+		result.type = p.type;
+	}
 
-    p.direction[0] = 1;
-    p.direction[1] = 0;
-    p.direction[2] = 0;
-
-    p.up[0] = 0;
-    p.up[1] = 0;
-    p.up[2] = 1;
-
-    p.position[0] = 0.0f;
-    p.position[1] = 0.0f;
-    p.position[2] = 0.0f;
-
-    p.scale[0] = 1.0f;
-    p.scale[1] = 1.0f;
-
-    Ptr<ASTOpConstantProjector> op = new ASTOpConstantProjector();
-    op->value = p;
-
-    result.op = op;
-    result.type = p.type;
 }
-
diff --git a/Engine/Plugins/Experimental/Mutable/Source/MutableTools/Private/MuT/CodeGenerator_Scalar.cpp b/Engine/Plugins/Experimental/Mutable/Source/MutableTools/Private/MuT/CodeGenerator_Scalar.cpp
--- a/Engine/Plugins/Experimental/Mutable/Source/MutableTools/Private/MuT/CodeGenerator_Scalar.cpp
+++ b/Engine/Plugins/Experimental/Mutable/Source/MutableTools/Private/MuT/CodeGenerator_Scalar.cpp
@@ -48,415 +48,415 @@
 #include <memory>
 #include <utility>
 
-namespace mu { class Node; }
 
-
-using namespace mu;
+namespace mu
+{
+	class Node;
 
 
-//-------------------------------------------------------------------------------------------------
-void CodeGenerator::GenerateScalar( SCALAR_GENERATION_RESULT& result, const NodeScalarPtrConst& untyped )
-{
-    if (!untyped)
-    {
-        result = SCALAR_GENERATION_RESULT();
-        return;
-    }
-
-    // See if it was already generated
-	VISITED_MAP_KEY key = GetCurrentCacheKey(untyped);
-	GeneratedScalarsMap::ValueType* it = m_generatedScalars.Find( key );
-    if ( it )
-    {
-        result = *it;
-        return;
-    }
-
-
-    // Generate for each different type of node
-    if ( auto Constant = dynamic_cast<const NodeScalarConstant*>(untyped.get()) )
-    {
-        GenerateScalar_Constant( result, Constant );
-    }
-    else if ( auto Param = dynamic_cast<const NodeScalarParameter*>(untyped.get()) )
-    {
-        GenerateScalar_Parameter( result, Param );
-    }
-	else if (auto Switch = dynamic_cast<const NodeScalarSwitch*>(untyped.get()))
-	{
-		GenerateScalar_Switch(result, Switch);
-	}
-	else if (auto EnumParam = dynamic_cast<const NodeScalarEnumParameter*>(untyped.get()))
-	{
-		GenerateScalar_EnumParameter(result, EnumParam);
-	}
-	else if (auto Curve = dynamic_cast<const NodeScalarCurve*>(untyped.get()))
-	{
-		GenerateScalar_Curve(result, Curve);
-	}
-	else if (auto Arithmetic = dynamic_cast<const NodeScalarArithmeticOperation*>(untyped.get()))
-	{
-		GenerateScalar_Arithmetic(result, Arithmetic);
-	}
-	else if (auto Variation = dynamic_cast<const NodeScalarVariation*>(untyped.get()))
+	//-------------------------------------------------------------------------------------------------
+	void CodeGenerator::GenerateScalar(SCALAR_GENERATION_RESULT& result, const NodeScalarPtrConst& untyped)
 	{
-		GenerateScalar_Variation(result, Variation);
-	}
-	else if (auto Table = dynamic_cast<const NodeScalarTable*>(untyped.get()))
-	{
-		GenerateScalar_Table(result, Table);
-	}
-	else
-    {
-        check(false);
-        mu::Halt();
-    }
+		if (!untyped)
+		{
+			result = SCALAR_GENERATION_RESULT();
+			return;
+		}
 
-    // Cache the result
-    m_generatedScalars.Add( key, result );
-}
+		// See if it was already generated
+		VISITED_MAP_KEY key = GetCurrentCacheKey(untyped);
+		GeneratedScalarsMap::ValueType* it = m_generatedScalars.Find(key);
+		if (it)
+		{
+			result = *it;
+			return;
+		}
 
 
-//-------------------------------------------------------------------------------------------------
-void CodeGenerator::GenerateScalar_Constant(SCALAR_GENERATION_RESULT& result, const Ptr<const NodeScalarConstant>& Typed )
-{
-    const NodeScalarConstant::Private& node = *Typed->GetPrivate();
+		// Generate for each different type of node
+		if (auto Constant = dynamic_cast<const NodeScalarConstant*>(untyped.get()))
+		{
+			GenerateScalar_Constant(result, Constant);
+		}
+		else if (auto Param = dynamic_cast<const NodeScalarParameter*>(untyped.get()))
+		{
+			GenerateScalar_Parameter(result, Param);
+		}
+		else if (auto Switch = dynamic_cast<const NodeScalarSwitch*>(untyped.get()))
+		{
+			GenerateScalar_Switch(result, Switch);
+		}
+		else if (auto EnumParam = dynamic_cast<const NodeScalarEnumParameter*>(untyped.get()))
+		{
+			GenerateScalar_EnumParameter(result, EnumParam);
+		}
+		else if (auto Curve = dynamic_cast<const NodeScalarCurve*>(untyped.get()))
+		{
+			GenerateScalar_Curve(result, Curve);
+		}
+		else if (auto Arithmetic = dynamic_cast<const NodeScalarArithmeticOperation*>(untyped.get()))
+		{
+			GenerateScalar_Arithmetic(result, Arithmetic);
+		}
+		else if (auto Variation = dynamic_cast<const NodeScalarVariation*>(untyped.get()))
+		{
+			GenerateScalar_Variation(result, Variation);
+		}
+		else if (auto Table = dynamic_cast<const NodeScalarTable*>(untyped.get()))
+		{
+			GenerateScalar_Table(result, Table);
+		}
+		else
+		{
+			check(false);
+			mu::Halt();
+		}
 
-	Ptr<ASTOpFixed> op = new ASTOpFixed();
-	op->op.type = OP_TYPE::SC_CONSTANT;
-	op->op.args.ScalarConstant.value = node.m_value;
+		// Cache the result
+		m_generatedScalars.Add(key, result);
+	}
 
-    result.op = op;
-}
 
+	//-------------------------------------------------------------------------------------------------
+	void CodeGenerator::GenerateScalar_Constant(SCALAR_GENERATION_RESULT& result, const Ptr<const NodeScalarConstant>& Typed)
+	{
+		const NodeScalarConstant::Private& node = *Typed->GetPrivate();
 
-//-------------------------------------------------------------------------------------------------
-void CodeGenerator::GenerateScalar_Parameter( SCALAR_GENERATION_RESULT& result, const Ptr<const NodeScalarParameter>& Typed)
-{
-    const NodeScalarParameter::Private& node = *Typed->GetPrivate();
+		Ptr<ASTOpFixed> op = new ASTOpFixed();
+		op->op.type = OP_TYPE::SC_CONSTANT;
+		op->op.args.ScalarConstant.value = node.m_value;
 
-    Ptr<ASTOpParameter> op;
+		result.op = op;
+	}
 
-	auto it = m_nodeVariables.find(node.m_pNode);
-	if (it == m_nodeVariables.end())
+
+	//-------------------------------------------------------------------------------------------------
+	void CodeGenerator::GenerateScalar_Parameter(SCALAR_GENERATION_RESULT& result, const Ptr<const NodeScalarParameter>& Typed)
 	{
-		PARAMETER_DESC param;
-		param.m_name = node.m_name;
-		param.m_uid = node.m_uid;
-		param.m_type = PARAMETER_TYPE::T_FLOAT;
-		param.m_defaultValue.m_float = node.m_defaultValue;
-		param.m_detailedType = node.m_detailedType;
-
-		op = new ASTOpParameter();
-		op->type = OP_TYPE::SC_PARAMETER;
-		op->parameter = param;
-
-		// Generate the code for the ranges
-		for (int32 a = 0; a < node.m_ranges.Num(); ++a)
-		{
-			RANGE_GENERATION_RESULT rangeResult;
-			GenerateRange(rangeResult, node.m_ranges[a]);
-			op->ranges.Emplace(op.get(), rangeResult.sizeOp, rangeResult.rangeName, rangeResult.rangeUID);
-		}
+		const NodeScalarParameter::Private& node = *Typed->GetPrivate();
 
-		// Generate the code for the additional images in the parameter
-		for (int32 a = 0; a < node.m_additionalImages.Num(); ++a)
+		Ptr<ASTOpParameter> op;
+
+		auto it = m_nodeVariables.find(node.m_pNode);
+		if (it == m_nodeVariables.end())
 		{
-			Ptr<ASTOp> descAd;
-			if (node.m_additionalImages[a])
+			PARAMETER_DESC param;
+			param.m_name = node.m_name;
+			param.m_uid = node.m_uid;
+			param.m_type = PARAMETER_TYPE::T_FLOAT;
+			param.m_defaultValue.m_float = node.m_defaultValue;
+			param.m_detailedType = node.m_detailedType;
+
+			op = new ASTOpParameter();
+			op->type = OP_TYPE::SC_PARAMETER;
+			op->parameter = param;
+
+			// Generate the code for the ranges
+			for (int32 a = 0; a < node.m_ranges.Num(); ++a)
 			{
-				// We take whatever size will be produced
-				IMAGE_STATE newState;
-				FImageDesc desc = CalculateImageDesc(*node.m_additionalImages[a]->GetBasePrivate());
-				newState.m_imageSize = desc.m_size;
-				newState.m_imageRect.min[0] = 0;
-				newState.m_imageRect.min[1] = 0;
-				newState.m_imageRect.size = desc.m_size;
-				newState.m_layoutBlock = -1;
-				m_imageState.Add(newState);
-
-				// Generate
-				descAd = Generate(node.m_additionalImages[a]);
-				check(descAd);
-
-				// Restore rect
-				m_imageState.Pop();
+				RANGE_GENERATION_RESULT rangeResult;
+				GenerateRange(rangeResult, node.m_ranges[a]);
+				op->ranges.Emplace(op.get(), rangeResult.sizeOp, rangeResult.rangeName, rangeResult.rangeUID);
 			}
-			op->additionalImages.Emplace(op, descAd);
+
+			// Generate the code for the additional images in the parameter
+			for (int32 a = 0; a < node.m_additionalImages.Num(); ++a)
+			{
+				Ptr<ASTOp> descAd;
+				if (node.m_additionalImages[a])
+				{
+					// We take whatever size will be produced
+					IMAGE_STATE newState;
+					FImageDesc desc = CalculateImageDesc(*node.m_additionalImages[a]->GetBasePrivate());
+					newState.m_imageSize = desc.m_size;
+					newState.m_imageRect.min[0] = 0;
+					newState.m_imageRect.min[1] = 0;
+					newState.m_imageRect.size = desc.m_size;
+					newState.m_layoutBlock = -1;
+					m_imageState.Add(newState);
+
+					// Generate
+					descAd = Generate(node.m_additionalImages[a]);
+					check(descAd);
+
+					// Restore rect
+					m_imageState.Pop();
+				}
+				op->additionalImages.Emplace(op, descAd);
+			}
+
+			m_nodeVariables[node.m_pNode] = op;
+		}
+		else
+		{
+			op = it->second;
 		}
 
-		m_nodeVariables[node.m_pNode] = op;
+		result.op = op;
 	}
-	else
-	{
-		op = it->second;
-	}
-
-    result.op = op;
-}
 
 
-//-------------------------------------------------------------------------------------------------
-void CodeGenerator::GenerateScalar_EnumParameter(SCALAR_GENERATION_RESULT& result, const Ptr<const NodeScalarEnumParameter>& Typed)
-{
-	const NodeScalarEnumParameter::Private& node = *Typed->GetPrivate();
+	//-------------------------------------------------------------------------------------------------
+	void CodeGenerator::GenerateScalar_EnumParameter(SCALAR_GENERATION_RESULT& result, const Ptr<const NodeScalarEnumParameter>& Typed)
+	{
+		const NodeScalarEnumParameter::Private& node = *Typed->GetPrivate();
 
-	Ptr<ASTOpParameter> op;
+		Ptr<ASTOpParameter> op;
 
-	auto it = m_nodeVariables.find(node.m_pNode);
-	if (it == m_nodeVariables.end())
-	{
-		PARAMETER_DESC param;
-		param.m_name = node.m_name;
-		param.m_uid = node.m_uid;
-		param.m_type = PARAMETER_TYPE::T_INT;
-		param.m_defaultValue.m_int = node.m_defaultValue;
-		param.m_detailedType = node.m_detailedType;
-
-		param.m_possibleValues.SetNum(node.m_options.Num());
-		for (int32 i = 0; i < node.m_options.Num(); ++i)
+		auto it = m_nodeVariables.find(node.m_pNode);
+		if (it == m_nodeVariables.end())
 		{
-			param.m_possibleValues[i].m_value = (int16_t)node.m_options[i].value;
-			param.m_possibleValues[i].m_name = node.m_options[i].name;
-		}
+			PARAMETER_DESC param;
+			param.m_name = node.m_name;
+			param.m_uid = node.m_uid;
+			param.m_type = PARAMETER_TYPE::T_INT;
+			param.m_defaultValue.m_int = node.m_defaultValue;
+			param.m_detailedType = node.m_detailedType;
+
+			param.m_possibleValues.SetNum(node.m_options.Num());
+			for (int32 i = 0; i < node.m_options.Num(); ++i)
+			{
+				param.m_possibleValues[i].m_value = (int16_t)node.m_options[i].value;
+				param.m_possibleValues[i].m_name = node.m_options[i].name;
+			}
 
-		op = new ASTOpParameter();
-		op->type = OP_TYPE::NU_PARAMETER;
-		op->parameter = param;
+			op = new ASTOpParameter();
+			op->type = OP_TYPE::NU_PARAMETER;
+			op->parameter = param;
 
-		// Generate the code for the ranges
-		for (int32 a = 0; a < node.m_ranges.Num(); ++a)
+			// Generate the code for the ranges
+			for (int32 a = 0; a < node.m_ranges.Num(); ++a)
+			{
+				RANGE_GENERATION_RESULT rangeResult;
+				GenerateRange(rangeResult, node.m_ranges[a]);
+				op->ranges.Emplace(op.get(), rangeResult.sizeOp, rangeResult.rangeName, rangeResult.rangeUID);
+			}
+
+			m_nodeVariables[node.m_pNode] = op;
+
+		}
+		else
 		{
-			RANGE_GENERATION_RESULT rangeResult;
-			GenerateRange(rangeResult, node.m_ranges[a]);
-			op->ranges.Emplace(op.get(), rangeResult.sizeOp, rangeResult.rangeName, rangeResult.rangeUID);
+			op = it->second;
 		}
 
-		m_nodeVariables[node.m_pNode] = op;
-
-	}
-	else
-	{
-		op = it->second;
+		result.op = op;
 	}
 
-	result.op = op;
-}
-
-
-//-------------------------------------------------------------------------------------------------
-void CodeGenerator::GenerateScalar_Switch(SCALAR_GENERATION_RESULT& result, const Ptr<const NodeScalarSwitch>& Typed)
-{
-	const NodeScalarSwitch::Private& node = *Typed->GetPrivate();
 
-	if (node.m_options.Num() == 0)
+	//-------------------------------------------------------------------------------------------------
+	void CodeGenerator::GenerateScalar_Switch(SCALAR_GENERATION_RESULT& result, const Ptr<const NodeScalarSwitch>& Typed)
 	{
-		// No options in the switch!
-		Ptr<ASTOp> missingOp = GenerateMissingScalarCode("Switch option",
-			1.0f,
-			node.m_errorContext);
-		result.op = missingOp;
-		return;
-	}
+		const NodeScalarSwitch::Private& node = *Typed->GetPrivate();
 
-	Ptr<ASTOpSwitch> op = new ASTOpSwitch();
-	op->type = OP_TYPE::SC_SWITCH;
+		if (node.m_options.Num() == 0)
+		{
+			// No options in the switch!
+			Ptr<ASTOp> missingOp = GenerateMissingScalarCode("Switch option",
+				1.0f,
+				node.m_errorContext);
+			result.op = missingOp;
+			return;
+		}
 
-	// Variable value
-	if (node.m_pParameter)
-	{
-		SCALAR_GENERATION_RESULT ChildResult;
-		GenerateScalar(ChildResult,node.m_pParameter.get());
-		op->variable = ChildResult.op;
-	}
-	else
-	{
-		// This argument is required
-		op->variable = GenerateMissingScalarCode("Switch variable", 0.0f, node.m_errorContext);
-	}
+		Ptr<ASTOpSwitch> op = new ASTOpSwitch();
+		op->type = OP_TYPE::SC_SWITCH;
 
-	// Options
-	for (int32 t = 0; t < node.m_options.Num(); ++t)
-	{
-		Ptr<ASTOp> branch;
-		if (node.m_options[t])
+		// Variable value
+		if (node.m_pParameter)
 		{
 			SCALAR_GENERATION_RESULT ChildResult;
-			GenerateScalar(ChildResult,node.m_options[t].get());
-			branch = ChildResult.op;
+			GenerateScalar(ChildResult, node.m_pParameter.get());
+			op->variable = ChildResult.op;
 		}
 		else
 		{
 			// This argument is required
-			branch = GenerateMissingScalarCode("Switch option", 1.0f, node.m_errorContext);
+			op->variable = GenerateMissingScalarCode("Switch variable", 0.0f, node.m_errorContext);
 		}
-		op->cases.Emplace((int16_t)t, op, branch);
-	}
-
-	result.op = op;
-}
-
-
-//-------------------------------------------------------------------------------------------------
-void CodeGenerator::GenerateScalar_Variation(SCALAR_GENERATION_RESULT& result, const Ptr<const NodeScalarVariation>& Typed)
-{
-	const NodeScalarVariation::Private& node = *Typed->GetPrivate();
 
-	Ptr<ASTOp> op;
-
-	// Default case
-	if (node.m_defaultScalar)
-	{
-		MESH_GENERATION_RESULT branchResults;
-
-		SCALAR_GENERATION_RESULT ChildResult;
-		GenerateScalar(ChildResult,node.m_defaultScalar);
-		op = ChildResult.op;
-	}
-	else
-	{
-		// This argument is required
-		op = GenerateMissingScalarCode("Variation default", 0.0f, node.m_errorContext);
-	}
-
-	// Process variations in reverse order, since conditionals are built bottom-up.
-	for (int t = node.m_variations.Num() - 1; t >= 0; --t)
-	{
-		int tagIndex = -1;
-		const string& tag = node.m_variations[t].m_tag;
-		for (int i = 0; i < m_firstPass.m_tags.Num(); ++i)
+		// Options
+		for (int32 t = 0; t < node.m_options.Num(); ++t)
 		{
-			if (m_firstPass.m_tags[i].tag == tag)
+			Ptr<ASTOp> branch;
+			if (node.m_options[t])
+			{
+				SCALAR_GENERATION_RESULT ChildResult;
+				GenerateScalar(ChildResult, node.m_options[t].get());
+				branch = ChildResult.op;
+			}
+			else
 			{
-				tagIndex = i;
+				// This argument is required
+				branch = GenerateMissingScalarCode("Switch option", 1.0f, node.m_errorContext);
 			}
+			op->cases.Emplace((int16_t)t, op, branch);
 		}
 
-		if (tagIndex < 0)
-		{
-			char buf[256];
-			mutable_snprintf(buf, 256, "Unknown tag found in image variation [%s].",
-				tag.c_str());
+		result.op = op;
+	}
 
-			m_pErrorLog->GetPrivate()->Add(buf, ELMT_WARNING, node.m_errorContext);
-			continue;
-		}
 
-		Ptr<ASTOp> variationOp;
-		if (node.m_variations[t].m_scalar)
+	//-------------------------------------------------------------------------------------------------
+	void CodeGenerator::GenerateScalar_Variation(SCALAR_GENERATION_RESULT& result, const Ptr<const NodeScalarVariation>& Typed)
+	{
+		const NodeScalarVariation::Private& node = *Typed->GetPrivate();
+
+		Ptr<ASTOp> op;
+
+		// Default case
+		if (node.m_defaultScalar)
 		{
+			FMeshGenerationResult branchResults;
+
 			SCALAR_GENERATION_RESULT ChildResult;
-			GenerateScalar(ChildResult,node.m_variations[t].m_scalar);
-			variationOp = ChildResult.op;
+			GenerateScalar(ChildResult, node.m_defaultScalar);
+			op = ChildResult.op;
 		}
 		else
 		{
 			// This argument is required
-			variationOp = GenerateMissingScalarCode("Variation option", 0.0f,
-				node.m_errorContext);
+			op = GenerateMissingScalarCode("Variation default", 0.0f, node.m_errorContext);
 		}
 
+		// Process variations in reverse order, since conditionals are built bottom-up.
+		for (int t = node.m_variations.Num() - 1; t >= 0; --t)
+		{
+			int tagIndex = -1;
+			const string& tag = node.m_variations[t].m_tag;
+			for (int i = 0; i < m_firstPass.m_tags.Num(); ++i)
+			{
+				if (m_firstPass.m_tags[i].tag == tag)
+				{
+					tagIndex = i;
+				}
+			}
 
-		Ptr<ASTOpConditional> conditional = new ASTOpConditional;
-		conditional->type = OP_TYPE::SC_CONDITIONAL;
-		conditional->no = op;
-		conditional->yes = variationOp;
-		conditional->condition = m_firstPass.m_tags[tagIndex].genericCondition;
+			if (tagIndex < 0)
+			{
+				char buf[256];
+				mutable_snprintf(buf, 256, "Unknown tag found in image variation [%s].",
+					tag.c_str());
 
-		op = conditional;
-	}
+				m_pErrorLog->GetPrivate()->Add(buf, ELMT_WARNING, node.m_errorContext);
+				continue;
+			}
 
-	result.op = op;
-}
+			Ptr<ASTOp> variationOp;
+			if (node.m_variations[t].m_scalar)
+			{
+				SCALAR_GENERATION_RESULT ChildResult;
+				GenerateScalar(ChildResult, node.m_variations[t].m_scalar);
+				variationOp = ChildResult.op;
+			}
+			else
+			{
+				// This argument is required
+				variationOp = GenerateMissingScalarCode("Variation option", 0.0f,
+					node.m_errorContext);
+			}
 
 
-//-------------------------------------------------------------------------------------------------
-void CodeGenerator::GenerateScalar_Curve(SCALAR_GENERATION_RESULT& result, const Ptr<const NodeScalarCurve>& Typed)
-{
-	const NodeScalarCurve::Private& node = *Typed->GetPrivate();
+			Ptr<ASTOpConditional> conditional = new ASTOpConditional;
+			conditional->type = OP_TYPE::SC_CONDITIONAL;
+			conditional->no = op;
+			conditional->yes = variationOp;
+			conditional->condition = m_firstPass.m_tags[tagIndex].genericCondition;
 
-	Ptr<ASTOpScalarCurve> op = new ASTOpScalarCurve();
+			op = conditional;
+		}
 
-	// T
-	if (Node* pA = node.m_input_scalar.get())
-	{
-		op->time = Generate(pA);
-	}
-	else
-	{
-		op->time = CodeGenerator::GenerateMissingScalarCode( "Curve T", 0.5f, node.m_errorContext );
+		result.op = op;
 	}
 
-	op->curve = node.m_curve;
 
-	result.op = op;
-}
+	//-------------------------------------------------------------------------------------------------
+	void CodeGenerator::GenerateScalar_Curve(SCALAR_GENERATION_RESULT& result, const Ptr<const NodeScalarCurve>& Typed)
+	{
+		const NodeScalarCurve::Private& node = *Typed->GetPrivate();
 
+		Ptr<ASTOpScalarCurve> op = new ASTOpScalarCurve();
 
-//-------------------------------------------------------------------------------------------------
-void CodeGenerator::GenerateScalar_Arithmetic(SCALAR_GENERATION_RESULT& result, const Ptr<const NodeScalarArithmeticOperation>& Typed)
-{
-	const NodeScalarArithmeticOperation::Private& node = *Typed->GetPrivate();
+		// T
+		if (Node* pA = node.m_input_scalar.get())
+		{
+			op->time = Generate(pA);
+		}
+		else
+		{
+			op->time = CodeGenerator::GenerateMissingScalarCode("Curve T", 0.5f, node.m_errorContext);
+		}
 
-	Ptr<ASTOpFixed> op = new ASTOpFixed();
-	op->op.type = OP_TYPE::SC_ARITHMETIC;
+		op->curve = node.m_curve;
 
-	switch (node.m_operation)
-	{
-	case NodeScalarArithmeticOperation::AO_ADD: op->op.args.ScalarArithmetic.operation = OP::ArithmeticArgs::ADD; break;
-	case NodeScalarArithmeticOperation::AO_SUBTRACT: op->op.args.ScalarArithmetic.operation = OP::ArithmeticArgs::SUBTRACT; break;
-	case NodeScalarArithmeticOperation::AO_MULTIPLY: op->op.args.ScalarArithmetic.operation = OP::ArithmeticArgs::MULTIPLY; break;
-	case NodeScalarArithmeticOperation::AO_DIVIDE: op->op.args.ScalarArithmetic.operation = OP::ArithmeticArgs::DIVIDE; break;
-	default:
-		checkf(false, TEXT("Unknown arithmetic operation."));
-		op->op.args.ScalarArithmetic.operation = OP::ArithmeticArgs::NONE;
-		break;
+		result.op = op;
 	}
 
-	// A
-	if (Node* pA = node.m_pA.get())
-	{
-		op->SetChild(op->op.args.ScalarArithmetic.a, Generate(pA));
-	}
-	else
-	{
-		op->SetChild(op->op.args.ScalarArithmetic.a,
-			CodeGenerator::GenerateMissingScalarCode
-			(
-				"ScalarArithmetic A",
-				1.0f,
-				node.m_errorContext
-			)
-		);
-	}
 
-	// B
-	if (Node* pB = node.m_pB.get())
-	{
-		op->SetChild(op->op.args.ScalarArithmetic.b, Generate(pB));
-	}
-	else
+	//-------------------------------------------------------------------------------------------------
+	void CodeGenerator::GenerateScalar_Arithmetic(SCALAR_GENERATION_RESULT& result, const Ptr<const NodeScalarArithmeticOperation>& Typed)
 	{
-		op->SetChild(op->op.args.ScalarArithmetic.b,
-			CodeGenerator::GenerateMissingScalarCode
-			(
-				"ScalarArithmetic B",
-				1.0f,
-				node.m_errorContext
-			)
-		);
-	}
+		const NodeScalarArithmeticOperation::Private& node = *Typed->GetPrivate();
 
-	result.op = op;
-}
+		Ptr<ASTOpFixed> op = new ASTOpFixed();
+		op->op.type = OP_TYPE::SC_ARITHMETIC;
 
+		switch (node.m_operation)
+		{
+		case NodeScalarArithmeticOperation::AO_ADD: op->op.args.ScalarArithmetic.operation = OP::ArithmeticArgs::ADD; break;
+		case NodeScalarArithmeticOperation::AO_SUBTRACT: op->op.args.ScalarArithmetic.operation = OP::ArithmeticArgs::SUBTRACT; break;
+		case NodeScalarArithmeticOperation::AO_MULTIPLY: op->op.args.ScalarArithmetic.operation = OP::ArithmeticArgs::MULTIPLY; break;
+		case NodeScalarArithmeticOperation::AO_DIVIDE: op->op.args.ScalarArithmetic.operation = OP::ArithmeticArgs::DIVIDE; break;
+		default:
+			checkf(false, TEXT("Unknown arithmetic operation."));
+			op->op.args.ScalarArithmetic.operation = OP::ArithmeticArgs::NONE;
+			break;
+		}
 
-//-------------------------------------------------------------------------------------------------
-void CodeGenerator::GenerateScalar_Table(SCALAR_GENERATION_RESULT& result, const Ptr<const NodeScalarTable>& Typed)
-{
-	const NodeScalarTable::Private& node = *Typed->GetPrivate();
+		// A
+		if (Node* pA = node.m_pA.get())
+		{
+			op->SetChild(op->op.args.ScalarArithmetic.a, Generate(pA));
+		}
+		else
+		{
+			op->SetChild(op->op.args.ScalarArithmetic.a,
+				CodeGenerator::GenerateMissingScalarCode
+				(
+					"ScalarArithmetic A",
+					1.0f,
+					node.m_errorContext
+				)
+			);
+		}
+
+		// B
+		if (Node* pB = node.m_pB.get())
+		{
+			op->SetChild(op->op.args.ScalarArithmetic.b, Generate(pB));
+		}
+		else
+		{
+			op->SetChild(op->op.args.ScalarArithmetic.b,
+				CodeGenerator::GenerateMissingScalarCode
+				(
+					"ScalarArithmetic B",
+					1.0f,
+					node.m_errorContext
+				)
+			);
+		}
 
-	Ptr<ASTOp> Op = GenerateTableSwitch<NodeScalarTable::Private, TCT_SCALAR, OP_TYPE::SC_SWITCH>(node,
+		result.op = op;
+	}
+
+
+	//-------------------------------------------------------------------------------------------------
+	void CodeGenerator::GenerateScalar_Table(SCALAR_GENERATION_RESULT& result, const Ptr<const NodeScalarTable>& Typed)
+	{
+		const NodeScalarTable::Private& node = *Typed->GetPrivate();
+
+		Ptr<ASTOp> Op = GenerateTableSwitch<NodeScalarTable::Private, TCT_SCALAR, OP_TYPE::SC_SWITCH>(node,
 			[this](const NodeScalarTable::Private& node, int colIndex, int row, ErrorLog* pErrorLog)
 			{
 				NodeScalarConstantPtr pCell = new NodeScalarConstant();
@@ -465,29 +465,30 @@ void CodeGenerator::GenerateScalar_Table(SCALAR_GENERATION_RESULT& result, const
 				return Generate(pCell);
 			});
 
-	result.op = Op;
-}
+		result.op = Op;
+	}
 
 
-//---------------------------------------------------------------------------------------------
-mu::Ptr<ASTOp> CodeGenerator::GenerateMissingScalarCode(const char* strWhere, float value, const void* errorContext)
-{
-	// Log a warning
-	char buf[256];
-	mutable_snprintf
-	(
-		buf, 256,
-		"Required connection not found: %s",
-		strWhere
-	);
-	m_pErrorLog->GetPrivate()->Add(buf, ELMT_ERROR, errorContext);
+	//---------------------------------------------------------------------------------------------
+	mu::Ptr<ASTOp> CodeGenerator::GenerateMissingScalarCode(const char* strWhere, float value, const void* errorContext)
+	{
+		// Log a warning
+		char buf[256];
+		mutable_snprintf
+		(
+			buf, 256,
+			"Required connection not found: %s",
+			strWhere
+		);
+		m_pErrorLog->GetPrivate()->Add(buf, ELMT_ERROR, errorContext);
 
-	// Create a constant node
-	NodeScalarConstantPtr pNode = new NodeScalarConstant();
-	pNode->SetValue(value);
+		// Create a constant node
+		NodeScalarConstantPtr pNode = new NodeScalarConstant();
+		pNode->SetValue(value);
 
-	Ptr<ASTOp> result = Generate(pNode);
+		Ptr<ASTOp> result = Generate(pNode);
 
-	return result;
-}
+		return result;
+	}
 
+}
\ No newline at end of file
