[mutable] Fix morphs not having any effect.

diff --git a/Engine/Plugins/Experimental/Mutable/Source/MutableTools/Private/MuT/CodeGenerator_SecondPass.cpp b/Engine/Plugins/Experimental/Mutable/Source/MutableTools/Private/MuT/CodeGenerator_SecondPass.cpp
--- a/Engine/Plugins/Experimental/Mutable/Source/MutableTools/Private/MuT/CodeGenerator_SecondPass.cpp
+++ b/Engine/Plugins/Experimental/Mutable/Source/MutableTools/Private/MuT/CodeGenerator_SecondPass.cpp
@@ -10,621 +10,623 @@
 #include "MuT/ErrorLog.h"
 #include "MuT/CodeGenerator_FirstPass.h"
 
-using namespace mu;
+namespace mu
+{
 
 
-//---------------------------------------------------------------------------------------------
-SecondPassGenerator::SecondPassGenerator( 
-    FirstPassGenerator* firstPass,
-    const CompilerOptions::Private* options )
-{
-    check(firstPass);
-    check(options);
-    m_pFirstPass = firstPass;
-    m_pCompilerOptions = options;
-
-    // Default conditions when there is no restriction accumulated.
-    CONDITION_CONTEXT noCondition;
-    m_currentCondition.push_back(noCondition);
-}
+	//---------------------------------------------------------------------------------------------
+	SecondPassGenerator::SecondPassGenerator(
+		FirstPassGenerator* firstPass,
+		const CompilerOptions::Private* options)
+	{
+		check(firstPass);
+		check(options);
+		m_pFirstPass = firstPass;
+		m_pCompilerOptions = options;
+
+		// Default conditions when there is no restriction accumulated.
+		CONDITION_CONTEXT noCondition;
+		m_currentCondition.push_back(noCondition);
+	}
+
+
+	//---------------------------------------------------------------------------------------------
+	mu::Ptr<ASTOp> SecondPassGenerator::GenerateTagCondition(size_t tagIndex,
+		const set<size_t>& posSurf,
+		const set<size_t>& negSurf,
+		const set<size_t>& posTag,
+		const set<size_t>& negTag)
+	{
+		auto& t = m_pFirstPass->m_tags[tagIndex];
+
+		// If this tag is already in the list of positive tags, return true as condition
+		if (posTag.find(tagIndex) != posTag.end())
+		{
+			return m_opPool.Add(new ASTOpConstantBool(true));
+		}
+
+		// If this tag is already in the list of negative tags, return false as condition
+		if (negTag.find(tagIndex) != negTag.end())
+		{
+			return m_opPool.Add(new ASTOpConstantBool(false));
+		}
+
+		// Cached?
+		CONDITION_GENERATION_KEY key;
+		key.tagOrSurfIndex = tagIndex;
+		//    key.negSurf = negSurf;
+		//    key.posSurf = posSurf;
+		//    key.negTag = negTag;
+		//    key.posTag = posTag;
+		for (auto s : negTag) { if (m_tagsPerTag[tagIndex].count(s) > 0) { key.negTag.insert(s); } }
+		for (auto s : posTag) { if (m_tagsPerTag[tagIndex].count(s) > 0) { key.posTag.insert(s); } }
+		for (auto s : negSurf) { if (m_surfacesPerTag[tagIndex].count(s) > 0) { key.negSurf.insert(s); } }
+		for (auto s : posSurf) { if (m_surfacesPerTag[tagIndex].count(s) > 0) { key.posSurf.insert(s); } }
+
+		{
+			auto it = m_tagConditionGenerationCache.find(key);
+			if (it != m_tagConditionGenerationCache.end())
+			{
+				return it->second;
+			}
+		}
 
+		Ptr<ASTOp> c;
 
-//---------------------------------------------------------------------------------------------
-mu::Ptr<ASTOp> SecondPassGenerator::GenerateTagCondition( size_t tagIndex,
-                                                      const set<size_t>& posSurf,
-                                                      const set<size_t>& negSurf,
-                                                      const set<size_t>& posTag,
-                                                      const set<size_t>& negTag )
-{
-    auto& t = m_pFirstPass->m_tags[tagIndex];
-
-    // If this tag is already in the list of positive tags, return true as condition
-    if ( posTag.find(tagIndex)!=posTag.end() )
-    {
-        return m_opPool.Add( new ASTOpConstantBool(true) );
-    }
-
-    // If this tag is already in the list of negative tags, return false as condition
-    if ( negTag.find(tagIndex)!=negTag.end() )
-    {
-        return m_opPool.Add( new ASTOpConstantBool(false) );
-    }
-
-    // Cached?
-    CONDITION_GENERATION_KEY key;
-    key.tagOrSurfIndex = tagIndex;
-//    key.negSurf = negSurf;
-//    key.posSurf = posSurf;
-//    key.negTag = negTag;
-//    key.posTag = posTag;
-	for (auto s : negTag) { if (m_tagsPerTag[tagIndex].count(s) > 0) { key.negTag.insert(s); } }
-	for (auto s : posTag) { if (m_tagsPerTag[tagIndex].count(s) > 0) { key.posTag.insert(s); } }
-	for (auto s : negSurf) { if (m_surfacesPerTag[tagIndex].count(s) > 0) { key.negSurf.insert(s); } }
-	for (auto s : posSurf) { if (m_surfacesPerTag[tagIndex].count(s) > 0) { key.posSurf.insert(s); } }
-
-    {
-        auto it = m_tagConditionGenerationCache.find( key );
-        if (it!=m_tagConditionGenerationCache.end())
-        {
-            return it->second;
-        }
-    }
-
-    Ptr<ASTOp> c;
-
-    // Condition expression for all the surfaces that activate the tag
-    for ( auto surfIndex: t.surfaces )
-    {
-        if ( posSurf.find(surfIndex) != posSurf.end() )
-        {
-            // This surface is already a positive requirement higher up in the condition so
-            // we can ignore it here.
-            continue;
-        }
-
-        if ( negSurf.find(surfIndex) != negSurf.end() )
-        {
-            // This surface is a negative requirement higher up in the condition so
-            // this branch never be true.
-            continue;
-        }
-
-        const auto& surface = m_pFirstPass->surfaces[surfIndex];
-
-        auto positiveTags = posTag;
-        positiveTags.insert( tagIndex );
-
-        Ptr<ASTOp> surfCondition = GenerateSurfaceCondition( surfIndex,
-                                                             posSurf,
-                                                             negSurf,
-                                                             positiveTags,
-                                                             negTag );
-
-        // If the surface is a constant false, we can skip adding it
-        if ( auto* constOp = dynamic_cast<const ASTOpConstantBool*>(surfCondition.get()) )
-        {
-            if (constOp->value==false)
-            {
-                continue;
-            }
-        }
-
-        Ptr<ASTOp> fullCondition;
-        if (surfCondition)
-        {
-            Ptr<ASTOpFixed> f = new ASTOpFixed();
-            f->op.type = OP_TYPE::BO_AND;
-            f->SetChild(f->op.args.BoolBinary.a, surface.objectCondition);
-            f->SetChild(f->op.args.BoolBinary.b, surfCondition);
-            fullCondition = m_opPool.Add( f );
-        }
-        else
-        {
-            fullCondition = m_opPool.Add( surface.objectCondition );
-        }
-
-
-        if (!c)
-        {
-            c = fullCondition;
-        }
-        else
-        {
-            Ptr<ASTOpFixed> o = new ASTOpFixed();
-            o->op.type = OP_TYPE::BO_OR;
-            o->SetChild( o->op.args.BoolBinary.a, fullCondition) ;
-            o->SetChild( o->op.args.BoolBinary.b, c );
-            c = m_opPool.Add( o );
-        }
-
-        // Optimise the condition now.
-        //PartialOptimise( c, m_pCompilerOptions->m_optimisationOptions );
-    }
-
-
-    // Condition expression for all the edit-surfaces that activate the tag
-    for ( auto editKey: t.edits )
-    {
-        if ( negSurf.find(editKey.Key) != negSurf.end() )
-        {
-            // The surface in the edit is a negative requirement higher up in the condition so
-            // this branch never be true.
-            continue;
-        }
-
-        Ptr<ASTOp> surfCondition;
-        if ( posSurf.find(editKey.Key) != posSurf.end() )
-        {
-            // This surface in the edit is already a positive requirement higher up in the condition
-            // so we don't need that part of the condition
-        }
-        else
-        {
-            auto positiveTags = posTag;
-            positiveTags.insert( tagIndex );
-
-            surfCondition = GenerateSurfaceCondition( editKey.Key,
-                                                      posSurf,
-                                                      negSurf,
-                                                      positiveTags,
-                                                      negTag );
-        }
-
-        const auto& surface = m_pFirstPass->surfaces[editKey.Key];
-        const auto& edit = surface.edits[editKey.Value];
-
-        // Combine object and surface conditions
-        Ptr<ASTOp> fullCondition;
-        if (surfCondition)
-        {
-            Ptr<ASTOpFixed> f = new ASTOpFixed();
-            f->op.type = OP_TYPE::BO_AND;
-            f->SetChild(f->op.args.BoolBinary.a, surface.objectCondition);
-            f->SetChild(f->op.args.BoolBinary.b, surfCondition);
-            fullCondition = m_opPool.Add( f );
-        }
-        else
-        {
-            fullCondition = m_opPool.Add( surface.objectCondition );
-        }
-
-        // Combine with edit condition
-        if (fullCondition)
-        {
-            Ptr<ASTOpFixed> f = new ASTOpFixed();
-            f->op.type = OP_TYPE::BO_AND;
-            f->SetChild(f->op.args.BoolBinary.a, edit.condition);
-            f->SetChild(f->op.args.BoolBinary.b, fullCondition);
-            fullCondition = m_opPool.Add( f );
-        }
-        else
-        {
-            fullCondition = m_opPool.Add( edit.condition );
-        }
-
-
-        if (!c)
-        {
-            c = fullCondition;
-        }
-        else
-        {
-            Ptr<ASTOpFixed> o = new ASTOpFixed();
-            o->op.type = OP_TYPE::BO_OR;
-            o->SetChild( o->op.args.BoolBinary.a, fullCondition) ;
-            o->SetChild( o->op.args.BoolBinary.b, c );
-            c = m_opPool.Add( o );
-        }
-    }
-
-
-    m_tagConditionGenerationCache.insert( std::make_pair<>(key,c) );
-
-    return c;
-}
+		// Condition expression for all the surfaces that activate the tag
+		for (auto surfIndex : t.surfaces)
+		{
+			if (posSurf.find(surfIndex) != posSurf.end())
+			{
+				// This surface is already a positive requirement higher up in the condition so
+				// we can ignore it here.
+				continue;
+			}
 
-//---------------------------------------------------------------------------------------------
-mu::Ptr<ASTOp> SecondPassGenerator::GenerateSurfaceCondition( size_t surfIndex,
-                                                    const set<size_t>& posSurf,
-                                                    const set<size_t>& negSurf,
-                                                    const set<size_t>& posTag,
-                                                    const set<size_t>& negTag )
-{
-    const auto& surf = m_pFirstPass->surfaces[surfIndex];
-
-    // If this surface is already in the list of positive surfaces, return true as condition
-    if ( posSurf.find(surfIndex)!=posSurf.end() )
-    {
-        return m_opPool.Add( new ASTOpConstantBool(true) );
-    }
-
-    // If this surface is already in the list of negative surfaces, return false as condition
-    if ( negSurf.find(surfIndex)!=negSurf.end() )
-    {
-        return m_opPool.Add( new ASTOpConstantBool(false) );
-    }
-
-    Ptr<ASTOp> c;
-
-    for (const auto& t: surf.positiveTags)
-    {
-        auto it = m_pFirstPass->m_tags.FindByPredicate([&](const FirstPassGenerator::TAG& e) { return e.tag == t; } );
-        if (!it)
-        {
-            // This could happen if a tag is in a variation but noone defines it.
-            // This surface depends on a tag that will never be active, so it will never be used.
-            return m_opPool.Add( new ASTOpConstantBool(false) );
-        }
-
-        size_t tagIndex = it - &m_pFirstPass->m_tags[0];
-
-        set<size_t> positiveSurfacesVisited = posSurf;
-        positiveSurfacesVisited.insert( surfIndex );
-
-        Ptr<ASTOp> tagCondition = GenerateTagCondition( tagIndex,
-                                                        positiveSurfacesVisited,
-                                                        negSurf,
-                                                        posTag,
-                                                        negTag );
-        // TODO: Optimise the tag condition here
-
-        
-
-        // If the tag is a constant ...
-        bool isConstant = false;
-        bool constantValue = false;
-        if ( auto* constOp = dynamic_cast<const ASTOpConstantBool*>(tagCondition.get()) )
-        {
-            isConstant = true;
-            constantValue = constOp->value;
-        }
-
-        if (tagCondition && !isConstant)
-        {
-            if (!c)
-            {
-                c = tagCondition;
-            }
-            else
-            {
-                Ptr<ASTOpFixed> o = new ASTOpFixed;
-                o->op.type = OP_TYPE::BO_AND;
-                o->SetChild( o->op.args.BoolBinary.a, tagCondition );
-                o->SetChild( o->op.args.BoolBinary.b, c );
-                c = m_opPool.Add( o );
-            }
-        }
-        else if (constantValue==true)
-        {
-            // No need to add it to the AND
-        }
-        else //if (constantValue==false)
-        {
-            // Entire expression will be false
-            Ptr<ASTOpConstantBool> f = new ASTOpConstantBool( false );
-
-            // No need to evaluate anything else.
-            c = m_opPool.Add( f );
-            break;
-        }
-    }
-
-
-    for (const auto& t: surf.negativeTags)
-    {
-        auto it = m_pFirstPass->m_tags.FindByPredicate([&](const FirstPassGenerator::TAG& e) { return e.tag == t; } );
-        if (!it)
-        {
-            // This could happen if a tag is in a variation but noone defines it.
-            continue;
-        }
-
-        size_t tagIndex = it - &m_pFirstPass->m_tags[0];
-
-        set<size_t> positiveSurfacesVisited = negSurf;
-        set<size_t> negativeSurfacesVisited = posSurf;
-        negativeSurfacesVisited.insert( surfIndex );
-        set<size_t> positiveTagsVisited = negTag;
-        set<size_t> negativeTagsVisited = posTag;
-        Ptr<ASTOp> tagCondition = GenerateTagCondition( tagIndex,
-                                                        positiveSurfacesVisited,
-                                                        negativeSurfacesVisited,
-                                                        positiveTagsVisited,
-                                                        negativeTagsVisited );
-
-        // TODO: Optimise the tag condition here
-
-        // If the tag is a constant ...
-        bool isConstant = false;
-        bool constantValue = false;
-        if ( auto* constOp = dynamic_cast<const ASTOpConstantBool*>(tagCondition.get()) )
-        {
-            isConstant = true;
-            constantValue = constOp->value;
-        }
-
-
-        if (!isConstant && tagCondition)
-        {
-            Ptr<ASTOpFixed> n = new ASTOpFixed;
-            n->op.type = OP_TYPE::BO_NOT;
-            n->SetChild( n->op.args.BoolNot.source, tagCondition );
-
-            if (!c)
-            {
-                c = m_opPool.Add( n );
-            }
-            else
-            {
-                Ptr<ASTOpFixed> o = new ASTOpFixed;
-                o->op.type = OP_TYPE::BO_AND;
-                o->SetChild( o->op.args.BoolBinary.a, n);
-                o->SetChild( o->op.args.BoolBinary.b, c);
-                c = m_opPool.Add( o );
-            }
-        }
-        else if (isConstant && constantValue==true)
-        {
-            // No expression here means always true which becomes always false
-            Ptr<ASTOpConstantBool> f = new ASTOpConstantBool( false );
-
-            // No need to evaluate anything else.
-            c = m_opPool.Add( f );
-            break;
-        }
-    }
-
-    return c;
-}
+			if (negSurf.find(surfIndex) != negSurf.end())
+			{
+				// This surface is a negative requirement higher up in the condition so
+				// this branch never be true.
+				continue;
+			}
 
-//---------------------------------------------------------------------------------------------
-mu::Ptr<ASTOp> SecondPassGenerator::GenerateModifierCondition( size_t modIndex )
-{
-    const auto& mod = m_pFirstPass->modifiers[modIndex];
-
-    Ptr<ASTOp> c;
-
-    bool done = false;
-    for (const auto& t: mod.positiveTags)
-    {
-        auto it = m_pFirstPass->m_tags.FindByPredicate([&](const FirstPassGenerator::TAG& e) { return e.tag == t; } );
-        if (!it)
-        {
-            // This could happen if a tag is in a variation but noone defines it.
-
-            // Entire expression will be false
-            Ptr<ASTOpConstantBool> f = new ASTOpConstantBool( false );
-            c = m_opPool.Add( f );
-            done = true;
-
-            continue;
-        }
-
-        size_t tagIndex = it - &m_pFirstPass->m_tags[0];
-
-        set<size_t> empty;
-        Ptr<ASTOp> tagCondition = GenerateTagCondition( tagIndex, empty, empty, empty, empty );
-
-        // TODO: Optimise the tag condition here
-
-        // If the tag is a constant ...
-        bool isConstant = false;
-        bool constantValue = false;
-        if ( auto* constOp = dynamic_cast<const ASTOpConstantBool*>(tagCondition.get()) )
-        {
-            isConstant = true;
-            constantValue = constOp->value;
-        }
-
-        if (tagCondition && !isConstant)
-        {
-            if (!c)
-            {
-                c = tagCondition;
-            }
-            else
-            {
-                Ptr<ASTOpFixed> o = new ASTOpFixed;
-                o->op.type = OP_TYPE::BO_AND;
-                o->SetChild( o->op.args.BoolBinary.a, tagCondition );
-                o->SetChild( o->op.args.BoolBinary.b, c );
-                c = m_opPool.Add( o );
-            }
-        }
-        else if (constantValue==true)
-        {
-            // No need to add it to the AND
-        }
-        else //if (constantValue==false)
-        {
-            // Entire expression will be false
-            Ptr<ASTOpConstantBool> f = new ASTOpConstantBool( false );
-
-            // No need to evaluate anything else.
-            c = m_opPool.Add( f );
-            done = true;
-            break;
-        }
-    }
-
-    if (!done)
-    {
-        for (const auto& t: mod.negativeTags)
-        {
-            auto it = m_pFirstPass->m_tags.FindByPredicate( [&](const FirstPassGenerator::TAG& e) { return e.tag == t; } );
-            if (!it)
-            {
-                // This could happen if a tag is in a variation but noone defines it.
-                continue;
-            }
-
-            size_t tagIndex = it - &m_pFirstPass->m_tags[0];
-
-            set<size_t> empty;
-            Ptr<ASTOp> tagCondition = GenerateTagCondition( tagIndex, empty, empty, empty, empty );
-
-            // TODO: Optimise the tag condition here
-
-            // If the tag is a constant ...
-            bool isConstant = false;
-            bool constantValue = false;
-            if ( auto* constOp = dynamic_cast<const ASTOpConstantBool*>(tagCondition.get()) )
-            {
-                isConstant = true;
-                constantValue = constOp->value;
-            }
-
-
-            if (!isConstant && tagCondition)
-            {
-                Ptr<ASTOpFixed> n = new ASTOpFixed;
-                n->op.type = OP_TYPE::BO_NOT;
-                n->SetChild( n->op.args.BoolNot.source, tagCondition );
-
-                if (!c)
-                {
-                    c = m_opPool.Add( n );
-                }
-                else
-                {
-                    Ptr<ASTOpFixed> o = new ASTOpFixed;
-                    o->op.type = OP_TYPE::BO_AND;
-                    o->SetChild( o->op.args.BoolBinary.a, n);
-                    o->SetChild( o->op.args.BoolBinary.b, c);
-                    c = m_opPool.Add( o );
-                }
-            }
-            else if (isConstant && constantValue==true)
-            {
-                // No expression here means always true which becomes always false
-                Ptr<ASTOpConstantBool> f = new ASTOpConstantBool( false );
-
-                // No need to evaluate anything else.
-                c = m_opPool.Add( f );
-                done = true;
-                break;
-            }
-        }
-    }
-
-    //PartialOptimise( c, m_pCompilerOptions->m_optimisationOptions );
-
-    return c;
-}
+			const auto& surface = m_pFirstPass->surfaces[surfIndex];
 
+			auto positiveTags = posTag;
+			positiveTags.insert(tagIndex);
 
-//---------------------------------------------------------------------------------------------
-bool SecondPassGenerator::Generate(
-    ErrorLogPtr pErrorLog,
-    const Node::Private* root)
-{
-    MUTABLE_CPUPROFILER_SCOPE(SecondPassGenerate);
-
-    check(root);
-    m_pErrorLog = pErrorLog;
-
-    // Find the list of surfaces every tag depends on
-	m_surfacesPerTag.clear();
-    m_surfacesPerTag.resize(m_pFirstPass->m_tags.Num());
-	m_tagsPerTag.clear();
-    m_tagsPerTag.resize(m_pFirstPass->m_tags.Num());
-    for( size_t t=0; t<m_pFirstPass->m_tags.Num(); ++t )
-    {
-        set<size_t> pendingSurfs;
-        for( size_t s: m_pFirstPass->m_tags[t].surfaces )
-        {
-            pendingSurfs.insert(s);
-        }
-
-		set<size_t> processedSurfs;
-        while( !pendingSurfs.empty() )
-        {
-            size_t cs = *pendingSurfs.begin();
-            pendingSurfs.erase( pendingSurfs.begin() );
-
-			if (std::find(processedSurfs.begin(), processedSurfs.end(), cs)
-				 != processedSurfs.end())
+			Ptr<ASTOp> surfCondition = GenerateSurfaceCondition(surfIndex,
+				posSurf,
+				negSurf,
+				positiveTags,
+				negTag);
+
+			// If the surface is a constant false, we can skip adding it
+			if (auto* constOp = dynamic_cast<const ASTOpConstantBool*>(surfCondition.get()))
+			{
+				if (constOp->value == false)
+				{
+					continue;
+				}
+			}
+
+			Ptr<ASTOp> fullCondition;
+			if (surfCondition)
+			{
+				Ptr<ASTOpFixed> f = new ASTOpFixed();
+				f->op.type = OP_TYPE::BO_AND;
+				f->SetChild(f->op.args.BoolBinary.a, surface.objectCondition);
+				f->SetChild(f->op.args.BoolBinary.b, surfCondition);
+				fullCondition = m_opPool.Add(f);
+			}
+			else
+			{
+				fullCondition = m_opPool.Add(surface.objectCondition);
+			}
+
+
+			if (!c)
+			{
+				c = fullCondition;
+			}
+			else
 			{
+				Ptr<ASTOpFixed> o = new ASTOpFixed();
+				o->op.type = OP_TYPE::BO_OR;
+				o->SetChild(o->op.args.BoolBinary.a, fullCondition);
+				o->SetChild(o->op.args.BoolBinary.b, c);
+				c = m_opPool.Add(o);
+			}
+
+			// Optimise the condition now.
+			//PartialOptimise( c, m_pCompilerOptions->m_optimisationOptions );
+		}
+
+
+		// Condition expression for all the edit-surfaces that activate the tag
+		for (auto editKey : t.edits)
+		{
+			if (negSurf.find(editKey.Key) != negSurf.end())
+			{
+				// The surface in the edit is a negative requirement higher up in the condition so
+				// this branch never be true.
 				continue;
 			}
 
-			processedSurfs.insert(cs);
-
-            m_surfacesPerTag[t].insert(cs);
-
-            auto& csurf = m_pFirstPass->surfaces[cs];
-            for ( auto sct: csurf.positiveTags )
-            {
-                auto it = m_pFirstPass->m_tags.FindByPredicate( [&](const FirstPassGenerator::TAG& e) { return e.tag == sct; } );
-                if (!it)
-                {
-                    // This could happen if a tag is in a variation but noone defines it.
-                    continue;
-                }
-
-                size_t ct = it - &m_pFirstPass->m_tags[0];
-
-				m_tagsPerTag[t].insert(ct);
-
-                for( size_t s: m_pFirstPass->m_tags[ct].surfaces )
-                {
-                    if (m_surfacesPerTag[t].find(s)!=m_surfacesPerTag[t].end())
-                    {
-                        pendingSurfs.insert(s);
-                    }
-                }
-            }
-            for ( auto sct: csurf.negativeTags )
-            {
-                auto it = m_pFirstPass->m_tags.FindByPredicate([&](const FirstPassGenerator::TAG& e) { return e.tag == sct; } );
-                if (!it)
-                {
-                    // This could happen if a tag is in a variation but noone defines it.
-                    continue;
-                }
-
-                size_t ct = it - &m_pFirstPass->m_tags[0];
-
-				m_tagsPerTag[t].insert(ct);
-
-                for( size_t s: m_pFirstPass->m_tags[ct].surfaces )
-                {
-                    if (m_surfacesPerTag[t].find(s)!=m_surfacesPerTag[t].end())
-                    {
-                        pendingSurfs.insert(s);
-                    }
-                }
-            }
-        }
-    }
-
-    // Create the conditions for every surface, modifier and individual tag.
-    m_tagConditionGenerationCache.clear();
-
-    for ( int32 s=0; s<m_pFirstPass->surfaces.Num(); ++s )
-    {
-        set<size_t> empty;
-        Ptr<ASTOp> c = GenerateSurfaceCondition( s, empty, empty, empty, empty );
-        m_pFirstPass->surfaces[s].surfaceCondition = c;
-    }
-
-    for ( int32 m=0; m<m_pFirstPass->modifiers.Num(); ++m )
-    {
-        Ptr<ASTOp> c = GenerateModifierCondition( m );
-        m_pFirstPass->modifiers[m].surfaceCondition = c;
-    }
-
-    for ( int32  s = 0; s < m_pFirstPass->m_tags.Num(); ++s )
-    {
-        set<size_t> empty;
-        Ptr<ASTOp> c = GenerateTagCondition( s, empty, empty, empty, empty );
-        m_pFirstPass->m_tags[s].genericCondition = c;
-    }
-
-
-    m_pFirstPass = nullptr;
-
-    return true;
-}
+			Ptr<ASTOp> surfCondition;
+			if (posSurf.find(editKey.Key) != posSurf.end())
+			{
+				// This surface in the edit is already a positive requirement higher up in the condition
+				// so we don't need that part of the condition
+			}
+			else
+			{
+				auto positiveTags = posTag;
+				positiveTags.insert(tagIndex);
+
+				surfCondition = GenerateSurfaceCondition(editKey.Key,
+					posSurf,
+					negSurf,
+					positiveTags,
+					negTag);
+			}
+
+			const auto& surface = m_pFirstPass->surfaces[editKey.Key];
+			const auto& edit = surface.edits[editKey.Value];
+
+			// Combine object and surface conditions
+			Ptr<ASTOp> fullCondition;
+			if (surfCondition)
+			{
+				Ptr<ASTOpFixed> f = new ASTOpFixed();
+				f->op.type = OP_TYPE::BO_AND;
+				f->SetChild(f->op.args.BoolBinary.a, surface.objectCondition);
+				f->SetChild(f->op.args.BoolBinary.b, surfCondition);
+				fullCondition = m_opPool.Add(f);
+			}
+			else
+			{
+				fullCondition = m_opPool.Add(surface.objectCondition);
+			}
+
+			// Combine with edit condition
+			if (fullCondition)
+			{
+				Ptr<ASTOpFixed> f = new ASTOpFixed();
+				f->op.type = OP_TYPE::BO_AND;
+				f->SetChild(f->op.args.BoolBinary.a, edit.condition);
+				f->SetChild(f->op.args.BoolBinary.b, fullCondition);
+				fullCondition = m_opPool.Add(f);
+			}
+			else
+			{
+				fullCondition = m_opPool.Add(edit.condition);
+			}
+
+
+			if (!c)
+			{
+				c = fullCondition;
+			}
+			else
+			{
+				Ptr<ASTOpFixed> o = new ASTOpFixed();
+				o->op.type = OP_TYPE::BO_OR;
+				o->SetChild(o->op.args.BoolBinary.a, fullCondition);
+				o->SetChild(o->op.args.BoolBinary.b, c);
+				c = m_opPool.Add(o);
+			}
+		}
+
+
+		m_tagConditionGenerationCache.insert(std::make_pair<>(key, c));
+
+		return c;
+	}
+
+	//---------------------------------------------------------------------------------------------
+	mu::Ptr<ASTOp> SecondPassGenerator::GenerateSurfaceCondition(size_t surfIndex,
+		const set<size_t>& posSurf,
+		const set<size_t>& negSurf,
+		const set<size_t>& posTag,
+		const set<size_t>& negTag)
+	{
+		const auto& surf = m_pFirstPass->surfaces[surfIndex];
+
+		// If this surface is already in the list of positive surfaces, return true as condition
+		if (posSurf.find(surfIndex) != posSurf.end())
+		{
+			return m_opPool.Add(new ASTOpConstantBool(true));
+		}
+
+		// If this surface is already in the list of negative surfaces, return false as condition
+		if (negSurf.find(surfIndex) != negSurf.end())
+		{
+			return m_opPool.Add(new ASTOpConstantBool(false));
+		}
 
+		Ptr<ASTOp> c;
 
+		for (const auto& t : surf.positiveTags)
+		{
+			auto it = m_pFirstPass->m_tags.FindByPredicate([&](const FirstPassGenerator::TAG& e) { return e.tag == t; });
+			if (!it)
+			{
+				// This could happen if a tag is in a variation but noone defines it.
+				// This surface depends on a tag that will never be active, so it will never be used.
+				return m_opPool.Add(new ASTOpConstantBool(false));
+			}
+
+			size_t tagIndex = it - &m_pFirstPass->m_tags[0];
+
+			set<size_t> positiveSurfacesVisited = posSurf;
+			positiveSurfacesVisited.insert(surfIndex);
+
+			Ptr<ASTOp> tagCondition = GenerateTagCondition(tagIndex,
+				positiveSurfacesVisited,
+				negSurf,
+				posTag,
+				negTag);
+			// TODO: Optimise the tag condition here
+
+
+
+			// If the tag is a constant ...
+			bool isConstant = false;
+			bool constantValue = false;
+			if (auto* constOp = dynamic_cast<const ASTOpConstantBool*>(tagCondition.get()))
+			{
+				isConstant = true;
+				constantValue = constOp->value;
+			}
+
+			if (tagCondition && !isConstant)
+			{
+				if (!c)
+				{
+					c = tagCondition;
+				}
+				else
+				{
+					Ptr<ASTOpFixed> o = new ASTOpFixed;
+					o->op.type = OP_TYPE::BO_AND;
+					o->SetChild(o->op.args.BoolBinary.a, tagCondition);
+					o->SetChild(o->op.args.BoolBinary.b, c);
+					c = m_opPool.Add(o);
+				}
+			}
+			else if (constantValue == true)
+			{
+				// No need to add it to the AND
+			}
+			else //if (constantValue==false)
+			{
+				// Entire expression will be false
+				Ptr<ASTOpConstantBool> f = new ASTOpConstantBool(false);
+
+				// No need to evaluate anything else.
+				c = m_opPool.Add(f);
+				break;
+			}
+		}
+
+
+		for (const auto& t : surf.negativeTags)
+		{
+			auto it = m_pFirstPass->m_tags.FindByPredicate([&](const FirstPassGenerator::TAG& e) { return e.tag == t; });
+			if (!it)
+			{
+				// This could happen if a tag is in a variation but noone defines it.
+				continue;
+			}
+
+			size_t tagIndex = it - &m_pFirstPass->m_tags[0];
+
+			set<size_t> positiveSurfacesVisited = negSurf;
+			set<size_t> negativeSurfacesVisited = posSurf;
+			negativeSurfacesVisited.insert(surfIndex);
+			set<size_t> positiveTagsVisited = negTag;
+			set<size_t> negativeTagsVisited = posTag;
+			Ptr<ASTOp> tagCondition = GenerateTagCondition(tagIndex,
+				positiveSurfacesVisited,
+				negativeSurfacesVisited,
+				positiveTagsVisited,
+				negativeTagsVisited);
+
+			// TODO: Optimise the tag condition here
+
+			// If the tag is a constant ...
+			bool isConstant = false;
+			bool constantValue = false;
+			if (auto* constOp = dynamic_cast<const ASTOpConstantBool*>(tagCondition.get()))
+			{
+				isConstant = true;
+				constantValue = constOp->value;
+			}
+
+
+			if (!isConstant && tagCondition)
+			{
+				Ptr<ASTOpFixed> n = new ASTOpFixed;
+				n->op.type = OP_TYPE::BO_NOT;
+				n->SetChild(n->op.args.BoolNot.source, tagCondition);
+
+				if (!c)
+				{
+					c = m_opPool.Add(n);
+				}
+				else
+				{
+					Ptr<ASTOpFixed> o = new ASTOpFixed;
+					o->op.type = OP_TYPE::BO_AND;
+					o->SetChild(o->op.args.BoolBinary.a, n);
+					o->SetChild(o->op.args.BoolBinary.b, c);
+					c = m_opPool.Add(o);
+				}
+			}
+			else if (isConstant && constantValue == true)
+			{
+				// No expression here means always true which becomes always false
+				Ptr<ASTOpConstantBool> f = new ASTOpConstantBool(false);
+
+				// No need to evaluate anything else.
+				c = m_opPool.Add(f);
+				break;
+			}
+		}
+
+		return c;
+	}
+
+	//---------------------------------------------------------------------------------------------
+	mu::Ptr<ASTOp> SecondPassGenerator::GenerateModifierCondition(size_t modIndex)
+	{
+		const auto& mod = m_pFirstPass->modifiers[modIndex];
+
+		Ptr<ASTOp> c;
+
+		bool done = false;
+		for (const auto& t : mod.positiveTags)
+		{
+			auto it = m_pFirstPass->m_tags.FindByPredicate([&](const FirstPassGenerator::TAG& e) { return e.tag == t; });
+			if (!it)
+			{
+				// This could happen if a tag is in a variation but noone defines it.
+
+				// Entire expression will be false
+				Ptr<ASTOpConstantBool> f = new ASTOpConstantBool(false);
+				c = m_opPool.Add(f);
+				done = true;
+
+				continue;
+			}
+
+			size_t tagIndex = it - &m_pFirstPass->m_tags[0];
+
+			set<size_t> empty;
+			Ptr<ASTOp> tagCondition = GenerateTagCondition(tagIndex, empty, empty, empty, empty);
+
+			// TODO: Optimise the tag condition here
+
+			// If the tag is a constant ...
+			bool isConstant = false;
+			bool constantValue = false;
+			if (auto* constOp = dynamic_cast<const ASTOpConstantBool*>(tagCondition.get()))
+			{
+				isConstant = true;
+				constantValue = constOp->value;
+			}
+
+			if (tagCondition && !isConstant)
+			{
+				if (!c)
+				{
+					c = tagCondition;
+				}
+				else
+				{
+					Ptr<ASTOpFixed> o = new ASTOpFixed;
+					o->op.type = OP_TYPE::BO_AND;
+					o->SetChild(o->op.args.BoolBinary.a, tagCondition);
+					o->SetChild(o->op.args.BoolBinary.b, c);
+					c = m_opPool.Add(o);
+				}
+			}
+			else if (constantValue == true)
+			{
+				// No need to add it to the AND
+			}
+			else //if (constantValue==false)
+			{
+				// Entire expression will be false
+				Ptr<ASTOpConstantBool> f = new ASTOpConstantBool(false);
+
+				// No need to evaluate anything else.
+				c = m_opPool.Add(f);
+				done = true;
+				break;
+			}
+		}
+
+		if (!done)
+		{
+			for (const auto& t : mod.negativeTags)
+			{
+				auto it = m_pFirstPass->m_tags.FindByPredicate([&](const FirstPassGenerator::TAG& e) { return e.tag == t; });
+				if (!it)
+				{
+					// This could happen if a tag is in a variation but noone defines it.
+					continue;
+				}
+
+				size_t tagIndex = it - &m_pFirstPass->m_tags[0];
+
+				set<size_t> empty;
+				Ptr<ASTOp> tagCondition = GenerateTagCondition(tagIndex, empty, empty, empty, empty);
+
+				// TODO: Optimise the tag condition here
+
+				// If the tag is a constant ...
+				bool isConstant = false;
+				bool constantValue = false;
+				if (auto* constOp = dynamic_cast<const ASTOpConstantBool*>(tagCondition.get()))
+				{
+					isConstant = true;
+					constantValue = constOp->value;
+				}
+
+
+				if (!isConstant && tagCondition)
+				{
+					Ptr<ASTOpFixed> n = new ASTOpFixed;
+					n->op.type = OP_TYPE::BO_NOT;
+					n->SetChild(n->op.args.BoolNot.source, tagCondition);
+
+					if (!c)
+					{
+						c = m_opPool.Add(n);
+					}
+					else
+					{
+						Ptr<ASTOpFixed> o = new ASTOpFixed;
+						o->op.type = OP_TYPE::BO_AND;
+						o->SetChild(o->op.args.BoolBinary.a, n);
+						o->SetChild(o->op.args.BoolBinary.b, c);
+						c = m_opPool.Add(o);
+					}
+				}
+				else if (isConstant && constantValue == true)
+				{
+					// No expression here means always true which becomes always false
+					Ptr<ASTOpConstantBool> f = new ASTOpConstantBool(false);
+
+					// No need to evaluate anything else.
+					c = m_opPool.Add(f);
+					done = true;
+					break;
+				}
+			}
+		}
+
+		//PartialOptimise( c, m_pCompilerOptions->m_optimisationOptions );
+
+		return c;
+	}
+
+
+	//---------------------------------------------------------------------------------------------
+	bool SecondPassGenerator::Generate(
+		ErrorLogPtr pErrorLog,
+		const Node::Private* root)
+	{
+		MUTABLE_CPUPROFILER_SCOPE(SecondPassGenerate);
+
+		check(root);
+		m_pErrorLog = pErrorLog;
+
+		// Find the list of surfaces every tag depends on
+		m_surfacesPerTag.clear();
+		m_surfacesPerTag.resize(m_pFirstPass->m_tags.Num());
+		m_tagsPerTag.clear();
+		m_tagsPerTag.resize(m_pFirstPass->m_tags.Num());
+		for (size_t t = 0; t < m_pFirstPass->m_tags.Num(); ++t)
+		{
+			set<size_t> pendingSurfs;
+			for (size_t s : m_pFirstPass->m_tags[t].surfaces)
+			{
+				pendingSurfs.insert(s);
+			}
+
+			set<size_t> processedSurfs;
+			while (!pendingSurfs.empty())
+			{
+				size_t cs = *pendingSurfs.begin();
+				pendingSurfs.erase(pendingSurfs.begin());
+
+				if (std::find(processedSurfs.begin(), processedSurfs.end(), cs)
+					!= processedSurfs.end())
+				{
+					continue;
+				}
+
+				processedSurfs.insert(cs);
+
+				m_surfacesPerTag[t].insert(cs);
+
+				auto& csurf = m_pFirstPass->surfaces[cs];
+				for (auto sct : csurf.positiveTags)
+				{
+					auto it = m_pFirstPass->m_tags.FindByPredicate([&](const FirstPassGenerator::TAG& e) { return e.tag == sct; });
+					if (!it)
+					{
+						// This could happen if a tag is in a variation but noone defines it.
+						continue;
+					}
+
+					size_t ct = it - &m_pFirstPass->m_tags[0];
+
+					m_tagsPerTag[t].insert(ct);
+
+					for (size_t s : m_pFirstPass->m_tags[ct].surfaces)
+					{
+						if (m_surfacesPerTag[t].find(s) != m_surfacesPerTag[t].end())
+						{
+							pendingSurfs.insert(s);
+						}
+					}
+				}
+				for (auto sct : csurf.negativeTags)
+				{
+					auto it = m_pFirstPass->m_tags.FindByPredicate([&](const FirstPassGenerator::TAG& e) { return e.tag == sct; });
+					if (!it)
+					{
+						// This could happen if a tag is in a variation but noone defines it.
+						continue;
+					}
+
+					size_t ct = it - &m_pFirstPass->m_tags[0];
+
+					m_tagsPerTag[t].insert(ct);
+
+					for (size_t s : m_pFirstPass->m_tags[ct].surfaces)
+					{
+						if (m_surfacesPerTag[t].find(s) != m_surfacesPerTag[t].end())
+						{
+							pendingSurfs.insert(s);
+						}
+					}
+				}
+			}
+		}
+
+		// Create the conditions for every surface, modifier and individual tag.
+		m_tagConditionGenerationCache.clear();
+
+		for (int32 s = 0; s < m_pFirstPass->surfaces.Num(); ++s)
+		{
+			set<size_t> empty;
+			Ptr<ASTOp> c = GenerateSurfaceCondition(s, empty, empty, empty, empty);
+			m_pFirstPass->surfaces[s].surfaceCondition = c;
+		}
+
+		for (int32 m = 0; m < m_pFirstPass->modifiers.Num(); ++m)
+		{
+			Ptr<ASTOp> c = GenerateModifierCondition(m);
+			m_pFirstPass->modifiers[m].surfaceCondition = c;
+		}
+
+		for (int32 s = 0; s < m_pFirstPass->m_tags.Num(); ++s)
+		{
+			set<size_t> empty;
+			Ptr<ASTOp> c = GenerateTagCondition(s, empty, empty, empty, empty);
+			m_pFirstPass->m_tags[s].genericCondition = c;
+		}
+
+
+		m_pFirstPass = nullptr;
+
+		return true;
+	}
+
+
+}
diff --git a/Engine/Plugins/Experimental/Mutable/Source/MutableTools/Private/MuT/CodeGenerator_String.cpp b/Engine/Plugins/Experimental/Mutable/Source/MutableTools/Private/MuT/CodeGenerator_String.cpp
--- a/Engine/Plugins/Experimental/Mutable/Source/MutableTools/Private/MuT/CodeGenerator_String.cpp
+++ b/Engine/Plugins/Experimental/Mutable/Source/MutableTools/Private/MuT/CodeGenerator_String.cpp
@@ -26,99 +26,101 @@
 #include <utility>
 
 
-using namespace mu;
+namespace mu
+{
 
 
-//-------------------------------------------------------------------------------------------------
-void CodeGenerator::GenerateString( STRING_GENERATION_RESULT& result, const NodeStringPtrConst& untyped )
-{
-    if (!untyped)
-    {
-        result = STRING_GENERATION_RESULT();
-        return;
-    }
-
-    // See if it was already generated
-	VISITED_MAP_KEY key = GetCurrentCacheKey(untyped);
-	GeneratedStringsMap::ValueType* it = m_generatedStrings.Find( key );
-    if ( it )
-    {
-        result = *it;
-        return;
-    }
-
-    // Generate for each different type of node
-    if ( auto Constant = dynamic_cast<const NodeStringConstant*>(untyped.get()) )
-    {
-        GenerateString_Constant( result, Constant );
-    }
-    else if ( auto Param = dynamic_cast<const NodeStringParameter*>(untyped.get()) )
-    {
-        GenerateString_Parameter( result, Param );
-    }
-	else
-    {
-        check(false);
-    }
-
-    // Cache the result
-    m_generatedStrings.Add( key, result );
-}
+	//-------------------------------------------------------------------------------------------------
+	void CodeGenerator::GenerateString(STRING_GENERATION_RESULT& result, const NodeStringPtrConst& untyped)
+	{
+		if (!untyped)
+		{
+			result = STRING_GENERATION_RESULT();
+			return;
+		}
 
+		// See if it was already generated
+		VISITED_MAP_KEY key = GetCurrentCacheKey(untyped);
+		GeneratedStringsMap::ValueType* it = m_generatedStrings.Find(key);
+		if (it)
+		{
+			result = *it;
+			return;
+		}
 
-//-------------------------------------------------------------------------------------------------
-void CodeGenerator::GenerateString_Constant(STRING_GENERATION_RESULT& result, const Ptr<const NodeStringConstant>& Typed )
-{
-    const NodeStringConstant::Private& node = *Typed->GetPrivate();
+		// Generate for each different type of node
+		if (auto Constant = dynamic_cast<const NodeStringConstant*>(untyped.get()))
+		{
+			GenerateString_Constant(result, Constant);
+		}
+		else if (auto Param = dynamic_cast<const NodeStringParameter*>(untyped.get()))
+		{
+			GenerateString_Parameter(result, Param);
+		}
+		else
+		{
+			check(false);
+		}
 
-	Ptr<ASTOpConstantString> op = new ASTOpConstantString();
-	op->value = node.m_value;
+		// Cache the result
+		m_generatedStrings.Add(key, result);
+	}
 
-    result.op = op;
-}
 
+	//-------------------------------------------------------------------------------------------------
+	void CodeGenerator::GenerateString_Constant(STRING_GENERATION_RESULT& result, const Ptr<const NodeStringConstant>& Typed)
+	{
+		const NodeStringConstant::Private& node = *Typed->GetPrivate();
 
-//-------------------------------------------------------------------------------------------------
-void CodeGenerator::GenerateString_Parameter( STRING_GENERATION_RESULT& result, const Ptr<const NodeStringParameter>& Typed)
-{
-    const NodeStringParameter::Private& node = *Typed->GetPrivate();
+		Ptr<ASTOpConstantString> op = new ASTOpConstantString();
+		op->value = node.m_value;
+
+		result.op = op;
+	}
 
-	Ptr<ASTOpParameter> op;
 
-	auto it = m_nodeVariables.find(node.m_pNode);
-	if (it == m_nodeVariables.end())
+	//-------------------------------------------------------------------------------------------------
+	void CodeGenerator::GenerateString_Parameter(STRING_GENERATION_RESULT& result, const Ptr<const NodeStringParameter>& Typed)
 	{
-		PARAMETER_DESC param;
-		param.m_name = node.m_name;
-		param.m_uid = node.m_uid;
-		param.m_type = PARAMETER_TYPE::T_FLOAT;
-		param.m_detailedType = node.m_detailedType;
-		// param.m_defaultValue.m_text = node.m_defaultValue;
-		int len = FMath::Min(MUTABLE_MAX_STRING_PARAM_LENGTH, int(node.m_defaultValue.size()));
-		FMemory::Memcpy(param.m_defaultValue.m_text, node.m_defaultValue.c_str(), len);
-		param.m_defaultValue.m_text[len] = 0;
-
-
-		op = new ASTOpParameter();
-		op->type = OP_TYPE::ST_PARAMETER;
-		op->parameter = param;
-
-		// Generate the code for the ranges
-		for (int32 a = 0; a < node.m_ranges.Num(); ++a)
+		const NodeStringParameter::Private& node = *Typed->GetPrivate();
+
+		Ptr<ASTOpParameter> op;
+
+		auto it = m_nodeVariables.find(node.m_pNode);
+		if (it == m_nodeVariables.end())
 		{
-			RANGE_GENERATION_RESULT rangeResult;
-			GenerateRange(rangeResult, node.m_ranges[a]);
-			op->ranges.Emplace(op.get(), rangeResult.sizeOp, rangeResult.rangeName, rangeResult.rangeUID);
+			PARAMETER_DESC param;
+			param.m_name = node.m_name;
+			param.m_uid = node.m_uid;
+			param.m_type = PARAMETER_TYPE::T_FLOAT;
+			param.m_detailedType = node.m_detailedType;
+			// param.m_defaultValue.m_text = node.m_defaultValue;
+			int len = FMath::Min(MUTABLE_MAX_STRING_PARAM_LENGTH, int(node.m_defaultValue.size()));
+			FMemory::Memcpy(param.m_defaultValue.m_text, node.m_defaultValue.c_str(), len);
+			param.m_defaultValue.m_text[len] = 0;
+
+
+			op = new ASTOpParameter();
+			op->type = OP_TYPE::ST_PARAMETER;
+			op->parameter = param;
+
+			// Generate the code for the ranges
+			for (int32 a = 0; a < node.m_ranges.Num(); ++a)
+			{
+				RANGE_GENERATION_RESULT rangeResult;
+				GenerateRange(rangeResult, node.m_ranges[a]);
+				op->ranges.Emplace(op.get(), rangeResult.sizeOp, rangeResult.rangeName, rangeResult.rangeUID);
+			}
+
+			m_nodeVariables[node.m_pNode] = op;
+		}
+		else
+		{
+			op = it->second;
 		}
 
-		m_nodeVariables[node.m_pNode] = op;
-	}
-	else
-	{
-		op = it->second;
+		result.op = op;
 	}
 
-    result.op = op;
-}
-
 
+}
