[mutable] Fix morphs not having any effect.

diff --git a/Engine/Plugins/Experimental/Mutable/Source/MutableTools/Private/MuT/CodeOptimiser.cpp b/Engine/Plugins/Experimental/Mutable/Source/MutableTools/Private/MuT/CodeOptimiser.cpp
--- a/Engine/Plugins/Experimental/Mutable/Source/MutableTools/Private/MuT/CodeOptimiser.cpp
+++ b/Engine/Plugins/Experimental/Mutable/Source/MutableTools/Private/MuT/CodeOptimiser.cpp
@@ -23,1332 +23,1332 @@
 #include <unordered_set>
 #include <shared_mutex>
 
-using namespace mu;
-
-//-------------------------------------------------------------------------------------------------
-//-------------------------------------------------------------------------------------------------
-//-------------------------------------------------------------------------------------------------
-namespace
+namespace mu
 {
 
-    struct MESH_ENTRY
-    {
-        MeshPtrConst mesh;
-		mu::Ptr<ASTOpConstantResource> op;
-
-        bool operator==(const MESH_ENTRY& o) const
-        {
-            return mesh==o.mesh || *mesh==*o.mesh;
-        }
-    };
-
-    struct custom_mesh_hash
-    {
-		uint64 operator()( const MeshPtrConst& k ) const
-        {
-			uint64 h =  std::hash<uint64>()(k->GetVertexCount());
-            hash_combine(h, k->GetFaceCount());
-            return h;
-        }
-    };
-
-    struct custom_mesh_equal
-    {
-        bool operator()( const MeshPtrConst& a, const MeshPtrConst& b ) const
-        {
-            return a==b || *a==*b;
-        }
-    };
-
-    struct custom_image_hash
-    {
-		uint64 operator()( const ImagePtrConst& k ) const
-        {
-			uint64 h =  std::hash<uint64>()(k->GetSizeX());
-            hash_combine(h, k->GetSizeY());
-            hash_combine(h, k->GetLODCount());
-            return h;
-        }
-    };
-
-    struct custom_image_equal
-    {
-        bool operator()( const ImagePtrConst& a, const ImagePtrConst& b ) const
-        {
-            return a==b || *a==*b;
-        }
-    };
-
-    struct custom_layout_hash
-    {
-		uint64 operator()( const LayoutPtrConst& k ) const
-        {
-			uint64 h =  std::hash<uint64>()(k->GetBlockCount());
-            FIntPoint s = k->GetGridSize();
-            hash_combine(h, s[0]);
-            hash_combine(h, s[1]);
-            return h;
-        }
-    };
-
-    struct custom_layout_equal
-    {
-        bool operator()( const LayoutPtrConst& a, const LayoutPtrConst& b ) const
-        {
-            return a==b || *a==*b;
-        }
-    };
-}
 
+	namespace
+	{
 
-//-------------------------------------------------------------------------------------------------
-void mu::DuplicatedDataRemoverAST( ASTOpList& roots )
-{
-    MUTABLE_CPUPROFILER_SCOPE(DuplicatedDataRemoverAST);
-
-    std::vector<Ptr<ASTOpConstantResource>> allMeshOps;
-    std::vector<Ptr<ASTOpConstantResource>> allImageOps;
-    std::vector<Ptr<ASTOpConstantResource>> allLayoutOps;
-
-    // Gather constants
-    {
-        //MUTABLE_CPUPROFILER_SCOPE(Gather);
-
-        ASTOp::Traverse_TopRandom_Unique_NonReentrant( roots, [&](Ptr<ASTOp> n)
-        {
-            switch ( n->GetOpType() )
-            {
-
-            case OP_TYPE::ME_CONSTANT:
-            {
-                auto typedNode = dynamic_cast<ASTOpConstantResource*>(n.get());
-                check(typedNode);
-
-                if (typedNode)
-                {
-                    allMeshOps.push_back(typedNode);
-                }
-                break;
-            }
-
-            case OP_TYPE::IM_CONSTANT:
-            {
-                auto typedNode = dynamic_cast<ASTOpConstantResource*>(n.get());
-                check(typedNode);
-
-                if (typedNode)
-                {
-                    allImageOps.push_back(typedNode);
-                }
-                break;
-            }
-
-            case OP_TYPE::LA_CONSTANT:
-            {
-                auto typedNode = dynamic_cast<ASTOpConstantResource*>(n.get());
-                check(typedNode);
-
-                if (typedNode)
-                {
-                    allLayoutOps.push_back(typedNode);
-                }
-                break;
-            }
-
-            //    These should be part of the duplicated code removal, in AST.
-            //            // Names
-            //            case OP_TYPE::IN_ADDMESH:
-            //            case OP_TYPE::IN_ADDIMAGE:
-            //            case OP_TYPE::IN_ADDVECTOR:
-            //            case OP_TYPE::IN_ADDSCALAR:
-            //            case OP_TYPE::IN_ADDCOMPONENT:
-            //            case OP_TYPE::IN_ADDSURFACE:
-            //            {
-            //                OP::ADDRESS value = program.m_code[at].args.InstanceAdd.name;
-            //                program.m_code[at].args.InstanceAdd.name = m_oldToNewStrings[value];
-            //                break;
-            //            }
-
-
-            default:
-                break;
-            }
-
-            return true;
-        });
-    }
-
-
-    // Compare meshes
-    {
-        //MUTABLE_CPUPROFILER_SCOPE(CompareMeshes);
-
-        std::unordered_multimap< size_t, MESH_ENTRY > meshes;
-
-        for (auto& typedNode: allMeshOps)
-        {
-            size_t key = typedNode->GetValueHash();
-
-            Ptr<ASTOp> found;
-
-            auto r = meshes.equal_range(key);
-            if (r.first!=r.second)
-            {
-                MeshPtrConst mesh = static_cast<const Mesh*>( typedNode->GetValue().get() );
-
-                for ( auto it=r.first; it!=r.second; ++it )
-                {
-                    if (!it->second.mesh)
-                    {
-                        it->second.mesh = static_cast<const Mesh*>( it->second.op->GetValue().get() );
-                    }
-
-                    if ( custom_mesh_equal()( mesh, it->second.mesh ) )
-                    {
-                        found = it->second.op;
-                        break;
-                    }
-                }
-            }
-
-            if (found)
-            {
-                ASTOp::Replace(typedNode,found);
-            }
-            else
-            {
-                MESH_ENTRY e;
-                e.op = typedNode;
-                meshes.insert( std::make_pair<>(key, e) );
-            }
-        }
-    }
-
-    // Compare images
-    {
-        //MUTABLE_CPUPROFILER_SCOPE(CompareImages);
-
-        // TODO Optimise like the mesh compare above
-        std::unordered_map< ImagePtrConst, Ptr<ASTOp>, custom_image_hash, custom_image_equal > images;
-
-        for (auto& typedNode: allImageOps)
-        {
-            ImagePtrConst image = static_cast<const Image*>( typedNode->GetValue().get() );
-
-            auto it = images.find(image);
-            if (it!=images.end())
-            {
-                ASTOp::Replace(typedNode,it->second);
-            }
-            else
-            {
-                images.insert( std::make_pair<>(image, typedNode) );
-            }
-        }
-    }
-
-    // Compare layouts
-    {
-        //MUTABLE_CPUPROFILER_SCOPE(CompareLayouts);
-
-        // TODO Optimise like the mesh compare above
-        std::unordered_map< LayoutPtrConst, Ptr<ASTOp>, custom_layout_hash, custom_layout_equal > constants;
-
-        for (auto& typedNode: allLayoutOps)
-        {
-            LayoutPtrConst r = static_cast<const Layout*>( typedNode->GetValue().get() );
-
-            auto it = constants.find(r);
-            if (it!=constants.end())
-            {
-                ASTOp::Replace(typedNode,it->second);
-            }
-            else
-            {
-                constants.insert( std::make_pair<>(r, typedNode) );
-            }
-        }
-    }
+		struct MESH_ENTRY
+		{
+			MeshPtrConst mesh;
+			mu::Ptr<ASTOpConstantResource> op;
 
-}
+			bool operator==(const MESH_ENTRY& o) const
+			{
+				return mesh==o.mesh || *mesh==*o.mesh;
+			}
+		};
 
+		struct custom_mesh_hash
+		{
+			uint64 operator()( const MeshPtrConst& k ) const
+			{
+				uint64 h =  std::hash<uint64>()(k->GetVertexCount());
+				hash_combine(h, k->GetFaceCount());
+				return h;
+			}
+		};
 
-//-------------------------------------------------------------------------------------------------
-//-------------------------------------------------------------------------------------------------
-//-------------------------------------------------------------------------------------------------
-namespace
-{
-    struct op_pointer_hash
-    {
-        inline size_t operator() (const mu::Ptr<ASTOp>& n) const
-        {
-            return n->Hash();
-        }
-    };
-
-    struct op_pointer_equal
-    {
-        inline bool operator() (const mu::Ptr<ASTOp>& a, const mu::Ptr<ASTOp>& b) const
-        {
-            return *a==*b;
-        }
-    };
-}
+		struct custom_mesh_equal
+		{
+			bool operator()( const MeshPtrConst& a, const MeshPtrConst& b ) const
+			{
+				return a==b || *a==*b;
+			}
+		};
 
-void mu::DuplicatedCodeRemoverAST( ASTOpList& roots )
-{
-    MUTABLE_CPUPROFILER_SCOPE(DuplicatedCodeRemoverAST);
-
-    // Visited nodes, per type
-    std::unordered_set<Ptr<ASTOp>,op_pointer_hash,op_pointer_equal> visited[int(OP_TYPE::COUNT)];
-
-    ASTOp::Traverse_BottomUp_Unique_NonReentrant( roots, [&](Ptr<ASTOp>& n)
-    {
-        auto& container = visited[(int)n->GetOpType()];
-
-        // debug
-//        for( const auto& o : container )
-//        {
-//            if (*n==*o)
-//            {
-//                auto thisHash = n->Hash();
-//                auto otherHash = o->Hash();
-//                check(thisHash==otherHash);
-//            }
-//        }
-
-        // Insert will tell us if it was already there
-        auto it = container.insert(n);
-        if( !it.second )
-        {
-            // It wasn't inserted, so it was already there
-            ASTOp::Replace(n,*it.first);
-        }
-    });
-}
+		struct custom_image_hash
+		{
+			uint64 operator()( const ImagePtrConst& k ) const
+			{
+				uint64 h =  std::hash<uint64>()(k->GetSizeX());
+				hash_combine(h, k->GetSizeY());
+				hash_combine(h, k->GetLODCount());
+				return h;
+			}
+		};
 
+		struct custom_image_equal
+		{
+			bool operator()( const ImagePtrConst& a, const ImagePtrConst& b ) const
+			{
+				return a==b || *a==*b;
+			}
+		};
 
-//---------------------------------------------------------------------------------------------
-//---------------------------------------------------------------------------------------------
-//---------------------------------------------------------------------------------------------
-bool FindOpTypeVisitor::Apply( PROGRAM& program, OP::ADDRESS rootAt )
-{
-    bool found = false;
-
-    // If the program added new operations, we haven't visited them.
-    m_visited.resize( program.m_opAddress.Num(), 0 );
-
-    m_pending.clear();
-    m_pending.reserve( program.m_opAddress.Num()/4 );
-    m_pending.push_back( std::make_pair(false,rootAt) );
-
-    // Don't early out to be able to complete parent op cached flags
-    while ( !m_pending.empty() )
-    {
-        std::pair<bool,int> item = m_pending.back();
-        m_pending.pop_back();
-        OP::ADDRESS at = item.second;
-
-        // Not cached?
-        if (m_visited[at]<=1)
-        {
-            if (item.first)
-            {
-                // Item indicating we finished with all the children of a parent
-                check(m_visited[at]==1);
-
-                bool subtreeFound = false;
-                ForEachReference( program, at, [&](OP::ADDRESS ref)
-                {
-                    subtreeFound = subtreeFound || m_visited[ref]==3;
-                });
-
-                m_visited[at] = subtreeFound ? 3 : 2;
-            }
-            else if (!found)
-            {
-                if ( std::find( m_typesToFind.begin(), m_typesToFind.end(), program.GetOpType(at))
-                     !=
-                     m_typesToFind.end() )
-                {
-                    m_visited[at] = 3; // visited, ops found
-                    found = true;
-                }
-                else
-                {
-                    check(m_visited[at]==0);
-
-                    m_visited[at] = 1;
-                    m_pending.push_back( std::make_pair(true,at) );
-
-                    ForEachReference( program, at, [&](OP::ADDRESS ref)
-                    {
-                        if (ref && m_visited[ref]==0)
-                        {
-                            m_pending.push_back( std::make_pair(false,ref) );
-                        }
-                    });
-                }
-            }
-            else
-            {
-                // We won't process it.
-                m_visited[at] = 0;
-            }
-        }
-    }
-
-    return m_visited[rootAt]==3;
-}
+		struct custom_layout_hash
+		{
+			uint64 operator()( const LayoutPtrConst& k ) const
+			{
+				uint64 h =  std::hash<uint64>()(k->GetBlockCount());
+				FIntPoint s = k->GetGridSize();
+				hash_combine(h, s[0]);
+				hash_combine(h, s[1]);
+				return h;
+			}
+		};
 
+		struct custom_layout_equal
+		{
+			bool operator()( const LayoutPtrConst& a, const LayoutPtrConst& b ) const
+			{
+				return a==b || *a==*b;
+			}
+		};
+	}
 
-//---------------------------------------------------------------------------------------------
-//---------------------------------------------------------------------------------------------
-//---------------------------------------------------------------------------------------------
-IsConstantVisitor::IsConstantVisitor()
-{
-    vector<OP_TYPE> parameterTypes;
-    parameterTypes.push_back( OP_TYPE::BO_PARAMETER );
-    parameterTypes.push_back( OP_TYPE::NU_PARAMETER );
-    parameterTypes.push_back( OP_TYPE::SC_PARAMETER );
-    parameterTypes.push_back( OP_TYPE::CO_PARAMETER );
-    parameterTypes.push_back( OP_TYPE::PR_PARAMETER );
-    parameterTypes.push_back( OP_TYPE::IM_PARAMETER );
-    m_findOpTypeVisitor= std::make_unique<FindOpTypeVisitor>(parameterTypes);
-}
 
+	//-------------------------------------------------------------------------------------------------
+	void DuplicatedDataRemoverAST( ASTOpList& roots )
+	{
+		MUTABLE_CPUPROFILER_SCOPE(DuplicatedDataRemoverAST);
 
-//---------------------------------------------------------------------------------------------
-bool IsConstantVisitor::Apply( PROGRAM& program, OP::ADDRESS at )
-{
-    return !m_findOpTypeVisitor->Apply(program,at);
-}
+		std::vector<Ptr<ASTOpConstantResource>> allMeshOps;
+		std::vector<Ptr<ASTOpConstantResource>> allImageOps;
+		std::vector<Ptr<ASTOpConstantResource>> allLayoutOps;
 
+		// Gather constants
+		{
+			//MUTABLE_CPUPROFILER_SCOPE(Gather);
 
-//---------------------------------------------------------------------------------------------
-//---------------------------------------------------------------------------------------------
-//---------------------------------------------------------------------------------------------
-class ConstantTask : public TaskManager::Task
-{
-public:
-    // input
-	mu::Ptr<ASTOp> m_source;
-    bool m_useDiskCache = false;
-    int m_imageCompressionQuality = 0;
-
-    std::shared_timed_mutex* m_codeAccessMutex = nullptr;
-
-private:
-
-    //
-	mu::Ptr<ASTOp> m_result;
-
-public:
-
-    // mu::Task interface
-    void Run() override
-    {
-        MUTABLE_CPUPROFILER_SCOPE(ConstantTask_Run);
-
-        // We need the clone because linking modifies ASTOp state
-        m_codeAccessMutex->lock_shared();
-		mu::Ptr<ASTOp> cloned = ASTOp::DeepClone( m_source );
-        m_codeAccessMutex->unlock_shared();
-
-        OP_TYPE type = cloned->GetOpType();
-        DATATYPE dtype = GetOpDataType(type);
-
-        SettingsPtr pSettings = new Settings;
-        pSettings->SetProfile( false );
-        pSettings->SetImageCompressionQuality( m_imageCompressionQuality );
-        SystemPtr pSystem = new System( pSettings );
-
-		// Don't generate mips suring linking here.
-		FLinkerOptions LinkerOptions;
-		LinkerOptions.MinTextureResidentMipCount = 255;
-
-        ModelPtrConst model = new Model;
-        ASTOp::FullLink( cloned, model->GetPrivate()->m_program, &LinkerOptions);
-        OP::ADDRESS at = cloned->linkedAddress;
-
-        PROGRAM::STATE state;
-        state.m_root = at;
-        model->GetPrivate()->m_program.m_states.Add(state);
-
-        ParametersPtr localParams = model->NewParameters();
-        pSystem->GetPrivate()->BeginBuild( model );
-
-        // Calculate the value and replace this op by a constant
-        switch( dtype )
-        {
-        case DT_MESH:
-        {
-            MUTABLE_CPUPROFILER_SCOPE(ConstantMesh);
-
-			mu::Ptr<const Mesh> pMesh = pSystem->GetPrivate()->BuildMesh( model, localParams.get(), at );
-
-            if (pMesh)
-            {
-				mu::Ptr<ASTOpConstantResource> constantOp = new ASTOpConstantResource();
-                constantOp->type = OP_TYPE::ME_CONSTANT;
-                constantOp->SetValue( pMesh, m_useDiskCache );
-                m_result = constantOp;
-              }
-            break;
-        }
-
-        case DT_IMAGE:
-        {
-            MUTABLE_CPUPROFILER_SCOPE(ConstantImage);
-
-			mu::Ptr<const Image> pImage = pSystem->GetPrivate()->BuildImage( model, localParams.get(), at, 0 );
-
-            if (pImage)
-            {
-				mu::Ptr<ASTOpConstantResource> constantOp = new ASTOpConstantResource();
-                constantOp->type = OP_TYPE::IM_CONSTANT;
-                constantOp->SetValue( pImage, m_useDiskCache );
-                m_result = constantOp;
-            }
-            break;
-        }
-
-        case DT_LAYOUT:
-        {
-            MUTABLE_CPUPROFILER_SCOPE(ConstantLayout);
-
-			mu::Ptr<const Layout> pLayout = pSystem->GetPrivate()->BuildLayout( model, localParams.get(), at );
-
-            if (pLayout)
-            {
-				mu::Ptr<ASTOpConstantResource> constantOp = new ASTOpConstantResource();
-                constantOp->type = OP_TYPE::LA_CONSTANT;
-                constantOp->SetValue( pLayout, m_useDiskCache );
-                m_result = constantOp;
-            }
-            break;
-        }
-
-        case DT_BOOL:
-        {
-            MUTABLE_CPUPROFILER_SCOPE(ConstantBool);
-
-            bool value = pSystem->GetPrivate()->BuildBool( model, localParams.get(), at );
-
-            {
-				mu::Ptr<ASTOpConstantBool> constantOp = new ASTOpConstantBool();
-                constantOp->value = value;
-                m_result = constantOp;
-            }
-            break;
-        }
-
-        case DT_COLOUR:
-        {
-            MUTABLE_CPUPROFILER_SCOPE(ConstantBool);
-
-            float r=0.0f, g=0.0f, b=0.0f;
-            pSystem->GetPrivate()->BuildColour( model, localParams.get(), at, &r,&g,&b );
-
-            {
-				mu::Ptr<ASTOpFixed> constantOp = new ASTOpFixed();
-                constantOp->op.type = OP_TYPE::CO_CONSTANT;
-                constantOp->op.args.ColourConstant.value[0] = r;
-                constantOp->op.args.ColourConstant.value[1] = g;
-                constantOp->op.args.ColourConstant.value[2] = b;
-                constantOp->op.args.ColourConstant.value[3] = 1.0f;
-                m_result = constantOp;
-            }
-            break;
-        }
-
-        case DT_INT:
-        case DT_SCALAR:
-        case DT_STRING:
-        case DT_PROJECTOR:
-            // TODO
-            break;
-
-        default:
-            break;
-        }
-
-		pSystem->GetPrivate()->EndBuild();
-
-    }
-
-    // TaskManager::Task interface
-    void Complete() override
-    {
-        // This runs in the managing thread
-        m_codeAccessMutex->lock();
-        ASTOp::Replace( m_source, m_result );
-        m_codeAccessMutex->unlock();
-    }
-};
-
-
-//---------------------------------------------------------------------------------------------
-bool mu::ConstantGeneratorAST( const CompilerOptions::Private* options, Ptr<ASTOp>& root,
-                               TaskManager* pTaskManager )
-{
-    MUTABLE_CPUPROFILER_SCOPE(ConstantGenerator);
-
-    bool modified = false;
-
-    // don't do this if constant optimization has been disabled, usually for debugging.
-    if (!options->m_optimisationOptions.m_constReduction)
-    {
-        return false;
-    }
-
-
-    // Calculate constant-subtree and special-op flags
-    {
-        MUTABLE_CPUPROFILER_SCOPE(ConstantGenerator_CalculateFlags);
-
-        ASTOp::Traverse_BottomUp_Unique( root, [&](Ptr<ASTOp>& n)
-        {
-            bool constantSubtree = true;
-            switch (n->GetOpType())
-            {
-            case OP_TYPE::BO_PARAMETER:
-            case OP_TYPE::NU_PARAMETER:
-            case OP_TYPE::SC_PARAMETER:
-            case OP_TYPE::CO_PARAMETER:
-            case OP_TYPE::PR_PARAMETER:
-            case OP_TYPE::IM_PARAMETER:
-                constantSubtree = false;
-                break;
-            default:
-                // Propagate from children
-                n->ForEachChild( [&constantSubtree](ASTChild& c)
-                {
-                    if (c)
-                    {
-                        constantSubtree = constantSubtree && c->m_constantSubtree;
-                    }
-                });
-                break;
-            }
-            n->m_constantSubtree = constantSubtree;
-
-            // We avoid generating constants for these operations, to avoid the memory
-            // explosion.
-            // TODO: Make compiler options for some of them
-            // TODO: Some of them are worth if the code below them is unique.
-            bool hasSpecialOpInSubtree = false;
-            switch (n->GetOpType())
-            {
-            case OP_TYPE::IM_BLANKLAYOUT:
-            case OP_TYPE::IM_COMPOSE:
-            //case OP_TYPE::IM_RASTERMESH:            // TODO review this one
-            case OP_TYPE::ME_MERGE:
-            case OP_TYPE::ME_CLIPWITHMESH:
-            case OP_TYPE::ME_CLIPMORPHPLANE:
-            case OP_TYPE::ME_APPLYPOSE:
-            case OP_TYPE::ME_REMOVEMASK:
-                hasSpecialOpInSubtree = true;
-                break;
-
-            default:
-                // Propagate from children
-                n->ForEachChild( [&](ASTChild& c)
-                {
-                    if (c)
-                    {
-                        hasSpecialOpInSubtree = hasSpecialOpInSubtree || c->m_hasSpecialOpInSubtree;
-                    }
-                });
-                break;
-            }
-            n->m_hasSpecialOpInSubtree = hasSpecialOpInSubtree;
-        });
-
-    }
-
-    std::shared_timed_mutex codeAccessMutex;
-
-    TArray< Ptr<ASTOp> > roots;
-    roots.Add(root);
-
-    // Generate constant operations
-    ASTOp::Traverse_TopDown_Unique_Imprecise( roots, [&](Ptr<ASTOp>& n)
-    {
-        bool recurse = true;
-
-        OP_TYPE type = n->GetOpType();
-        DATATYPE dtype = GetOpDataType(type);
-
-        bool constant = false;
-        if (dtype!=DT_INSTANCE)
-        {
-            constant = n->m_constantSubtree;
-        }
-
-        // See if it is a case of instructions we want to avoid, but with special parameters that make
-        // them ok
-        bool specialCase = false;
-        if (constant)
-        {
-            //const OP& op = program.m_code[at];
-            // TODO
-//                if (n->GetOpType()==OP_TYPE::IM_COMPOSE)
-//                {
-//                    // Get the layout
-//                    // In case we fill all of the target image anyway
-//                    ParametersPtr pParams = m_pModel->NewParameters();
-
-//                    const auto& args = op.args.ImageCompose;
-//                    LayoutPtrConst pLayout = m_pSystem->GetPrivate()->BuildLayout
-//                        ( m_pModel.get(), pParams.get(), args.layout );
-
-//                    if ( pLayout &&
-//                         pLayout->GetPrivate()->IsSingleBlockAndFull() )
-//                    {
-//                        specialCase = true;
-//                    }
-//                }
-
-//                else
-                if (n->GetOpType()==OP_TYPE::IM_RASTERMESH)
-            {
-                const auto& args = dynamic_cast<const ASTOpFixed*>(n.get())->op.args.ImageRasterMesh;
-                if ( !args.image )
-                {
-                    specialCase = true;
-                }
-            }
-        }
-
-        // See if the subtree contains operations that we don't want to collapse even if
-        // they are constant
-        if ( constant && !specialCase)
-        {
-            if ( n->m_hasSpecialOpInSubtree )
-            {
-                constant = false;
-            }
-        }
-
-        // If children are constant
-        if ( constant
-             && type!=OP_TYPE::BO_CONSTANT
-             && type!=OP_TYPE::NU_CONSTANT
-             && type!=OP_TYPE::SC_CONSTANT
-             && type!=OP_TYPE::CO_CONSTANT
-             && type!=OP_TYPE::IM_CONSTANT
-             && type!=OP_TYPE::ME_CONSTANT
-             && type!=OP_TYPE::LA_CONSTANT
-             && type!=OP_TYPE::PR_CONSTANT
-             && type!=OP_TYPE::NONE
-            )
-        {
-            switch( dtype )
-            {
-            case DT_MESH:
-            case DT_IMAGE:
-            case DT_LAYOUT:
-            case DT_BOOL:
-            case DT_COLOUR:
-            {
-                recurse = false;
-                modified = true;
-
-                ConstantTask* constantTask = new ConstantTask;
-                constantTask->m_useDiskCache =
-                    options->m_optimisationOptions.m_useDiskCache;
-                constantTask->m_source = n;
-                constantTask->m_codeAccessMutex = &codeAccessMutex;
-                constantTask->m_imageCompressionQuality =
-                    options->m_imageCompressionQuality;
-
-                if ( pTaskManager )
-                {
-                    pTaskManager->AddTask( constantTask );
-                }
-                else
-                {
-                    constantTask->Run();
-                    constantTask->Complete();
-                    delete constantTask;
-                }
-
-                break;
-            }
-
-            case DT_INT:
-            case DT_SCALAR:
-            case DT_STRING:
-            case DT_PROJECTOR:
-                // TODO
-                break;
-
-            case DT_INSTANCE:
-                // nothing to do
-                break;
-
-            default:
-                break;
-            }
-
-        }
-
-        return recurse;
-    });
-
-    if (pTaskManager)
-    {
-        pTaskManager->CompleteTasks();
-    }
-
-    return modified;
-}
+			ASTOp::Traverse_TopRandom_Unique_NonReentrant( roots, [&](Ptr<ASTOp> n)
+			{
+				switch ( n->GetOpType() )
+				{
+
+				case OP_TYPE::ME_CONSTANT:
+				{
+					auto typedNode = dynamic_cast<ASTOpConstantResource*>(n.get());
+					check(typedNode);
+
+					if (typedNode)
+					{
+						allMeshOps.push_back(typedNode);
+					}
+					break;
+				}
+
+				case OP_TYPE::IM_CONSTANT:
+				{
+					auto typedNode = dynamic_cast<ASTOpConstantResource*>(n.get());
+					check(typedNode);
+
+					if (typedNode)
+					{
+						allImageOps.push_back(typedNode);
+					}
+					break;
+				}
+
+				case OP_TYPE::LA_CONSTANT:
+				{
+					auto typedNode = dynamic_cast<ASTOpConstantResource*>(n.get());
+					check(typedNode);
+
+					if (typedNode)
+					{
+						allLayoutOps.push_back(typedNode);
+					}
+					break;
+				}
+
+				//    These should be part of the duplicated code removal, in AST.
+				//            // Names
+				//            case OP_TYPE::IN_ADDMESH:
+				//            case OP_TYPE::IN_ADDIMAGE:
+				//            case OP_TYPE::IN_ADDVECTOR:
+				//            case OP_TYPE::IN_ADDSCALAR:
+				//            case OP_TYPE::IN_ADDCOMPONENT:
+				//            case OP_TYPE::IN_ADDSURFACE:
+				//            {
+				//                OP::ADDRESS value = program.m_code[at].args.InstanceAdd.name;
+				//                program.m_code[at].args.InstanceAdd.name = m_oldToNewStrings[value];
+				//                break;
+				//            }
+
+
+				default:
+					break;
+				}
+
+				return true;
+			});
+		}
 
 
-//-------------------------------------------------------------------------------------------------
-//-------------------------------------------------------------------------------------------------
-//-------------------------------------------------------------------------------------------------
-CodeOptimiser::CodeOptimiser( CompilerOptionsPtr options, vector<STATE_COMPILATION_DATA>& states )
-    : m_states( states )
-{
-    m_options = options;
-}
+		// Compare meshes
+		{
+			//MUTABLE_CPUPROFILER_SCOPE(CompareMeshes);
 
+			std::unordered_multimap< size_t, MESH_ENTRY > meshes;
 
-//-------------------------------------------------------------------------------------------------
-void CodeOptimiser::FullOptimiseAST( ASTOpList& roots,
-                                     TaskManager* pTaskManager )
-{
-    bool modified = true;
-    int numIterations = 0;
-    while (modified && (!m_optimizeIterationsMax || (m_optimizeIterationsLeft>0) || !numIterations) )
-    {
-        --m_optimizeIterationsLeft;
-        ++numIterations;
-        UE_LOG(LogMutableCore, Verbose, TEXT("Main optimise iteration %d, max %d, left %d"),
-                numIterations, m_optimizeIterationsMax, m_optimizeIterationsLeft);
-
-        modified = false;
-
-        // All kind of optimisations that depend on the meaning of each operation
-        // \TODO: We are doing it for all states.
-        UE_LOG(LogMutableCore, Verbose, TEXT(" - semantic optimiser"));
-        modified |= SemanticOptimiserAST( roots, m_options->GetPrivate()->m_optimisationOptions );
-        //AXE_INT_VALUE("Mutable", Verbose, "ast size", (int64_t)ASTOp::CountNodes(roots));
-		UE_LOG(LogMutableCore, Verbose, TEXT("(int) %s : %ld"), TEXT("ast size"), int64(ASTOp::CountNodes(roots)));
-		ASTOp::LogHistogram(roots);
+			for (auto& typedNode: allMeshOps)
+			{
+				size_t key = typedNode->GetValueHash();
+
+				Ptr<ASTOp> found;
+
+				auto r = meshes.equal_range(key);
+				if (r.first!=r.second)
+				{
+					MeshPtrConst mesh = static_cast<const Mesh*>( typedNode->GetValue().get() );
+
+					for ( auto it=r.first; it!=r.second; ++it )
+					{
+						if (!it->second.mesh)
+						{
+							it->second.mesh = static_cast<const Mesh*>( it->second.op->GetValue().get() );
+						}
+
+						if ( custom_mesh_equal()( mesh, it->second.mesh ) )
+						{
+							found = it->second.op;
+							break;
+						}
+					}
+				}
+
+				if (found)
+				{
+					ASTOp::Replace(typedNode,found);
+				}
+				else
+				{
+					MESH_ENTRY e;
+					e.op = typedNode;
+					meshes.insert( std::make_pair<>(key, e) );
+				}
+			}
+		}
 
-        UE_LOG(LogMutableCore, Verbose, TEXT(" - sink optimiser"));
-        modified |= SinkOptimiserAST( roots, m_options->GetPrivate()->m_optimisationOptions );
-        //AXE_INT_VALUE("Mutable", Verbose, "ast size", (int64_t)ASTOp::CountNodes(roots));
-        ASTOp::LogHistogram(roots);
+		// Compare images
+		{
+			//MUTABLE_CPUPROFILER_SCOPE(CompareImages);
 
-        // Image size operations are treated separately
-        UE_LOG(LogMutableCore, Verbose, TEXT(" - size optimiser"));
-        modified |= SizeOptimiserAST( roots );
-        //AXE_INT_VALUE("Mutable", Verbose, "ast size", (int64_t)ASTOp::CountNodes(roots));
+			// TODO Optimise like the mesh compare above
+			std::unordered_map< ImagePtrConst, Ptr<ASTOp>, custom_image_hash, custom_image_equal > images;
 
-        // We cannot run the logic optimiser here because the constant memory explodes.
-        // TODO: Conservative logic optimiser
-        if (!modified)
-        {
-//                AXE_INT_VALUE("Mutable", Verbose, "program size", (int64_t)program.m_opAddress.Num());
-//                UE_LOG(LogMutableCore, Verbose, " - logic optimiser");
-//                LocalLogicOptimiser log;
-//                modified |= log.Apply( program, (int)s );
-        }
-    }
+			for (auto& typedNode: allImageOps)
+			{
+				ImagePtrConst image = static_cast<const Image*>( typedNode->GetValue().get() );
+
+				auto it = images.find(image);
+				if (it!=images.end())
+				{
+					ASTOp::Replace(typedNode,it->second);
+				}
+				else
+				{
+					images.insert( std::make_pair<>(image, typedNode) );
+				}
+			}
+		}
 
-//    ASTOp::LogHistogram(roots);
+		// Compare layouts
+		{
+			//MUTABLE_CPUPROFILER_SCOPE(CompareLayouts);
 
-    UE_LOG(LogMutableCore, Verbose, TEXT(" - duplicated code remover"));
-    DuplicatedCodeRemoverAST( roots );
-	UE_LOG(LogMutableCore, Verbose, TEXT("(int) %s : %ld"), TEXT("ast size"), int64(ASTOp::CountNodes(roots)));
+			// TODO Optimise like the mesh compare above
+			std::unordered_map< LayoutPtrConst, Ptr<ASTOp>, custom_layout_hash, custom_layout_equal > constants;
 
-    ASTOp::LogHistogram(roots);
+			for (auto& typedNode: allLayoutOps)
+			{
+				LayoutPtrConst r = static_cast<const Layout*>( typedNode->GetValue().get() );
+
+				auto it = constants.find(r);
+				if (it!=constants.end())
+				{
+					ASTOp::Replace(typedNode,it->second);
+				}
+				else
+				{
+					constants.insert( std::make_pair<>(r, typedNode) );
+				}
+			}
+		}
 
-    UE_LOG(LogMutableCore, Verbose, TEXT(" - duplicated data remover"));
-    DuplicatedDataRemoverAST( roots );
-	UE_LOG(LogMutableCore, Verbose, TEXT("(int) %s : %ld"), TEXT("ast size"), int64(ASTOp::CountNodes(roots)));
+	}
 
-    ASTOp::LogHistogram(roots);
 
-    // Generate constants
-    for ( auto& r: roots )
-    {
-		UE_LOG(LogMutableCore, Verbose, TEXT("(int) %s : %ld"), TEXT("ast size"), int64(ASTOp::CountNodes(roots)));
-		UE_LOG(LogMutableCore, Verbose, TEXT(" - constant generator"));
+	//-------------------------------------------------------------------------------------------------
+	//-------------------------------------------------------------------------------------------------
+	//-------------------------------------------------------------------------------------------------
+	namespace
+	{
+		struct op_pointer_hash
+		{
+			inline size_t operator() (const mu::Ptr<ASTOp>& n) const
+			{
+				return n->Hash();
+			}
+		};
 
-        // Constant subtree generation
-        modified = ConstantGeneratorAST( m_options->GetPrivate(), r, pTaskManager );
-    }
+		struct op_pointer_equal
+		{
+			inline bool operator() (const mu::Ptr<ASTOp>& a, const mu::Ptr<ASTOp>& b) const
+			{
+				return *a==*b;
+			}
+		};
+	}
 
-	UE_LOG(LogMutableCore, Verbose, TEXT("(int) %s : %ld"), TEXT("ast size"), int64(ASTOp::CountNodes(roots)));
+	void DuplicatedCodeRemoverAST( ASTOpList& roots )
+	{
+		MUTABLE_CPUPROFILER_SCOPE(DuplicatedCodeRemoverAST);
 
-//    ASTOp::LogHistogram(roots);
+		// Visited nodes, per type
+		std::unordered_set<Ptr<ASTOp>,op_pointer_hash,op_pointer_equal> visited[int(OP_TYPE::COUNT)];
 
-    // \todo: this code deduplication here, takes very long and doesn't do much... investigate!
-//    UE_LOG(LogMutableCore, Verbose, " - duplicated code remover");
-//    DuplicatedCodeRemoverAST( roots );
-//    AXE_INT_VALUE("Mutable", Verbose, "ast size", (int64_t)ASTOp::CountNodes(roots));
+		ASTOp::Traverse_BottomUp_Unique_NonReentrant( roots, [&](Ptr<ASTOp>& n)
+		{
+			auto& container = visited[(int)n->GetOpType()];
+
+			// debug
+	//        for( const auto& o : container )
+	//        {
+	//            if (*n==*o)
+	//            {
+	//                auto thisHash = n->Hash();
+	//                auto otherHash = o->Hash();
+	//                check(thisHash==otherHash);
+	//            }
+	//        }
+
+			// Insert will tell us if it was already there
+			auto it = container.insert(n);
+			if( !it.second )
+			{
+				// It wasn't inserted, so it was already there
+				ASTOp::Replace(n,*it.first);
+			}
+		});
+	}
 
-    ASTOp::LogHistogram(roots);
 
-    UE_LOG(LogMutableCore, Verbose, TEXT(" - duplicated data remover"));
-    DuplicatedDataRemoverAST( roots );
-	UE_LOG(LogMutableCore, Verbose, TEXT("(int) %s : %ld"), TEXT("ast size"), int64(ASTOp::CountNodes(roots)));
+	//---------------------------------------------------------------------------------------------
+	//---------------------------------------------------------------------------------------------
+	//---------------------------------------------------------------------------------------------
+	bool FindOpTypeVisitor::Apply( PROGRAM& program, OP::ADDRESS rootAt )
+	{
+		bool found = false;
 
-    //if (!modified)
-    {
-        //for ( size_t s=0;  s<program.m_states.size(); ++s )
-        {
-            ASTOp::LogHistogram(roots);
-            UE_LOG(LogMutableCore, Verbose, TEXT(" - logic optimiser"));
-            modified |= LocalLogicOptimiserAST( roots );
-            ASTOp::LogHistogram(roots);
-        }
+		// If the program added new operations, we haven't visited them.
+		m_visited.resize( program.m_opAddress.Num(), 0 );
 
-        //AXE_INT_VALUE("Mutable", Verbose, "ast size", (int64_t)ASTOp::CountNodes(roots));
-    }
+		m_pending.clear();
+		m_pending.reserve( program.m_opAddress.Num()/4 );
+		m_pending.push_back( std::make_pair(false,rootAt) );
 
-    ASTOp::LogHistogram(roots);
-}
+		// Don't early out to be able to complete parent op cached flags
+		while ( !m_pending.empty() )
+		{
+			std::pair<bool,int> item = m_pending.back();
+			m_pending.pop_back();
+			OP::ADDRESS at = item.second;
 
+			// Not cached?
+			if (m_visited[at]<=1)
+			{
+				if (item.first)
+				{
+					// Item indicating we finished with all the children of a parent
+					check(m_visited[at]==1);
+
+					bool subtreeFound = false;
+					ForEachReference( program, at, [&](OP::ADDRESS ref)
+					{
+						subtreeFound = subtreeFound || m_visited[ref]==3;
+					});
+
+					m_visited[at] = subtreeFound ? 3 : 2;
+				}
+				else if (!found)
+				{
+					if ( std::find( m_typesToFind.begin(), m_typesToFind.end(), program.GetOpType(at))
+						 !=
+						 m_typesToFind.end() )
+					{
+						m_visited[at] = 3; // visited, ops found
+						found = true;
+					}
+					else
+					{
+						check(m_visited[at]==0);
+
+						m_visited[at] = 1;
+						m_pending.push_back( std::make_pair(true,at) );
+
+						ForEachReference( program, at, [&](OP::ADDRESS ref)
+						{
+							if (ref && m_visited[ref]==0)
+							{
+								m_pending.push_back( std::make_pair(false,ref) );
+							}
+						});
+					}
+				}
+				else
+				{
+					// We won't process it.
+					m_visited[at] = 0;
+				}
+			}
+		}
 
-//---------------------------------------------------------------------------------------------
-// The state represents if there is a parent operation requiring skeleton for current mesh subtree.
-class CollectAllMeshesForSkeletonVisitorAST : public Visitor_TopDown_Unique_Const<uint8_t>
-{
-public:
-
-    CollectAllMeshesForSkeletonVisitorAST( const ASTOpList& roots  )
-    {
-        Traverse( roots, false );
-    }
-
-    // List of meshes that require a skeleton
-    vector<mu::Ptr<ASTOpConstantResource>> m_meshesRequiringSkeleton;
-
-private:
-
-    // Visitor_TopDown_Unique_Const<uint8_t> interface
-    bool Visit( const mu::Ptr<ASTOp>& node ) override
-    {
-        // \todo: refine to avoid instruction branches with irrelevant skeletons.
-
-        uint8_t currentProtected = GetCurrentState();
-
-        switch (node->GetOpType())
-        {
-
-        case OP_TYPE::ME_CONSTANT:
-        {
-			mu::Ptr<ASTOpConstantResource> typedOp = dynamic_cast<ASTOpConstantResource*>(node.get());
-            check(typedOp);
-
-            if (currentProtected)
-            {
-                if ( std::find(m_meshesRequiringSkeleton.begin(), m_meshesRequiringSkeleton.end(), typedOp)
-                     ==
-                     m_meshesRequiringSkeleton.end() )
-                {
-                    m_meshesRequiringSkeleton.push_back(typedOp);
-                }
-            }
-
-            return false;
-        }
-
-        case OP_TYPE::ME_CLIPMORPHPLANE:
-        {
-            auto typedOp = dynamic_cast<ASTOpMeshClipMorphPlane*>(node.get());
-            if (typedOp)
-            {
-                if (typedOp->vertexSelectionType
-                        ==
-                        OP::MeshClipMorphPlaneArgs::VS_BONE_HIERARCHY)
-                {
-                    // We need the skeleton for the source mesh
-                    RecurseWithState( typedOp->source.child(), true );
-                    return false;
-                }
-            }
-
-            return true;
-        }
-
-		case OP_TYPE::ME_APPLYPOSE:
+		return m_visited[rootAt]==3;
+	}
+
+
+	//---------------------------------------------------------------------------------------------
+	//---------------------------------------------------------------------------------------------
+	//---------------------------------------------------------------------------------------------
+	IsConstantVisitor::IsConstantVisitor()
+	{
+		vector<OP_TYPE> parameterTypes;
+		parameterTypes.push_back( OP_TYPE::BO_PARAMETER );
+		parameterTypes.push_back( OP_TYPE::NU_PARAMETER );
+		parameterTypes.push_back( OP_TYPE::SC_PARAMETER );
+		parameterTypes.push_back( OP_TYPE::CO_PARAMETER );
+		parameterTypes.push_back( OP_TYPE::PR_PARAMETER );
+		parameterTypes.push_back( OP_TYPE::IM_PARAMETER );
+		m_findOpTypeVisitor= std::make_unique<FindOpTypeVisitor>(parameterTypes);
+	}
+
+
+	//---------------------------------------------------------------------------------------------
+	bool IsConstantVisitor::Apply( PROGRAM& program, OP::ADDRESS at )
+	{
+		return !m_findOpTypeVisitor->Apply(program,at);
+	}
+
+
+	//---------------------------------------------------------------------------------------------
+	//---------------------------------------------------------------------------------------------
+	//---------------------------------------------------------------------------------------------
+	class ConstantTask : public TaskManager::Task
+	{
+	public:
+		// input
+		mu::Ptr<ASTOp> m_source;
+		bool m_useDiskCache = false;
+		int m_imageCompressionQuality = 0;
+
+		std::shared_timed_mutex* m_codeAccessMutex = nullptr;
+
+	private:
+
+		//
+		mu::Ptr<ASTOp> m_result;
+
+	public:
+
+		// mu::Task interface
+		void Run() override
 		{
-			auto typedOp = dynamic_cast<ASTOpMeshApplyPose*>(node.get());
-			check(typedOp);
+			MUTABLE_CPUPROFILER_SCOPE(ConstantTask_Run);
+
+			// We need the clone because linking modifies ASTOp state
+			m_codeAccessMutex->lock_shared();
+			mu::Ptr<ASTOp> cloned = ASTOp::DeepClone( m_source );
+			m_codeAccessMutex->unlock_shared();
+
+			OP_TYPE type = cloned->GetOpType();
+			DATATYPE dtype = GetOpDataType(type);
+
+			SettingsPtr pSettings = new Settings;
+			pSettings->SetProfile( false );
+			pSettings->SetImageCompressionQuality( m_imageCompressionQuality );
+			SystemPtr pSystem = new System( pSettings );
+
+			// Don't generate mips suring linking here.
+			FLinkerOptions LinkerOptions;
+			LinkerOptions.MinTextureResidentMipCount = 255;
 
-			if (typedOp)
+			ModelPtrConst model = new Model;
+			ASTOp::FullLink( cloned, model->GetPrivate()->m_program, &LinkerOptions);
+			OP::ADDRESS at = cloned->linkedAddress;
+
+			PROGRAM::STATE state;
+			state.m_root = at;
+			model->GetPrivate()->m_program.m_states.Add(state);
+
+			ParametersPtr localParams = model->NewParameters();
+			pSystem->GetPrivate()->BeginBuild( model );
+
+			// Calculate the value and replace this op by a constant
+			switch( dtype )
 			{
-				// We need the skeleton for both meshes
-				RecurseWithState(typedOp->base.child(), true);
-				RecurseWithState(typedOp->pose.child(), true);
-				return false;
+			case DT_MESH:
+			{
+				MUTABLE_CPUPROFILER_SCOPE(ConstantMesh);
+
+				mu::Ptr<const Mesh> pMesh = pSystem->GetPrivate()->BuildMesh( model, localParams.get(), at );
+
+				if (pMesh)
+				{
+					mu::Ptr<ASTOpConstantResource> constantOp = new ASTOpConstantResource();
+					constantOp->type = OP_TYPE::ME_CONSTANT;
+					constantOp->SetValue( pMesh, m_useDiskCache );
+					m_result = constantOp;
+				  }
+				break;
+			}
+
+			case DT_IMAGE:
+			{
+				MUTABLE_CPUPROFILER_SCOPE(ConstantImage);
+
+				mu::Ptr<const Image> pImage = pSystem->GetPrivate()->BuildImage( model, localParams.get(), at, 0 );
+
+				if (pImage)
+				{
+					mu::Ptr<ASTOpConstantResource> constantOp = new ASTOpConstantResource();
+					constantOp->type = OP_TYPE::IM_CONSTANT;
+					constantOp->SetValue( pImage, m_useDiskCache );
+					m_result = constantOp;
+				}
+				break;
+			}
+
+			case DT_LAYOUT:
+			{
+				MUTABLE_CPUPROFILER_SCOPE(ConstantLayout);
+
+				mu::Ptr<const Layout> pLayout = pSystem->GetPrivate()->BuildLayout( model, localParams.get(), at );
+
+				if (pLayout)
+				{
+					mu::Ptr<ASTOpConstantResource> constantOp = new ASTOpConstantResource();
+					constantOp->type = OP_TYPE::LA_CONSTANT;
+					constantOp->SetValue( pLayout, m_useDiskCache );
+					m_result = constantOp;
+				}
+				break;
+			}
+
+			case DT_BOOL:
+			{
+				MUTABLE_CPUPROFILER_SCOPE(ConstantBool);
+
+				bool value = pSystem->GetPrivate()->BuildBool( model, localParams.get(), at );
+
+				{
+					mu::Ptr<ASTOpConstantBool> constantOp = new ASTOpConstantBool();
+					constantOp->value = value;
+					m_result = constantOp;
+				}
+				break;
+			}
+
+			case DT_COLOUR:
+			{
+				MUTABLE_CPUPROFILER_SCOPE(ConstantBool);
+
+				float r=0.0f, g=0.0f, b=0.0f;
+				pSystem->GetPrivate()->BuildColour( model, localParams.get(), at, &r,&g,&b );
+
+				{
+					mu::Ptr<ASTOpFixed> constantOp = new ASTOpFixed();
+					constantOp->op.type = OP_TYPE::CO_CONSTANT;
+					constantOp->op.args.ColourConstant.value[0] = r;
+					constantOp->op.args.ColourConstant.value[1] = g;
+					constantOp->op.args.ColourConstant.value[2] = b;
+					constantOp->op.args.ColourConstant.value[3] = 1.0f;
+					m_result = constantOp;
+				}
+				break;
+			}
+
+			case DT_INT:
+			case DT_SCALAR:
+			case DT_STRING:
+			case DT_PROJECTOR:
+				// TODO
+				break;
+
+			default:
+				break;
 			}
 
-			break;
+			pSystem->GetPrivate()->EndBuild();
+
+		}
+
+		// TaskManager::Task interface
+		void Complete() override
+		{
+			// This runs in the managing thread
+			m_codeAccessMutex->lock();
+			ASTOp::Replace( m_source, m_result );
+			m_codeAccessMutex->unlock();
+		}
+	};
+
+
+	//---------------------------------------------------------------------------------------------
+	bool ConstantGeneratorAST( const CompilerOptions::Private* options, Ptr<ASTOp>& root,
+								   TaskManager* pTaskManager )
+	{
+		MUTABLE_CPUPROFILER_SCOPE(ConstantGenerator);
+
+		bool modified = false;
+
+		// don't do this if constant optimization has been disabled, usually for debugging.
+		if (!options->m_optimisationOptions.m_constReduction)
+		{
+			return false;
 		}
 
-		case OP_TYPE::ME_BINDSHAPE:
+
+		// Calculate constant-subtree and special-op flags
 		{
-			auto typedOp = dynamic_cast<ASTOpMeshBindShape*>(node.get());
-			check(typedOp);
+			MUTABLE_CPUPROFILER_SCOPE(ConstantGenerator_CalculateFlags);
 
-			if (typedOp && typedOp->m_reshapeSkeleton)
+			ASTOp::Traverse_BottomUp_Unique( root, [&](Ptr<ASTOp>& n)
 			{
-				RecurseWithState(typedOp->Mesh.child(), true);
-				return false;
-			}
+				bool constantSubtree = true;
+				switch (n->GetOpType())
+				{
+				case OP_TYPE::BO_PARAMETER:
+				case OP_TYPE::NU_PARAMETER:
+				case OP_TYPE::SC_PARAMETER:
+				case OP_TYPE::CO_PARAMETER:
+				case OP_TYPE::PR_PARAMETER:
+				case OP_TYPE::IM_PARAMETER:
+					constantSubtree = false;
+					break;
+				default:
+					// Propagate from children
+					n->ForEachChild( [&constantSubtree](ASTChild& c)
+					{
+						if (c)
+						{
+							constantSubtree = constantSubtree && c->m_constantSubtree;
+						}
+					});
+					break;
+				}
+				n->m_constantSubtree = constantSubtree;
+
+				// We avoid generating constants for these operations, to avoid the memory
+				// explosion.
+				// TODO: Make compiler options for some of them
+				// TODO: Some of them are worth if the code below them is unique.
+				bool hasSpecialOpInSubtree = false;
+				switch (n->GetOpType())
+				{
+				case OP_TYPE::IM_BLANKLAYOUT:
+				case OP_TYPE::IM_COMPOSE:
+				//case OP_TYPE::IM_RASTERMESH:            // TODO review this one
+				case OP_TYPE::ME_MERGE:
+				case OP_TYPE::ME_CLIPWITHMESH:
+				case OP_TYPE::ME_CLIPMORPHPLANE:
+				case OP_TYPE::ME_APPLYPOSE:
+				case OP_TYPE::ME_REMOVEMASK:
+					hasSpecialOpInSubtree = true;
+					break;
+
+				default:
+					// Propagate from children
+					n->ForEachChild( [&](ASTChild& c)
+					{
+						if (c)
+						{
+							hasSpecialOpInSubtree = hasSpecialOpInSubtree || c->m_hasSpecialOpInSubtree;
+						}
+					});
+					break;
+				}
+				n->m_hasSpecialOpInSubtree = hasSpecialOpInSubtree;
+			});
 
-			break;
 		}
 
-		case OP_TYPE::ME_APPLYSHAPE:
+		std::shared_timed_mutex codeAccessMutex;
+
+		TArray< Ptr<ASTOp> > roots;
+		roots.Add(root);
+
+		// Generate constant operations
+		ASTOp::Traverse_TopDown_Unique_Imprecise( roots, [&](Ptr<ASTOp>& n)
 		{
-			auto typedOp = dynamic_cast<ASTOpMeshApplyShape*>(node.get());
-			check(typedOp);
+			bool recurse = true;
+
+			OP_TYPE type = n->GetOpType();
+			DATATYPE dtype = GetOpDataType(type);
 
-			if (typedOp && typedOp->m_reshapeSkeleton)
+			bool constant = false;
+			if (dtype!=DT_INSTANCE)
 			{
-				RecurseWithState(typedOp->Mesh.child(), true);
-				return false;
+				constant = n->m_constantSubtree;
 			}
 
-			break;
+			// See if it is a case of instructions we want to avoid, but with special parameters that make
+			// them ok
+			bool specialCase = false;
+			if (constant)
+			{
+				//const OP& op = program.m_code[at];
+				// TODO
+	//                if (n->GetOpType()==OP_TYPE::IM_COMPOSE)
+	//                {
+	//                    // Get the layout
+	//                    // In case we fill all of the target image anyway
+	//                    ParametersPtr pParams = m_pModel->NewParameters();
+
+	//                    const auto& args = op.args.ImageCompose;
+	//                    LayoutPtrConst pLayout = m_pSystem->GetPrivate()->BuildLayout
+	//                        ( m_pModel.get(), pParams.get(), args.layout );
+
+	//                    if ( pLayout &&
+	//                         pLayout->GetPrivate()->IsSingleBlockAndFull() )
+	//                    {
+	//                        specialCase = true;
+	//                    }
+	//                }
+
+	//                else
+					if (n->GetOpType()==OP_TYPE::IM_RASTERMESH)
+				{
+					const auto& args = dynamic_cast<const ASTOpFixed*>(n.get())->op.args.ImageRasterMesh;
+					if ( !args.image )
+					{
+						specialCase = true;
+					}
+				}
+			}
+
+			// See if the subtree contains operations that we don't want to collapse even if
+			// they are constant
+			if ( constant && !specialCase)
+			{
+				if ( n->m_hasSpecialOpInSubtree )
+				{
+					constant = false;
+				}
+			}
+
+			// If children are constant
+			if ( constant
+				 && type!=OP_TYPE::BO_CONSTANT
+				 && type!=OP_TYPE::NU_CONSTANT
+				 && type!=OP_TYPE::SC_CONSTANT
+				 && type!=OP_TYPE::CO_CONSTANT
+				 && type!=OP_TYPE::IM_CONSTANT
+				 && type!=OP_TYPE::ME_CONSTANT
+				 && type!=OP_TYPE::LA_CONSTANT
+				 && type!=OP_TYPE::PR_CONSTANT
+				 && type!=OP_TYPE::NONE
+				)
+			{
+				switch( dtype )
+				{
+				case DT_MESH:
+				case DT_IMAGE:
+				case DT_LAYOUT:
+				case DT_BOOL:
+				case DT_COLOUR:
+				{
+					recurse = false;
+					modified = true;
+
+					ConstantTask* constantTask = new ConstantTask;
+					constantTask->m_useDiskCache =
+						options->m_optimisationOptions.m_useDiskCache;
+					constantTask->m_source = n;
+					constantTask->m_codeAccessMutex = &codeAccessMutex;
+					constantTask->m_imageCompressionQuality =
+						options->m_imageCompressionQuality;
+
+					if ( pTaskManager )
+					{
+						pTaskManager->AddTask( constantTask );
+					}
+					else
+					{
+						constantTask->Run();
+						constantTask->Complete();
+						delete constantTask;
+					}
+
+					break;
+				}
+
+				case DT_INT:
+				case DT_SCALAR:
+				case DT_STRING:
+				case DT_PROJECTOR:
+					// TODO
+					break;
+
+				case DT_INSTANCE:
+					// nothing to do
+					break;
+
+				default:
+					break;
+				}
+
+			}
+
+			return recurse;
+		});
+
+		if (pTaskManager)
+		{
+			pTaskManager->CompleteTasks();
 		}
 
-        default:
-            break;
-        }
+		return modified;
+	}
 
-        return true;
-    }
 
-};
+	//-------------------------------------------------------------------------------------------------
+	//-------------------------------------------------------------------------------------------------
+	//-------------------------------------------------------------------------------------------------
+	CodeOptimiser::CodeOptimiser( CompilerOptionsPtr options, vector<STATE_COMPILATION_DATA>& states )
+		: m_states( states )
+	{
+		m_options = options;
+	}
 
 
-//---------------------------------------------------------------------------------------------
-// This stores an ADD_MESH op with the child meshes collected and the final skeleton to use
-// for this op.
-struct ADDMESH_SKELETON
-{
-	mu::Ptr<ASTOp> m_pAddMeshOp;
-	TArray<mu::Ptr<ASTOpConstantResource>> m_contributingMeshes;
-	mu::Ptr<Skeleton> m_pFinalSkeleton;
-
-    ADDMESH_SKELETON( const mu::Ptr<ASTOp>& pAddMeshOp,
-                      TArray<mu::Ptr<ASTOpConstantResource>>& contributingMeshes,
-                      const mu::Ptr<Skeleton>& pFinalSkeleton )
-    {
-        m_pAddMeshOp = pAddMeshOp;
-        m_contributingMeshes = std::move(contributingMeshes);
-        m_pFinalSkeleton = pFinalSkeleton;
-    }
-};
-
-
-//---------------------------------------------------------------------------------------------
-void SkeletonCleanerAST( TArray<mu::Ptr<ASTOp>>& roots, const MODEL_OPTIMIZATION_OPTIONS& options )
-{
-    // This collects all the meshes that require a skeleton because they are used in operations
-    // that require it.
-    CollectAllMeshesForSkeletonVisitorAST requireSkeletonCollector( roots );
-
-    TArray<ADDMESH_SKELETON> replacementsFound;
-
-    ASTOp::Traverse_TopDown_Unique_Imprecise( roots, [&](mu::Ptr<ASTOp>& at )
-    {
-        // Only recurse instance construction ops.
-        bool processChildren = GetOpDataType(at->GetOpType())==DT_INSTANCE;
-
-        if ( at->GetOpType() == OP_TYPE::IN_ADDMESH )
-        {
-            auto typedNode = dynamic_cast<ASTOpInstanceAdd*>(at.get());
-			mu::Ptr<ASTOp> meshRoot = typedNode->value.child();
-
-            if (meshRoot)
-            {
-                // Gather constant meshes contributing to the final mesh
-                TArray<mu::Ptr<ASTOpConstantResource>> subtreeMeshes;
-				TArray<mu::Ptr<ASTOp>> tempRoots;
-                tempRoots.Add(meshRoot);
-                ASTOp::Traverse_TopDown_Unique_Imprecise( tempRoots, [&](mu::Ptr<ASTOp>& lat )
-                {
-                    // \todo: refine to avoid instruction branches with irrelevant skeletons.
-                    if ( lat->GetOpType() == OP_TYPE::ME_CONSTANT )
-                    {
-						mu::Ptr<ASTOpConstantResource> typedOp = dynamic_cast<ASTOpConstantResource*>(lat.get());
-                        check(typedOp);
-
-                        if ( subtreeMeshes.Find(typedOp)
-                             ==
-                             INDEX_NONE )
-                        {
-                            subtreeMeshes.Add(typedOp);
-                        }
-                    }
-                    return true;
-                });
-
-                // Create a mesh just with the unified skeleton
-                SkeletonPtr pFinalSkeleton = new Skeleton;
-                for (const auto& meshAt: subtreeMeshes)
-                {
-                    MeshPtrConst pMesh = static_cast<const Mesh*>(meshAt->GetValue().get());
-					mu::Ptr<const Skeleton> sourceSkeleton = pMesh ? pMesh->GetSkeleton() : nullptr;
-                    if (sourceSkeleton)
-                    {
-                        ExtendSkeleton(pFinalSkeleton.get(),sourceSkeleton.get());
-                    }
-                }
-
-                replacementsFound.Emplace( at, subtreeMeshes, pFinalSkeleton );
-            }
-        }
-
-        return processChildren;
-    });
-
-
-    // Iterate all meshes again
-    ASTOp::Traverse_TopDown_Unique_Imprecise( roots, [&](mu::Ptr<ASTOp>& at )
-    {
-        if (at->GetOpType()==OP_TYPE::ME_CONSTANT)
-        {
-            auto typedOp = dynamic_cast<ASTOpConstantResource*>(at.get());
-            check(typedOp);
-            if (!typedOp)
-            {
-                return true;
-            }
-
-            for( auto& rep: replacementsFound )
-            {
-                if ( rep.m_contributingMeshes.Contains(at) )
-                {
-					mu::Ptr<const Mesh> pMesh =  static_cast<const Mesh*>(typedOp->GetValue().get());
-                    pMesh->CheckIntegrity();
-
-					mu::Ptr<Mesh> pNewMesh = MeshRemapSkeleton( pMesh.get(),
-                                                            rep.m_pFinalSkeleton.get());
-
-                    // It returns null if there is no need to remap.
-                    if (pNewMesh)
-                    {
-                        pNewMesh->CheckIntegrity();
-						mu::Ptr<ASTOpConstantResource> newOp = new ASTOpConstantResource();
-                        newOp->type = OP_TYPE::ME_CONSTANT;
-                        newOp->SetValue( pNewMesh, options.m_useDiskCache );
-
-                        ASTOp::Replace( at, newOp );
-                    }
-                }
-            }
-        }
-        return true;
-    });
-
-//    for( auto& rep: replacementsFound )
-//    {
-//        check(rep.m_pAddMeshOp->GetOpType()==OP_TYPE::IN_ADDMESH);
-//        auto newAddMeshRootTyped = dynamic_cast<ASTOpInstanceAdd*>(rep.m_pAddMeshOp.get());
-//        auto meshRoot = newAddMeshRootTyped->value.child();
-
-//        // Add new skeleton constant instruction
-//        Ptr<ASTOpConstantResource> cop = new ASTOpConstantResource();
-//        cop->type = OP_TYPE::ME_CONSTANT;
-//        cop->SetValue( rep.m_pFinalSkeleton, options.m_useDiskCache );
-
-//        // Add new "apply skeleton" instruction
-//        Ptr<ASTOpFixed> skop = new ASTOpFixed();
-//        skop->op.type = OP_TYPE::ME_SETSKELETON;
-//        skop->SetChild( skop->op.args.MeshSetSkeleton.source, meshRoot);
-//        skop->SetChild( skop->op.args.MeshSetSkeleton.skeleton, cop);
-
-//        // We need a new mesh add instruction, it is not safe to edit the current one.
-//        auto newOp = mu::Clone<ASTOpInstanceAdd>(newAddMeshRootTyped);
-//        newOp->value = skop;
-
-//        // TODO BLEH UNSAFE?
-//        ASTOp::Replace(rep.m_pAddMeshOp,newOp);
-//    }
-}
+	//-------------------------------------------------------------------------------------------------
+	void CodeOptimiser::FullOptimiseAST( ASTOpList& roots,
+										 TaskManager* pTaskManager )
+	{
+		bool modified = true;
+		int numIterations = 0;
+		while (modified && (!m_optimizeIterationsMax || (m_optimizeIterationsLeft>0) || !numIterations) )
+		{
+			--m_optimizeIterationsLeft;
+			++numIterations;
+			UE_LOG(LogMutableCore, Verbose, TEXT("Main optimise iteration %d, max %d, left %d"),
+					numIterations, m_optimizeIterationsMax, m_optimizeIterationsLeft);
+
+			modified = false;
+
+			// All kind of optimisations that depend on the meaning of each operation
+			// \TODO: We are doing it for all states.
+			UE_LOG(LogMutableCore, Verbose, TEXT(" - semantic optimiser"));
+			modified |= SemanticOptimiserAST( roots, m_options->GetPrivate()->m_optimisationOptions );
+			//AXE_INT_VALUE("Mutable", Verbose, "ast size", (int64_t)ASTOp::CountNodes(roots));
+			UE_LOG(LogMutableCore, Verbose, TEXT("(int) %s : %ld"), TEXT("ast size"), int64(ASTOp::CountNodes(roots)));
+			ASTOp::LogHistogram(roots);
 
+			UE_LOG(LogMutableCore, Verbose, TEXT(" - sink optimiser"));
+			modified |= SinkOptimiserAST( roots, m_options->GetPrivate()->m_optimisationOptions );
+			//AXE_INT_VALUE("Mutable", Verbose, "ast size", (int64_t)ASTOp::CountNodes(roots));
+			ASTOp::LogHistogram(roots);
 
-//---------------------------------------------------------------------------------------------
-void CodeOptimiser::OptimiseAST( TaskManager* pTaskManager )
-{
-    MUTABLE_CPUPROFILER_SCOPE(OptimiseAST);
-
-    // Gather all the roots
-    TArray<Ptr<ASTOp>> roots;
-    for(const auto& s:m_states)
-    {
-        roots.Add(s.root);
-    }
-
-	UE_LOG(LogMutableCore, Verbose, TEXT("(int) %s : %ld"), TEXT("ast size"), int64(ASTOp::CountNodes(roots)));
-
-    if ( m_options->GetPrivate()->m_optimisationOptions.m_enabled )
-    {
-        // We use 4 times the count because at the time we moved to sharing this count it
-        // was being used 4 times, and we want to keep the tests consistent.
-        m_optimizeIterationsLeft =
-            m_options->GetPrivate()->m_optimisationOptions.m_maxOptimisationLoopCount * 4;
-        m_optimizeIterationsMax = m_optimizeIterationsLeft;
-
-        // The first duplicated data remover has the special mission of removing
-        // duplicated data (meshes) that may have been specified in the source
-        // data, before we make it diverge because of different uses, like layout
-        // creation
-        UE_LOG(LogMutableCore, Verbose, TEXT(" - duplicated data remover"));
-        DuplicatedDataRemoverAST( roots );
+			// Image size operations are treated separately
+			UE_LOG(LogMutableCore, Verbose, TEXT(" - size optimiser"));
+			modified |= SizeOptimiserAST( roots );
+			//AXE_INT_VALUE("Mutable", Verbose, "ast size", (int64_t)ASTOp::CountNodes(roots));
+
+			// We cannot run the logic optimiser here because the constant memory explodes.
+			// TODO: Conservative logic optimiser
+			if (!modified)
+			{
+	//                AXE_INT_VALUE("Mutable", Verbose, "program size", (int64_t)program.m_opAddress.Num());
+	//                UE_LOG(LogMutableCore, Verbose, " - logic optimiser");
+	//                LocalLogicOptimiser log;
+	//                modified |= log.Apply( program, (int)s );
+			}
+		}
+
+	//    ASTOp::LogHistogram(roots);
+
+		UE_LOG(LogMutableCore, Verbose, TEXT(" - duplicated code remover"));
+		DuplicatedCodeRemoverAST( roots );
 		UE_LOG(LogMutableCore, Verbose, TEXT("(int) %s : %ld"), TEXT("ast size"), int64(ASTOp::CountNodes(roots)));
 
-        ASTOp::LogHistogram(roots);
+		ASTOp::LogHistogram(roots);
 
-        UE_LOG(LogMutableCore, Verbose, TEXT(" - duplicated code remover"));
-        DuplicatedCodeRemoverAST( roots );
+		UE_LOG(LogMutableCore, Verbose, TEXT(" - duplicated data remover"));
+		DuplicatedDataRemoverAST( roots );
 		UE_LOG(LogMutableCore, Verbose, TEXT("(int) %s : %ld"), TEXT("ast size"), int64(ASTOp::CountNodes(roots)));
 
-        // Special optimization stages
-        if ( m_options->GetPrivate()->m_optimisationOptions.m_uniformizeSkeleton )
-        {
-            UE_LOG(LogMutableCore, Verbose, TEXT(" - skeleton cleaner"));
-            ASTOp::LogHistogram(roots);
+		ASTOp::LogHistogram(roots);
 
-            SkeletonCleanerAST( roots, m_options->GetPrivate()->m_optimisationOptions );
+		// Generate constants
+		for ( auto& r: roots )
+		{
 			UE_LOG(LogMutableCore, Verbose, TEXT("(int) %s : %ld"), TEXT("ast size"), int64(ASTOp::CountNodes(roots)));
-			ASTOp::LogHistogram(roots);
-        }
+			UE_LOG(LogMutableCore, Verbose, TEXT(" - constant generator"));
+
+			// Constant subtree generation
+			modified = ConstantGeneratorAST( m_options->GetPrivate(), r, pTaskManager );
+		}
 
-        // First optimisation stage. It tries to resolve all the image sizes. This is necessary
-        // because some operations cannot be applied correctly until the image size is known
-        // like the grow-map generation.
-        bool modified = true;
-        int numIterations = 0;
-        while (modified && (!m_optimizeIterationsMax || (m_optimizeIterationsLeft>0) || !numIterations) )
-        {
-            MUTABLE_CPUPROFILER_SCOPE(FirstStage);
+		UE_LOG(LogMutableCore, Verbose, TEXT("(int) %s : %ld"), TEXT("ast size"), int64(ASTOp::CountNodes(roots)));
 
-            --m_optimizeIterationsLeft;
-            ++numIterations;
-            UE_LOG(LogMutableCore, Verbose, TEXT("First optimise iteration %d, max %d, left %d"),
-                    numIterations, m_optimizeIterationsMax, m_optimizeIterationsLeft);
+	//    ASTOp::LogHistogram(roots);
 
-            modified = false;
+		// \todo: this code deduplication here, takes very long and doesn't do much... investigate!
+	//    UE_LOG(LogMutableCore, Verbose, " - duplicated code remover");
+	//    DuplicatedCodeRemoverAST( roots );
+	//    AXE_INT_VALUE("Mutable", Verbose, "ast size", (int64_t)ASTOp::CountNodes(roots));
 
-            UE_LOG(LogMutableCore, Verbose, TEXT(" - size optimiser"));
-            modified |= SizeOptimiserAST( roots );
-        }
+		ASTOp::LogHistogram(roots);
 
+		UE_LOG(LogMutableCore, Verbose, TEXT(" - duplicated data remover"));
+		DuplicatedDataRemoverAST( roots );
 		UE_LOG(LogMutableCore, Verbose, TEXT("(int) %s : %ld"), TEXT("ast size"), int64(ASTOp::CountNodes(roots)));
 
-        // Main optimisation stage
-        {
-            MUTABLE_CPUPROFILER_SCOPE(MainStage);
-            FullOptimiseAST( roots, pTaskManager );
-			UE_LOG(LogMutableCore, Verbose, TEXT("(int) %s : %ld"), TEXT("ast size"), int64(ASTOp::CountNodes(roots)));
+		//if (!modified)
+		{
+			//for ( size_t s=0;  s<program.m_states.size(); ++s )
+			{
+				ASTOp::LogHistogram(roots);
+				UE_LOG(LogMutableCore, Verbose, TEXT(" - logic optimiser"));
+				modified |= LocalLogicOptimiserAST( roots );
+				ASTOp::LogHistogram(roots);
+			}
+
+			//AXE_INT_VALUE("Mutable", Verbose, "ast size", (int64_t)ASTOp::CountNodes(roots));
+		}
+
+		ASTOp::LogHistogram(roots);
+	}
+
+
+	//---------------------------------------------------------------------------------------------
+	// The state represents if there is a parent operation requiring skeleton for current mesh subtree.
+	class CollectAllMeshesForSkeletonVisitorAST : public Visitor_TopDown_Unique_Const<uint8_t>
+	{
+	public:
+
+		CollectAllMeshesForSkeletonVisitorAST( const ASTOpList& roots  )
+		{
+			Traverse( roots, false );
+		}
+
+		// List of meshes that require a skeleton
+		vector<mu::Ptr<ASTOpConstantResource>> m_meshesRequiringSkeleton;
+
+	private:
+
+		// Visitor_TopDown_Unique_Const<uint8_t> interface
+		bool Visit( const mu::Ptr<ASTOp>& node ) override
+		{
+			// \todo: refine to avoid instruction branches with irrelevant skeletons.
+
+			uint8_t currentProtected = GetCurrentState();
+
+			switch (node->GetOpType())
+			{
+
+			case OP_TYPE::ME_CONSTANT:
+			{
+				mu::Ptr<ASTOpConstantResource> typedOp = dynamic_cast<ASTOpConstantResource*>(node.get());
+				check(typedOp);
+
+				if (currentProtected)
+				{
+					if ( std::find(m_meshesRequiringSkeleton.begin(), m_meshesRequiringSkeleton.end(), typedOp)
+						 ==
+						 m_meshesRequiringSkeleton.end() )
+					{
+						m_meshesRequiringSkeleton.push_back(typedOp);
+					}
+				}
+
+				return false;
+			}
+
+			case OP_TYPE::ME_CLIPMORPHPLANE:
+			{
+				auto typedOp = dynamic_cast<ASTOpMeshClipMorphPlane*>(node.get());
+				if (typedOp)
+				{
+					if (typedOp->vertexSelectionType
+							==
+							OP::MeshClipMorphPlaneArgs::VS_BONE_HIERARCHY)
+					{
+						// We need the skeleton for the source mesh
+						RecurseWithState( typedOp->source.child(), true );
+						return false;
+					}
+				}
+
+				return true;
+			}
+
+			case OP_TYPE::ME_APPLYPOSE:
+			{
+				auto typedOp = dynamic_cast<ASTOpMeshApplyPose*>(node.get());
+				check(typedOp);
+
+				if (typedOp)
+				{
+					// We need the skeleton for both meshes
+					RecurseWithState(typedOp->base.child(), true);
+					RecurseWithState(typedOp->pose.child(), true);
+					return false;
+				}
+
+				break;
+			}
+
+			case OP_TYPE::ME_BINDSHAPE:
+			{
+				auto typedOp = dynamic_cast<ASTOpMeshBindShape*>(node.get());
+				check(typedOp);
+
+				if (typedOp && typedOp->m_reshapeSkeleton)
+				{
+					RecurseWithState(typedOp->Mesh.child(), true);
+					return false;
+				}
+
+				break;
+			}
+
+			case OP_TYPE::ME_APPLYSHAPE:
+			{
+				auto typedOp = dynamic_cast<ASTOpMeshApplyShape*>(node.get());
+				check(typedOp);
+
+				if (typedOp && typedOp->m_reshapeSkeleton)
+				{
+					RecurseWithState(typedOp->Mesh.child(), true);
+					return false;
+				}
+
+				break;
+			}
+
+			default:
+				break;
+			}
+
+			return true;
+		}
+
+	};
+
+
+	//---------------------------------------------------------------------------------------------
+	// This stores an ADD_MESH op with the child meshes collected and the final skeleton to use
+	// for this op.
+	struct ADDMESH_SKELETON
+	{
+		mu::Ptr<ASTOp> m_pAddMeshOp;
+		TArray<mu::Ptr<ASTOpConstantResource>> m_contributingMeshes;
+		mu::Ptr<Skeleton> m_pFinalSkeleton;
+
+		ADDMESH_SKELETON( const mu::Ptr<ASTOp>& pAddMeshOp,
+						  TArray<mu::Ptr<ASTOpConstantResource>>& contributingMeshes,
+						  const mu::Ptr<Skeleton>& pFinalSkeleton )
+		{
+			m_pAddMeshOp = pAddMeshOp;
+			m_contributingMeshes = std::move(contributingMeshes);
+			m_pFinalSkeleton = pFinalSkeleton;
+		}
+	};
+
+
+	//---------------------------------------------------------------------------------------------
+	void SkeletonCleanerAST( TArray<mu::Ptr<ASTOp>>& roots, const MODEL_OPTIMIZATION_OPTIONS& options )
+	{
+		// This collects all the meshes that require a skeleton because they are used in operations
+		// that require it.
+		CollectAllMeshesForSkeletonVisitorAST requireSkeletonCollector( roots );
+
+		TArray<ADDMESH_SKELETON> replacementsFound;
+
+		ASTOp::Traverse_TopDown_Unique_Imprecise( roots, [&](mu::Ptr<ASTOp>& at )
+		{
+			// Only recurse instance construction ops.
+			bool processChildren = GetOpDataType(at->GetOpType())==DT_INSTANCE;
+
+			if ( at->GetOpType() == OP_TYPE::IN_ADDMESH )
+			{
+				auto typedNode = dynamic_cast<ASTOpInstanceAdd*>(at.get());
+				mu::Ptr<ASTOp> meshRoot = typedNode->value.child();
+
+				if (meshRoot)
+				{
+					// Gather constant meshes contributing to the final mesh
+					TArray<mu::Ptr<ASTOpConstantResource>> subtreeMeshes;
+					TArray<mu::Ptr<ASTOp>> tempRoots;
+					tempRoots.Add(meshRoot);
+					ASTOp::Traverse_TopDown_Unique_Imprecise( tempRoots, [&](mu::Ptr<ASTOp>& lat )
+					{
+						// \todo: refine to avoid instruction branches with irrelevant skeletons.
+						if ( lat->GetOpType() == OP_TYPE::ME_CONSTANT )
+						{
+							mu::Ptr<ASTOpConstantResource> typedOp = dynamic_cast<ASTOpConstantResource*>(lat.get());
+							check(typedOp);
+
+							if ( subtreeMeshes.Find(typedOp)
+								 ==
+								 INDEX_NONE )
+							{
+								subtreeMeshes.Add(typedOp);
+							}
+						}
+						return true;
+					});
+
+					// Create a mesh just with the unified skeleton
+					SkeletonPtr pFinalSkeleton = new Skeleton;
+					for (const auto& meshAt: subtreeMeshes)
+					{
+						MeshPtrConst pMesh = static_cast<const Mesh*>(meshAt->GetValue().get());
+						mu::Ptr<const Skeleton> sourceSkeleton = pMesh ? pMesh->GetSkeleton() : nullptr;
+						if (sourceSkeleton)
+						{
+							ExtendSkeleton(pFinalSkeleton.get(),sourceSkeleton.get());
+						}
+					}
+
+					replacementsFound.Emplace( at, subtreeMeshes, pFinalSkeleton );
+				}
+			}
+
+			return processChildren;
+		});
+
+
+		// Iterate all meshes again
+		ASTOp::Traverse_TopDown_Unique_Imprecise( roots, [&](mu::Ptr<ASTOp>& at )
+		{
+			if (at->GetOpType()==OP_TYPE::ME_CONSTANT)
+			{
+				auto typedOp = dynamic_cast<ASTOpConstantResource*>(at.get());
+				check(typedOp);
+				if (!typedOp)
+				{
+					return true;
+				}
+
+				for( auto& rep: replacementsFound )
+				{
+					if ( rep.m_contributingMeshes.Contains(at) )
+					{
+						mu::Ptr<const Mesh> pMesh =  static_cast<const Mesh*>(typedOp->GetValue().get());
+						pMesh->CheckIntegrity();
+
+						mu::Ptr<Mesh> pNewMesh = MeshRemapSkeleton( pMesh.get(),
+																rep.m_pFinalSkeleton.get());
+
+						// It returns null if there is no need to remap.
+						if (pNewMesh)
+						{
+							pNewMesh->CheckIntegrity();
+							mu::Ptr<ASTOpConstantResource> newOp = new ASTOpConstantResource();
+							newOp->type = OP_TYPE::ME_CONSTANT;
+							newOp->SetValue( pNewMesh, options.m_useDiskCache );
+
+							ASTOp::Replace( at, newOp );
+						}
+					}
+				}
+			}
+			return true;
+		});
+
+	//    for( auto& rep: replacementsFound )
+	//    {
+	//        check(rep.m_pAddMeshOp->GetOpType()==OP_TYPE::IN_ADDMESH);
+	//        auto newAddMeshRootTyped = dynamic_cast<ASTOpInstanceAdd*>(rep.m_pAddMeshOp.get());
+	//        auto meshRoot = newAddMeshRootTyped->value.child();
+
+	//        // Add new skeleton constant instruction
+	//        Ptr<ASTOpConstantResource> cop = new ASTOpConstantResource();
+	//        cop->type = OP_TYPE::ME_CONSTANT;
+	//        cop->SetValue( rep.m_pFinalSkeleton, options.m_useDiskCache );
+
+	//        // Add new "apply skeleton" instruction
+	//        Ptr<ASTOpFixed> skop = new ASTOpFixed();
+	//        skop->op.type = OP_TYPE::ME_SETSKELETON;
+	//        skop->SetChild( skop->op.args.MeshSetSkeleton.source, meshRoot);
+	//        skop->SetChild( skop->op.args.MeshSetSkeleton.skeleton, cop);
+
+	//        // We need a new mesh add instruction, it is not safe to edit the current one.
+	//        auto newOp = mu::Clone<ASTOpInstanceAdd>(newAddMeshRootTyped);
+	//        newOp->value = skop;
+
+	//        // TODO BLEH UNSAFE?
+	//        ASTOp::Replace(rep.m_pAddMeshOp,newOp);
+	//    }
+	}
+
+
+	//---------------------------------------------------------------------------------------------
+	void CodeOptimiser::OptimiseAST( TaskManager* pTaskManager )
+	{
+		MUTABLE_CPUPROFILER_SCOPE(OptimiseAST);
+
+		// Gather all the roots
+		TArray<Ptr<ASTOp>> roots;
+		for(const auto& s:m_states)
+		{
+			roots.Add(s.root);
 		}
 
-//            // Third optimisation stage: introduce image compose conditions
-//            {
-//                // TODO: We have disable this since the new operation "layout remove block" is used.
-//                // Proper support needs to be implemented for this case.
-//    //                MUTABLE_CPUPROFILER_SCOPE(ImageComposeConditionGenerator);
-//    //                AXE_INT_VALUE("Mutable", Verbose, "program size", (int64_t)program.m_opAddress.Num());
-//    //                UE_LOG(LogMutableCore, Verbose, " - image compose conditions generator");
-//    //                ImageComposeConditionGenerator composeGen( program );
-//            }
-
-
-        {
-            MUTABLE_CPUPROFILER_SCOPE(ConditionalsStage);
-
-            //program.Check();
-
-            // Optimise the conditionals added by the image compose condition generator
-            // also the constants generator at the end of previous fulloptimiseAST.
-            FullOptimiseAST( roots, pTaskManager );
-        }
-
-        // Analyse mesh constants to see which of them are in optimised mesh formats, and set the flags.
-        ASTOp::Traverse_BottomUp_Unique_NonReentrant( roots, [&](Ptr<ASTOp>& n)
-        {
-            if (n->GetOpType()==OP_TYPE::ME_CONSTANT)
-            {
-                auto typed = dynamic_cast<ASTOpConstantResource*>(n.get());
-                auto pMesh = static_cast<const Mesh*>(typed->GetValue().get());
-                pMesh->ResetStaticFormatFlags();
-                typed->SetValue( pMesh,
-                                 m_options->GetPrivate()->m_optimisationOptions.m_useDiskCache );
-            }
-        });
-
-        // Make sure we didn't lose track of pointers
-        for ( size_t s=0;  s<m_states.size(); ++s )
-        {
-            check( roots.Contains( m_states[s].root ) );
-        }
-
-        ASTOp::LogHistogram(roots);
-
-        {
-            MUTABLE_CPUPROFILER_SCOPE(StatesStage);
-
-            // Optimise for every state
-            OptimiseStatesAST( pTaskManager );
-
-            // Optimise the data formats (TODO)
-            //OperationFlagGenerator flagGen( pResult.get() );
-        }
-
-        ASTOp::LogHistogram(roots);
-    }
-
-//        // Minimal optimisation of constant subtrees
-    else if ( m_options->GetPrivate()->m_optimisationOptions.m_constReduction )
-    {
-        // The first duplicated data remover has the special mission of removing
-        // duplicated data (meshes) that may have been specified in the source
-        // data, before we make it diverge because of different uses, like layout
-        // creation
-        UE_LOG(LogMutableCore, Verbose, TEXT(" - duplicated data remover"));
-        DuplicatedDataRemoverAST( roots );
-        //AXE_INT_VALUE("Mutable", Verbose, "ast size", (int64_t)ASTOp::CountNodes(roots));
-
-        UE_LOG(LogMutableCore, Verbose, TEXT(" - duplicated code remover"));
-        DuplicatedCodeRemoverAST( roots );
 		UE_LOG(LogMutableCore, Verbose, TEXT("(int) %s : %ld"), TEXT("ast size"), int64(ASTOp::CountNodes(roots)));
 
-        for ( size_t s=0;  s<m_states.size(); ++s )
-        {
-            UE_LOG(LogMutableCore, Verbose, TEXT(" - constant generator"));
-            ConstantGeneratorAST( m_options->GetPrivate(), m_states[s].root, pTaskManager );
-            //AXE_INT_VALUE("Mutable", Verbose, "ast size", (int64_t)ASTOp::CountNodes(roots));
-        }
+		if ( m_options->GetPrivate()->m_optimisationOptions.m_enabled )
+		{
+			// We use 4 times the count because at the time we moved to sharing this count it
+			// was being used 4 times, and we want to keep the tests consistent.
+			m_optimizeIterationsLeft =
+				m_options->GetPrivate()->m_optimisationOptions.m_maxOptimisationLoopCount * 4;
+			m_optimizeIterationsMax = m_optimizeIterationsLeft;
+
+			// The first duplicated data remover has the special mission of removing
+			// duplicated data (meshes) that may have been specified in the source
+			// data, before we make it diverge because of different uses, like layout
+			// creation
+			UE_LOG(LogMutableCore, Verbose, TEXT(" - duplicated data remover"));
+			DuplicatedDataRemoverAST( roots );
+			UE_LOG(LogMutableCore, Verbose, TEXT("(int) %s : %ld"), TEXT("ast size"), int64(ASTOp::CountNodes(roots)));
 
-        UE_LOG(LogMutableCore, Verbose, TEXT(" - duplicated data remover"));
-        DuplicatedDataRemoverAST( roots );
-        //AXE_INT_VALUE("Mutable", Verbose, "ast size", (int64_t)ASTOp::CountNodes(roots));
+			ASTOp::LogHistogram(roots);
 
-        UE_LOG(LogMutableCore, Verbose, TEXT(" - duplicated code remover"));
-        DuplicatedCodeRemoverAST( roots );
-        //AXE_INT_VALUE("Mutable", Verbose, "ast size", (int64_t)ASTOp::CountNodes(roots));
+			UE_LOG(LogMutableCore, Verbose, TEXT(" - duplicated code remover"));
+			DuplicatedCodeRemoverAST( roots );
+			UE_LOG(LogMutableCore, Verbose, TEXT("(int) %s : %ld"), TEXT("ast size"), int64(ASTOp::CountNodes(roots)));
 
-        // Make sure we didn't lose track of pointers
-        for ( size_t s=0;  s<m_states.size(); ++s )
-        {
-            check( roots.Contains( m_states[s].root ) );
-        }
-    }
+			// Special optimization stages
+			if ( m_options->GetPrivate()->m_optimisationOptions.m_uniformizeSkeleton )
+			{
+				UE_LOG(LogMutableCore, Verbose, TEXT(" - skeleton cleaner"));
+				ASTOp::LogHistogram(roots);
 
-    ASTOp::LogHistogram(roots);
+				SkeletonCleanerAST( roots, m_options->GetPrivate()->m_optimisationOptions );
+				UE_LOG(LogMutableCore, Verbose, TEXT("(int) %s : %ld"), TEXT("ast size"), int64(ASTOp::CountNodes(roots)));
+				ASTOp::LogHistogram(roots);
+			}
+
+			// First optimisation stage. It tries to resolve all the image sizes. This is necessary
+			// because some operations cannot be applied correctly until the image size is known
+			// like the grow-map generation.
+			bool modified = true;
+			int numIterations = 0;
+			while (modified && (!m_optimizeIterationsMax || (m_optimizeIterationsLeft>0) || !numIterations) )
+			{
+				MUTABLE_CPUPROFILER_SCOPE(FirstStage);
+
+				--m_optimizeIterationsLeft;
+				++numIterations;
+				UE_LOG(LogMutableCore, Verbose, TEXT("First optimise iteration %d, max %d, left %d"),
+						numIterations, m_optimizeIterationsMax, m_optimizeIterationsLeft);
+
+				modified = false;
+
+				UE_LOG(LogMutableCore, Verbose, TEXT(" - size optimiser"));
+				modified |= SizeOptimiserAST( roots );
+			}
+
+			UE_LOG(LogMutableCore, Verbose, TEXT("(int) %s : %ld"), TEXT("ast size"), int64(ASTOp::CountNodes(roots)));
+
+			// Main optimisation stage
+			{
+				MUTABLE_CPUPROFILER_SCOPE(MainStage);
+				FullOptimiseAST( roots, pTaskManager );
+				UE_LOG(LogMutableCore, Verbose, TEXT("(int) %s : %ld"), TEXT("ast size"), int64(ASTOp::CountNodes(roots)));
+			}
+
+	//            // Third optimisation stage: introduce image compose conditions
+	//            {
+	//                // TODO: We have disable this since the new operation "layout remove block" is used.
+	//                // Proper support needs to be implemented for this case.
+	//    //                MUTABLE_CPUPROFILER_SCOPE(ImageComposeConditionGenerator);
+	//    //                AXE_INT_VALUE("Mutable", Verbose, "program size", (int64_t)program.m_opAddress.Num());
+	//    //                UE_LOG(LogMutableCore, Verbose, " - image compose conditions generator");
+	//    //                ImageComposeConditionGenerator composeGen( program );
+	//            }
 
-}
 
+			{
+				MUTABLE_CPUPROFILER_SCOPE(ConditionalsStage);
+
+				//program.Check();
+
+				// Optimise the conditionals added by the image compose condition generator
+				// also the constants generator at the end of previous fulloptimiseAST.
+				FullOptimiseAST( roots, pTaskManager );
+			}
+
+			// Analyse mesh constants to see which of them are in optimised mesh formats, and set the flags.
+			ASTOp::Traverse_BottomUp_Unique_NonReentrant( roots, [&](Ptr<ASTOp>& n)
+			{
+				if (n->GetOpType()==OP_TYPE::ME_CONSTANT)
+				{
+					auto typed = dynamic_cast<ASTOpConstantResource*>(n.get());
+					auto pMesh = static_cast<const Mesh*>(typed->GetValue().get());
+					pMesh->ResetStaticFormatFlags();
+					typed->SetValue( pMesh,
+									 m_options->GetPrivate()->m_optimisationOptions.m_useDiskCache );
+				}
+			});
+
+			// Make sure we didn't lose track of pointers
+			for ( size_t s=0;  s<m_states.size(); ++s )
+			{
+				check( roots.Contains( m_states[s].root ) );
+			}
+
+			ASTOp::LogHistogram(roots);
+
+			{
+				MUTABLE_CPUPROFILER_SCOPE(StatesStage);
+
+				// Optimise for every state
+				OptimiseStatesAST( pTaskManager );
+
+				// Optimise the data formats (TODO)
+				//OperationFlagGenerator flagGen( pResult.get() );
+			}
+
+			ASTOp::LogHistogram(roots);
+		}
+
+	//        // Minimal optimisation of constant subtrees
+		else if ( m_options->GetPrivate()->m_optimisationOptions.m_constReduction )
+		{
+			// The first duplicated data remover has the special mission of removing
+			// duplicated data (meshes) that may have been specified in the source
+			// data, before we make it diverge because of different uses, like layout
+			// creation
+			UE_LOG(LogMutableCore, Verbose, TEXT(" - duplicated data remover"));
+			DuplicatedDataRemoverAST( roots );
+			//AXE_INT_VALUE("Mutable", Verbose, "ast size", (int64_t)ASTOp::CountNodes(roots));
+
+			UE_LOG(LogMutableCore, Verbose, TEXT(" - duplicated code remover"));
+			DuplicatedCodeRemoverAST( roots );
+			UE_LOG(LogMutableCore, Verbose, TEXT("(int) %s : %ld"), TEXT("ast size"), int64(ASTOp::CountNodes(roots)));
+
+			for ( size_t s=0;  s<m_states.size(); ++s )
+			{
+				UE_LOG(LogMutableCore, Verbose, TEXT(" - constant generator"));
+				ConstantGeneratorAST( m_options->GetPrivate(), m_states[s].root, pTaskManager );
+				//AXE_INT_VALUE("Mutable", Verbose, "ast size", (int64_t)ASTOp::CountNodes(roots));
+			}
+
+			UE_LOG(LogMutableCore, Verbose, TEXT(" - duplicated data remover"));
+			DuplicatedDataRemoverAST( roots );
+			//AXE_INT_VALUE("Mutable", Verbose, "ast size", (int64_t)ASTOp::CountNodes(roots));
+
+			UE_LOG(LogMutableCore, Verbose, TEXT(" - duplicated code remover"));
+			DuplicatedCodeRemoverAST( roots );
+			//AXE_INT_VALUE("Mutable", Verbose, "ast size", (int64_t)ASTOp::CountNodes(roots));
+
+			// Make sure we didn't lose track of pointers
+			for ( size_t s=0;  s<m_states.size(); ++s )
+			{
+				check( roots.Contains( m_states[s].root ) );
+			}
+		}
+
+		ASTOp::LogHistogram(roots);
+
+	}
+
+}
diff --git a/Engine/Plugins/Experimental/Mutable/Source/MutableTools/Private/MuT/ErrorLog.cpp b/Engine/Plugins/Experimental/Mutable/Source/MutableTools/Private/MuT/ErrorLog.cpp
--- a/Engine/Plugins/Experimental/Mutable/Source/MutableTools/Private/MuT/ErrorLog.cpp
+++ b/Engine/Plugins/Experimental/Mutable/Source/MutableTools/Private/MuT/ErrorLog.cpp
@@ -64,16 +64,17 @@ namespace mu
 
 
 	//---------------------------------------------------------------------------------------------
-	const char* ErrorLog::GetMessageText( int index ) const
+	const FString& ErrorLog::GetMessageText( int index ) const
 	{
 		const char* strResult = "";
 
 		if ( index >=0 && index<(int)m_pD->m_messages.Num() )
 		{
-			strResult = m_pD->m_messages[index].m_text.c_str();
+			return m_pD->m_messages[index].m_text;
 		}
 
-		return strResult;
+		static FString Empty;
+		return Empty;
 	}
 
 
@@ -124,15 +125,24 @@ namespace mu
 	}
 
 	//---------------------------------------------------------------------------------------------
-	void ErrorLog::Private::Add( const char* strMessage, ErrorLogMessageType type,
-								 const void* context )
+	void ErrorLog::Private::Add(const char* strMessage, ErrorLogMessageType type,
+		const void* context)
 	{
-		m_messages.Add( MSG() );
+		m_messages.Add(MSG());
 		m_messages.Last().m_type = type;
 		m_messages.Last().m_text = strMessage;
 		m_messages.Last().m_context = context;
 	}
 
+	//---------------------------------------------------------------------------------------------
+	void ErrorLog::Private::Add(const FString& InMessage, ErrorLogMessageType InType, const void* InContext)
+	{
+		m_messages.Add(MSG());
+		m_messages.Last().m_type = InType;
+		m_messages.Last().m_text = InMessage;
+		m_messages.Last().m_context = InContext;
+	}
+
 	//---------------------------------------------------------------------------------------------
 	void ErrorLog::Private::Add( const char* strMessage, 
                                  const ErrorLogMessageAttachedDataView& dataView,
@@ -160,10 +170,10 @@ namespace mu
 		{
 			switch ( msg.m_type )
 			{
-			case ELMT_ERROR: 	UE_LOG(LogMutableCore, Log, TEXT("  ERR  %s\n"), msg.m_text.c_str()); break;
-			case ELMT_WARNING: 	UE_LOG(LogMutableCore, Log, TEXT("  WRN  %s\n"), msg.m_text.c_str()); break;
-			case ELMT_INFO: 	UE_LOG(LogMutableCore, Log, TEXT("  INF  %s\n"), msg.m_text.c_str()); break;
-			default: 			UE_LOG(LogMutableCore, Log, TEXT("  NON  %s\n"), msg.m_text.c_str()); break;
+			case ELMT_ERROR: 	UE_LOG(LogMutableCore, Log, TEXT("  ERR  %s\n"), *msg.m_text); break;
+			case ELMT_WARNING: 	UE_LOG(LogMutableCore, Log, TEXT("  WRN  %s\n"), *msg.m_text); break;
+			case ELMT_INFO: 	UE_LOG(LogMutableCore, Log, TEXT("  INF  %s\n"), *msg.m_text); break;
+			default: 			UE_LOG(LogMutableCore, Log, TEXT("  NON  %s\n"), *msg.m_text); break;
 			}
 		}
 	}
diff --git a/Engine/Plugins/Experimental/Mutable/Source/MutableTools/Private/MuT/ErrorLogPrivate.h b/Engine/Plugins/Experimental/Mutable/Source/MutableTools/Private/MuT/ErrorLogPrivate.h
--- a/Engine/Plugins/Experimental/Mutable/Source/MutableTools/Private/MuT/ErrorLogPrivate.h
+++ b/Engine/Plugins/Experimental/Mutable/Source/MutableTools/Private/MuT/ErrorLogPrivate.h
@@ -37,7 +37,7 @@ namespace mu
 			}
 
 			ErrorLogMessageType m_type;
-			string m_text;
+			FString m_text;
             std::shared_ptr< DATA > m_data;
 			const void* m_context;
 		};
@@ -48,7 +48,10 @@ namespace mu
 		//-----------------------------------------------------------------------------------------
 
 		//!
-		void Add( const char* strMessage, ErrorLogMessageType type, const void* context );
+		void Add(const FString& Message, ErrorLogMessageType Type, const void* Context);
+
+		//!
+		void Add(const char* strMessage, ErrorLogMessageType type, const void* context);
 
         //!
         void Add( const char* strMessage, const ErrorLogMessageAttachedDataView& data, 
diff --git a/Engine/Plugins/Experimental/Mutable/Source/MutableTools/Private/MuT/NodeMeshMorph.cpp b/Engine/Plugins/Experimental/Mutable/Source/MutableTools/Private/MuT/NodeMeshMorph.cpp
--- a/Engine/Plugins/Experimental/Mutable/Source/MutableTools/Private/MuT/NodeMeshMorph.cpp
+++ b/Engine/Plugins/Experimental/Mutable/Source/MutableTools/Private/MuT/NodeMeshMorph.cpp
@@ -150,10 +150,10 @@ namespace mu
 	}
 
 	//---------------------------------------------------------------------------------------------
-	void NodeMeshMorph::SetMorphIndicesAreRelative( bool relative )
-	{
-		m_pD->m_vertexIndicesAreRelative = relative;
-	}
+	//void NodeMeshMorph::SetMorphIndicesAreRelative( bool relative )
+	//{
+	//	m_pD->m_vertexIndicesAreRelative = relative;
+	//}
 
 	//---------------------------------------------------------------------------------------------
 	void NodeMeshMorph::SetReshapeSkeleton(bool bEnable)
diff --git a/Engine/Plugins/Experimental/Mutable/Source/MutableTools/Private/MuT/NodeMeshMorphPrivate.h b/Engine/Plugins/Experimental/Mutable/Source/MutableTools/Private/MuT/NodeMeshMorphPrivate.h
--- a/Engine/Plugins/Experimental/Mutable/Source/MutableTools/Private/MuT/NodeMeshMorphPrivate.h
+++ b/Engine/Plugins/Experimental/Mutable/Source/MutableTools/Private/MuT/NodeMeshMorphPrivate.h
@@ -24,7 +24,7 @@ namespace mu
 		NodeScalarPtr m_pFactor;
 		NodeMeshPtr m_pBase;
 		TArray<NodeMeshPtr> m_morphs;
-        bool m_vertexIndicesAreRelative = false;
+        bool m_vertexIndicesAreRelative_Deprecated = false;
 		
 		bool m_reshapeSkeleton = false;
 		bool m_reshapePhysicsVolumes = false;
@@ -44,7 +44,7 @@ namespace mu
 			arch << m_pFactor;
 			arch << m_pBase;
 			arch << m_morphs;
-            arch << m_vertexIndicesAreRelative;
+            arch << m_vertexIndicesAreRelative_Deprecated;
 
 			arch << m_reshapeSkeleton;
 			arch << m_reshapePhysicsVolumes;
@@ -67,7 +67,7 @@ namespace mu
 			arch >> m_pFactor;
 			arch >> m_pBase;
 			arch >> m_morphs;
-            arch >> m_vertexIndicesAreRelative;
+            arch >> m_vertexIndicesAreRelative_Deprecated;
 	
 			if (ver >= 2)
 			{
diff --git a/Engine/Plugins/Experimental/Mutable/Source/MutableTools/Private/MuT/TypeInfo.cpp b/Engine/Plugins/Experimental/Mutable/Source/MutableTools/Private/MuT/TypeInfo.cpp
--- a/Engine/Plugins/Experimental/Mutable/Source/MutableTools/Private/MuT/TypeInfo.cpp
+++ b/Engine/Plugins/Experimental/Mutable/Source/MutableTools/Private/MuT/TypeInfo.cpp
@@ -5,127 +5,129 @@
 #include "MuR/ParametersPrivate.h"
 #include "MuR/Types.h"
 
-using namespace mu;
-
-
-const char* TypeInfo::s_imageFormatName[ size_t(EImageFormat::IF_COUNT) ] =
-{
-	"None",
-	"RGB U888",
-	"RGBA U8888",
-	"L U8",
-
-    "PVRTC2",
-    "PVRTC4",
-    "ETC1",
-    "ETC2",
-
-	"L U8 RLE",
-	"RGB U888 RLE",
-	"RGBA U8888 RLE",
-	"L U1 RLE",
-
-    "BC1",
-    "BC2",
-    "BC3",
-    "BC4",
-    "BC5",
-    "BC6",
-    "BC7",
-
-    "BGRA U8888",
-
-    "ASTC_4x4_RGB_LDR",
-    "ASTC_4x4_RGBA_LDR",
-    "ASTC_4x4_RG_LDR",
-
-};
-
-
-const char* TypeInfo::s_meshBufferSemanticName[ MBS_COUNT ] =
+namespace mu
 {
-	"None",
-
-	"VertexIndex",
-
-	"Position",
-	"Normal",
-	"Tangent",
-	"Binormal",
-	"Tex Coords",
-	"Colour",
-	"Bone Weights",
-	"Bone Indices",
-
-	"Layout Block",
 
-	"Chart",
-
-    "Other",
-
-    "Tangent Space Sign",
-
-	"TriangleIndex",
-	"BarycentricCoords",
-	"Distance"
-
-};
-
-const char* TypeInfo::s_meshBufferFormatName[ MBF_COUNT ] =
-{
-	"None",
-	"Float16",
-	"Float32",
-
-	"UInt8",
-	"UInt16",
-	"UInt32",
-	"Int8",
-	"Int16",
-	"Int32",
-
-	"NUInt8",
-	"NUInt16",
-	"NUInt32",
-	"NInt8",
-	"NInt16",
-	"NInt32",
-
-    "PackedDir8",
-    "PackedDir8_WTangentSign",
-    "PackedDirS8",
-    "PackedDirS8_WTangentSign",
-
-	"Float64",
-};
-
-const char* TypeInfo::s_projectorTypeName[  static_cast<uint32>(mu::PROJECTOR_TYPE::COUNT) ] =
-{
-	"Planar",
-	"Cylindrical",
-	"Wrapping",
-};
-
-const char* TypeInfo::s_curveInterpolationModeName [ static_cast<uint8_t>(mu::CurveKeyFrame::InterpMode::Count) ] =
-{
-	"Linear",
-	"Constant",
-	"Cubic",
-	"None",
-};
-
-const char* TypeInfo::s_curveTangentModeName [static_cast<uint8_t> (mu::CurveKeyFrame::TangentMode::Count)] =
-{
-	"Auto",
-	"Manual",
-	"Broken",
-	"None",
-};
-
-const char* TypeInfo::s_curveTangentWeightModeName [ static_cast<uint8_t> (mu::CurveKeyFrame::TangentWeightMode::Count)] =
-{
-	"Auto",
-	"Manual",
-	"Broken",
-	"None",
-};
\ No newline at end of file
+	const char* TypeInfo::s_imageFormatName[size_t(EImageFormat::IF_COUNT)] =
+	{
+		"None",
+		"RGB U888",
+		"RGBA U8888",
+		"L U8",
+
+		"PVRTC2",
+		"PVRTC4",
+		"ETC1",
+		"ETC2",
+
+		"L U8 RLE",
+		"RGB U888 RLE",
+		"RGBA U8888 RLE",
+		"L U1 RLE",
+
+		"BC1",
+		"BC2",
+		"BC3",
+		"BC4",
+		"BC5",
+		"BC6",
+		"BC7",
+
+		"BGRA U8888",
+
+		"ASTC_4x4_RGB_LDR",
+		"ASTC_4x4_RGBA_LDR",
+		"ASTC_4x4_RG_LDR",
+
+	};
+
+
+	const char* TypeInfo::s_meshBufferSemanticName[MBS_COUNT] =
+	{
+		"None",
+
+		"VertexIndex",
+
+		"Position",
+		"Normal",
+		"Tangent",
+		"Binormal",
+		"Tex Coords",
+		"Colour",
+		"Bone Weights",
+		"Bone Indices",
+
+		"Layout Block",
+
+		"Chart",
+
+		"Other",
+
+		"Tangent Space Sign",
+
+		"TriangleIndex",
+		"BarycentricCoords",
+		"Distance"
+
+	};
+
+	const char* TypeInfo::s_meshBufferFormatName[MBF_COUNT] =
+	{
+		"None",
+		"Float16",
+		"Float32",
+
+		"UInt8",
+		"UInt16",
+		"UInt32",
+		"Int8",
+		"Int16",
+		"Int32",
+
+		"NUInt8",
+		"NUInt16",
+		"NUInt32",
+		"NInt8",
+		"NInt16",
+		"NInt32",
+
+		"PackedDir8",
+		"PackedDir8_WTangentSign",
+		"PackedDirS8",
+		"PackedDirS8_WTangentSign",
+
+		"Float64",
+	};
+
+	const char* TypeInfo::s_projectorTypeName[static_cast<uint32>(mu::PROJECTOR_TYPE::COUNT)] =
+	{
+		"Planar",
+		"Cylindrical",
+		"Wrapping",
+	};
+
+	const char* TypeInfo::s_curveInterpolationModeName[static_cast<uint8_t>(mu::CurveKeyFrame::InterpMode::Count)] =
+	{
+		"Linear",
+		"Constant",
+		"Cubic",
+		"None",
+	};
+
+	const char* TypeInfo::s_curveTangentModeName[static_cast<uint8_t> (mu::CurveKeyFrame::TangentMode::Count)] =
+	{
+		"Auto",
+		"Manual",
+		"Broken",
+		"None",
+	};
+
+	const char* TypeInfo::s_curveTangentWeightModeName[static_cast<uint8_t> (mu::CurveKeyFrame::TangentWeightMode::Count)] =
+	{
+		"Auto",
+		"Manual",
+		"Broken",
+		"None",
+	};
+
+}
\ No newline at end of file
diff --git a/Engine/Plugins/Experimental/Mutable/Source/MutableTools/Public/MuT/ErrorLog.h b/Engine/Plugins/Experimental/Mutable/Source/MutableTools/Public/MuT/ErrorLog.h
--- a/Engine/Plugins/Experimental/Mutable/Source/MutableTools/Public/MuT/ErrorLog.h
+++ b/Engine/Plugins/Experimental/Mutable/Source/MutableTools/Public/MuT/ErrorLog.h
@@ -54,7 +54,7 @@ namespace mu
 
 		//! Get the text of a message.
 		//! \param index index of the message from 0 to GetMessageCount(ELMT_NONE)-1
-		const char* GetMessageText( int index ) const;
+		const FString& GetMessageText( int index ) const;
 
 		//! Get the opaque context of a message.
 		//! \param index index of the message from 0 to GetMessageCount(ELMT_NONE)-1
diff --git a/Engine/Plugins/Experimental/Mutable/Source/MutableTools/Public/MuT/NodeMeshMorph.h b/Engine/Plugins/Experimental/Mutable/Source/MutableTools/Public/MuT/NodeMeshMorph.h
--- a/Engine/Plugins/Experimental/Mutable/Source/MutableTools/Public/MuT/NodeMeshMorph.h
+++ b/Engine/Plugins/Experimental/Mutable/Source/MutableTools/Public/MuT/NodeMeshMorph.h
@@ -75,7 +75,7 @@ namespace mu
 		//! Set the domain of the vertex indices in the morph connections. By default they are 
 		//! absolute as generated by the MakeMeshMorph node. But if they are provided externally, 
 		//! they may be relative to the main mesh connected in the Base node.
-        void SetMorphIndicesAreRelative( bool relative );
+        //void SetMorphIndicesAreRelative( bool relative );
 
 		/** Also deform the mesh skeleton. Disabled by default. */
 		void SetReshapeSkeleton(bool);	
