[mutable] Added new BeginUpdateDelegate.

- Removed old unused UpdateBeginDelegate (was commented).
- Standerized public delegates:
    - Renamed BeginDestroyNativeDelegate.
    - Created dynamic and native versions of existing delegates.

diff --git a/Engine/Plugins/Experimental/Mutable/Source/CustomizableObject/Classes/MuCO/CustomizableObjectInstance.h b/Engine/Plugins/Experimental/Mutable/Source/CustomizableObject/Classes/MuCO/CustomizableObjectInstance.h
--- a/Engine/Plugins/Experimental/Mutable/Source/CustomizableObject/Classes/MuCO/CustomizableObjectInstance.h
+++ b/Engine/Plugins/Experimental/Mutable/Source/CustomizableObject/Classes/MuCO/CustomizableObjectInstance.h
@@ -100,15 +100,30 @@ enum class EUpdateResult : uint8
 };
 
 
-DECLARE_DELEGATE(FObjectInstanceUpdateBeginDelegate);
+/* When creating new delegates use the following conventions:
+ *
+ * - All delegates must be multicast.
+ * - If the delegate is exposed to the API create both, dynamic and native versions (non-dynamic).
+ * - Dynamic delegates should not be transient. Use the native version if you do not want it to be saved.
+ * - Native delegates names should end with "NativeDelegate".
+ * - Dynamic delegates broadcast before native delegates. */
+
+DECLARE_DYNAMIC_MULTICAST_DELEGATE_OneParam(FBeginUpdateDelegate, UCustomizableObjectInstance*, Instance);
+DECLARE_MULTICAST_DELEGATE_OneParam(FBeginUpdateNativeDelegate, UCustomizableObjectInstance*);
+
 DECLARE_DYNAMIC_MULTICAST_DELEGATE_OneParam(FObjectInstanceUpdatedDelegate, UCustomizableObjectInstance*, Instance);
 DECLARE_MULTICAST_DELEGATE_OneParam(FObjectInstanceUpdatedNativeDelegate, UCustomizableObjectInstance*);
-DECLARE_DELEGATE_OneParam( FProjectorStateChangedDelegate, FString );
-DECLARE_MULTICAST_DELEGATE_OneParam(FBeginDestroyDelegate, UCustomizableObjectInstance*);
+
+DECLARE_DYNAMIC_MULTICAST_DELEGATE_OneParam(FBeginDestroyDelegate, UCustomizableObjectInstance*, Instance);
+DECLARE_MULTICAST_DELEGATE_OneParam(FBeginDestroyNativeDelegate, UCustomizableObjectInstance*);
+
+DECLARE_DELEGATE_OneParam(FProjectorStateChangedDelegate, FString);
 
 DECLARE_DYNAMIC_DELEGATE_TwoParams(FEachComponentAnimInstanceClassDelegate, int32, SlotIndex, TSubclassOf<UAnimInstance>, AnimInstClass);
+
 DECLARE_DELEGATE_TwoParams(FEachComponentAnimInstanceClassNativeDelegate, int32 /*SlotIndex*/, TSubclassOf<UAnimInstance> /*AnimInstClass*/);
 
+
 UCLASS( Blueprintable, BlueprintType, HideCategories=(CustomizableObjectInstance) )
 class CUSTOMIZABLEOBJECT_API UCustomizableObjectInstance : public UObject
 {
@@ -124,28 +139,34 @@ public:
 	const FCustomizableObjectInstanceDescriptor& GetDescriptor() const;
 	
 	void SetDescriptor(const FCustomizableObjectInstanceDescriptor& InDescriptor);
-	
-	/** The generated skeletal meshes for this Instance, one for each component */
-	UPROPERTY(Transient, VisibleAnywhere, Category = CustomizableSkeletalMesh)
-	TArray< TObjectPtr<USkeletalMesh> > SkeletalMeshes;
 
-	/** Broadcasts when the Customizable Object Instance is updated. */
+	/** Broadcast at the beginning of an Instance update. */
+	UPROPERTY(BlueprintAssignable, Category = CustomizableObjectInstance)
+	FBeginUpdateDelegate BeginUpdateDelegate;
+
+	/** Broadcast at the beginning of an Instance update. */
+	FBeginUpdateNativeDelegate BeginUpdateNativeDelegate;
+
+	/** Broadcast when the Customizable Object Instance is updated. */
 	UPROPERTY(Transient, BlueprintAssignable, Category = CustomizableObjectInstance)
 	FObjectInstanceUpdatedDelegate UpdatedDelegate;
 
-	/** Non-dynamic version. */
+	/** Broadcast when the Customizable Object Instance is updated. */
 	FObjectInstanceUpdatedNativeDelegate UpdatedNativeDelegate;
 
-	// Delegate to notify the beginning of an update of this instance
-	// NOTE: this delegate can be called outside the main thread, make sure operations done in the
-	//       callbacks are safe, avoid using it outside editor functionality
-	FObjectInstanceUpdateBeginDelegate UpdateBeginDelegate;
+	/** Broadcast when UObject::BeginDestroy is being called. */	
+	UPROPERTY(BlueprintAssignable, Category = CustomizableObjectInstance)
+	FBeginDestroyDelegate BeginDestroyDelegate;
 
 	/** Broadcast when UObject::BeginDestroy is being called. */
-	FBeginDestroyDelegate BeginDestroyDelegate;
+	FBeginDestroyNativeDelegate BeginDestroyNativeDelegate;
 
 	TMap<FString, bool> ParamNameToExpandedMap; // Used to check whether a mutable param is expanded in the editor to show its child params
 
+	/** The generated skeletal meshes for this Instance, one for each component */
+	UPROPERTY(Transient, VisibleAnywhere, Category = CustomizableSkeletalMesh)
+	TArray< TObjectPtr<USkeletalMesh> > SkeletalMeshes;
+
 	// Will store status description of current skeletal mesh generation (for instance, "EmptyLOD0" or "EmptyMesh"
 	UPROPERTY()
 	FString SkeletalMeshStatus;
diff --git a/Engine/Plugins/Experimental/Mutable/Source/CustomizableObject/Private/MuCO/CustomizableObjectInstance.cpp b/Engine/Plugins/Experimental/Mutable/Source/CustomizableObject/Private/MuCO/CustomizableObjectInstance.cpp
--- a/Engine/Plugins/Experimental/Mutable/Source/CustomizableObject/Private/MuCO/CustomizableObjectInstance.cpp
+++ b/Engine/Plugins/Experimental/Mutable/Source/CustomizableObject/Private/MuCO/CustomizableObjectInstance.cpp
@@ -360,6 +360,7 @@ bool UCustomizableObjectInstance::IsEditorOnly() const
 void UCustomizableObjectInstance::BeginDestroy()
 {
 	BeginDestroyDelegate.Broadcast(this);
+	BeginDestroyNativeDelegate.Broadcast(this);
 
 	if (PrivateData && PrivateData->StreamingHandle)
 	{
diff --git a/Engine/Plugins/Experimental/Mutable/Source/CustomizableObject/Private/MuCO/CustomizableObjectSystem.cpp b/Engine/Plugins/Experimental/Mutable/Source/CustomizableObject/Private/MuCO/CustomizableObjectSystem.cpp
--- a/Engine/Plugins/Experimental/Mutable/Source/CustomizableObject/Private/MuCO/CustomizableObjectSystem.cpp
+++ b/Engine/Plugins/Experimental/Mutable/Source/CustomizableObject/Private/MuCO/CustomizableObjectSystem.cpp
@@ -837,6 +837,9 @@ void FCustomizableObjectSystemPrivate::InitUpdateSkeletalMesh(UCustomizableObjec
 		}
 
 		MutableOperationQueue.Enqueue(FMutableQueueElem::Create(Operation, Priority, InCustomizableObjectInstance->GetPrivate()->MinSquareDistFromComponentToPlayer));
+
+		InCustomizableObjectInstance->BeginUpdateDelegate.Broadcast(InCustomizableObjectInstance);
+		InCustomizableObjectInstance->BeginUpdateNativeDelegate.Broadcast(InCustomizableObjectInstance);
 	}
 	else
 	{
diff --git a/Engine/Plugins/Experimental/Mutable/Source/CustomizableObjectEditor/Private/MuCOE/StatePerformanceTesting.cpp b/Engine/Plugins/Experimental/Mutable/Source/CustomizableObjectEditor/Private/MuCOE/StatePerformanceTesting.cpp
--- a/Engine/Plugins/Experimental/Mutable/Source/CustomizableObjectEditor/Private/MuCOE/StatePerformanceTesting.cpp
+++ b/Engine/Plugins/Experimental/Mutable/Source/CustomizableObjectEditor/Private/MuCOE/StatePerformanceTesting.cpp
@@ -124,9 +124,6 @@ void FTestingCustomizableObject::StartTest(UCustomizableObject* CustomizableObje
 
 	if (SkeletalMeshComponents.Num() && SkeletalMeshComponents.Num() == CustomizableSkeletalComponents.Num())
 	{
-		Instance->UpdateSkeletalMeshAsync(true);
-		Instance->UpdateBeginDelegate.BindRaw(this, &FTestingCustomizableObject::InstanceUpdateBegin);
-
 		Instance->UpdatedDelegate.AddDynamic(HelperCallback, &UUpdateClassHelper::DelegatedCallback);
 		Instance->UpdateSkeletalMeshAsync(true);
 		InitStateAndParameterSequence();
@@ -160,7 +157,7 @@ void FTestingCustomizableObject::FinishTest()
 	StateTestEnded.ExecuteIfBound();
 	if (Instance)
 	{
-		Instance->UpdateBeginDelegate.Unbind();
+		//Instance->UpdateBeginDelegate.Unbind();
 		Instance->UpdatedDelegate.RemoveDynamic(HelperCallback, &UUpdateClassHelper::DelegatedCallback);
 	}
 
@@ -197,15 +194,6 @@ void FTestingCustomizableObject::FinishTest()
 }
 
 
-void FTestingCustomizableObject::InstanceUpdateBegin()
-{
-	if (ActiveBatch != nullptr)
-	{
-		ActiveBatch->TestStarted = true;
-	}
-}
-
-
 void FTestingCustomizableObject::InstanceUpdateEnd()
 {
 	const int32 MutableRuntimeCycles = Instance->LastUpdateMutableRuntimeCycles;
diff --git a/Engine/Plugins/Experimental/Mutable/Source/CustomizableObjectEditor/Private/MuCOE/StatePerformanceTesting.h b/Engine/Plugins/Experimental/Mutable/Source/CustomizableObjectEditor/Private/MuCOE/StatePerformanceTesting.h
--- a/Engine/Plugins/Experimental/Mutable/Source/CustomizableObjectEditor/Private/MuCOE/StatePerformanceTesting.h
+++ b/Engine/Plugins/Experimental/Mutable/Source/CustomizableObjectEditor/Private/MuCOE/StatePerformanceTesting.h
@@ -351,9 +351,6 @@ public:
 	/** Finish the test */
 	void FinishTest();
 
-	/** Callback for instance update begin delegate */
-	void InstanceUpdateBegin();
-
 	/** Callback for instance update end delegate */
 	void InstanceUpdateEnd();
 
