[mutable] Support multidimensional image parameters at core level.
Additional style cleanup (removing autos, legacy math, ...)
Additional optimization to object compilation: mesh bind now pushed down conditionals and mesh-remove-mask operations to run them at compile time.
Also moved mesh-difference to its own AST class. Added optimization for switches and conditionals.
#rnx

diff --git a/Engine/Plugins/Experimental/Mutable/Source/CustomizableObject/Classes/MuCO/CustomizableObject.h b/Engine/Plugins/Experimental/Mutable/Source/CustomizableObject/Classes/MuCO/CustomizableObject.h
--- a/Engine/Plugins/Experimental/Mutable/Source/CustomizableObject/Classes/MuCO/CustomizableObject.h
+++ b/Engine/Plugins/Experimental/Mutable/Source/CustomizableObject/Classes/MuCO/CustomizableObject.h
@@ -1132,7 +1132,7 @@ private:
 	// This is a manual version number for the binary blobs in this asset.
 	// Increasing it invalidates all the previously compiled models.
 	// Warning: If while merging code both versions have changed, take the highest+1.
-	static const int32 CurrentSupportedVersion = 348;
+	static const int32 CurrentSupportedVersion = 349;
 
 public:
 
diff --git a/Engine/Plugins/Experimental/Mutable/Source/CustomizableObject/Private/MuCO/MutableMeshBufferUtils.cpp b/Engine/Plugins/Experimental/Mutable/Source/CustomizableObject/Private/MuCO/MutableMeshBufferUtils.cpp
--- a/Engine/Plugins/Experimental/Mutable/Source/CustomizableObject/Private/MuCO/MutableMeshBufferUtils.cpp
+++ b/Engine/Plugins/Experimental/Mutable/Source/CustomizableObject/Private/MuCO/MutableMeshBufferUtils.cpp
@@ -107,6 +107,7 @@ void MutableMeshBufferUtils::SetupSkinBuffer(const int32& InCurrentVertexBuffer,
 	MESH_BUFFER_FORMAT Formats[ChannelCount] = {MBF_UINT8, MBF_NUINT8};
 	switch (MaxBoneIndexTypeSizeBytes)
 	{
+	case 0: // Fallback to something in this case.
 	case 1: Formats[0] = mu::MBF_UINT8;
 		break;
 	case 2: Formats[0] = mu::MBF_UINT16;
@@ -122,6 +123,7 @@ void MutableMeshBufferUtils::SetupSkinBuffer(const int32& InCurrentVertexBuffer,
 
 	switch (MaxBoneWeightTypeSizeBytes)
 	{
+	case 0: // Fallback to something in this case.
 	case 1: Formats[1] = mu::MBF_NUINT8;
 		break;
 	case 2: Formats[1] = mu::MBF_NUINT16;
diff --git a/Engine/Plugins/Experimental/Mutable/Source/CustomizableObjectEditor/Private/MuCOE/GenerateMutableSource/GenerateMutableSourceImage.cpp b/Engine/Plugins/Experimental/Mutable/Source/CustomizableObjectEditor/Private/MuCOE/GenerateMutableSource/GenerateMutableSourceImage.cpp
--- a/Engine/Plugins/Experimental/Mutable/Source/CustomizableObjectEditor/Private/MuCOE/GenerateMutableSource/GenerateMutableSourceImage.cpp
+++ b/Engine/Plugins/Experimental/Mutable/Source/CustomizableObjectEditor/Private/MuCOE/GenerateMutableSource/GenerateMutableSourceImage.cpp
@@ -207,7 +207,7 @@ mu::NodeImagePtr GenerateMutableSourceImage(const UEdGraphPin* Pin, FMutableGrap
 
 	else if (const UCustomizableObjectNodeTextureParameter* TypedNodeParam = Cast<UCustomizableObjectNodeTextureParameter>(Node))
 	{
-		mu::NodeImageParameterPtr TextureNode = new mu::NodeImageParameter();
+		mu::Ptr<mu::NodeImageParameter> TextureNode = new mu::NodeImageParameter();
 
 		GenerationContext.AddParameterNameUnique(Node, TypedNodeParam->ParameterName);
 
diff --git a/Engine/Plugins/Experimental/Mutable/Source/MutableRuntime/Private/MuR/CodeRunner.cpp b/Engine/Plugins/Experimental/Mutable/Source/MutableRuntime/Private/MuR/CodeRunner.cpp
--- a/Engine/Plugins/Experimental/Mutable/Source/MutableRuntime/Private/MuR/CodeRunner.cpp
+++ b/Engine/Plugins/Experimental/Mutable/Source/MutableRuntime/Private/MuR/CodeRunner.cpp
@@ -907,15 +907,22 @@ namespace mu
 
         case OP_TYPE::ME_DIFFERENCE:
         {
-            auto args = pModel->GetPrivate()->m_program.GetOpArgs<OP::MeshDifferenceArgs>(item.at);
+			const uint8* data = pModel->GetPrivate()->m_program.GetOpArgsPointer(item.at);
+
+			OP::ADDRESS BaseAt = 0;
+			FMemory::Memcpy(&BaseAt, data, sizeof(OP::ADDRESS)); data += sizeof(OP::ADDRESS);
+
+			OP::ADDRESS TargetAt = 0;
+			FMemory::Memcpy(&TargetAt, data, sizeof(OP::ADDRESS)); data += sizeof(OP::ADDRESS);
+
             switch (item.stage)
             {
             case 0:
-                if ( args.base && args.target )
+                if (BaseAt && TargetAt)
                 {
                     AddOp( SCHEDULED_OP( item.at, item, 1),
-                            SCHEDULED_OP( args.base, item),
-                            SCHEDULED_OP( args.target, item) );
+                            SCHEDULED_OP( BaseAt, item),
+                            SCHEDULED_OP( TargetAt, item) );
                 }
                 else
                 {
@@ -925,25 +932,32 @@ namespace mu
 
             case 1:
             {
-                Ptr<const Mesh> pBase = GetMemory().GetMesh(CACHE_ADDRESS(args.base,item));
-                Ptr<const Mesh> pTarget = GetMemory().GetMesh(CACHE_ADDRESS(args.target,item));
+                Ptr<const Mesh> pBase = GetMemory().GetMesh(CACHE_ADDRESS(BaseAt,item));
+                Ptr<const Mesh> pTarget = GetMemory().GetMesh(CACHE_ADDRESS(TargetAt,item));
 
-                MESH_BUFFER_SEMANTIC semantics[ MUTABLE_OP_MAX_MORPH_CHANNELS ];
-                int semanticIndices[ MUTABLE_OP_MAX_MORPH_CHANNELS ];
+				TArray<MESH_BUFFER_SEMANTIC, TInlineAllocator<8>> Semantics;
+				TArray<int32, TInlineAllocator<8>> SemanticIndices;
 
-                int numChannels = 0;
-                for ( int i=0;
-                      i<MUTABLE_OP_MAX_MORPH_CHANNELS && args.channelSemantic[i];
-                      ++i )
+				uint8 bIgnoreTextureCoords = 0;
+				FMemory::Memcpy(&bIgnoreTextureCoords, data, sizeof(uint8)); data += sizeof(uint8);
+
+				uint8 NumChannels = 0;
+				FMemory::Memcpy(&NumChannels, data, sizeof(uint8)); data += sizeof(uint8);
+
+                for ( uint8 i=0; i< NumChannels; ++i )
                 {
-					semantics[numChannels] = MESH_BUFFER_SEMANTIC(args.channelSemantic[i]);
-					semanticIndices[numChannels] = args.channelSemanticIndex[i];
-					++numChannels;
+					uint8 Semantic = 0;
+					FMemory::Memcpy(&Semantic, data, sizeof(uint8)); data += sizeof(uint8);
+					uint8 SemanticIndex = 0;
+					FMemory::Memcpy(&SemanticIndex, data, sizeof(uint8)); data += sizeof(uint8);
+
+					Semantics.Add(MESH_BUFFER_SEMANTIC(Semantic));
+					SemanticIndices.Add(SemanticIndex);
                 }
 
                 MeshPtr pResult = MeshDifference( pBase.get(), pTarget.get(),
-                                          numChannels, semantics, semanticIndices,
-                                          args.ignoreTextureCoords!=0 );
+                                          NumChannels, Semantics.GetData(), SemanticIndices.GetData(),
+                                          bIgnoreTextureCoords!=0 );
 
                 GetMemory().SetMesh( item, pResult );
                 break;
@@ -2375,16 +2389,17 @@ namespace mu
 
         case OP_TYPE::IM_PARAMETER:
         {
-				auto args = pModel->GetPrivate()->m_program.GetOpArgs<OP::ParameterArgs>(item.at);
-				EXTERNAL_IMAGE_ID id = pParams->GetImageValue(args.variable);
-				ImagePtr pResult = LoadExternalImage(id);
-				GetMemory().SetImage(item, pResult);
+			OP::ParameterArgs args = pModel->GetPrivate()->m_program.GetOpArgs<OP::ParameterArgs>(item.at);
+			Ptr<RangeIndex> Index = BuildCurrentOpRangeIndex(item, pParams, pModel, args.variable);
+			EXTERNAL_IMAGE_ID id = pParams->GetImageValue(args.variable, Index);
+			ImagePtr pResult = LoadExternalImage(id);
+			GetMemory().SetImage(item, pResult);
             break;
         }
 
         case OP_TYPE::IM_LAYERCOLOUR:
         {
-            auto args = pModel->GetPrivate()->m_program.GetOpArgs<OP::ImageLayerColourArgs>(item.at);
+			OP::ImageLayerColourArgs args = pModel->GetPrivate()->m_program.GetOpArgs<OP::ImageLayerColourArgs>(item.at);
             switch (item.stage)
             {
             case 0:
@@ -2408,7 +2423,7 @@ namespace mu
 
         case OP_TYPE::IM_LAYER:
         {
-            auto args = pModel->GetPrivate()->m_program.GetOpArgs<OP::ImageLayerArgs>(item.at);
+			OP::ImageLayerArgs args = pModel->GetPrivate()->m_program.GetOpArgs<OP::ImageLayerArgs>(item.at);
             switch (item.stage)
             {
             case 0:
@@ -2432,7 +2447,7 @@ namespace mu
 
         case OP_TYPE::IM_MULTILAYER:
         {
-            auto args = pModel->GetPrivate()->m_program.GetOpArgs<OP::ImageMultiLayerArgs>(item.at);
+			OP::ImageMultiLayerArgs args = pModel->GetPrivate()->m_program.GetOpArgs<OP::ImageMultiLayerArgs>(item.at);
             switch (item.stage)
             {
             case 0:
diff --git a/Engine/Plugins/Experimental/Mutable/Source/MutableRuntime/Private/MuR/OpImageProject.cpp b/Engine/Plugins/Experimental/Mutable/Source/MutableRuntime/Private/MuR/OpImageProject.cpp
--- a/Engine/Plugins/Experimental/Mutable/Source/MutableRuntime/Private/MuR/OpImageProject.cpp
+++ b/Engine/Plugins/Experimental/Mutable/Source/MutableRuntime/Private/MuR/OpImageProject.cpp
@@ -1484,8 +1484,8 @@ float getTriangleRatio(const vec2f& a, const vec2f& b, const vec2f& c)
 	float lenSide2 = length(c - a);
 	float lenSide3 = length(b - c);
 
-	float maxLen = max(max(lenSide1, lenSide2), lenSide3);
-	float minLen = min(min(lenSide1, lenSide2), lenSide3);
+	float maxLen = FMath::Max3(lenSide1, lenSide2, lenSide3);
+	float minLen = FMath::Min3(lenSide1, lenSide2, lenSide3);
 
 	return maxLen / minLen;
 }
diff --git a/Engine/Plugins/Experimental/Mutable/Source/MutableRuntime/Private/MuR/Operations.cpp b/Engine/Plugins/Experimental/Mutable/Source/MutableRuntime/Private/MuR/Operations.cpp
--- a/Engine/Plugins/Experimental/Mutable/Source/MutableRuntime/Private/MuR/Operations.cpp
+++ b/Engine/Plugins/Experimental/Mutable/Source/MutableRuntime/Private/MuR/Operations.cpp
@@ -374,11 +374,6 @@ namespace mu
             f(&op.args.MeshApplyLayout.mesh );
             break;
 
-        case OP_TYPE::ME_DIFFERENCE:
-            f(&op.args.MeshDifference.base );
-            f(&op.args.MeshDifference.target );
-            break;
-
         case OP_TYPE::ME_MORPH2:
             f(&op.args.MeshMorph2.factor );
             f(&op.args.MeshMorph2.base );
@@ -916,7 +911,7 @@ namespace mu
         //-------------------------------------------------------------------------------------
         case OP_TYPE::ME_APPLYLAYOUT:
         {
-            auto args = program.GetOpArgs<OP::MeshApplyLayoutArgs>(at);
+			OP::MeshApplyLayoutArgs args = program.GetOpArgs<OP::MeshApplyLayoutArgs>(at);
             f(args.layout );
             f(args.mesh );
             break;
@@ -924,15 +919,21 @@ namespace mu
 
         case OP_TYPE::ME_DIFFERENCE:
         {
-            auto args = program.GetOpArgs<OP::MeshDifferenceArgs>(at);
-            f(args.base );
-            f(args.target );
+			const uint8_t* data = program.GetOpArgsPointer(at);
+			OP::ADDRESS BaseAt = 0;
+			FMemory::Memcpy(&BaseAt, data, sizeof(OP::ADDRESS)); data += sizeof(OP::ADDRESS);
+
+			OP::ADDRESS TargetAt = 0;
+			FMemory::Memcpy(&TargetAt, data, sizeof(OP::ADDRESS)); data += sizeof(OP::ADDRESS);
+
+			f(BaseAt);
+            f(TargetAt);
             break;
         }
 
         case OP_TYPE::ME_MORPH2:
         {
-            auto args = program.GetOpArgs<OP::MeshMorph2Args>(at);
+			OP::MeshMorph2Args args = program.GetOpArgs<OP::MeshMorph2Args>(at);
             f(args.factor );
             f(args.base );
 
@@ -945,7 +946,7 @@ namespace mu
 
         case OP_TYPE::ME_MERGE:
         {
-            auto args = program.GetOpArgs<OP::MeshMergeArgs>(at);
+			OP::MeshMergeArgs args = program.GetOpArgs<OP::MeshMergeArgs>(at);
             f(args.base );
             f(args.added );
             break;
@@ -953,7 +954,7 @@ namespace mu
 
         case OP_TYPE::ME_INTERPOLATE:
         {
-            auto args = program.GetOpArgs<OP::MeshInterpolateArgs>(at);
+			OP::MeshInterpolateArgs args = program.GetOpArgs<OP::MeshInterpolateArgs>(at);
             f(args.factor );
             f(args.base );
 
@@ -966,7 +967,7 @@ namespace mu
 
         case OP_TYPE::ME_MASKCLIPMESH:
         {
-            auto args = program.GetOpArgs<OP::MeshMaskClipMeshArgs>(at);
+			OP::MeshMaskClipMeshArgs args = program.GetOpArgs<OP::MeshMaskClipMeshArgs>(at);
             f(args.source );
             f(args.clip );
             break;
@@ -974,7 +975,7 @@ namespace mu
 
         case OP_TYPE::ME_MASKDIFF:
         {
-            auto args = program.GetOpArgs<OP::MeshMaskDiffArgs>(at);
+			OP::MeshMaskDiffArgs args = program.GetOpArgs<OP::MeshMaskDiffArgs>(at);
             f(args.source );
             f(args.fragment );
             break;
@@ -982,7 +983,7 @@ namespace mu
 
         case OP_TYPE::ME_SUBTRACT:
         {
-            auto args = program.GetOpArgs<OP::MeshSubtractArgs>(at);
+			OP::MeshSubtractArgs args = program.GetOpArgs<OP::MeshSubtractArgs>(at);
             f(args.a );
             f(args.b );
             break;
@@ -1012,7 +1013,7 @@ namespace mu
 
         case OP_TYPE::ME_FORMAT:
         {
-            auto args = program.GetOpArgs<OP::MeshFormatArgs>(at);
+			OP::MeshFormatArgs args = program.GetOpArgs<OP::MeshFormatArgs>(at);
             f(args.source );
             f(args.format );
             break;
@@ -1020,7 +1021,7 @@ namespace mu
 
         case OP_TYPE::ME_TRANSFORM:
         {
-            auto args = program.GetOpArgs<OP::MeshTransformArgs>(at);
+			OP::MeshTransformArgs args = program.GetOpArgs<OP::MeshTransformArgs>(at);
             f(args.source );
             break;
         }
@@ -1036,21 +1037,21 @@ namespace mu
 
         case OP_TYPE::ME_EXTRACTFACEGROUP:
         {
-            auto args = program.GetOpArgs<OP::MeshExtractFaceGroupArgs>(at);
+			OP::MeshExtractFaceGroupArgs args = program.GetOpArgs<OP::MeshExtractFaceGroupArgs>(at);
             f(args.source );
             break;
         }
 
         case OP_TYPE::ME_CLIPMORPHPLANE:
         {
-            auto args = program.GetOpArgs<OP::MeshClipMorphPlaneArgs>(at);
+			OP::MeshClipMorphPlaneArgs args = program.GetOpArgs<OP::MeshClipMorphPlaneArgs>(at);
             f(args.source);
             break;
         }
 
         case OP_TYPE::ME_CLIPWITHMESH :
         {
-            auto args = program.GetOpArgs<OP::MeshClipWithMeshArgs>(at);
+			OP::MeshClipWithMeshArgs args = program.GetOpArgs<OP::MeshClipWithMeshArgs>(at);
             f(args.source);
             f(args.clipMesh);
             break;
@@ -1058,7 +1059,7 @@ namespace mu
 
         case OP_TYPE::ME_CLIPDEFORM:
         {
-            auto args = program.GetOpArgs<OP::MeshClipDeformArgs>(at);
+			OP::MeshClipDeformArgs args = program.GetOpArgs<OP::MeshClipDeformArgs>(at);
             f(args.mesh);
             f(args.clipShape);
             break;
@@ -1075,7 +1076,7 @@ namespace mu
  	
         case OP_TYPE::ME_REMAPINDICES :
         {
-            auto args = program.GetOpArgs<OP::MeshRemapIndicesArgs>(at);
+			OP::MeshRemapIndicesArgs args = program.GetOpArgs<OP::MeshRemapIndicesArgs>(at);
             f(args.source);
             f(args.reference);
             break;
@@ -1083,7 +1084,7 @@ namespace mu
 
         case OP_TYPE::ME_SETSKELETON :
         {
-            auto args = program.GetOpArgs<OP::MeshSetSkeletonArgs>(at);
+			OP::MeshSetSkeletonArgs args = program.GetOpArgs<OP::MeshSetSkeletonArgs>(at);
             f(args.source);
             f(args.skeleton);
             break;
@@ -1091,7 +1092,7 @@ namespace mu
 
         case OP_TYPE::ME_PROJECT :
         {
-            auto args = program.GetOpArgs<OP::MeshProjectArgs>(at);
+			OP::MeshProjectArgs args = program.GetOpArgs<OP::MeshProjectArgs>(at);
             f(args.mesh);
             f(args.projector);
             break;
@@ -1099,7 +1100,7 @@ namespace mu
 
 		case OP_TYPE::ME_APPLYPOSE:
 		{
-			auto args = program.GetOpArgs<OP::MeshApplyPoseArgs>(at);
+			OP::MeshApplyPoseArgs args = program.GetOpArgs<OP::MeshApplyPoseArgs>(at);
 			f(args.base);
 			f(args.pose);
 			break;
@@ -1107,7 +1108,7 @@ namespace mu
 
 		case OP_TYPE::ME_GEOMETRYOPERATION:
 		{
-			auto args = program.GetOpArgs<OP::MeshGeometryOperationArgs>(at);
+			OP::MeshGeometryOperationArgs args = program.GetOpArgs<OP::MeshGeometryOperationArgs>(at);
 			f(args.meshA);
 			f(args.meshB);
 			f(args.scalarA);
@@ -1117,7 +1118,7 @@ namespace mu
 
 		case OP_TYPE::ME_BINDSHAPE:
 		{
-			auto args = program.GetOpArgs<OP::MeshBindShapeArgs>(at);
+			OP::MeshBindShapeArgs args = program.GetOpArgs<OP::MeshBindShapeArgs>(at);
 			f(args.mesh);
 			f(args.shape);
 			break;
@@ -1125,7 +1126,7 @@ namespace mu
 
 		case OP_TYPE::ME_APPLYSHAPE:
 		{
-			auto args = program.GetOpArgs<OP::MeshApplyShapeArgs>(at);
+			OP::MeshApplyShapeArgs args = program.GetOpArgs<OP::MeshApplyShapeArgs>(at);
 			f(args.mesh);
 			f(args.shape);
 			break;
@@ -1140,7 +1141,7 @@ namespace mu
         case OP_TYPE::IN_ADDCOMPONENT:
         case OP_TYPE::IN_ADDSURFACE:
         {
-            auto args = program.GetOpArgs<OP::InstanceAddArgs>(at);
+			OP::InstanceAddArgs args = program.GetOpArgs<OP::InstanceAddArgs>(at);
             f(args.instance );
             f(args.value );
             break;
@@ -1148,7 +1149,7 @@ namespace mu
 
         case OP_TYPE::IN_ADDLOD:
         {
-            auto args = program.GetOpArgs<OP::InstanceAddLODArgs>(at);
+			OP::InstanceAddLODArgs args = program.GetOpArgs<OP::InstanceAddLODArgs>(at);
             for (int t=0;t<MUTABLE_OP_MAX_ADD_COUNT;++t)
             {
                 f(args.lod[t] );
@@ -1159,14 +1160,14 @@ namespace mu
         //-------------------------------------------------------------------------------------
         case OP_TYPE::LA_PACK:
         {
-            auto args = program.GetOpArgs<OP::LayoutPackArgs>(at);
+			OP::LayoutPackArgs args = program.GetOpArgs<OP::LayoutPackArgs>(at);
             f(args.layout );
             break;
         }
 
         case OP_TYPE::LA_MERGE:
         {
-            auto args = program.GetOpArgs<OP::LayoutMergeArgs>(at);
+			OP::LayoutMergeArgs args = program.GetOpArgs<OP::LayoutMergeArgs>(at);
             f(args.base );
             f(args.added );
             break;
@@ -1174,7 +1175,7 @@ namespace mu
 
         case OP_TYPE::LA_REMOVEBLOCKS:
         {
-            auto args = program.GetOpArgs<OP::LayoutRemoveBlocksArgs>(at);
+			OP::LayoutRemoveBlocksArgs args = program.GetOpArgs<OP::LayoutRemoveBlocksArgs>(at);
             f(args.source );
             f(args.mesh );
             break;
diff --git a/Engine/Plugins/Experimental/Mutable/Source/MutableRuntime/Private/MuR/Operations.h b/Engine/Plugins/Experimental/Mutable/Source/MutableRuntime/Private/MuR/Operations.h
--- a/Engine/Plugins/Experimental/Mutable/Source/MutableRuntime/Private/MuR/Operations.h
+++ b/Engine/Plugins/Experimental/Mutable/Source/MutableRuntime/Private/MuR/Operations.h
@@ -15,7 +15,6 @@
 #define MUTABLE_OP_MAX_ADD_COUNT			7
 #define MUTABLE_OP_MAX_SWIZZLE_CHANNELS		4
 #define MUTABLE_OP_MAX_MORPH2_TARGETS		5
-#define MUTABLE_OP_MAX_MORPH_CHANNELS		8
 
 namespace mu
 {
@@ -780,20 +779,21 @@ namespace mu
             uint16 channel;
         };
 
-        struct MeshDifferenceArgs
-        {
-            ADDRESS base;
-            ADDRESS target;
+		// Moved to dynamically-sized link data
+        //struct MeshDifferenceArgs
+        //{
+        //    ADDRESS base;
+        //    ADDRESS target;
 
-            //! If true the texture coordinates will not be morphed. This is only relevant if the list
-            //! in channelSemantic is empty.
-            uint8 ignoreTextureCoords;
+        //    //! If true the texture coordinates will not be morphed. This is only relevant if the list
+        //    //! in channelSemantic is empty.
+        //    uint8 ignoreTextureCoords;
 
-            //! Channels to diff. Meaningful until the first null channel. If all are null, it means
-            //! that all the channels should be diff-ed
-            uint8 channelSemantic[MUTABLE_OP_MAX_MORPH_CHANNELS];
-            uint8 channelSemanticIndex[MUTABLE_OP_MAX_MORPH_CHANNELS];
-        };
+        //    //! Channels to diff. Meaningful until the first null channel. If all are null, it means
+        //    //! that all the channels should be diff-ed
+        //    uint8 channelSemantic[MUTABLE_OP_MAX_MORPH_CHANNELS];
+        //    uint8 channelSemanticIndex[MUTABLE_OP_MAX_MORPH_CHANNELS];
+        //};
 
         struct MeshMergeArgs
         {
@@ -1111,7 +1111,7 @@ namespace mu
 
             //-------------------------------------------------------------------------------------
             MeshApplyLayoutArgs MeshApplyLayout;
-            MeshDifferenceArgs MeshDifference;
+            //MeshDifferenceArgs MeshDifference;
             MeshMorph2Args MeshMorph2;
             MeshMergeArgs MeshMerge;
             MeshInterpolateArgs MeshInterpolate;
diff --git a/Engine/Plugins/Experimental/Mutable/Source/MutableRuntime/Private/MuR/Parameters.cpp b/Engine/Plugins/Experimental/Mutable/Source/MutableRuntime/Private/MuR/Parameters.cpp
--- a/Engine/Plugins/Experimental/Mutable/Source/MutableRuntime/Private/MuR/Parameters.cpp
+++ b/Engine/Plugins/Experimental/Mutable/Source/MutableRuntime/Private/MuR/Parameters.cpp
@@ -71,7 +71,7 @@ namespace mu
 
         for (size_t p=0; p<valueCount; ++p)
         {
-            const auto& desc = params->GetPrivate()->m_pModel->GetPrivate()->m_program.m_parameters[p];
+            const PARAMETER_DESC& desc = params->GetPrivate()->m_pModel->GetPrivate()->m_program.m_parameters[p];
             arch << desc.m_name;
             arch << desc.m_uid;
             arch << desc.m_type;
@@ -117,7 +117,7 @@ namespace mu
             {
                 for (size_t mp=0; mp<modelParameters; ++mp)
                 {
-                    const auto& desc = pModel->GetPrivate()->m_program.m_parameters[mp];
+                    const PARAMETER_DESC& desc = pModel->GetPrivate()->m_program.m_parameters[mp];
                     if (desc.m_uid==uid)
                     {
                         modelParam = int(mp);
@@ -131,7 +131,7 @@ namespace mu
             {
                 for (size_t mp=0; mp<modelParameters; ++mp)
                 {
-                    const auto& desc = pModel->GetPrivate()->m_program.m_parameters[mp];
+                    const PARAMETER_DESC& desc = pModel->GetPrivate()->m_program.m_parameters[mp];
                     if (desc.m_name==name && desc.m_type==type)
                     {
                         modelParam = int(mp);
@@ -145,7 +145,7 @@ namespace mu
             {
                 for (size_t mp=0; mp<modelParameters; ++mp)
                 {
-                    const auto& desc = pModel->GetPrivate()->m_program.m_parameters[mp];
+                    const PARAMETER_DESC& desc = pModel->GetPrivate()->m_program.m_parameters[mp];
                     if (desc.m_name==name)
                     {
                         modelParam = int(mp);
@@ -291,7 +291,7 @@ namespace mu
             return nullptr;
         }
 
-        auto it = m_pD->m_multiValues[paramIndex].begin();
+		TMap< TArray<int32_t>, PARAMETER_VALUE >::TRangedForIterator it = m_pD->m_multiValues[paramIndex].begin();
         for ( int i=0; i<valueIndex; ++i )
         {
             ++it;
@@ -354,8 +354,8 @@ namespace mu
 
         if ( index<int(m_pD->m_multiValues.Num()))
         {
-            const auto& m = m_pD->m_multiValues[index];
-            auto it = m.Find(pos->m_pD->m_values);
+            const TMap< TArray<int32_t>, PARAMETER_VALUE >& m = m_pD->m_multiValues[index];
+			const PARAMETER_VALUE* it = m.Find(pos->m_pD->m_values);
             if (it)
             {
                 return it->m_bool;
@@ -406,18 +406,9 @@ namespace mu
                 m_pD->m_multiValues.SetNum(index+1);
             }
 
-            auto& m = m_pD->m_multiValues[index];
-            auto it = m.Find(pos->m_pD->m_values);
-            if (it)
-            {
-                it->m_bool = value;
-            }
-            else
-            {
-                PARAMETER_VALUE valueData;
-                valueData.m_bool = value;
-                m.Add( pos->m_pD->m_values, valueData );
-            }
+			TMap< TArray<int32_t>, PARAMETER_VALUE >& m = m_pD->m_multiValues[index];
+			PARAMETER_VALUE& it = m.FindOrAdd(pos->m_pD->m_values);
+            it.m_bool = value;
         }
 	}
 
@@ -524,8 +515,8 @@ namespace mu
 
         if ( index<int(m_pD->m_multiValues.Num()))
         {
-            const auto& m = m_pD->m_multiValues[index];
-            auto it = m.Find(pos->m_pD->m_values);
+            const TMap< TArray<int32_t>, PARAMETER_VALUE >& m = m_pD->m_multiValues[index];
+			const PARAMETER_VALUE* it = m.Find(pos->m_pD->m_values);
             if (it)
             {
                 return it->m_int;
@@ -576,18 +567,9 @@ namespace mu
                 m_pD->m_multiValues.SetNum(index+1);
             }
 
-            auto& m = m_pD->m_multiValues[index];
-            auto it = m.Find(pos->m_pD->m_values);
-            if (it)
-            {
-                it->m_int = value;
-            }
-            else
-            {
-                PARAMETER_VALUE valueData;
-                valueData.m_int = value;
-                m.Add( pos->m_pD->m_values, valueData );
-            }
+			TMap< TArray<int32_t>, PARAMETER_VALUE >& m = m_pD->m_multiValues[index];
+			PARAMETER_VALUE& it = m.FindOrAdd(pos->m_pD->m_values);
+            it.m_int = value;
         }
     }
 
@@ -621,8 +603,8 @@ namespace mu
 
         if ( index<int(m_pD->m_multiValues.Num()))
         {
-            const auto& m = m_pD->m_multiValues[index];
-            auto it = m.Find(pos->m_pD->m_values);
+            const TMap< TArray<int32_t>, PARAMETER_VALUE >& m = m_pD->m_multiValues[index];
+			const PARAMETER_VALUE* it = m.Find(pos->m_pD->m_values);
             if (it)
             {
                 return it->m_float;
@@ -673,18 +655,9 @@ namespace mu
                 m_pD->m_multiValues.SetNum(index+1);
             }
 
-            auto& m = m_pD->m_multiValues[index];
-            auto it = m.Find(pos->m_pD->m_values);
-            if (it)
-            {
-                it->m_float = value;
-            }
-            else
-            {
-                PARAMETER_VALUE valueData;
-                valueData.m_float = value;
-                m.Add( pos->m_pD->m_values, valueData );
-            }
+			TMap< TArray<int32_t>, PARAMETER_VALUE >& m = m_pD->m_multiValues[index];
+			PARAMETER_VALUE& it = m.FindOrAdd(pos->m_pD->m_values);
+            it.m_float = value;
         }
     }
 
@@ -721,8 +694,8 @@ namespace mu
 
         if ( index<int(m_pD->m_multiValues.Num()))
         {
-            const auto& m = m_pD->m_multiValues[index];
-            auto it = m.Find(pos->m_pD->m_values);
+            const TMap< TArray<int32_t>, PARAMETER_VALUE >& m = m_pD->m_multiValues[index];
+			const PARAMETER_VALUE* it = m.Find(pos->m_pD->m_values);
             if (it)
             {
                 if (pR) *pR = it->m_colour[0];
@@ -781,43 +754,72 @@ namespace mu
                 m_pD->m_multiValues.SetNum(index+1);
             }
 
-            auto& m = m_pD->m_multiValues[index];
-            auto it = m.Find(pos->m_pD->m_values);
-            if (it)
-            {
-                it->m_colour[0] = r;
-                it->m_colour[1] = g;
-                it->m_colour[2] = b;
-            }
-            else
-            {
-                PARAMETER_VALUE valueData;
-                valueData.m_colour[0] = r;
-                valueData.m_colour[1] = g;
-                valueData.m_colour[2] = b;
-                m.Add( pos->m_pD->m_values, valueData );
-            }
+			TMap< TArray<int32_t>, PARAMETER_VALUE >& m = m_pD->m_multiValues[index];
+			PARAMETER_VALUE& it = m.FindOrAdd(pos->m_pD->m_values);
+            it.m_colour[0] = r;
+            it.m_colour[1] = g;
+            it.m_colour[2] = b;
         }
     }
 
 
     //---------------------------------------------------------------------------------------------
-    EXTERNAL_IMAGE_ID Parameters::GetImageValue( int index ) const
+    EXTERNAL_IMAGE_ID Parameters::GetImageValue( int index, const Ptr<const RangeIndex>& pos ) const
     {
         check( index >= 0 && index < (int)m_pD->m_values.Num() );
         check( GetType( index ) == PARAMETER_TYPE::T_IMAGE );
 
-        return m_pD->m_values[index].m_image;
+		// Single value case
+		if (!pos)
+		{
+			return m_pD->m_values[index].m_image;
+		}
+
+		// Multivalue case
+		check(pos->m_pD->m_parameter == index);
+
+		if (index < int(m_pD->m_multiValues.Num()))
+		{
+			const TMap< TArray<int32_t>, PARAMETER_VALUE >& m = m_pD->m_multiValues[index];
+			const PARAMETER_VALUE* it = m.Find(pos->m_pD->m_values);
+			if (it)
+			{
+				return it->m_image;
+			}
+		}
+
+		// Multivalue parameter, but no multivalue set. Return single value.
+		return m_pD->m_values[index].m_image;
+
     }
 
 
     //---------------------------------------------------------------------------------------------
-    void Parameters::SetImageValue( int index, EXTERNAL_IMAGE_ID id )
+    void Parameters::SetImageValue( int index, EXTERNAL_IMAGE_ID id, const Ptr<const RangeIndex>& pos )
     {
         check( index >= 0 && index < (int)m_pD->m_values.Num() );
         check( GetType( index ) == PARAMETER_TYPE::T_IMAGE );
 
-        m_pD->m_values[index].m_image = id;
+		// Single value case
+		if (!pos)
+		{
+			m_pD->m_values[index].m_image = id;
+		}
+
+		// Multivalue case
+		else
+		{
+			check(pos->m_pD->m_parameter == index);
+
+			if (index >= int(m_pD->m_multiValues.Num()))
+			{
+				m_pD->m_multiValues.SetNum(index + 1);
+			}
+
+			TMap< TArray<int32_t>, PARAMETER_VALUE >& m = m_pD->m_multiValues[index];
+			PARAMETER_VALUE& it = m.FindOrAdd(pos->m_pD->m_values);
+			it.m_image = id;
+		}
     }
 
 
@@ -846,8 +848,8 @@ namespace mu
 
         if ( index < int( m_pD->m_multiValues.Num() ) )
         {
-            const auto& m = m_pD->m_multiValues[index];
-            auto it = m.Find( pos->m_pD->m_values );
+            const TMap< TArray<int32_t>, PARAMETER_VALUE >& m = m_pD->m_multiValues[index];
+			const PARAMETER_VALUE* it = m.Find( pos->m_pD->m_values );
             if ( it )
             {
                 return it->m_text;
@@ -874,8 +876,7 @@ namespace mu
             return;
         }
 
-        int len = FMath::Min( MUTABLE_MAX_STRING_PARAM_LENGTH,
-                            int( strnlen( value, MUTABLE_MAX_STRING_PARAM_LENGTH ) ) );
+        int32 len = FMath::Min( MUTABLE_MAX_STRING_PARAM_LENGTH, int32( strnlen( value, MUTABLE_MAX_STRING_PARAM_LENGTH ) ) );
 
         // Single value case
         if ( !pos )
@@ -886,7 +887,6 @@ namespace mu
                 m_pD->m_multiValues[index].Empty();
             }
 
-            //m_pD->m_values[index].m_float = value;
 			FMemory::Memcpy( m_pD->m_values[index].m_text, value, len );
             m_pD->m_values[index].m_text[len] = 0;
         }
@@ -901,22 +901,10 @@ namespace mu
                 m_pD->m_multiValues.SetNum( index + 1 );
             }
 
-            auto& m = m_pD->m_multiValues[index];
-            auto it = m.Find( pos->m_pD->m_values );
-            if ( it )
-            {
-                //it->second.m_float = value;
-				FMemory::Memcpy( it->m_text, value, len );
-                it->m_text[len] = 0;
-            }
-            else
-            {
-                PARAMETER_VALUE valueData;
-                //valueData.m_float = value;
-				FMemory::Memcpy( valueData.m_text, value, len );
-                valueData.m_text[len] = 0;
-                m.Add( pos->m_pD->m_values, valueData );
-            }
+			TMap< TArray<int32_t>, PARAMETER_VALUE >& m = m_pD->m_multiValues[index];
+			PARAMETER_VALUE& it = m.FindOrAdd( pos->m_pD->m_values );
+			FMemory::Memcpy( it.m_text, value, len );
+            it.m_text[len] = 0;
         }
     }
 
@@ -1070,16 +1058,8 @@ namespace mu
             }
 
 			TMap< TArray<int32_t>, PARAMETER_VALUE >& m = m_pD->m_multiValues[index];
-			PARAMETER_VALUE* it = m.Find(pos->m_pD->m_values);
-            if (it)
-            {
-                result = &it->m_projector;
-            }
-            else
-            {
-                PARAMETER_VALUE& it2 = m.Add( pos->m_pD->m_values, PARAMETER_VALUE() );
-                result = &it2.m_projector;
-            }
+			PARAMETER_VALUE& it = m.FindOrAdd(pos->m_pD->m_values);
+            result = &it.m_projector;
         }
 
         check( result );
diff --git a/Engine/Plugins/Experimental/Mutable/Source/MutableRuntime/Private/MuR/Raster.h b/Engine/Plugins/Experimental/Mutable/Source/MutableRuntime/Private/MuR/Raster.h
--- a/Engine/Plugins/Experimental/Mutable/Source/MutableRuntime/Private/MuR/Raster.h
+++ b/Engine/Plugins/Experimental/Mutable/Source/MutableRuntime/Private/MuR/Raster.h
@@ -2,23 +2,11 @@
 
 #pragma once
 
+#include "Math/UnrealMathUtility.h"
+
 namespace mu
 {
 
-    inline int32_t min( int32_t a, int32_t b ) { return a<b? a : b; }
-    inline int32_t max( int32_t a, int32_t b ) { return a<b? b : a; }
-    inline int32_t min( int32_t a, int32_t b, int32_t c ) { return a<b? ( a<c? a : c ) : ( b<c? b : c ); }
-    inline int32_t max( int32_t a, int32_t b, int32_t c ) { return a>b? ( a>c? a : c ) : ( b>c? b : c ); }
-
-    inline int64_t min( int64_t a, int64_t b ) { return a<b? a : b; }
-    inline int64_t max( int64_t a, int64_t b ) { return a<b? b : a; }
-    inline int64_t min( int64_t a, int64_t b, int64_t c ) { return a<b? ( a<c? a : c ) : ( b<c? b : c ); }
-    inline int64_t max( int64_t a, int64_t b, int64_t c ) { return a>b? ( a>c? a : c ) : ( b>c? b : c ); }
-
-    inline float min( float a, float b ) { return a<b? a : b; }
-    inline float max( float a, float b ) { return a<b? b : a; }
-
-
 	template<int NUM_INTERPOLATORS>
 	class RasterVertex
 	{
@@ -99,10 +87,10 @@ namespace mu
         const INT FDY31 = DY31 * 16;
 
 	    // Bounding rectangle
-        INT minx = max(INT(0),        ( min(X1, X2, X3) + 0xF) / 16 );
-        INT maxx = min(INT(aWidth-1), ( max(X1, X2, X3) + 0xF) / 16 );
-        INT miny = max(INT(0),		  ( min(Y1, Y2, Y3) + 0xF) / 16 );
-        INT maxy = min(INT(aHeight-1),( max(Y1, Y2, Y3) + 0xF) / 16 );
+        INT minx = FMath::Max(INT(0),        (FMath::Min3(X1, X2, X3) + 0xF) / 16 );
+        INT maxx = FMath::Min(INT(aWidth-1), (FMath::Max3(X1, X2, X3) + 0xF) / 16 );
+        INT miny = FMath::Max(INT(0),		 (FMath::Min3(Y1, Y2, Y3) + 0xF) / 16 );
+        INT maxy = FMath::Min(INT(aHeight-1),(FMath::Max3(Y1, Y2, Y3) + 0xF) / 16 );
 
 	    // Block size, standard 8x8 (must be power of two)
         const INT q = 8;
@@ -181,9 +169,9 @@ namespace mu
 	                        float py = y+iy-v1.y;
 							float alpha = px*vy+py*-vx;
 							float beta = px*-uy+py*ux;
-							alpha = max(0.0,min(1.0f, alpha ));
-							beta = max(0.0,min(1.0f, beta ));
-							float gamma = max(0.0, 1.0f-alpha-beta );
+							alpha = FMath::Max(0.0, FMath::Min(1.0f, alpha ));
+							beta = FMath::Max(0.0, FMath::Min(1.0f, beta ));
+							float gamma = FMath::Max(0.0, 1.0f-alpha-beta );
 
 	                        for (int ii=0; ii<NUM_INTERPOLATORS; ii++)
 							{
@@ -220,9 +208,9 @@ namespace mu
 		                        float py = iy-v1.y;
 								float alpha = px*vy+py*-vx;
 								float beta = px*-uy+py*ux;
-								alpha = max(0.0,min(1.0f, alpha ));
-								beta = max(0.0,min(1.0f, beta ));
-								float gamma = max(0.0, 1.0f-alpha-beta );
+								alpha = FMath::Max(0.0, FMath::Min(1.0f, alpha ));
+								beta = FMath::Max(0.0, FMath::Min(1.0f, beta ));
+								float gamma = FMath::Max(0.0, 1.0f-alpha-beta );
 
 		                        for (int ii=0; ii<NUM_INTERPOLATORS; ii++)
 								{
diff --git a/Engine/Plugins/Experimental/Mutable/Source/MutableRuntime/Public/MuR/Model.h b/Engine/Plugins/Experimental/Mutable/Source/MutableRuntime/Public/MuR/Model.h
--- a/Engine/Plugins/Experimental/Mutable/Source/MutableRuntime/Public/MuR/Model.h
+++ b/Engine/Plugins/Experimental/Mutable/Source/MutableRuntime/Public/MuR/Model.h
@@ -16,7 +16,7 @@ namespace mu { class OutputArchive; }
 //! data structures. Compiled models are not necessarily compatible when the runtime is updated,
 //! so this version number can be used externally to verify this. It is not used internally, and
 //! serializing models from different versions than this runtime will probably result in a crash.
-#define MUTABLE_COMPILED_MODEL_CODE_VERSION		uint32( 47 )
+#define MUTABLE_COMPILED_MODEL_CODE_VERSION		uint32( 48 )
 #define MUTABLE_PARAMETERS_VERSION              uint32( 1 )
 
 
diff --git a/Engine/Plugins/Experimental/Mutable/Source/MutableRuntime/Public/MuR/Parameters.h b/Engine/Plugins/Experimental/Mutable/Source/MutableRuntime/Public/MuR/Parameters.h
--- a/Engine/Plugins/Experimental/Mutable/Source/MutableRuntime/Public/MuR/Parameters.h
+++ b/Engine/Plugins/Experimental/Mutable/Source/MutableRuntime/Public/MuR/Parameters.h
@@ -7,8 +7,6 @@
 #include "MuR/RefCounted.h"
 #include "MuR/Types.h"
 
-namespace mu { class InputArchive; }
-namespace mu { class OutputArchive; }
 
 #define MUTABLE_MAX_STRING_PARAM_LENGTH     40
 
@@ -16,8 +14,10 @@ namespace mu
 {
 
 	//!
-	class Parameters;
+	class InputArchive;
+	class OutputArchive;
 
+	class Parameters;
     using ParametersPtr = Ptr<Parameters>;
     using ParametersPtrConst = Ptr<const Parameters>;
 
@@ -359,13 +359,15 @@ namespace mu
         //! Return the value of an image parameter.
         //! \pre The parameter specified by index is a T_IMAGE.
         //! \param index Index of the parameter from 0 to GetCount()-1
-        //! \return The externalId specified when setting the image value (\see SetImageValue)
-        EXTERNAL_IMAGE_ID GetImageValue( int index ) const;
+		//! \param pos Only for multidimensional parametres: relevant position to set in the ranges
+	    //! \return The externalId specified when setting the image value (\see SetImageValue)
+        EXTERNAL_IMAGE_ID GetImageValue( int index, const Ptr<const RangeIndex>& pos = nullptr) const;
 
         //! If the parameter is of the image type, set its value.
         //! \param index Index of the parameter from 0 to GetCount()-1
         //! \param externalId Application-specific id used to identify this image during replication.
-        void SetImageValue( int index, EXTERNAL_IMAGE_ID externalId );
+		//! \param pos Only for multidimensional parametres: relevant position to set in the ranges
+		void SetImageValue( int index, EXTERNAL_IMAGE_ID externalId, const Ptr<const RangeIndex>& pos = nullptr);
 
         //! Return the value of a float parameter.
         //! \pre The parameter specified by index is a T_FLOAT.
diff --git a/Engine/Plugins/Experimental/Mutable/Source/MutableTools/Private/MuT/ASTOpMeshBindShape.cpp b/Engine/Plugins/Experimental/Mutable/Source/MutableTools/Private/MuT/ASTOpMeshBindShape.cpp
--- a/Engine/Plugins/Experimental/Mutable/Source/MutableTools/Private/MuT/ASTOpMeshBindShape.cpp
+++ b/Engine/Plugins/Experimental/Mutable/Source/MutableTools/Private/MuT/ASTOpMeshBindShape.cpp
@@ -7,7 +7,9 @@
 #include "MuR/ModelPrivate.h"
 #include "MuR/RefCounted.h"
 #include "MuR/Types.h"
+#include "MuT/ASTOpConditional.h"
 #include "MuT/ASTOpSwitch.h"
+#include "MuT/ASTOpMeshRemoveMask.h"
 #include "MuT/StreamsPrivate.h"
 
 #include <memory>
@@ -215,7 +217,7 @@ namespace mu
 
 			case OP_TYPE::ME_SWITCH:
 			{
-				// Id the switch variable and structure is the same
+				// If the switch variable and structure is the same
 				const ASTOpSwitch* MeshSwitch = reinterpret_cast<const ASTOpSwitch*>(MeshAt.get());
 				const ASTOpSwitch* ShapeSwitch = reinterpret_cast<const ASTOpSwitch*>(ShapeAt.get());
 				bool bIsSimilarSwitch = MeshSwitch->IsCompatibleWith(ShapeSwitch);
@@ -235,7 +237,6 @@ namespace mu
 					NewSwitch->def = NewBind;
 				}
 
-				// We need to copy the options because we change them
 				for (int32 v = 0; v < NewSwitch->cases.Num(); ++v)
 				{
 					if (NewSwitch->cases[v].branch)
@@ -251,6 +252,40 @@ namespace mu
 				break;
 			}
 
+
+			case OP_TYPE::ME_CONDITIONAL:
+			{
+				const ASTOpConditional* MeshConditional = reinterpret_cast<const ASTOpConditional*>(MeshAt.get());
+				const ASTOpConditional* ShapeConditional = reinterpret_cast<const ASTOpConditional*>(ShapeAt.get());
+				bool bIsSimilar = MeshConditional->condition == ShapeConditional->condition;
+				if (!bIsSimilar)
+				{
+					break;
+				}
+
+				Ptr<ASTOpConditional> NewConditional = mu::Clone<ASTOpConditional>(MeshAt);
+
+				if (NewConditional->yes)
+				{
+					Ptr<ASTOpMeshBindShape> NewBind = mu::Clone<ASTOpMeshBindShape>(this);
+					NewBind->Mesh = MeshConditional->yes.child();
+					NewBind->Shape = ShapeConditional->yes.child();
+					NewConditional->yes = NewBind;
+				}
+
+				if (NewConditional->no)
+				{
+					Ptr<ASTOpMeshBindShape> NewBind = mu::Clone<ASTOpMeshBindShape>(this);
+					NewBind->Mesh = MeshConditional->no.child();
+					NewBind->Shape = ShapeConditional->no.child();
+					NewConditional->no = NewBind;
+				}
+
+				NewOp = NewConditional;
+				break;
+			}
+
+
 			default:
 				break;
 
@@ -277,7 +312,6 @@ namespace mu
 					NewSwitch->def = NewBind;
 				}
 
-				// We need to copy the options because we change them
 				for (int32 v = 0; v < NewSwitch->cases.Num(); ++v)
 				{
 					if (NewSwitch->cases[v].branch)
@@ -292,6 +326,46 @@ namespace mu
 				break;
 			}
 
+			case OP_TYPE::ME_CONDITIONAL:
+			{
+				// Move the operation down all the paths
+				Ptr<ASTOpConditional> NewConditional = mu::Clone<ASTOpConditional>(MeshAt);
+
+				if (NewConditional->yes)
+				{
+					Ptr<ASTOpMeshBindShape> NewBind = mu::Clone<ASTOpMeshBindShape>(this);
+					NewBind->Mesh = NewConditional->yes.child();
+					NewConditional->yes = NewBind;
+				}
+
+				if (NewConditional->no)
+				{
+					Ptr<ASTOpMeshBindShape> NewBind = mu::Clone<ASTOpMeshBindShape>(this);
+					NewBind->Mesh = NewConditional->no.child();
+					NewConditional->no = NewBind;
+				}
+
+				NewOp = NewConditional;
+				break;
+			}
+
+			case OP_TYPE::ME_REMOVEMASK:
+			{
+				// We bind something that could have a part removed: we can reorder to bind the entire mesh
+				// and apply remove later at runtime.
+
+				Ptr<ASTOpMeshRemoveMask> NewRemove = mu::Clone<ASTOpMeshRemoveMask>(MeshAt);
+				if (NewRemove->source)
+				{
+					Ptr<ASTOpMeshBindShape> NewBind = mu::Clone<ASTOpMeshBindShape>(this);
+					NewBind->Mesh = NewRemove->source.child();
+					NewRemove->source = NewBind;
+				}
+
+				NewOp = NewRemove;
+				break;
+			}
+
 			default:
 				break;
 
@@ -317,7 +391,6 @@ namespace mu
 					NewSwitch->def = NewBind;
 				}
 
-				// We need to copy the options because we change them
 				for (int32 v = 0; v < NewSwitch->cases.Num(); ++v)
 				{
 					if (NewSwitch->cases[v].branch)
@@ -332,6 +405,29 @@ namespace mu
 				break;
 			}
 
+			case OP_TYPE::ME_CONDITIONAL:
+			{
+				// Move the operation down all the paths
+				Ptr<ASTOpConditional> NewConditional = mu::Clone<ASTOpConditional>(ShapeAt);
+
+				if (NewConditional->yes)
+				{
+					Ptr<ASTOpMeshBindShape> NewBind = mu::Clone<ASTOpMeshBindShape>(this);
+					NewBind->Shape = NewConditional->yes.child();
+					NewConditional->yes = NewBind;
+				}
+
+				if (NewConditional->no)
+				{
+					Ptr<ASTOpMeshBindShape> NewBind = mu::Clone<ASTOpMeshBindShape>(this);
+					NewBind->Shape = NewConditional->no.child();
+					NewConditional->no = NewBind;
+				}
+
+				NewOp = NewConditional;
+				break;
+			}
+
 			default:
 				break;
 
diff --git a/Engine/Plugins/Experimental/Mutable/Source/MutableTools/Private/MuT/ASTOpMeshDifference.cpp b/Engine/Plugins/Experimental/Mutable/Source/MutableTools/Private/MuT/ASTOpMeshDifference.cpp
new file mode 100644
--- /dev/null
+++ b/Engine/Plugins/Experimental/Mutable/Source/MutableTools/Private/MuT/ASTOpMeshDifference.cpp
@@ -0,0 +1,338 @@
+// Copyright Epic Games, Inc. All Rights Reserved.
+
+#include "MuT/ASTOpMeshDifference.h"
+
+#include "HAL/PlatformMath.h"
+#include "MuR/ModelPrivate.h"
+#include "MuR/RefCounted.h"
+#include "MuR/Types.h"
+#include "MuT/StreamsPrivate.h"
+#include "MuT/ASTOpConditional.h"
+#include "MuT/ASTOpSwitch.h"
+#include "MuT/ASTOpMeshRemoveMask.h"
+
+#include <memory>
+
+
+namespace mu
+{
+
+	//---------------------------------------------------------------------------------------------
+	ASTOpMeshDifference::ASTOpMeshDifference()
+		: Base(this), Target(this)
+	{
+	}
+
+
+	//---------------------------------------------------------------------------------------------
+	ASTOpMeshDifference::~ASTOpMeshDifference()
+	{
+		// Explicit call needed to avoid recursive destruction
+		ASTOp::RemoveChildren();
+	}
+
+
+	//---------------------------------------------------------------------------------------------
+	bool ASTOpMeshDifference::IsEqual(const ASTOp& otherUntyped) const
+	{
+		if (const ASTOpMeshDifference* Other = dynamic_cast<const ASTOpMeshDifference*>(&otherUntyped))
+		{
+			return Base == Other->Base && Target == Other->Target
+				&& bIgnoreTextureCoords == Other->bIgnoreTextureCoords
+				&& Channels == Other->Channels;
+		}
+		return false;
+	}
+
+
+	//---------------------------------------------------------------------------------------------
+	mu::Ptr<ASTOp> ASTOpMeshDifference::Clone(MapChildFuncRef mapChild) const
+	{
+		Ptr<ASTOpMeshDifference> n = new ASTOpMeshDifference();
+		n->Base = mapChild(Base.child());
+		n->Target = mapChild(Target.child());
+		n->bIgnoreTextureCoords = bIgnoreTextureCoords;
+		n->Channels = Channels;		
+		return n;
+	}
+
+
+	//---------------------------------------------------------------------------------------------
+	void ASTOpMeshDifference::ForEachChild(const TFunctionRef<void(ASTChild&)> f)
+	{
+		f(Base);
+		f(Target);
+	}
+
+
+	//---------------------------------------------------------------------------------------------
+	uint64 ASTOpMeshDifference::Hash() const
+	{
+		uint64 res = std::hash<ASTOp*>()(Base.child().get());
+		hash_combine(res, Target.child().get());
+		return res;
+	}
+
+
+	//---------------------------------------------------------------------------------------------
+	void ASTOpMeshDifference::Link(PROGRAM& program, const FLinkerOptions*)
+	{
+		// Already linked?
+		if (!linkedAddress)
+		{
+			linkedAddress = (OP::ADDRESS)program.m_opAddress.Num();
+
+			program.m_opAddress.Add((uint32_t)program.m_byteCode.Num());
+			AppendCode(program.m_byteCode, OP_TYPE::ME_DIFFERENCE);
+
+			OP::ADDRESS BaseAt = Base ? Base->linkedAddress : 0;
+			AppendCode(program.m_byteCode, BaseAt);
+
+			OP::ADDRESS TargetAt = Target ? Target->linkedAddress : 0;
+			AppendCode(program.m_byteCode, TargetAt);
+
+			AppendCode(program.m_byteCode, (uint8)bIgnoreTextureCoords);
+
+			AppendCode(program.m_byteCode, (uint8)Channels.Num());
+			for (const FChannel& b : Channels)
+			{
+				AppendCode(program.m_byteCode, b.Semantic);
+				AppendCode(program.m_byteCode, b.SemanticIndex);
+			}
+		}
+	}
+
+
+	//-------------------------------------------------------------------------------------------------
+	mu::Ptr<ASTOp> ASTOpMeshDifference::OptimiseSink(const MODEL_OPTIMIZATION_OPTIONS&, OPTIMIZE_SINK_CONTEXT&) const
+	{
+		Ptr<ASTOp> NewOp;
+
+		Ptr<ASTOp> BaseAt = Base.child();
+		if (!BaseAt)
+		{
+			return nullptr;
+		}
+
+		Ptr<ASTOp> TargetAt = Target.child();
+		if (!TargetAt)
+		{
+			return nullptr;
+		}
+
+		OP_TYPE BaseType = BaseAt->GetOpType();
+		OP_TYPE TargetType = TargetAt->GetOpType();
+
+		// See if both base and target have an operation that can be optimized in a combined way
+		if (BaseType == TargetType)
+		{
+			switch (BaseType)
+			{
+
+			case OP_TYPE::ME_SWITCH:
+			{
+				// If the switch variable and structure is the same
+				const ASTOpSwitch* BaseSwitch = reinterpret_cast<const ASTOpSwitch*>(BaseAt.get());
+				const ASTOpSwitch* TargetSwitch = reinterpret_cast<const ASTOpSwitch*>(TargetAt.get());
+				bool bIsSimilarSwitch = BaseSwitch->IsCompatibleWith(TargetSwitch);
+				if (!bIsSimilarSwitch)
+				{
+					break;
+				}
+
+				// Move the operation down all the paths
+				Ptr<ASTOpSwitch> NewSwitch = mu::Clone<ASTOpSwitch>(BaseAt);
+
+				if (NewSwitch->def)
+				{
+					Ptr<ASTOpMeshDifference> NewDiff = mu::Clone<ASTOpMeshDifference>(this);
+					NewDiff->Base = BaseSwitch->def.child();
+					NewDiff->Target = TargetSwitch->def.child();
+					NewSwitch->def = NewDiff;
+				}
+
+				for (int32 v = 0; v < NewSwitch->cases.Num(); ++v)
+				{
+					if (NewSwitch->cases[v].branch)
+					{
+						Ptr<ASTOpMeshDifference> NewDiff = mu::Clone<ASTOpMeshDifference>(this);
+						NewDiff->Base = BaseSwitch->cases[v].branch.child();
+						NewDiff->Target = TargetSwitch->FindBranch(BaseSwitch->cases[v].condition);
+						NewSwitch->cases[v].branch = NewDiff;
+					}
+				}
+
+				NewOp = NewSwitch;
+				break;
+			}
+
+
+			case OP_TYPE::ME_CONDITIONAL:
+			{
+				const ASTOpConditional* BaseConditional = reinterpret_cast<const ASTOpConditional*>(BaseAt.get());
+				const ASTOpConditional* TargetConditional = reinterpret_cast<const ASTOpConditional*>(TargetAt.get());
+				bool bIsSimilar = BaseConditional->condition == TargetConditional->condition;
+				if (!bIsSimilar)
+				{
+					break;
+				}
+
+				Ptr<ASTOpConditional> NewConditional = mu::Clone<ASTOpConditional>(BaseAt);
+
+				if (NewConditional->yes)
+				{
+					Ptr<ASTOpMeshDifference> NewDiff = mu::Clone<ASTOpMeshDifference>(this);
+					NewDiff->Base = BaseConditional->yes.child();
+					NewDiff->Target = TargetConditional->yes.child();
+					NewConditional->yes = NewDiff;
+				}
+
+				if (NewConditional->no)
+				{
+					Ptr<ASTOpMeshDifference> NewDiff = mu::Clone<ASTOpMeshDifference>(this);
+					NewDiff->Base = BaseConditional->no.child();
+					NewDiff->Target = TargetConditional->no.child();
+					NewConditional->no = NewDiff;
+				}
+
+				NewOp = NewConditional;
+				break;
+			}
+
+
+			default:
+				break;
+
+			}
+		}
+
+
+		// If not already optimized
+		if (!NewOp)
+		{
+			// Optimize only the mesh parameter
+			switch (BaseType)
+			{
+
+			case OP_TYPE::ME_SWITCH:
+			{
+				// Move the operation down all the paths
+				Ptr<ASTOpSwitch> NewSwitch = mu::Clone<ASTOpSwitch>(BaseAt);
+
+				if (NewSwitch->def)
+				{
+					Ptr<ASTOpMeshDifference> NewDiff = mu::Clone<ASTOpMeshDifference>(this);
+					NewDiff->Base = NewSwitch->def.child();
+					NewSwitch->def = NewDiff;
+				}
+
+				for (int32 v = 0; v < NewSwitch->cases.Num(); ++v)
+				{
+					if (NewSwitch->cases[v].branch)
+					{
+						Ptr<ASTOpMeshDifference> NewDiff = mu::Clone<ASTOpMeshDifference>(this);
+						NewDiff->Base = NewSwitch->cases[v].branch.child();
+						NewSwitch->cases[v].branch = NewDiff;
+					}
+				}
+
+				NewOp = NewSwitch;
+				break;
+			}
+
+			case OP_TYPE::ME_CONDITIONAL:
+			{
+				// Move the operation down all the paths
+				Ptr<ASTOpConditional> NewConditional = mu::Clone<ASTOpConditional>(BaseAt);
+
+				if (NewConditional->yes)
+				{
+					Ptr<ASTOpMeshDifference> NewDiff = mu::Clone<ASTOpMeshDifference>(this);
+					NewDiff->Base = NewConditional->yes.child();
+					NewConditional->yes = NewDiff;
+				}
+
+				if (NewConditional->no)
+				{
+					Ptr<ASTOpMeshDifference> NewDiff = mu::Clone<ASTOpMeshDifference>(this);
+					NewDiff->Base = NewConditional->no.child();
+					NewConditional->no = NewDiff;
+				}
+
+				NewOp = NewConditional;
+				break;
+			}
+
+			default:
+				break;
+
+			}
+		}
+
+		// If not already optimized
+		if (!NewOp)
+		{
+			// Optimize only the shape parameter
+			switch (TargetType)
+			{
+
+			case OP_TYPE::ME_SWITCH:
+			{
+				// Move the operation down all the paths
+				Ptr<ASTOpSwitch> NewSwitch = mu::Clone<ASTOpSwitch>(TargetAt);
+
+				if (NewSwitch->def)
+				{
+					Ptr<ASTOpMeshDifference> NewDiff = mu::Clone<ASTOpMeshDifference>(this);
+					NewDiff->Target = NewSwitch->def.child();
+					NewSwitch->def = NewDiff;
+				}
+
+				for (int32 v = 0; v < NewSwitch->cases.Num(); ++v)
+				{
+					if (NewSwitch->cases[v].branch)
+					{
+						Ptr<ASTOpMeshDifference> NewDiff = mu::Clone<ASTOpMeshDifference>(this);
+						NewDiff->Target = NewSwitch->cases[v].branch.child();
+						NewSwitch->cases[v].branch = NewDiff;
+					}
+				}
+
+				NewOp = NewSwitch;
+				break;
+			}
+
+			case OP_TYPE::ME_CONDITIONAL:
+			{
+				// Move the operation down all the paths
+				Ptr<ASTOpConditional> NewConditional = mu::Clone<ASTOpConditional>(TargetAt);
+
+				if (NewConditional->yes)
+				{
+					Ptr<ASTOpMeshDifference> NewDiff = mu::Clone<ASTOpMeshDifference>(this);
+					NewDiff->Target = NewConditional->yes.child();
+					NewConditional->yes = NewDiff;
+				}
+
+				if (NewConditional->no)
+				{
+					Ptr<ASTOpMeshDifference> NewDiff = mu::Clone<ASTOpMeshDifference>(this);
+					NewDiff->Target = NewConditional->no.child();
+					NewConditional->no = NewDiff;
+				}
+
+				NewOp = NewConditional;
+				break;
+			}
+
+			default:
+				break;
+
+			}
+		}
+
+
+		return NewOp;
+	}
+
+}
diff --git a/Engine/Plugins/Experimental/Mutable/Source/MutableTools/Private/MuT/ASTOpMeshDifference.h b/Engine/Plugins/Experimental/Mutable/Source/MutableTools/Private/MuT/ASTOpMeshDifference.h
new file mode 100644
--- /dev/null
+++ b/Engine/Plugins/Experimental/Mutable/Source/MutableTools/Private/MuT/ASTOpMeshDifference.h
@@ -0,0 +1,61 @@
+// Copyright Epic Games, Inc. All Rights Reserved.
+
+#pragma once
+
+#include "HAL/Platform.h"
+#include "MuR/MemoryPrivate.h"
+#include "MuR/Operations.h"
+#include "MuR/Ptr.h"
+#include "MuT/AST.h"
+
+
+namespace mu
+{
+	struct PROGRAM;
+
+	/** Calculate the difference between two meshes (usually to calculate a morph delta). */
+	class ASTOpMeshDifference : public ASTOp
+	{
+	public:
+
+		/** Base mesh of the morph. */
+		ASTChild Base;
+
+		/** Target mesh of the morph. */
+		ASTChild Target;
+
+		/** If true the texture coordinates will not be morphed.This is only relevant if the list in channelSemantic is empty. */
+		bool bIgnoreTextureCoords;
+
+		/** Channels to diff. */
+		struct FChannel
+		{
+			uint8 Semantic=0;
+			uint8 SemanticIndex=0;
+
+			inline bool operator==( const FChannel& Other ) const
+			{
+				return Semantic == Other.Semantic && SemanticIndex == Other.SemanticIndex;
+			}
+		};
+		TArray<FChannel> Channels;
+
+	public:
+
+		ASTOpMeshDifference();
+		ASTOpMeshDifference(const ASTOpMeshDifference&) = delete;
+		~ASTOpMeshDifference() override;
+
+		// ASTOp interface
+		OP_TYPE GetOpType() const override { return OP_TYPE::ME_DIFFERENCE; }
+		uint64 Hash() const override;
+		void ForEachChild(const TFunctionRef<void(ASTChild&)>) override;
+		bool IsEqual(const ASTOp& otherUntyped) const override;
+		Ptr<ASTOp> Clone(MapChildFuncRef mapChild) const override;
+		void Link(PROGRAM& program, const FLinkerOptions*) override;
+		Ptr<ASTOp> OptimiseSink(const MODEL_OPTIMIZATION_OPTIONS&, OPTIMIZE_SINK_CONTEXT&) const override;
+
+	};
+
+}
+
diff --git a/Engine/Plugins/Experimental/Mutable/Source/MutableTools/Private/MuT/CodeGenerator.cpp b/Engine/Plugins/Experimental/Mutable/Source/MutableTools/Private/MuT/CodeGenerator.cpp
--- a/Engine/Plugins/Experimental/Mutable/Source/MutableTools/Private/MuT/CodeGenerator.cpp
+++ b/Engine/Plugins/Experimental/Mutable/Source/MutableTools/Private/MuT/CodeGenerator.cpp
@@ -26,6 +26,7 @@
 #include "MuT/ASTOpMeshClipMorphPlane.h"
 #include "MuT/ASTOpMeshMaskClipMesh.h"
 #include "MuT/ASTOpMeshRemoveMask.h"
+#include "MuT/ASTOpMeshDifference.h"
 #include "MuT/ASTOpParameter.h"
 #include "MuT/CodeGenerator_SecondPass.h"
 #include "MuT/CodeOptimiser.h"
@@ -1074,33 +1075,32 @@ namespace mu
                     GenerateMesh(MorphTargetMeshOptions, morphResult, pMorph );
 
 					// BaseMorph generation through mesh diff
-					Ptr<ASTOpFixed> diffBase;
+					Ptr<ASTOpMeshDifference> diffBase;
 					{
-						Ptr<ASTOpFixed> op = new ASTOpFixed();
-						op->op.type = OP_TYPE::ME_DIFFERENCE;
-						op->SetChild(op->op.args.MeshDifference.base, meshResults.meshOp);
-						op->SetChild(op->op.args.MeshDifference.target, meshResults.meshOp);
+						// \TODO: Optimize by setting the identity morph constant instead.
+						Ptr<ASTOpMeshDifference> op = new ASTOpMeshDifference();
+						op->Base = meshResults.meshOp;
+						op->Target = meshResults.meshOp;
 
 						// Morphing tex coords here is not supported:
 						// Generating the homogoneous UVs is difficult since we don't have the base
 						// layout yet.                       
-						op->op.args.MeshDifference.ignoreTextureCoords = 1;
+						op->bIgnoreTextureCoords = true;
 						diffBase = op;
 					}
 
 					// Morph generation through mesh diff
 					Ptr<ASTOp> targetAd = morphResult.meshOp;
-                    Ptr<ASTOpFixed> diffAd;
+                    Ptr<ASTOpMeshDifference> diffAd;
                     {
-                        Ptr<ASTOpFixed> op = new ASTOpFixed();
-                        op->op.type = OP_TYPE::ME_DIFFERENCE;
-                        op->SetChild(op->op.args.MeshDifference.base, meshResults.baseMeshOp );
-                        op->SetChild(op->op.args.MeshDifference.target, targetAd );
+                        Ptr<ASTOpMeshDifference> op = new ASTOpMeshDifference();
+                        op->Base = meshResults.baseMeshOp;
+                        op->Target = targetAd;
 					
                         // Morphing tex coords here is not supported:
                         // Generating the homogoneous UVs is difficult since we don't have the base
                         // layout yet.                       
-                        op->op.args.MeshDifference.ignoreTextureCoords = 1;
+                        op->bIgnoreTextureCoords = true;
                         diffAd = op;
                     }
 
diff --git a/Engine/Plugins/Experimental/Mutable/Source/MutableTools/Private/MuT/CodeGenerator_Image.cpp b/Engine/Plugins/Experimental/Mutable/Source/MutableTools/Private/MuT/CodeGenerator_Image.cpp
--- a/Engine/Plugins/Experimental/Mutable/Source/MutableTools/Private/MuT/CodeGenerator_Image.cpp
+++ b/Engine/Plugins/Experimental/Mutable/Source/MutableTools/Private/MuT/CodeGenerator_Image.cpp
@@ -271,6 +271,14 @@ namespace mu
 			op->parameter.m_uid = node.m_uid;
             op->parameter.m_type = PARAMETER_TYPE::T_IMAGE;
 
+			// Generate the code for the ranges
+			for (int32 a = 0; a < node.m_ranges.Num(); ++a)
+			{
+				RANGE_GENERATION_RESULT rangeResult;
+				GenerateRange(rangeResult, node.m_ranges[a]);
+				op->ranges.Emplace(op.get(), rangeResult.sizeOp, rangeResult.rangeName, rangeResult.rangeUID);
+			}
+
             m_nodeVariables[node.m_pNode] = op;
         }
         else
diff --git a/Engine/Plugins/Experimental/Mutable/Source/MutableTools/Private/MuT/CodeGenerator_Mesh.cpp b/Engine/Plugins/Experimental/Mutable/Source/MutableTools/Private/MuT/CodeGenerator_Mesh.cpp
--- a/Engine/Plugins/Experimental/Mutable/Source/MutableTools/Private/MuT/CodeGenerator_Mesh.cpp
+++ b/Engine/Plugins/Experimental/Mutable/Source/MutableTools/Private/MuT/CodeGenerator_Mesh.cpp
@@ -27,8 +27,8 @@
 #include "MuT/ASTOpMeshFormat.h"
 #include "MuT/ASTOpMeshGeometryOperation.h"
 #include "MuT/ASTOpMeshMorphReshape.h"
-//#include "MuT/ASTOpMeshRemapIndices.h"
 #include "MuT/ASTOpMeshTransform.h"
+#include "MuT/ASTOpMeshDifference.h"
 #include "MuT/ASTOpSwitch.h"
 #include "MuT/CodeGenerator.h"
 #include "MuT/CodeGenerator_FirstPass.h"
@@ -468,11 +468,10 @@ class Node;
     {
         NodeMeshMakeMorph::Private& node = *InMakeMorphNode->GetPrivate();
 
-        Ptr<ASTOpFixed> op = new ASTOpFixed();
-        op->op.type = OP_TYPE::ME_DIFFERENCE;
+        Ptr<ASTOpMeshDifference> op = new ASTOpMeshDifference();
 
         // \todo Texcoords are broken?
-        op->op.args.MeshDifference.ignoreTextureCoords = 1;
+        op->bIgnoreTextureCoords = true;
 
         // Base
         FMeshGenerationResult BaseResult;
@@ -483,7 +482,7 @@ class Node;
 			BaseOptions.bLayouts = false;
 			GenerateMesh(BaseOptions, BaseResult, node.m_pBase );
 
-            op->SetChild( op->op.args.MeshDifference.base, BaseResult.meshOp );
+            op->Base = BaseResult.meshOp;
         }
         else
         {
@@ -503,7 +502,7 @@ class Node;
 			FMeshGenerationResult TargetResult;
             GenerateMesh( TargetOptions, TargetResult, node.m_pTarget );
 
-            op->SetChild( op->op.args.MeshDifference.target, TargetResult.meshOp );
+            op->Target = TargetResult.meshOp;
         }
         else
         {
@@ -644,35 +643,22 @@ class Node;
                 }
                 else
                 {
-                    Ptr<ASTOpFixed> dop = new ASTOpFixed();
-                    dop->op.type = OP_TYPE::ME_DIFFERENCE;
-                    dop->SetChild( dop->op.args.MeshDifference.base, base );
-                    dop->SetChild( dop->op.args.MeshDifference.target, TargetResult.meshOp );
+                    Ptr<ASTOpMeshDifference> dop = new ASTOpMeshDifference();
+                    dop->Base = base;
+                    dop->Target = TargetResult.meshOp;
 
                     // \todo Texcoords are broken?
-                    dop->op.args.MeshDifference.ignoreTextureCoords = 1;
+                    dop->bIgnoreTextureCoords = true;
 
-                    if ( node.m_channels.Num()>MUTABLE_OP_MAX_MORPH_CHANNELS )
+                    for ( size_t c=0; c<node.m_channels.Num(); ++c)
                     {
-                        char temp[256];
-                        mutable_snprintf( temp, 256,
-                                          "Morph uses too many channels [%d]. The maximum is [%d].",
-                                          node.m_channels.Num(),
-                                          MUTABLE_OP_MAX_MORPH_CHANNELS );
-                        m_pErrorLog->GetPrivate()->Add( temp, ELMT_ERROR, node.m_errorContext );
-                    }
-
-                    for ( size_t c=0;
-                          c<node.m_channels.Num() && c<MUTABLE_OP_MAX_MORPH_CHANNELS;
-                          ++c)
-                    {
-                        check( node.m_channels[c].semantic<256 );
-                        dop->op.args.MeshDifference.channelSemantic[c] =
-                                uint8_t( node.m_channels[c].semantic );
-
-                        check( node.m_channels[c].semanticIndex<256 );
-                        dop->op.args.MeshDifference.channelSemanticIndex[c] =
-                                uint8_t( node.m_channels[c].semanticIndex );
+                        check( node.m_channels[c].semantic < 256 );
+						check(node.m_channels[c].semanticIndex < 256);
+						
+						ASTOpMeshDifference::FChannel Channel;
+						Channel.Semantic = uint8(node.m_channels[c].semantic);
+						Channel.SemanticIndex = uint8(node.m_channels[c].semanticIndex);
+						dop->Channels.Add(Channel);
                     }
 
                     op->SetChild( op->op.args.MeshInterpolate.targets[count-1], dop );
diff --git a/Engine/Plugins/Experimental/Mutable/Source/MutableTools/Private/MuT/DataPacker.cpp b/Engine/Plugins/Experimental/Mutable/Source/MutableTools/Private/MuT/DataPacker.cpp
--- a/Engine/Plugins/Experimental/Mutable/Source/MutableTools/Private/MuT/DataPacker.cpp
+++ b/Engine/Plugins/Experimental/Mutable/Source/MutableTools/Private/MuT/DataPacker.cpp
@@ -22,6 +22,7 @@
 #include "MuT/ASTOpInstanceAdd.h"
 #include "MuT/ASTOpMeshExtractLayoutBlocks.h"
 #include "MuT/ASTOpMeshRemoveMask.h"
+#include "MuT/ASTOpMeshDifference.h"
 
 #include <cstdint>
 #include <memory>
@@ -547,13 +548,13 @@ namespace mu
             {
                 recurse = false;
 
-				const ASTOpFixed* op = dynamic_cast<const ASTOpFixed*>(node.get());
+				const ASTOpMeshDifference* op = dynamic_cast<const ASTOpMeshDifference*>(node.get());
 
                 uint64_t newState = currentSemantics;
                 newState |= (UINT64_C(1)<<MBS_VERTEXINDEX);
-                RecurseWithState( op->children[op->op.args.MeshDifference.base].child(), newState );
+                RecurseWithState( op->Base.child(), newState );
 
-                RecurseWithState( op->children[op->op.args.MeshDifference.target].child(), currentSemantics );
+                RecurseWithState( op->Target.child(), currentSemantics );
                 break;
              }
 
diff --git a/Engine/Plugins/Experimental/Mutable/Source/MutableTools/Private/MuT/NodeImageParameter.cpp b/Engine/Plugins/Experimental/Mutable/Source/MutableTools/Private/MuT/NodeImageParameter.cpp
--- a/Engine/Plugins/Experimental/Mutable/Source/MutableTools/Private/MuT/NodeImageParameter.cpp
+++ b/Engine/Plugins/Experimental/Mutable/Source/MutableTools/Private/MuT/NodeImageParameter.cpp
@@ -92,6 +92,25 @@ namespace mu
 		}
 	}
 
+
+	//---------------------------------------------------------------------------------------------
+	void NodeImageParameter::SetRangeCount(int i)
+	{
+		check(i >= 0);
+		m_pD->m_ranges.SetNum(i);
+	}
+
+
+	//---------------------------------------------------------------------------------------------
+	void NodeImageParameter::SetRange(int i, NodeRangePtr pRange)
+	{
+		check(i >= 0 && i<int(m_pD->m_ranges.Num()));
+		if (i >= 0 && i<int(m_pD->m_ranges.Num()))
+		{
+			m_pD->m_ranges[i] = pRange;
+		}
+	}
+
 }
 
 
diff --git a/Engine/Plugins/Experimental/Mutable/Source/MutableTools/Private/MuT/NodeImageParameterPrivate.h b/Engine/Plugins/Experimental/Mutable/Source/MutableTools/Private/MuT/NodeImageParameterPrivate.h
--- a/Engine/Plugins/Experimental/Mutable/Source/MutableTools/Private/MuT/NodeImageParameterPrivate.h
+++ b/Engine/Plugins/Experimental/Mutable/Source/MutableTools/Private/MuT/NodeImageParameterPrivate.h
@@ -4,6 +4,7 @@
 
 #include "MuT/NodeImageParameter.h"
 #include "MuT/NodeImagePrivate.h"
+#include "MuT/NodeRange.h"
 
 
 namespace mu
@@ -22,25 +23,33 @@ namespace mu
 		string m_name;
 		string m_uid;
 
+		TArray<Ptr<NodeRange>> m_ranges;
+
 		//!
 		void Serialise( OutputArchive& arch ) const
 		{
-            uint32_t ver = 1;
+            uint32 ver = 2;
 			arch << ver;
 
 			arch << m_name;
 			arch << m_uid;
+			arch << m_ranges;
 		}
 
 		//!
 		void Unserialise( InputArchive& arch )
 		{
-            uint32_t ver;
+            uint32 ver;
 			arch >> ver;
-            check(ver==1);
+            check(ver<=2);
 
 			arch >> m_name;
             arch >> m_uid;
+
+			if (ver >= 2)
+			{
+				arch >> m_ranges;
+			}
 		}
 	};
 
diff --git a/Engine/Plugins/Experimental/Mutable/Source/MutableTools/Public/MuT/NodeImageParameter.h b/Engine/Plugins/Experimental/Mutable/Source/MutableTools/Public/MuT/NodeImageParameter.h
--- a/Engine/Plugins/Experimental/Mutable/Source/MutableTools/Public/MuT/NodeImageParameter.h
+++ b/Engine/Plugins/Experimental/Mutable/Source/MutableTools/Public/MuT/NodeImageParameter.h
@@ -10,14 +10,9 @@
 
 namespace mu
 {
-
-    class NodeImageParameter;
-    typedef Ptr<NodeImageParameter> NodeImageParameterPtr;
-    typedef Ptr<const NodeImageParameter> NodeImageParameterPtrConst;
-
 	class InputArchive;
 	class OutputArchive;
-
+	class NodeRange;
 
     //! Node that defines a Image model parameter.
 	//! \ingroup model
@@ -33,21 +28,19 @@ namespace mu
 
 		void SerialiseWrapper(OutputArchive& arch) const override;
 		static void Serialise( const NodeImageParameter* pNode, OutputArchive& arch );
-        static NodeImageParameterPtr StaticUnserialise( InputArchive& arch );
+        static Ptr<NodeImageParameter> StaticUnserialise( InputArchive& arch );
 
 
 		//-----------------------------------------------------------------------------------------
 		// Node Interface
 		//-----------------------------------------------------------------------------------------
 
-		
-
 		const NODE_TYPE* GetType() const override;
 		static const NODE_TYPE* GetStaticType();
 
 		int GetInputCount() const override;
 		Node* GetInputNode( int i ) const override;
-		void SetInputNode( int i, NodePtr pNode ) override;
+		void SetInputNode( int i, Ptr<Node> Node ) override;
 
 		//-----------------------------------------------------------------------------------------
 		// Own Interface
@@ -65,11 +58,10 @@ namespace mu
 		//! Set the uid of the parameter.
 		void SetUid( const char* );
 
-		//! Get the default value of the parameter.
-        //ImagePtrConst GetDefaultValue() const;
-
-		//! Set the default value of the parameter.
-        //void SetDefaultValue( ImagePtr );
+		//! Set the number of ranges (dimensions) for this parameter.
+		//! By default a parameter has 0 ranges, meaning it only has one value.
+		void SetRangeCount(int Index);
+		void SetRange(int Index, Ptr<NodeRange> Range);
 
 		//-----------------------------------------------------------------------------------------
 		// Interface pattern
