#rnx

diff --git a/Engine/Plugins/Experimental/Mutable/Source/CustomizableObjectEditor/Private/MuCOE/GenerateMutableSource/GenerateMutableSourceMesh.cpp b/Engine/Plugins/Experimental/Mutable/Source/CustomizableObjectEditor/Private/MuCOE/GenerateMutableSource/GenerateMutableSourceMesh.cpp
--- a/Engine/Plugins/Experimental/Mutable/Source/CustomizableObjectEditor/Private/MuCOE/GenerateMutableSource/GenerateMutableSourceMesh.cpp
+++ b/Engine/Plugins/Experimental/Mutable/Source/CustomizableObjectEditor/Private/MuCOE/GenerateMutableSource/GenerateMutableSourceMesh.cpp
@@ -2893,21 +2893,28 @@ mu::NodeMeshPtr GenerateMutableSourceMesh(const UEdGraphPin * Pin,
 		{
 			mu::Ptr<mu::NodeMesh> InputMeshNode = GenerateMutableSourceMesh(ConnectedPin, GenerationContext, MeshData);
 
-			TArray<FString> ArrayBoneName;
-			TArray<FTransform> ArrayTransform;
-			UCustomizableObjectNodeAnimationPose::StaticRetrievePoseInformation(TypedNode->PoseAsset, GenerationContext.GetCurrentComponentInfo().RefSkeletalMesh, ArrayBoneName, ArrayTransform);
-			mu::NodeMeshApplyPosePtr NodeMeshApplyPose = CreateNodeMeshApplyPose(InputMeshNode, GenerationContext.Object, ArrayBoneName, ArrayTransform);
-
-			if (NodeMeshApplyPose)
+			if (TypedNode->PoseAsset && GenerationContext.GetCurrentComponentInfo().RefSkeletalMesh)
 			{
-				Result = NodeMeshApplyPose;
+				TArray<FString> ArrayBoneName;
+				TArray<FTransform> ArrayTransform;
+				UCustomizableObjectNodeAnimationPose::StaticRetrievePoseInformation(TypedNode->PoseAsset, GenerationContext.GetCurrentComponentInfo().RefSkeletalMesh, ArrayBoneName, ArrayTransform);
+				mu::NodeMeshApplyPosePtr NodeMeshApplyPose = CreateNodeMeshApplyPose(InputMeshNode, GenerationContext.Object, ArrayBoneName, ArrayTransform);
+
+				if (NodeMeshApplyPose)
+				{
+					Result = NodeMeshApplyPose;
+				}
+				else
+				{
+					FString msg = FString::Printf(TEXT("Couldn't get bone transform information from a Pose Asset."));
+					GenerationContext.Compiler->CompilerLog(FText::FromString(msg), Node);
+
+					Result = nullptr;
+				}
 			}
 			else
 			{
-				FString msg = FString::Printf(TEXT("Couldn't get bone transform information from a Pose Asset."));
-				GenerationContext.Compiler->CompilerLog(FText::FromString(msg), Node);
-
-				Result = nullptr;
+				Result = InputMeshNode;
 			}
 		}
 	}
diff --git a/Engine/Plugins/Experimental/Mutable/Source/CustomizableObjectEditor/Private/MuCOE/GraphTraversal.cpp b/Engine/Plugins/Experimental/Mutable/Source/CustomizableObjectEditor/Private/MuCOE/GraphTraversal.cpp
--- a/Engine/Plugins/Experimental/Mutable/Source/CustomizableObjectEditor/Private/MuCOE/GraphTraversal.cpp
+++ b/Engine/Plugins/Experimental/Mutable/Source/CustomizableObjectEditor/Private/MuCOE/GraphTraversal.cpp
@@ -29,6 +29,7 @@
 #include "MuCOE/Nodes/CustomizableObjectNodeSkeletalMesh.h"
 #include "MuCOE/Nodes/CustomizableObjectNodeStaticMesh.h"
 #include "MuCOE/Nodes/CustomizableObjectNodeTable.h"
+#include "MuCOE/Nodes/CustomizableObjectNodeAnimationPose.h"
 #include "Templates/Casts.h"
 #include "Templates/Function.h"
 #include "UObject/NameTypes.h"
@@ -497,6 +498,14 @@ const UEdGraphPin* FindMeshBaseSource(const UEdGraphPin& Pin, const bool bOnlyLo
 		}
 	}
 
+	else if (const UCustomizableObjectNodeAnimationPose* NodeMeshPose = Cast<UCustomizableObjectNodeAnimationPose>(Node))
+	{
+		if (const UEdGraphPin* ConnectedPin = FollowInputPin(*NodeMeshPose->GetInputMeshPin()))
+		{
+			return FindMeshBaseSource(*ConnectedPin, bOnlyLookForStaticMesh);
+		}
+	}
+
 	else
 	{
 		unimplemented(); // Case missing.
diff --git a/Engine/Plugins/Experimental/Mutable/Source/CustomizableObjectEditor/Private/MuCOE/Nodes/CustomizableObjectNodeAnimationPose.cpp b/Engine/Plugins/Experimental/Mutable/Source/CustomizableObjectEditor/Private/MuCOE/Nodes/CustomizableObjectNodeAnimationPose.cpp
--- a/Engine/Plugins/Experimental/Mutable/Source/CustomizableObjectEditor/Private/MuCOE/Nodes/CustomizableObjectNodeAnimationPose.cpp
+++ b/Engine/Plugins/Experimental/Mutable/Source/CustomizableObjectEditor/Private/MuCOE/Nodes/CustomizableObjectNodeAnimationPose.cpp
@@ -34,9 +34,6 @@ struct FPropertyChangedEvent;
 #define LOCTEXT_NAMESPACE "CustomizableObjectEditor"
 
 
-PRAGMA_DISABLE_DEPRECATION_WARNINGS
-
-
 UCustomizableObjectNodeAnimationPose::UCustomizableObjectNodeAnimationPose() : Super()
 	, PoseAsset(nullptr)
 {
@@ -98,97 +95,99 @@ FLinearColor UCustomizableObjectNodeAnimationPose::GetNodeTitleColor() const
 
 void UCustomizableObjectNodeAnimationPose::StaticRetrievePoseInformation(UPoseAsset* PoseAsset, USkeletalMesh* RefSkeletalMesh, TArray<FString>& OutArrayBoneName, TArray<FTransform>& OutArrayTransform)
 {
-
-	// Need this for FCompactPose
-	FMemMark Mark(FMemStack::Get());
-	
-	UDebugSkelMeshComponent* SkeletalMeshComponent = NewObject<UDebugSkelMeshComponent>();
-	SkeletalMeshComponent->SetSkeletalMesh(RefSkeletalMesh);
-	SkeletalMeshComponent->SetAnimationMode(EAnimationMode::AnimationBlueprint);
-	SkeletalMeshComponent->AllocateTransformData();
-	SkeletalMeshComponent->SetAnimation(PoseAsset);
-	SkeletalMeshComponent->RefreshBoneTransforms();
-	SkeletalMeshComponent->InitAnim(false);
-
-	UAnimInstance* AnimInstance = SkeletalMeshComponent->GetAnimInstance();
-	FBoneContainer& BoneContainerCopy = AnimInstance->GetRequiredBones();
-
-	FCompactPose OutPose; // Needs a FMemMark declared before in the stack context so that the memory allocated by the FCompactPose is freed correctly
-	OutPose.SetBoneContainer(&BoneContainerCopy);
-
-	FBlendedCurve OutCurve;
-	FStackCustomAttributes OutAttributes;
-	FAnimationPoseData OutAnimData(OutPose, OutCurve, OutAttributes);
-	PoseAsset->GetBaseAnimationPose(OutAnimData);
-	OutPose = OutAnimData.GetPose();
-	OutCurve = OutAnimData.GetCurve();
-
-	FBlendedHeapCurve AnimCurves = SkeletalMeshComponent->AnimCurves;
-	const int32 NumElement = AnimCurves.CurveWeights.Num();
-
-	for (int i = 0; i < NumElement; ++i)
+	if (PoseAsset && RefSkeletalMesh)
 	{
-		OutCurve.CurveWeights.Add(AnimCurves.CurveWeights[i]);
-	}
+		// Need this for the FCompactPose below
+		FMemMark Mark(FMemStack::Get());
+
+		UDebugSkelMeshComponent* SkeletalMeshComponent = NewObject<UDebugSkelMeshComponent>();
+		SkeletalMeshComponent->SetSkeletalMesh(RefSkeletalMesh);
+		SkeletalMeshComponent->SetAnimationMode(EAnimationMode::AnimationBlueprint);
+		SkeletalMeshComponent->AllocateTransformData();
+		SkeletalMeshComponent->SetAnimation(PoseAsset);
+		SkeletalMeshComponent->RefreshBoneTransforms();
+		SkeletalMeshComponent->InitAnim(false);
+
+		UAnimInstance* AnimInstance = SkeletalMeshComponent->GetAnimInstance();
+		FBoneContainer& BoneContainerCopy = AnimInstance->GetRequiredBones();
+
+		// Needs a FMemMark declared before in the stack context so that the memory allocated by the FCompactPose is freed correctly
+		FCompactPose OutPose;
+		OutPose.SetBoneContainer(&BoneContainerCopy);
+
+		FBlendedCurve OutCurve;
+		UE::Anim::FStackAttributeContainer OutAttributes;
+		FAnimationPoseData OutAnimData(OutPose, OutCurve, OutAttributes);
+		PoseAsset->GetBaseAnimationPose(OutAnimData);
+		OutPose = OutAnimData.GetPose();
+		OutCurve = OutAnimData.GetCurve();
+
+		FBlendedHeapCurve AnimCurves = SkeletalMeshComponent->AnimCurves;
+		const int32 NumElement = AnimCurves.CurveWeights.Num();
+
+		for (int32 i = 0; i < NumElement; ++i)
+		{
+			OutCurve.CurveWeights.Add(AnimCurves.CurveWeights[i]);
+		}
 
-	OutCurve.bInitialized = AnimCurves.bInitialized;
+		OutCurve.bInitialized = AnimCurves.bInitialized;
 
-	// Assuming one single pose, with a weigth set to 1.0
-	FAnimExtractContext ExtractionContext;
-	ExtractionContext.bExtractRootMotion = false;
-	ExtractionContext.CurrentTime = 0.0f;
+		// Assuming one single pose, with a weigth set to 1.0
+		FAnimExtractContext ExtractionContext;
+		ExtractionContext.bExtractRootMotion = false;
+		ExtractionContext.CurrentTime = 0.0f;
 
-	OutCurve.UIDToArrayIndexLUT = AnimCurves.UIDToArrayIndexLUT;
+		OutCurve.UIDToArrayIndexLUT = AnimCurves.UIDToArrayIndexLUT;
 
-	const TArray<FSmartName>& PoseNames = PoseAsset->GetPoseNames();
-	ExtractionContext.PoseCurves.Add(FPoseCurve(0, PoseNames[0].UID, 1.0f));
-	FAnimationPoseData SecondOutAnimData(OutPose, OutCurve, OutAttributes);
-	PoseAsset->GetAnimationPose(SecondOutAnimData, ExtractionContext);
-	OutPose = OutAnimData.GetPose();
-	OutCurve = OutAnimData.GetCurve();
+		const TArray<FSmartName>& PoseNames = PoseAsset->GetPoseNames();
+		ExtractionContext.PoseCurves.Add(FPoseCurve(0, PoseNames[0].UID, 1.0f));
+		FAnimationPoseData SecondOutAnimData(OutPose, OutCurve, OutAttributes);
+		PoseAsset->GetAnimationPose(SecondOutAnimData, ExtractionContext);
+		OutPose = OutAnimData.GetPose();
+		OutCurve = OutAnimData.GetCurve();
 
-	const TArray<FTransform, FAnimStackAllocator>& ArrayPoseBoneTransform = OutPose.GetBones();
-	const TArray<FBoneIndexType>& ArrayPoseBoneIndices = OutPose.GetBoneContainer().GetBoneIndicesArray();
+		const TArray<FTransform, FAnimStackAllocator>& ArrayPoseBoneTransform = OutPose.GetBones();
+		const TArray<FBoneIndexType>& ArrayPoseBoneIndices = OutPose.GetBoneContainer().GetBoneIndicesArray();
 
-	for (int32 ArrayIndex = 0; ArrayIndex < ArrayPoseBoneIndices.Num(); ArrayIndex++)
-	{
-		int32 const& PoseBoneIndex = ArrayPoseBoneIndices[ArrayIndex];
+		for (int32 ArrayIndex = 0; ArrayIndex < ArrayPoseBoneIndices.Num(); ArrayIndex++)
+		{
+			FBoneIndexType const& PoseBoneIndex = ArrayPoseBoneIndices[ArrayIndex];
 
-		checkSlow(PoseBoneIndex != INDEX_NONE);
+			checkSlow(PoseBoneIndex != INDEX_NONE);
 
-		FTransform CumulativePoseTransform;
-		int32 ParentIndex = PoseBoneIndex;
+			FTransform CumulativePoseTransform;
+			FBoneIndexType ParentIndex = PoseBoneIndex;
 
-		while (ParentIndex > 0)
-		{
-			for (int32 IndicesIndex = 0; IndicesIndex < ArrayPoseBoneIndices.Num(); ++IndicesIndex)
+			while (ParentIndex > 0)
 			{
-				const uint16 BoneIndex = ArrayPoseBoneIndices[IndicesIndex];
-				if (BoneIndex == ParentIndex)
+				for (int32 IndicesIndex = 0; IndicesIndex < ArrayPoseBoneIndices.Num(); ++IndicesIndex)
 				{
-					CumulativePoseTransform = CumulativePoseTransform * ArrayPoseBoneTransform[IndicesIndex];
-					break;
-				} 
-				else if (BoneIndex > ParentIndex)
-				{
-					break; // Break for, ArrayPoseBoneIndices is sorted.
+					const FBoneIndexType BoneIndex = ArrayPoseBoneIndices[IndicesIndex];
+
+					if (BoneIndex == ParentIndex)
+					{
+						CumulativePoseTransform = CumulativePoseTransform * ArrayPoseBoneTransform[IndicesIndex];
+						break;
+					}
+					else if (BoneIndex > ParentIndex)
+					{
+						break; // Break for, ArrayPoseBoneIndices is sorted.
+					}
 				}
-			}
 
-			ParentIndex = RefSkeletalMesh->GetRefSkeleton().GetParentIndex(ParentIndex);
-		}
+				ParentIndex = RefSkeletalMesh->GetRefSkeleton().GetParentIndex(ParentIndex);
+			}
 
-		FString Name = RefSkeletalMesh->GetRefSkeleton().GetBoneName(PoseBoneIndex).ToString();
-		int32 SMBoneIndex = SkeletalMeshComponent->GetBoneIndex(FName(*Name));
-		FTransform BoneToComponentTranform = SkeletalMeshComponent->GetEditableComponentSpaceTransforms()[SMBoneIndex];
-		FTransform TransformToAdd = BoneToComponentTranform.Inverse() * CumulativePoseTransform;
+			FString Name = RefSkeletalMesh->GetRefSkeleton().GetBoneName(PoseBoneIndex).ToString();
+			int32 SMBoneIndex = SkeletalMeshComponent->GetBoneIndex(FName(*Name));
+			FTransform BoneToComponentTranform = SkeletalMeshComponent->GetEditableComponentSpaceTransforms()[SMBoneIndex];
+			FTransform TransformToAdd = BoneToComponentTranform.Inverse() * CumulativePoseTransform;
 
-		OutArrayBoneName.Add(Name);
-		OutArrayTransform.Add(TransformToAdd);
+			OutArrayBoneName.Add(Name);
+			OutArrayTransform.Add(TransformToAdd);
+		}
 	}
-
 }
 
-#undef LOCTEXT_NAMESPACE
 
-PRAGMA_ENABLE_DEPRECATION_WARNINGS
+#undef LOCTEXT_NAMESPACE
