[mutable] Enable clothing long range tethers.

#rnx

diff --git a/Engine/Plugins/Experimental/Mutable/Source/CustomizableObject/Private/MuCO/CustomizableObjectInstance.cpp b/Engine/Plugins/Experimental/Mutable/Source/CustomizableObject/Private/MuCO/CustomizableObjectInstance.cpp
--- a/Engine/Plugins/Experimental/Mutable/Source/CustomizableObject/Private/MuCO/CustomizableObjectInstance.cpp
+++ b/Engine/Plugins/Experimental/Mutable/Source/CustomizableObject/Private/MuCO/CustomizableObjectInstance.cpp
@@ -3543,8 +3543,7 @@ void UCustomizableInstancePrivateData::BuildClothingData(const TSharedPtr<FMutab
 			NewSelfCollisionIndices.SetNum(SrcSelfCollisionIndices.Num());
 			NewSelfCollisionIndices.SetNum(TrimAndRemapIndices(NewSelfCollisionIndices, SrcSelfCollisionIndices), false);
 						
-			// TODO: Remap tethers. This is important with the chaos implementation to remove stretchiness to the cloth.
-			// Possible implementation, not checked thoroughly or benchmarked.
+			// TODO: Possible implementation, not checked thoroughly or benchmarked.
 			auto TrimAndRemapTethers = [&IndexMap](FClothTetherData& Dst, const FClothTetherData& Src)
 			{
 				if (!IndexMap.Num())
@@ -3576,9 +3575,8 @@ void UCustomizableInstancePrivateData::BuildClothingData(const TSharedPtr<FMutab
 				}
 			};
 
-			// Not enabled yet.
-			//TrimAndRemapTethers( NewLodData.PhysicalMeshData.GeodesicTethers, SrcLodData.PhysicalMeshData.GeodesicTethers );
-			//TrimAndRemapTethers( NewLodData.PhysicalMeshData.EuclideanTethers, SrcLodData.PhysicalMeshData.EuclideanTethers );
+			TrimAndRemapTethers( NewLodData.PhysicalMeshData.GeodesicTethers, SrcLodData.PhysicalMeshData.GeodesicTethers );
+			TrimAndRemapTethers( NewLodData.PhysicalMeshData.EuclideanTethers, SrcLodData.PhysicalMeshData.EuclideanTethers );
 		}
 	}
 
