[mutable] Fix Instances not being updated due to LOD changes.

diff --git a/Engine/Plugins/Experimental/Mutable/Source/CustomizableObject/Classes/MuCO/CustomizableObjectInstance.h b/Engine/Plugins/Experimental/Mutable/Source/CustomizableObject/Classes/MuCO/CustomizableObjectInstance.h
--- a/Engine/Plugins/Experimental/Mutable/Source/CustomizableObject/Classes/MuCO/CustomizableObjectInstance.h
+++ b/Engine/Plugins/Experimental/Mutable/Source/CustomizableObject/Classes/MuCO/CustomizableObjectInstance.h
@@ -583,6 +583,8 @@ public:
 	float GetMinSquareDistToPlayer() const;
 	void SetMinSquareDistToPlayer(float NewValue);
 	void SetMinMaxLODToLoad(int32 NewMinLOD, int32 NewMaxLOD, bool LimitLODUpgrades = true);
+	int32 GetMinLODToLoad() const;
+	int32 GetMaxLODToLoad() const;
 	int32 GetNumLODsAvailable() const;
 
 	/* If enabled, CurrentMinLOD will be the first LOD of the generated SkeletalMesh. Otherwise the number of LODs will remain constant and LODs [0 .. CurrentMinLOD] will share the same RenderData. 
diff --git a/Engine/Plugins/Experimental/Mutable/Source/CustomizableObject/Private/MuCO/CustomizableInstancePrivateData.h b/Engine/Plugins/Experimental/Mutable/Source/CustomizableObject/Private/MuCO/CustomizableInstancePrivateData.h
--- a/Engine/Plugins/Experimental/Mutable/Source/CustomizableObject/Private/MuCO/CustomizableInstancePrivateData.h
+++ b/Engine/Plugins/Experimental/Mutable/Source/CustomizableObject/Private/MuCO/CustomizableInstancePrivateData.h
@@ -216,9 +216,7 @@ public:
 	// Copy data generated in the mutable thread over to the instance and initializes additional data required during the update
 	void PrepareForUpdate(const TSharedPtr<FMutableOperationData>& OperationData);
 
-	void SetMinMaxLODToLoad(int32 NewMinLOD, int32 NewMaxLOD, bool bLimitLODUpgrades = true);
-	int32 GetMinLODToLoad() { return CurrentMinLOD; }
-	int32 GetMaxLODToLoad() { return CurrentMaxLOD; }
+	void SetMinMaxLODToLoad(UCustomizableObjectInstance* Public, int32 NewMinLOD, int32 NewMaxLOD, bool bLimitLODUpgrades = true);
 	
 	int32 GetNumLODsAvailable() const { return NumLODsAvailable; }
 	
@@ -270,13 +268,12 @@ public:
 												
 	double LastUpdateTime;
 
-	// These are the LODs that MUST be used in an update through GetMinLODToLoad and GetMaxLODToLoad, they are thread safe
-	int32 CurrentMinLOD;
-	int32 CurrentMaxLOD;
+	// LODs applied on the last update. Represent the actual LODs the Instance is using.
+	int32 LastUpdateMinLOD = -1;
+	int32 LastUpdateMaxLOD = -1;
 
-	// These are the LODs we want to have, they MUSN'T be used in an update, they aren't thread safe to use from the mutable thread
-	int32 MinLODToLoad;
-	int32 MaxLODToLoad;
+	/** Saves the LODs requested on the update. */
+	void SaveMinMaxLODToLoad(const UCustomizableObjectInstance* Public);
 	
 	// This is the LODs that the Customizable Object has
 	int32 NumLODsAvailable;
diff --git a/Engine/Plugins/Experimental/Mutable/Source/CustomizableObject/Private/MuCO/CustomizableObjectInstance.cpp b/Engine/Plugins/Experimental/Mutable/Source/CustomizableObject/Private/MuCO/CustomizableObjectInstance.cpp
--- a/Engine/Plugins/Experimental/Mutable/Source/CustomizableObject/Private/MuCO/CustomizableObjectInstance.cpp
+++ b/Engine/Plugins/Experimental/Mutable/Source/CustomizableObject/Private/MuCO/CustomizableObjectInstance.cpp
@@ -142,6 +142,13 @@ UTexture2D* UCustomizableInstancePrivateData::CreateTexture()
 }
 
 
+void UCustomizableInstancePrivateData::SaveMinMaxLODToLoad(const UCustomizableObjectInstance* Public)
+{
+	LastUpdateMinLOD = Public->Descriptor.MinLODToLoad;
+	LastUpdateMaxLOD = Public->Descriptor.MaxLODToLoad;
+}
+
+
 int32 UCustomizableInstancePrivateData::GetLastMeshId(int32 ComponentIndex, int32 LODIndex) const
 {
 	const FCustomizableInstanceComponentData* ComponentData = GetComponentData(ComponentIndex);
@@ -212,16 +219,17 @@ void UCustomizableObjectInstance::SetDescriptor(const FCustomizableObjectInstanc
 }
 
 
-void UCustomizableInstancePrivateData::SetMinMaxLODToLoad(int32 NewMinLOD, int32 NewMaxLOD, bool bLimitLODUpgrades)
+void UCustomizableInstancePrivateData::SetMinMaxLODToLoad(UCustomizableObjectInstance* Public, int32 NewMinLOD, int32 NewMaxLOD, bool bLimitLODUpgrades)
 {
 	if (!HasCOInstanceFlags(LODsStreamingEnabled))
 	{
 		return;
 	}
 
-	NewMinLOD = FMath::Min(FMath::Max(NewMinLOD, (int32)FirstLODAvailable), FirstLODAvailable + NumMaxLODsToStream);
+	// Min LOD
+	NewMinLOD = FMath::Min(FMath::Max(NewMinLOD, static_cast<int32>(FirstLODAvailable)), FirstLODAvailable + NumMaxLODsToStream);
 
-	const bool bIsDowngradeLODUpdate = CurrentMinLOD >= 0 && NewMinLOD > CurrentMinLOD;
+	const bool bIsDowngradeLODUpdate = LastUpdateMinLOD >= 0 && NewMinLOD > LastUpdateMinLOD;
 	SetCOInstanceFlags(bIsDowngradeLODUpdate ? PendingLODsDowngrade : None);
 
 	if (!bIsDowngradeLODUpdate && bLimitLODUpgrades)
@@ -232,28 +240,29 @@ void UCustomizableInstancePrivateData::SetMinMaxLODToLoad(int32 NewMinLOD, int32
 		}
 	}
 
-	if(NumLODsAvailable != INT32_MAX)
+	// Max LOD
+	if (NumLODsAvailable != INT32_MAX)
 	{
-		if(MaxLODToLoad == INT32_MAX)
+		if (Public->Descriptor.MaxLODToLoad == INT32_MAX)
 		{
-			MaxLODToLoad = NumLODsAvailable - 1;
+			Public->Descriptor.MaxLODToLoad = NumLODsAvailable - 1;
 		}
 
 		NewMaxLOD = FMath::Min(NewMaxLOD, NumLODsAvailable - 1);
 	}
 	
-	SetCOInstanceFlags(MinLODToLoad != NewMinLOD || MaxLODToLoad != NewMaxLOD ? PendingLODsUpdate : None);
+	NewMaxLOD = FMath::Max(NewMaxLOD, NewMinLOD);
+
+	// Save the new LODs
+	SetCOInstanceFlags(Public->Descriptor.MinLODToLoad != NewMinLOD || Public->Descriptor.MaxLODToLoad != NewMaxLOD ? PendingLODsUpdate : None);
 
-	MinLODToLoad = NewMinLOD;
-	MaxLODToLoad = NewMaxLOD;
+	Public->Descriptor.MinLODToLoad = NewMinLOD;
+	Public->Descriptor.MaxLODToLoad = NewMaxLOD;
 }
 
 
 void UCustomizableInstancePrivateData::PrepareForUpdate(const TSharedPtr<FMutableOperationData>& OperationData)
 {
-	CurrentMinLOD = OperationData->CurrentMinLOD;
-	CurrentMaxLOD = OperationData->CurrentMaxLOD;
-
 	const bool bNumLODsUpdated = NumLODsAvailable != OperationData->NumLODsAvailable;
 	
 	SetCOInstanceFlags(bNumLODsUpdated ? PendingMeshUpdate : None);
@@ -786,7 +795,7 @@ void UCustomizableInstancePrivateData::DoUpdateSkeletalMeshAsync(UCustomizableOb
 
 	if (HasCOInstanceFlags(PendingLODsUpdateSecondStage))
 	{
-		UE_LOG(LogMutable, Verbose, TEXT("LOD change: %d, %d -> %d, %d"), CurrentMinLOD, CurrentMaxLOD, MinLODToLoad, MaxLODToLoad);
+		UE_LOG(LogMutable, Verbose, TEXT("LOD change: %d, %d -> %d, %d"), LastUpdateMinLOD, LastUpdateMaxLOD, Instance->Descriptor.MinLODToLoad, Instance->Descriptor.MaxLODToLoad);
 	}
 
 	UCustomizableObjectSystem::GetInstance()->GetPrivate()->InitUpdateSkeletalMesh(Instance, Priority);
@@ -1437,8 +1446,8 @@ bool UCustomizableInstancePrivateData::UpdateSkeletalMesh_PostBeginUpdate0(UCust
 		UE_LOG(LogMutable, Warning, TEXT("Updating skeletal mesh error for mesh %s"), *Public->GetName());
 	}
 
-	const int32 FirstComponent = OperationData->InstanceUpdateData.LODs[CurrentMinLOD].FirstComponent;
-	const int32 NumComponents = OperationData->InstanceUpdateData.LODs[CurrentMinLOD].ComponentCount;
+	const int32 FirstComponent = OperationData->InstanceUpdateData.LODs[OperationData->CurrentMinLOD].FirstComponent;
+	const int32 NumComponents = OperationData->InstanceUpdateData.LODs[OperationData->CurrentMinLOD].ComponentCount;
 
 	bool bHasSkeletalMesh = false;
 
@@ -1622,12 +1631,12 @@ bool UCustomizableInstancePrivateData::UpdateSkeletalMesh_PostBeginUpdate0(UCust
 
 					for (int32 LODIndex = 0; LODIndex <= OperationData->CurrentMaxLOD && LODIndex < OperationData->InstanceUpdateData.LODs.Num(); ++LODIndex)
 					{
-						if (HasCOInstanceFlags(ReduceLODs) && LODIndex < CurrentMinLOD)
+						if (HasCOInstanceFlags(ReduceLODs) && LODIndex < OperationData->CurrentMinLOD)
 						{
 							continue;
 						}
 						
-						Helper_GetLODData(SkeletalMesh).Add((LODIndex <= CurrentMinLOD && Helper_GetLODData(SkeletalMesh).Num()) ? &Helper_GetLODData(SkeletalMesh)[0] : new Helper_LODDataType());
+						Helper_GetLODData(SkeletalMesh).Add((LODIndex <= OperationData->CurrentMinLOD && Helper_GetLODData(SkeletalMesh).Num()) ? &Helper_GetLODData(SkeletalMesh)[0] : new Helper_LODDataType());
 
 						const FMutableRefLODData& LODData = RefSkeletalMeshData->LODData[LODIndex];
 						
@@ -1701,12 +1710,6 @@ UCustomizableInstancePrivateData::UCustomizableInstancePrivateData()
 	
 	LastUpdateTime = 0.f;
 
-	CurrentMinLOD = -1;
-	CurrentMaxLOD = -1;
-
-	MinLODToLoad = 0;
-	MaxLODToLoad = INT32_MAX;
-
 	NumLODsAvailable = INT32_MAX;
 	FirstLODAvailable = 0;
 
@@ -2447,6 +2450,11 @@ void SetTexturePropertiesFromMutableImageProps(UTexture2D* Texture, const FMutab
 
 void UCustomizableObjectInstance::Updated(EUpdateResult Result)
 {
+	if (Result == EUpdateResult::Success)
+	{
+		DescriptorHash = UpdateDescriptorHash;
+	}
+
 	// Call Customizable Skeletal Components updated callbacks.
 	for (TObjectIterator<UCustomizableSkeletalComponent> It; It; ++It)
 	{
@@ -2458,11 +2466,9 @@ void UCustomizableObjectInstance::Updated(EUpdateResult Result)
 		}
 	}
 
-	// Call Instance updated (this) callbacks.
-	DescriptorHash = UpdateDescriptorHash;
-
 	if (Result == EUpdateResult::Success)
 	{
+		// Call Instance updated (this) callbacks.
 		UpdatedDelegate.Broadcast(this);
 		UpdatedNativeDelegate.Broadcast(this);
 	}	
@@ -2801,10 +2807,10 @@ void UCustomizableInstancePrivateData::BuildSkeletalMeshElementData(const TShare
 	SkeletalMesh->GetMaterials().SetNum(1);
 	SkeletalMesh->GetMaterials()[0] = UnrealMaterial;
 
-	int32 MeshLODIndex = HasCOInstanceFlags(ReduceLODs) ? 0 : CurrentMinLOD;
+	int32 MeshLODIndex = HasCOInstanceFlags(ReduceLODs) ? 0 : OperationData->CurrentMinLOD;
 
 	int32 LODCount = OperationData->InstanceUpdateData.LODs.Num();
-	for (int32 LODIndex = CurrentMinLOD; LODIndex <= CurrentMaxLOD && LODIndex < LODCount; ++LODIndex, ++MeshLODIndex)
+	for (int32 LODIndex = OperationData->CurrentMinLOD; LODIndex <= OperationData->CurrentMaxLOD && LODIndex < LODCount; ++LODIndex, ++MeshLODIndex)
 	{
 		const FInstanceUpdateData::FLOD& LOD = OperationData->InstanceUpdateData.LODs[LODIndex];
 		const mu::MeshPtrConst MutableMesh = OperationData->InstanceUpdateData.Components[LOD.FirstComponent + ComponentIndex].Mesh;
@@ -2979,7 +2985,7 @@ void UCustomizableInstancePrivateData::BuildMorphTargetsData(const TSharedPtr<FM
 	const int32 SkeletalMeshMorphTargetsNum = SkeletalMesh->GetMorphTargets().Num();
 
 	int32 LODCount = OperationData->InstanceUpdateData.LODs.Num();
-	for (int32 LODIndex = CurrentMinLOD; LODIndex <= CurrentMaxLOD && LODIndex < LODCount; ++LODIndex)
+	for (int32 LODIndex = OperationData->CurrentMinLOD; LODIndex <= OperationData->CurrentMaxLOD && LODIndex < LODCount; ++LODIndex)
 	{
 		const FInstanceUpdateData::FLOD& LOD = OperationData->InstanceUpdateData.LODs[LODIndex];
 		if (mu::MeshPtrConst MutableMesh = OperationData->InstanceUpdateData.Components[LOD.FirstComponent+ComponentIndex].Mesh)
@@ -3118,8 +3124,8 @@ void UCustomizableInstancePrivateData::BuildClothingData(const TSharedPtr<FMutab
 	{
 		MUTABLE_CPUPROFILER_SCOPE(DiscoverSectionsWithCloth);
 
-		int32 MeshLODIndex = HasCOInstanceFlags(ReduceLODs) ? 0 : CurrentMinLOD;
-		for (int32 LODIndex = CurrentMinLOD; LODIndex <= CurrentMaxLOD && LODIndex < LODCount; ++LODIndex, ++MeshLODIndex)
+		int32 MeshLODIndex = HasCOInstanceFlags(ReduceLODs) ? 0 : OperationData->CurrentMinLOD;
+		for (int32 LODIndex = OperationData->CurrentMinLOD; LODIndex <= OperationData->CurrentMaxLOD && LODIndex < LODCount; ++LODIndex, ++MeshLODIndex)
 		{
 			const FInstanceUpdateData::FLOD& LOD = OperationData->InstanceUpdateData.LODs[LODIndex];
 			const int32 ComponentCount = LOD.ComponentCount;
@@ -3775,15 +3781,15 @@ void UCustomizableInstancePrivateData::BuildClothingData(const TSharedPtr<FMutab
 
 	// Based on FSkeletalMeshLODModel::GetClothMappingData().
 	TArray<TArray<FMeshToMeshVertData>> LodMappingData;
-	LodMappingData.SetNum(FMath::Max<int32>(CurrentMaxLOD + 1, LODCount));
+	LodMappingData.SetNum(FMath::Max<int32>(OperationData->CurrentMaxLOD + 1, LODCount));
 
 	TArray<TArray<FClothBufferIndexMapping>> LodClothingIndexMapping;
-	LodClothingIndexMapping.SetNum(FMath::Max<int32>(CurrentMaxLOD + 1, LODCount));
+	LodClothingIndexMapping.SetNum(FMath::Max<int32>(OperationData->CurrentMaxLOD + 1, LODCount));
 	{
 		int32 NumSectionsWithClothProcessed = 0;
-		int32 MeshLODIndex = HasCOInstanceFlags(ReduceLODs) ? 0 : CurrentMinLOD;
+		int32 MeshLODIndex = HasCOInstanceFlags(ReduceLODs) ? 0 : OperationData->CurrentMinLOD;
 
-		for (int32 LODIndex = CurrentMinLOD; LODIndex <= CurrentMaxLOD && LODIndex < LODCount; ++LODIndex, ++MeshLODIndex)
+		for (int32 LODIndex = OperationData->CurrentMinLOD; LODIndex <= OperationData->CurrentMaxLOD && LODIndex < LODCount; ++LODIndex, ++MeshLODIndex)
 		{
 			const FInstanceUpdateData::FLOD& LOD = OperationData->InstanceUpdateData.LODs[LODIndex];
 			const int32 ComponentCount = LOD.ComponentCount;
@@ -3837,8 +3843,8 @@ void UCustomizableInstancePrivateData::BuildClothingData(const TSharedPtr<FMutab
 	{
 		MUTABLE_CPUPROFILER_SCOPE(InitClothRenderData)
 		// Based on FSkeletalMeshLODModel::GetClothMappingData().
-		int32 MeshLODIndex = HasCOInstanceFlags(ReduceLODs) ? 0 : CurrentMinLOD;
-		for (int32 LODIndex = CurrentMinLOD; LODIndex <= CurrentMaxLOD && LODIndex < LODCount; ++LODIndex, ++MeshLODIndex)
+		int32 MeshLODIndex = HasCOInstanceFlags(ReduceLODs) ? 0 : OperationData->CurrentMinLOD;
+		for (int32 LODIndex = OperationData->CurrentMinLOD; LODIndex <= OperationData->CurrentMaxLOD && LODIndex < LODCount; ++LODIndex, ++MeshLODIndex)
 		{
 			FSkeletalMeshLODRenderData& LODModel = RenderResource->LODRenderData[MeshLODIndex];
 	
@@ -3937,10 +3943,10 @@ bool UCustomizableInstancePrivateData::BuildSkeletalMeshRenderData(const TShared
 	SkeletalMesh->SetHasBeenSimplified(false);
 	SkeletalMesh->SetHasVertexColors(false);
 
-	int32 MeshLODIndex = HasCOInstanceFlags(ReduceLODs) ? 0 : CurrentMinLOD;
+	int32 MeshLODIndex = HasCOInstanceFlags(ReduceLODs) ? 0 : OperationData->CurrentMinLOD;
 
 	int32 LODCount = OperationData->InstanceUpdateData.LODs.Num();
-	for (int32 LODIndex = CurrentMinLOD; LODIndex <= CurrentMaxLOD && LODIndex < LODCount; ++LODIndex, ++MeshLODIndex)
+	for (int32 LODIndex = OperationData->CurrentMinLOD; LODIndex <= OperationData->CurrentMaxLOD && LODIndex < LODCount; ++LODIndex, ++MeshLODIndex)
 	{
 		MUTABLE_CPUPROFILER_SCOPE(LODLoop);
 
@@ -4121,7 +4127,7 @@ FGraphEventRef UCustomizableInstancePrivateData::LoadAdditionalAssetsAsync(const
 	// Load assets coming from SubMeshes of the newly generated Mesh
 	if (OperationData->InstanceUpdateData.LODs.Num())
 	{
-		const FInstanceUpdateData::FLOD& LOD = OperationData->InstanceUpdateData.LODs[CurrentMinLOD];
+		const FInstanceUpdateData::FLOD& LOD = OperationData->InstanceUpdateData.LODs[OperationData->CurrentMinLOD];
 		for (int32 ComponentIndex = 0; ComponentIndex < LOD.ComponentCount; ++ComponentIndex)
 		{
 			const FInstanceUpdateData::FComponent& Component = Components[LOD.FirstComponent+ComponentIndex];
@@ -4477,7 +4483,7 @@ void UCustomizableInstancePrivateData::BuildMaterials(const TSharedPtr<FMutableO
 	const bool bUseCurrentMinLODAsBaseLOD = HasCOInstanceFlags(ReduceLODs);
 	const bool bReuseTextures = HasCOInstanceFlags(ReuseTextures);
 
-	const FInstanceUpdateData::FLOD& FirstLOD = OperationData->InstanceUpdateData.LODs[CurrentMinLOD];
+	const FInstanceUpdateData::FLOD& FirstLOD = OperationData->InstanceUpdateData.LODs[OperationData->CurrentMinLOD];
 	const int32 NumComponents = FirstLOD.ComponentCount;
 
 	for (int32 ComponentIndex = 0; ComponentIndex < NumComponents; ++ComponentIndex)
@@ -4500,9 +4506,9 @@ void UCustomizableInstancePrivateData::BuildMaterials(const TSharedPtr<FMutableO
 
 		MUTABLE_CPUPROFILER_SCOPE(LODLoop);
 
-		int32 MeshLODIndex = bUseCurrentMinLODAsBaseLOD ? 0 : CurrentMinLOD;
+		int32 MeshLODIndex = bUseCurrentMinLODAsBaseLOD ? 0 : OperationData->CurrentMinLOD;
 
-		for (int32 LODIndex = CurrentMinLOD; LODIndex <= CurrentMaxLOD && LODIndex < LODCount; ++LODIndex, ++MeshLODIndex)
+		for (int32 LODIndex = OperationData->CurrentMinLOD; LODIndex <= OperationData->CurrentMaxLOD && LODIndex < LODCount; ++LODIndex, ++MeshLODIndex)
 		{
 			const FInstanceUpdateData::FLOD& LOD = OperationData->InstanceUpdateData.LODs[LODIndex];
 
@@ -5147,7 +5153,19 @@ void UCustomizableObjectInstance::SetMinSquareDistToPlayer(float NewValue)
 
 void UCustomizableObjectInstance::SetMinMaxLODToLoad(int32 NewMinLOD, int32 NewMaxLOD, bool bLimitLODUpgrades)
 {
-	GetPrivate()->SetMinMaxLODToLoad(NewMinLOD, NewMaxLOD, bLimitLODUpgrades);
+	GetPrivate()->SetMinMaxLODToLoad(this, NewMinLOD, NewMaxLOD, bLimitLODUpgrades);
+}
+
+
+int32 UCustomizableObjectInstance::GetMinLODToLoad() const
+{
+	return Descriptor.MinLODToLoad;	
+}
+
+
+int32 UCustomizableObjectInstance::GetMaxLODToLoad() const
+{
+	return Descriptor.MaxLODToLoad;	
 }
 
 
diff --git a/Engine/Plugins/Experimental/Mutable/Source/CustomizableObject/Private/MuCO/CustomizableObjectInstanceDescriptor.cpp b/Engine/Plugins/Experimental/Mutable/Source/CustomizableObject/Private/MuCO/CustomizableObjectInstanceDescriptor.cpp
--- a/Engine/Plugins/Experimental/Mutable/Source/CustomizableObject/Private/MuCO/CustomizableObjectInstanceDescriptor.cpp
+++ b/Engine/Plugins/Experimental/Mutable/Source/CustomizableObject/Private/MuCO/CustomizableObjectInstanceDescriptor.cpp
@@ -2129,8 +2129,9 @@ uint32 GetTypeHash(const FCustomizableObjectInstanceDescriptor& Key)
 	}
 	
 	Hash = HashCombine(Hash, GetTypeHash(Key.State));
-
 	Hash = HashCombine(Hash, GetTypeHash(Key.bBuildParameterDecorations));
+	Hash = HashCombine(Hash, GetTypeHash(Key.MinLODToLoad));
+	Hash = HashCombine(Hash, GetTypeHash(Key.MaxLODToLoad));
 
 	for (const TTuple<FName, FMultilayerProjector>& Pair : Key.MultilayerProjectors)
 	{
diff --git a/Engine/Plugins/Experimental/Mutable/Source/CustomizableObject/Private/MuCO/CustomizableObjectSystem.cpp b/Engine/Plugins/Experimental/Mutable/Source/CustomizableObject/Private/MuCO/CustomizableObjectSystem.cpp
--- a/Engine/Plugins/Experimental/Mutable/Source/CustomizableObject/Private/MuCO/CustomizableObjectSystem.cpp
+++ b/Engine/Plugins/Experimental/Mutable/Source/CustomizableObject/Private/MuCO/CustomizableObjectSystem.cpp
@@ -526,11 +526,11 @@ void FCustomizableObjectSystemPrivate::UpdateStats()
 					{
 						CountAllocatedSkeletalMesh++;
 
-						if (CustomizableObjectInstance->GetPrivate()->GetMinLODToLoad() < 1)
+						if (CustomizableObjectInstance->GetPrivate()->LastUpdateMinLOD < 1)
 						{
 							++CountLOD0;
 						}
-						else if (CustomizableObjectInstance->GetPrivate()->GetMinLODToLoad() < 2)
+						else if (CustomizableObjectInstance->GetPrivate()->LastUpdateMaxLOD < 2)
 						{
 							++CountLOD1;
 						}
@@ -776,24 +776,24 @@ void UpdateSkeletalMesh(UCustomizableObjectInstance* CustomizableObjectInstance)
 }
 
 
-void FCustomizableObjectSystemPrivate::InitUpdateSkeletalMesh(UCustomizableObjectInstance* InCustomizableObjectInstance, FMutableQueueElem::EQueuePriorityType Priority)
+void FCustomizableObjectSystemPrivate::InitUpdateSkeletalMesh(UCustomizableObjectInstance* Public, FMutableQueueElem::EQueuePriorityType Priority)
 {
 	MUTABLE_CPUPROFILER_SCOPE(FCustomizableObjectSystemPrivate::InitUpdateSkeletalMesh);
 
 	check(IsInGameThread());
 
-	if (!InCustomizableObjectInstance ||
-		!InCustomizableObjectInstance->IsValidLowLevel())
+	if (!Public ||
+		!Public->IsValidLowLevel())
 	{
-		InCustomizableObjectInstance->Updated(EUpdateResult::Error);
+		Public->Updated(EUpdateResult::Error);
 		return;
 	}
 	
-	if (UCustomizableObject* CustomizableObject = InCustomizableObjectInstance->GetCustomizableObject();
+	if (UCustomizableObject* CustomizableObject = Public->GetCustomizableObject();
 		CustomizableObject &&
 		!CustomizableObject->IsLocked())
 	{
-		FString CurrentState = InCustomizableObjectInstance->GetCurrentState();
+		FString CurrentState = Public->GetCurrentState();
 		FParameterUIData* State = CustomizableObject->StateUIDataMap.Find(CurrentState);
 		bool bNeverStream = State ? State->bDontCompressRuntimeTextures : false;
 		const bool bUseMipmapStreaming = !bNeverStream;
@@ -804,14 +804,16 @@ void FCustomizableObjectSystemPrivate::InitUpdateSkeletalMesh(UCustomizableObjec
 			MipsToSkip = 255; // This means skip all possible mips until only UTexture::GetStaticMinTextureResidentMipCount() are left
 		}
 		
-		const TSharedPtr<FMutableOperation> Operation = MakeShared<FMutableOperation>(FMutableOperation::CreateInstanceUpdate(InCustomizableObjectInstance, bNeverStream, MipsToSkip));
+		const TSharedPtr<FMutableOperation> Operation = MakeShared<FMutableOperation>(FMutableOperation::CreateInstanceUpdate(Public, bNeverStream, MipsToSkip));
+
+		Public->GetPrivate()->SaveMinMaxLODToLoad(Public);
 
 		// Skip the update if the Instance has the same descriptor (pending or applied).
 		// If there is a hash collision, at worse, we would add an unnecessary update.
 		{
-			const uint32 UpdateDescriptorHash = InCustomizableObjectInstance->GetUpdateDescriptorHash();
+			const uint32 UpdateDescriptorHash = Public->GetUpdateDescriptorHash();
 
-			if (const FMutableQueueElem* QueueElem = MutableOperationQueue.Get(InCustomizableObjectInstance))
+			if (const FMutableQueueElem* QueueElem = MutableOperationQueue.Get(Public))
 			{
 				if (const TSharedPtr<FMutableOperation>& QueuedOperation = QueueElem->Operation;
 					QueuedOperation &&
@@ -829,21 +831,21 @@ void FCustomizableObjectSystemPrivate::InitUpdateSkeletalMesh(UCustomizableObjec
 				return;
 			}
 			
-			if (InCustomizableObjectInstance->GetDescriptorHash() == UpdateDescriptorHash)
+			if (Public->GetDescriptorHash() == UpdateDescriptorHash)
 			{
-				UpdateSkeletalMesh(InCustomizableObjectInstance);
+				UpdateSkeletalMesh(Public);
 				return;
 			}
 		}
 
-		MutableOperationQueue.Enqueue(FMutableQueueElem::Create(Operation, Priority, InCustomizableObjectInstance->GetPrivate()->MinSquareDistFromComponentToPlayer));
+		MutableOperationQueue.Enqueue(FMutableQueueElem::Create(Operation, Priority, Public->GetPrivate()->MinSquareDistFromComponentToPlayer));
 
-		InCustomizableObjectInstance->BeginUpdateDelegate.Broadcast(InCustomizableObjectInstance);
-		InCustomizableObjectInstance->BeginUpdateNativeDelegate.Broadcast(InCustomizableObjectInstance);
+		Public->BeginUpdateDelegate.Broadcast(Public);
+		Public->BeginUpdateNativeDelegate.Broadcast(Public);
 	}
 	else
 	{
-		InCustomizableObjectInstance->Updated(EUpdateResult::Error);
+		Public->Updated(EUpdateResult::Error);
 	}
 }
 
@@ -2099,12 +2101,6 @@ namespace impl
 
 		UCustomizableInstancePrivateData* PrivateInstance = CandidateInstance->GetPrivate();
 
-		PrivateInstance->CurrentMinLOD = FMath::Max(PrivateInstance->MinLODToLoad, (int32)PrivateInstance->FirstLODAvailable);
-		PrivateInstance->CurrentMinLOD = FMath::Min(PrivateInstance->CurrentMinLOD, PrivateInstance->FirstLODAvailable + PrivateInstance->NumMaxLODsToStream);
-
-		PrivateInstance->CurrentMaxLOD = FMath::Max(PrivateInstance->MaxLODToLoad, PrivateInstance->CurrentMinLOD);
-
-
 		// Prepare streaming for the current customizable object
 		System->GetPrivate()->Streamer->PrepareStreamingForObject(CustomizableObject);
 
@@ -2114,8 +2110,8 @@ namespace impl
 		TSharedPtr<FMutableOperationData> CurrentOperationData = MakeShared<FMutableOperationData>();
 		CurrentOperationData->TextureCoverageQueries_MutableThreadParams = PrivateInstance->TextureCoverageQueries;
 		CurrentOperationData->TextureCoverageQueries_MutableThreadResults.Empty();
-		CurrentOperationData->CurrentMinLOD = PrivateInstance->CurrentMinLOD;
-		CurrentOperationData->CurrentMaxLOD = PrivateInstance->CurrentMaxLOD;
+		CurrentOperationData->CurrentMinLOD = PrivateInstance->LastUpdateMinLOD;
+		CurrentOperationData->CurrentMaxLOD = PrivateInstance->LastUpdateMaxLOD;
 		CurrentOperationData->bNeverStream = Operation->bNeverStream;
 		CurrentOperationData->MipsToSkip = Operation->MipsToSkip;
 		CurrentOperationData->MutableParameters = Parameters;
diff --git a/Engine/Plugins/Experimental/Mutable/Source/CustomizableObject/Private/MuCO/CustomizableObjectSystemPrivate.h b/Engine/Plugins/Experimental/Mutable/Source/CustomizableObject/Private/MuCO/CustomizableObjectSystemPrivate.h
--- a/Engine/Plugins/Experimental/Mutable/Source/CustomizableObject/Private/MuCO/CustomizableObjectSystemPrivate.h
+++ b/Engine/Plugins/Experimental/Mutable/Source/CustomizableObject/Private/MuCO/CustomizableObjectSystemPrivate.h
@@ -676,7 +676,7 @@ public:
 
 
 	// Init the async Skeletal Mesh creation/update
-	void InitUpdateSkeletalMesh(UCustomizableObjectInstance* InCustomizableObjectInstance, FMutableQueueElem::EQueuePriorityType Priority);
+	void InitUpdateSkeletalMesh(UCustomizableObjectInstance* Public, FMutableQueueElem::EQueuePriorityType Priority);
 		
 	// Init an async and safe release of the UE4 and Mutable resources used by the instance without actually destroying the instance, for example if it's very far away
 	void InitDiscardResourcesSkeletalMesh(UCustomizableObjectInstance* InCustomizableObjectInstance);
diff --git a/Engine/Plugins/Experimental/Mutable/Source/CustomizableObject/Private/MuCO/LogInformationUtil.cpp b/Engine/Plugins/Experimental/Mutable/Source/CustomizableObject/Private/MuCO/LogInformationUtil.cpp
--- a/Engine/Plugins/Experimental/Mutable/Source/CustomizableObject/Private/MuCO/LogInformationUtil.cpp
+++ b/Engine/Plugins/Experimental/Mutable/Source/CustomizableObject/Private/MuCO/LogInformationUtil.cpp
@@ -125,10 +125,10 @@ void LogInformationUtil::LogShowInstanceData(const UCustomizableObjectInstance*
 	LogData += FString::Printf(TEXT("bIsUpdating=%d "), InstanceFlags & ECOInstanceFlags::Updating);
 	LogData += FString::Printf(TEXT("bShouldUpdateLODs=%d "), InstanceFlags & ECOInstanceFlags::PendingLODsUpdate);
 
-	LogData += FString::Printf(TEXT("CurrentMinLOD=%d "), CustomizableObjectInstance->GetPrivate()->CurrentMinLOD);
-	LogData += FString::Printf(TEXT("CurrentMaxLOD=%d "), CustomizableObjectInstance->GetPrivate()->CurrentMaxLOD);
-	LogData += FString::Printf(TEXT("MinLODToLoad=%d "), CustomizableObjectInstance->GetPrivate()->MinLODToLoad);
-	LogData += FString::Printf(TEXT("MaxLODToLoad=%d\n"), CustomizableObjectInstance->GetPrivate()->MaxLODToLoad);
+	LogData += FString::Printf(TEXT("CurrentMinLOD=%d "), CustomizableObjectInstance->GetPrivate()->LastUpdateMinLOD);
+	LogData += FString::Printf(TEXT("CurrentMaxLOD=%d "), CustomizableObjectInstance->GetPrivate()->LastUpdateMaxLOD);
+	LogData += FString::Printf(TEXT("MinLODToLoad=%d "), CustomizableObjectInstance->GetMinLODToLoad());
+	LogData += FString::Printf(TEXT("MaxLODToLoad=%d\n"), CustomizableObjectInstance->GetMaxLODToLoad());
 }
 
 
@@ -174,10 +174,10 @@ void LogInformationUtil::LogShowInstanceDataFull(const UCustomizableObjectInstan
 	MessageChunk += FString::Printf(TEXT("        LastMinSquareDistFromComponentToPlayer = %.2f\n"), CustomizableObjectInstance->GetPrivate()->LastMinSquareDistFromComponentToPlayer);
 	LogData += MessageChunk;
 
-	MessageChunk = FString::Printf(TEXT("        CurrentMinLOD = %d\n"), CustomizableObjectInstance->GetPrivate()->CurrentMinLOD);
-	MessageChunk += FString::Printf(TEXT("        CurrentMaxLOD = %d\n"), CustomizableObjectInstance->GetPrivate()->CurrentMaxLOD);
-	MessageChunk += FString::Printf(TEXT("        MinLODToLoad = %d\n"), CustomizableObjectInstance->GetPrivate()->MinLODToLoad);
-	MessageChunk += FString::Printf(TEXT("        MaxLODToLoad = %d\n"), CustomizableObjectInstance->GetPrivate()->MaxLODToLoad);
+	MessageChunk = FString::Printf(TEXT("        CurrentMinLOD = %d\n"), CustomizableObjectInstance->GetPrivate()->LastUpdateMinLOD);
+	MessageChunk += FString::Printf(TEXT("        CurrentMaxLOD = %d\n"), CustomizableObjectInstance->GetPrivate()->LastUpdateMaxLOD);
+	MessageChunk += FString::Printf(TEXT("        MinLODToLoad = %d\n"), CustomizableObjectInstance->GetMinLODToLoad());
+	MessageChunk += FString::Printf(TEXT("        MaxLODToLoad = %d\n"), CustomizableObjectInstance->GetMaxLODToLoad());
 	MessageChunk += FString::Printf(TEXT("        bShouldUpdateLODs = %d\n"), InstanceFlags & ECOInstanceFlags::PendingLODsUpdate);
 	LogData += MessageChunk;
 
@@ -192,11 +192,11 @@ void LogInformationUtil::LogShowInstanceDataFull(const UCustomizableObjectInstan
 	{
 		if ((CustomizableObjectInstance->GetSkeletalMesh(ComponentIndex) != nullptr) && (CustomizableObjectInstance->GetSkeletalMesh(ComponentIndex)->GetResourceForRendering()))
 		{
-			if (CustomizableObjectInstance->GetPrivate()->GetMinLODToLoad() < 1)
+			if (CustomizableObjectInstance->GetPrivate()->LastUpdateMinLOD < 1)
 			{
 				CountLOD0++;
 			}
-			else if (CustomizableObjectInstance->GetPrivate()->GetMinLODToLoad() < 2)
+			else if (CustomizableObjectInstance->GetPrivate()->LastUpdateMinLOD < 2)
 			{
 				CountLOD1++;
 			}
diff --git a/Engine/Plugins/Experimental/Mutable/Source/CustomizableObject/Public/MuCO/CustomizableObjectInstanceDescriptor.h b/Engine/Plugins/Experimental/Mutable/Source/CustomizableObject/Public/MuCO/CustomizableObjectInstanceDescriptor.h
--- a/Engine/Plugins/Experimental/Mutable/Source/CustomizableObject/Public/MuCO/CustomizableObjectInstanceDescriptor.h
+++ b/Engine/Plugins/Experimental/Mutable/Source/CustomizableObject/Public/MuCO/CustomizableObjectInstanceDescriptor.h
@@ -313,12 +313,18 @@ private:
 	UPROPERTY()
 	TArray<FCustomizableObjectProjectorParameterValue> ProjectorParameters;
 
+	/** Mutable parameters optimization state. */
 	int32 State = 0;
 	
 	/** Flag to control the build of the parameter description images required for customization UI.
      * These descriptions get generated with every instance update, so it should be disabled when not needed. */
 	bool bBuildParameterDecorations = false;
 
+	/** These are the LODs we want to have, they MUST NOT be used in an update (Mutable thread). */
+	int32 MinLODToLoad = 0;
+	int32 MaxLODToLoad = INT32_MAX;
+
+	/** Lookup of UCustomizableObjectDescriptor::IntParameters. */
 	TMap<FString, int32> IntParametersLookupTable;
 	
 	/** Multilayer Projector helpers. See FMultilayerProjector.*/
