#rnx

diff --git a/Engine/Plugins/Experimental/Mutable/Source/CustomizableObject/Private/MuCO/CustomizableObjectInstance.cpp b/Engine/Plugins/Experimental/Mutable/Source/CustomizableObject/Private/MuCO/CustomizableObjectInstance.cpp
--- a/Engine/Plugins/Experimental/Mutable/Source/CustomizableObject/Private/MuCO/CustomizableObjectInstance.cpp
+++ b/Engine/Plugins/Experimental/Mutable/Source/CustomizableObject/Private/MuCO/CustomizableObjectInstance.cpp
@@ -4628,7 +4628,7 @@ void UCustomizableInstancePrivateData::BuildMaterials(const TSharedPtr<FMutableO
 							FString TextureReuseCacheRef = bReuseTextures ? FString::Printf(TEXT("%d-%d-%d-%d"), LODIndex, ComponentIndex, Surface.SurfaceId, ImageIndex) : FString();
 
 							// If the mutable image is null, it must be in the cache
-							FMutableImageCacheKey ImageCacheKey = { Image.ImageID };//, OperationData->MipsToSkip }; // Todo: Reenable extended FMutableImageCacheKey
+							FMutableImageCacheKey ImageCacheKey = { Image.ImageID, OperationData->MipsToSkip };
 							if (!MutableImage)
 							{
 								TWeakObjectPtr<UTexture2D>* CachedPointerPtr = Cache.Images.Find(ImageCacheKey);
diff --git a/Engine/Plugins/Experimental/Mutable/Source/CustomizableObject/Private/MuCO/CustomizableObjectSystem.cpp b/Engine/Plugins/Experimental/Mutable/Source/CustomizableObject/Private/MuCO/CustomizableObjectSystem.cpp
--- a/Engine/Plugins/Experimental/Mutable/Source/CustomizableObject/Private/MuCO/CustomizableObjectSystem.cpp
+++ b/Engine/Plugins/Experimental/Mutable/Source/CustomizableObject/Private/MuCO/CustomizableObjectSystem.cpp
@@ -1967,7 +1967,8 @@ namespace impl
 
 		for (const FInstanceUpdateData::FImage& Image : OperationData->InstanceUpdateData.Images)
 		{
-			TWeakObjectPtr<UTexture2D>* TexturePtr = Cache.Images.Find({ Image.ImageID });
+			FMutableImageCacheKey Key(Image.ImageID, OperationData->MipsToSkip);
+			TWeakObjectPtr<UTexture2D>* TexturePtr = Cache.Images.Find(Key);
 
 			if (TexturePtr && TexturePtr->Get() && System->GetPrivate()->TextureHasReferences(Image.ImageID))
 			{
diff --git a/Engine/Plugins/Experimental/Mutable/Source/CustomizableObject/Private/MuCO/CustomizableObjectSystemPrivate.h b/Engine/Plugins/Experimental/Mutable/Source/CustomizableObject/Private/MuCO/CustomizableObjectSystemPrivate.h
--- a/Engine/Plugins/Experimental/Mutable/Source/CustomizableObject/Private/MuCO/CustomizableObjectSystemPrivate.h
+++ b/Engine/Plugins/Experimental/Mutable/Source/CustomizableObject/Private/MuCO/CustomizableObjectSystemPrivate.h
@@ -168,19 +168,21 @@ public:
 struct FMutableImageCacheKey
 {
 	mu::RESOURCE_ID Resource = 0;
-	//int32 SkippedMips = 0; // Todo: Reenable extended FMutableImageCacheKey
+	int32 SkippedMips = 0;
 
-	inline bool operator==(const FMutableImageCacheKey& o) const
+	FMutableImageCacheKey(mu::RESOURCE_ID InResource, int32 InSkippedMips)
+		: Resource(InResource), SkippedMips(InSkippedMips) {}
+
+	inline bool operator==(const FMutableImageCacheKey& Other) const
 	{
-		return Resource == o.Resource;//&& SkippedMips == o.SkippedMips; // Todo: Reenable extended FMutableImageCacheKey
+		return Resource == Other.Resource && SkippedMips == Other.SkippedMips;
 	}
 };
 
 
 inline uint32 GetTypeHash(const FMutableImageCacheKey& Key)
 {
-	//return HashCombine(GetTypeHash(Key.Resource), GetTypeHash(Key.SkippedMips)); // Todo: Reenable extended FMutableImageCacheKey
-	return GetTypeHash(Key.Resource);
+	return HashCombine(GetTypeHash(Key.Resource), GetTypeHash(Key.SkippedMips));
 }
 
 
