[mutable] Cleanup the layout code generation to prevent bugs related to mismatched blocks.
Moved MeshMorph operation to its own AST node.
#rnx

diff --git a/Engine/Plugins/Experimental/Mutable/Source/MutableRuntime/Private/MuR/CodeRunner.cpp b/Engine/Plugins/Experimental/Mutable/Source/MutableRuntime/Private/MuR/CodeRunner.cpp
--- a/Engine/Plugins/Experimental/Mutable/Source/MutableRuntime/Private/MuR/CodeRunner.cpp
+++ b/Engine/Plugins/Experimental/Mutable/Source/MutableRuntime/Private/MuR/CodeRunner.cpp
@@ -972,15 +972,32 @@ namespace mu
 
         case OP_TYPE::ME_MORPH2:
         {
-            auto args = pModel->GetPrivate()->m_program.GetOpArgs<OP::MeshMorph2Args>(item.at);
-            switch (item.stage)
+			const uint8* data = pModel->GetPrivate()->m_program.GetOpArgsPointer(item.at);
+
+			OP::ADDRESS FactorAt = 0;
+			FMemory::Memcpy(&FactorAt, data, sizeof(OP::ADDRESS)); data += sizeof(OP::ADDRESS);
+			
+			OP::ADDRESS BaseAt = 0;
+			FMemory::Memcpy(&BaseAt, data, sizeof(OP::ADDRESS)); data += sizeof(OP::ADDRESS);
+
+			uint8 NumTargets = 0;
+			FMemory::Memcpy(&NumTargets, data, sizeof(uint8)); data += sizeof(uint8);
+
+			TArray<MESH_BUFFER_SEMANTIC, TInlineAllocator<8>> Targets;
+			Targets.SetNum(NumTargets);
+			for (uint8 T = 0; T < NumTargets; ++T)
+			{
+				FMemory::Memcpy(&Targets[T], data, sizeof(OP::ADDRESS)); data += sizeof(OP::ADDRESS);
+			}
+
+			switch (item.stage)
             {
             case 0:
-                if ( args.base )
+                if (BaseAt)
                 {
-                    AddOp( SCHEDULED_OP( item.at, item, 1),
-                           SCHEDULED_OP( args.base, item),
-                           SCHEDULED_OP( args.factor, item) );
+                    AddOp( SCHEDULED_OP(item.at, item, 1),
+                           SCHEDULED_OP(BaseAt, item),
+                           SCHEDULED_OP(FactorAt, item) );
                 }
                 else
                 {
@@ -990,24 +1007,16 @@ namespace mu
 
             case 1:
             {
-                bool baseValid = GetMemory().IsValid( CACHE_ADDRESS(args.base,item) );
-                float factor = GetMemory().GetScalar( CACHE_ADDRESS(args.factor,item) );
-
-                int count = 0;
-                for ( int i=0
-                    ; i<MUTABLE_OP_MAX_MORPH2_TARGETS && args.targets[i]
-                    ; ++i )
-                {
-                    count++;
-                }
+                bool baseValid = GetMemory().IsValid( CACHE_ADDRESS(BaseAt,item) );
+                float factor = GetMemory().GetScalar( CACHE_ADDRESS(FactorAt,item) );
 
                 // Factor goes from -1 to 1 across all targets. [0 - 1] represents positive morphs, while [-1, 0) represent negative morphs.
 				factor = FMath::Clamp(factor, -1.0f, 1.0f); // Is the factor not in range [-1, 1], it will index a non existing morph.
 
-				float absFactor = fabsf(factor);
-                float delta = 1.0f/(count-1);
-				int min = (int)floorf(absFactor/delta);
-				int max = (int)ceilf(absFactor/delta);
+				float absFactor = FMath::Abs(factor);
+                float delta = 1.0f/(NumTargets -1);
+				int min = (int)FMath::FloorToFloat(absFactor/delta);
+				int max = (int)FMath::CeilToFloat(absFactor/delta);
 
                 // Factor from 0 to 1 between the two targets
 				float bifactor = absFactor/delta - min;
@@ -1015,7 +1024,7 @@ namespace mu
 				// From [0,1] to [-1,0] for negative factors
 				if (factor < 0.0f)
 				{
-					float threshold = -1.0f + SMALL_NUMBER;
+					float threshold = -1.0f + UE_SMALL_NUMBER;
 
 					if (factor <= threshold)
 					{
@@ -1031,36 +1040,36 @@ namespace mu
 
                 if (baseValid)
                 {
-                    SCHEDULED_OP_DATA data;
-                    data.bifactor = bifactor;
-					data.min = FMath::Clamp(min, 0, count - 1);
-					data.max = FMath::Clamp(max, 0, count - 1);
-					uint32 dataAddress = uint32(m_heapData.Add(data));
+                    SCHEDULED_OP_DATA HeapData;
+					HeapData.bifactor = bifactor;
+					HeapData.min = FMath::Clamp(min, 0, NumTargets - 1);
+					HeapData.max = FMath::Clamp(max, 0, NumTargets - 1);
+					uint32 dataAddress = uint32(m_heapData.Add(HeapData));
 
                     // Just the first of the targets
-					if ( bifactor < SMALL_NUMBER && bifactor > -SMALL_NUMBER )
+					if ( bifactor < UE_SMALL_NUMBER && bifactor > -UE_SMALL_NUMBER )
                     {                        
                         AddOp( SCHEDULED_OP( item.at, item, 2, dataAddress),
-                                SCHEDULED_OP( args.base, item),
-                                SCHEDULED_OP( args.targets[min], item) );
+                                SCHEDULED_OP(BaseAt, item),
+                                SCHEDULED_OP(Targets[min], item) );
                     }
                     // Just the second of the targets
-					else if ( bifactor > 1.0f - SMALL_NUMBER || bifactor <= -1.0f + SMALL_NUMBER )
+					else if ( bifactor > 1.0f - UE_SMALL_NUMBER || bifactor <= -1.0f + UE_SMALL_NUMBER )
                     {
                         check( max>0 );
 
                         AddOp( SCHEDULED_OP( item.at, item, 2, dataAddress),
-                                SCHEDULED_OP( args.base, item),
-                                SCHEDULED_OP( args.targets[max], item) );
+                                SCHEDULED_OP(BaseAt, item),
+                                SCHEDULED_OP(Targets[max], item) );
                     }
                     // Mix two targets on the base
                     else
                     {
                         // We will need the base again
                         AddOp( SCHEDULED_OP( item.at, item, 2, dataAddress),
-                                SCHEDULED_OP( args.base, item),
-                                SCHEDULED_OP( args.targets[min], item),
-                                SCHEDULED_OP( args.targets[max], item) );
+                                SCHEDULED_OP(BaseAt, item),
+                                SCHEDULED_OP(Targets[min], item),
+                                SCHEDULED_OP(Targets[max], item) );
                     }
                 }
 
@@ -1069,43 +1078,43 @@ namespace mu
 
             case 2:
             {
-                Ptr<const Mesh> pBase = GetMemory().GetMesh( CACHE_ADDRESS(args.base,item) );
+                Ptr<const Mesh> pBase = GetMemory().GetMesh( CACHE_ADDRESS(BaseAt,item) );
 
                 // Factor from 0 to 1 between the two targets
-                const SCHEDULED_OP_DATA& data = m_heapData[ (size_t)item.customState ];
-                float bifactor = data.bifactor;
-                int min = data.min;
-                int max = data.max;
+                const SCHEDULED_OP_DATA& HeapData = m_heapData[ (size_t)item.customState ];
+                float bifactor = HeapData.bifactor;
+                int min = HeapData.min;
+                int max = HeapData.max;
 
                 MeshPtrConst pResult;
 
                 if (pBase)
                 {
                     // Just the first of the targets
-					if ( bifactor < SMALL_NUMBER && bifactor > -SMALL_NUMBER )
+					if ( bifactor < UE_SMALL_NUMBER && bifactor > -UE_SMALL_NUMBER )
                     {
                         // Base with one full morph
-                        Ptr<const Mesh> pMorph = GetMemory().GetMesh( CACHE_ADDRESS(args.targets[min],item) );
+                        Ptr<const Mesh> pMorph = GetMemory().GetMesh( CACHE_ADDRESS(Targets[min],item) );
                         if (pMorph)
                         {
 							pResult = MeshMorph(pBase.get(), pMorph.get());
                         }
                     }
                     // Just the second of the targets
-                    else if ( bifactor > 1.0f - SMALL_NUMBER )
+                    else if ( bifactor > 1.0f - UE_SMALL_NUMBER )
                     {
                         check( max>0 );
-                        Ptr<const Mesh> pMorph = GetMemory().GetMesh( CACHE_ADDRESS(args.targets[max],item) );
+                        Ptr<const Mesh> pMorph = GetMemory().GetMesh( CACHE_ADDRESS(Targets[max],item) );
                         if (pMorph)
                         {
 							pResult = MeshMorph(pBase.get(), pMorph.get());
                         }
                     }
 					// Negative target
-					else if (bifactor <= -1.0f + SMALL_NUMBER)
+					else if (bifactor <= -1.0f + UE_SMALL_NUMBER)
 					{
 						check( max > 0 );
-						Ptr<const Mesh> pMorph = GetMemory().GetMesh(CACHE_ADDRESS(args.targets[max], item));
+						Ptr<const Mesh> pMorph = GetMemory().GetMesh(CACHE_ADDRESS(Targets[max], item));
 						if (pMorph)
 						{
 							pResult = MeshMorph(pBase.get(), pMorph.get(), bifactor);
@@ -1114,8 +1123,8 @@ namespace mu
                     // Mix two targets on the base
                     else
                     {
-                        Ptr<const Mesh> pMin = GetMemory().GetMesh( CACHE_ADDRESS(args.targets[min],item) );
-                        Ptr<const Mesh> pMax = GetMemory().GetMesh( CACHE_ADDRESS(args.targets[max],item) );
+                        Ptr<const Mesh> pMin = GetMemory().GetMesh( CACHE_ADDRESS(Targets[min],item) );
+                        Ptr<const Mesh> pMax = GetMemory().GetMesh( CACHE_ADDRESS(Targets[max],item) );
                         if (pMin && pMax)
                         {
                             pResult = MeshMorph2( pBase.get(), pMin.get(), pMax.get(), bifactor );
@@ -1247,7 +1256,7 @@ namespace mu
 				uint32 dataAddress = uint32(m_heapData.Num());
 
                 // Just the first of the targets
-                if ( bifactor < SMALL_NUMBER )
+                if ( bifactor < UE_SMALL_NUMBER )
                 {
                     if( min==0 )
                     {
@@ -1270,7 +1279,7 @@ namespace mu
                         }
                     }
                 // Just the second of the targets
-                else if ( bifactor > 1.0f-SMALL_NUMBER )
+                else if ( bifactor > 1.0f-UE_SMALL_NUMBER )
                 {
                     m_heapData.Add(data);
                         AddOp( SCHEDULED_OP( item.at, item, 2, dataAddress),
@@ -1323,7 +1332,7 @@ namespace mu
                 if (pBase)
                 {
                     // Just the first of the targets
-                    if ( bifactor < SMALL_NUMBER )
+                    if ( bifactor < UE_SMALL_NUMBER )
                     {
                         if( min==0 )
                         {
@@ -1338,7 +1347,7 @@ namespace mu
                         }
                     }
                     // Just the second of the targets
-                    else if ( bifactor > 1.0f-SMALL_NUMBER )
+                    else if ( bifactor > 1.0f-UE_SMALL_NUMBER )
                     {
                         check( max>0 );
                         Ptr<const Mesh> pMorph = GetMemory().GetMesh( CACHE_ADDRESS(args.targets[max-1],item) );
@@ -3139,12 +3148,12 @@ namespace mu
 				data.max = FMath::Clamp(max, 0, count - 1);
 				uint32 dataPos = uint32(m_heapData.Add(data));
 
-                if ( bifactor < SMALL_NUMBER )
+                if ( bifactor < UE_SMALL_NUMBER )
                 {
                         AddOp( SCHEDULED_OP( item.at, item, 2, dataPos),
                                SCHEDULED_OP( args.targets[min], item) );
                     }
-                else if ( bifactor > 1.0f-SMALL_NUMBER )
+                else if ( bifactor > 1.0f-UE_SMALL_NUMBER )
                 {
                         AddOp( SCHEDULED_OP( item.at, item, 2, dataPos),
                                SCHEDULED_OP( args.targets[max], item) );
@@ -3176,12 +3185,12 @@ namespace mu
                 int max = data.max;
 
                 ImagePtr pResult;
-                if ( bifactor < SMALL_NUMBER )
+                if ( bifactor < UE_SMALL_NUMBER )
                 {
                     Ptr<const Image> pSource = GetMemory().GetImage( CACHE_ADDRESS(args.targets[min],item) );
                     pResult = pSource->Clone();
                 }
-                else if ( bifactor > 1.0f-SMALL_NUMBER )
+                else if ( bifactor > 1.0f-UE_SMALL_NUMBER )
                 {
                     Ptr<const Image> pSource = GetMemory().GetImage( CACHE_ADDRESS(args.targets[max],item) );
                     pResult = pSource->Clone();
diff --git a/Engine/Plugins/Experimental/Mutable/Source/MutableRuntime/Private/MuR/Operations.cpp b/Engine/Plugins/Experimental/Mutable/Source/MutableRuntime/Private/MuR/Operations.cpp
--- a/Engine/Plugins/Experimental/Mutable/Source/MutableRuntime/Private/MuR/Operations.cpp
+++ b/Engine/Plugins/Experimental/Mutable/Source/MutableRuntime/Private/MuR/Operations.cpp
@@ -374,16 +374,6 @@ namespace mu
             f(&op.args.MeshApplyLayout.mesh );
             break;
 
-        case OP_TYPE::ME_MORPH2:
-            f(&op.args.MeshMorph2.factor );
-            f(&op.args.MeshMorph2.base );
-
-            for (int t=0;t<MUTABLE_OP_MAX_MORPH2_TARGETS;++t)
-            {
-                f(&op.args.MeshMorph2.targets[t] );
-            }
-            break;
-
         case OP_TYPE::ME_MERGE:
             f(&op.args.MeshMerge.base );
             f(&op.args.MeshMerge.added );
@@ -920,28 +910,39 @@ namespace mu
         case OP_TYPE::ME_DIFFERENCE:
         {
 			const uint8_t* data = program.GetOpArgsPointer(at);
+
 			OP::ADDRESS BaseAt = 0;
 			FMemory::Memcpy(&BaseAt, data, sizeof(OP::ADDRESS)); data += sizeof(OP::ADDRESS);
+			f(BaseAt);
 
 			OP::ADDRESS TargetAt = 0;
 			FMemory::Memcpy(&TargetAt, data, sizeof(OP::ADDRESS)); data += sizeof(OP::ADDRESS);
-
-			f(BaseAt);
-            f(TargetAt);
-            break;
+			f(TargetAt);
+			break;
         }
 
         case OP_TYPE::ME_MORPH2:
         {
-			OP::MeshMorph2Args args = program.GetOpArgs<OP::MeshMorph2Args>(at);
-            f(args.factor );
-            f(args.base );
+			const uint8_t* data = program.GetOpArgsPointer(at);
 
-            for (int t=0;t<MUTABLE_OP_MAX_MORPH2_TARGETS;++t)
-            {
-                f(args.targets[t] );
-            }
-            break;
+			OP::ADDRESS FactorAt = 0;
+			FMemory::Memcpy(&FactorAt, data, sizeof(OP::ADDRESS)); data += sizeof(OP::ADDRESS);
+			f(FactorAt);
+
+			OP::ADDRESS BaseAt = 0;
+			FMemory::Memcpy(&BaseAt, data, sizeof(OP::ADDRESS)); data += sizeof(OP::ADDRESS);
+			f(BaseAt);
+
+			uint8 NumTargets = 0;
+			FMemory::Memcpy(&NumTargets, data, sizeof(uint8)); data += sizeof(uint8);
+
+			for (uint8 T = 0; T < NumTargets; ++T)
+			{
+				OP::ADDRESS TargetAt = 0;
+				FMemory::Memcpy(&TargetAt, data, sizeof(OP::ADDRESS)); data += sizeof(OP::ADDRESS);
+				f(TargetAt);
+			}
+			break;
         }
 
         case OP_TYPE::ME_MERGE:
diff --git a/Engine/Plugins/Experimental/Mutable/Source/MutableRuntime/Private/MuR/Operations.h b/Engine/Plugins/Experimental/Mutable/Source/MutableRuntime/Private/MuR/Operations.h
--- a/Engine/Plugins/Experimental/Mutable/Source/MutableRuntime/Private/MuR/Operations.h
+++ b/Engine/Plugins/Experimental/Mutable/Source/MutableRuntime/Private/MuR/Operations.h
@@ -14,7 +14,6 @@
 #define MUTABLE_OP_MAX_INTERPOLATE_COUNT	6
 #define MUTABLE_OP_MAX_ADD_COUNT			7
 #define MUTABLE_OP_MAX_SWIZZLE_CHANNELS		4
-#define MUTABLE_OP_MAX_MORPH2_TARGETS		5
 
 namespace mu
 {
@@ -779,22 +778,6 @@ namespace mu
             uint16 channel;
         };
 
-		// Moved to dynamically-sized link data
-        //struct MeshDifferenceArgs
-        //{
-        //    ADDRESS base;
-        //    ADDRESS target;
-
-        //    //! If true the texture coordinates will not be morphed. This is only relevant if the list
-        //    //! in channelSemantic is empty.
-        //    uint8 ignoreTextureCoords;
-
-        //    //! Channels to diff. Meaningful until the first null channel. If all are null, it means
-        //    //! that all the channels should be diff-ed
-        //    uint8 channelSemantic[MUTABLE_OP_MAX_MORPH_CHANNELS];
-        //    uint8 channelSemanticIndex[MUTABLE_OP_MAX_MORPH_CHANNELS];
-        //};
-
         struct MeshMergeArgs
         {
             ADDRESS base;
@@ -812,13 +795,6 @@ namespace mu
             ADDRESS targets[ MUTABLE_OP_MAX_INTERPOLATE_COUNT-1 ];
         };
 
-        struct MeshMorph2Args
-        {
-            ADDRESS factor;
-            ADDRESS base;
-            ADDRESS targets[ MUTABLE_OP_MAX_MORPH2_TARGETS ];
-        };
-
         struct MeshMaskClipMeshArgs
         {
             ADDRESS source;
@@ -1111,8 +1087,6 @@ namespace mu
 
             //-------------------------------------------------------------------------------------
             MeshApplyLayoutArgs MeshApplyLayout;
-            //MeshDifferenceArgs MeshDifference;
-            MeshMorph2Args MeshMorph2;
             MeshMergeArgs MeshMerge;
             MeshInterpolateArgs MeshInterpolate;
             MeshMaskDiffArgs MeshMaskDiff;
diff --git a/Engine/Plugins/Experimental/Mutable/Source/MutableTools/Private/MuT/ASTOpMeshFormat.cpp b/Engine/Plugins/Experimental/Mutable/Source/MutableTools/Private/MuT/ASTOpMeshFormat.cpp
--- a/Engine/Plugins/Experimental/Mutable/Source/MutableTools/Private/MuT/ASTOpMeshFormat.cpp
+++ b/Engine/Plugins/Experimental/Mutable/Source/MutableTools/Private/MuT/ASTOpMeshFormat.cpp
@@ -16,6 +16,7 @@
 #include "MuT/ASTOpConstantResource.h"
 #include "MuT/ASTOpMeshClipMorphPlane.h"
 #include "MuT/ASTOpMeshRemoveMask.h"
+#include "MuT/ASTOpMeshMorph.h"
 #include "MuT/ASTOpSwitch.h"
 #include "MuT/StreamsPrivate.h"
 
@@ -273,8 +274,8 @@ mu::Ptr<ASTOp> Sink_MeshFormatAST::Visit(const mu::Ptr<ASTOp>& at, const ASTOpMe
 	case OP_TYPE::ME_MORPH2:
 	{
 		// Move the format down the base of the morph
-		auto newOp = mu::Clone<ASTOpFixed>(at);
-		newOp->SetChild(newOp->op.args.MeshMorph2.base, Visit(newOp->children[newOp->op.args.MeshMorph2.base].child(), currentFormatOp));
+		Ptr<ASTOpMeshMorph> newOp = mu::Clone<ASTOpMeshMorph>(at);
+		newOp->Base = Visit(newOp->Base.child(), currentFormatOp);
 
 		// Reformat the morph targets to match the new format.
 		// \TODO: Cache pTargetMorphFormat? motaop?
@@ -286,9 +287,9 @@ mu::Ptr<ASTOp> Sink_MeshFormatAST::Visit(const mu::Ptr<ASTOp>& at, const ASTOpMe
 		motaop->SetValue(pTargetMorphFormat, false /* useDiskCache */);
 		auto targetMorphFormatAt = motaop;
 
-		for (int t = 0; t < MUTABLE_OP_MAX_MORPH2_TARGETS; ++t)
+		for (int32 t=0; t<newOp->Targets.Num(); ++t)
 		{
-			if (newOp->children[newOp->op.args.MeshMorph2.targets[t]])
+			if (newOp->Targets[t])
 			{
 				mu::Ptr<ASTOpMeshFormat> newFormat = mu::Clone<ASTOpMeshFormat>(currentFormatOp);
 				newFormat->Buffers =
@@ -296,7 +297,7 @@ mu::Ptr<ASTOp> Sink_MeshFormatAST::Visit(const mu::Ptr<ASTOp>& at, const ASTOpMe
 					| OP::MeshFormatArgs::BT_IGNORE_MISSING;
 				newFormat->Format = targetMorphFormatAt;
 
-				newOp->SetChild(newOp->op.args.MeshMorph2.targets[t], Visit(newOp->children[newOp->op.args.MeshMorph2.targets[t]].child(), newFormat.get()));
+				newOp->Targets[t] = Visit(newOp->Targets[t].child(), newFormat.get());
 			}
 		}
 
diff --git a/Engine/Plugins/Experimental/Mutable/Source/MutableTools/Private/MuT/ASTOpMeshMaskClipMesh.cpp b/Engine/Plugins/Experimental/Mutable/Source/MutableTools/Private/MuT/ASTOpMeshMaskClipMesh.cpp
--- a/Engine/Plugins/Experimental/Mutable/Source/MutableTools/Private/MuT/ASTOpMeshMaskClipMesh.cpp
+++ b/Engine/Plugins/Experimental/Mutable/Source/MutableTools/Private/MuT/ASTOpMeshMaskClipMesh.cpp
@@ -150,16 +150,13 @@ namespace mu
 				switch (at->GetOpType())
 				{
 
-					// This cannot be sunk since the result is different. Since the clipping is now correctly
-					// generated at the end of the chain when really necessary, this wrong optimisation is no 
-					// longer needed.
-					//case OP_TYPE::ME_MORPH2:
-		   //         {
-		   //             // Remove this op: Let's clip the unmorphed mesh only
-		   //             auto typedAt = dynamic_cast<const ASTOpFixed*>(at.get());
-		   //             newAt = Visit( typedAt->children[typedAt->op.args.MeshMorph2.base].child() );
-		   //             break;
-		   //         }
+				// This cannot be sunk since the result is different. Since the clipping is now correctly
+				// generated at the end of the chain when really necessary, this wrong optimisation is no 
+				// longer needed.
+				//case OP_TYPE::ME_MORPH2:
+		        //{
+		        //    break;
+		        //}
 
 				case OP_TYPE::ME_REMOVEMASK:
 				{
diff --git a/Engine/Plugins/Experimental/Mutable/Source/MutableTools/Private/MuT/ASTOpMeshMorph.cpp b/Engine/Plugins/Experimental/Mutable/Source/MutableTools/Private/MuT/ASTOpMeshMorph.cpp
new file mode 100644
--- /dev/null
+++ b/Engine/Plugins/Experimental/Mutable/Source/MutableTools/Private/MuT/ASTOpMeshMorph.cpp
@@ -0,0 +1,121 @@
+// Copyright Epic Games, Inc. All Rights Reserved.
+
+#include "MuT/ASTOpMeshMorph.h"
+
+#include "HAL/PlatformMath.h"
+#include "MuR/ModelPrivate.h"
+#include "MuR/RefCounted.h"
+#include "MuR/Types.h"
+#include "MuT/StreamsPrivate.h"
+
+#include <memory>
+
+namespace mu
+{
+
+	//---------------------------------------------------------------------------------------------
+	ASTOpMeshMorph::ASTOpMeshMorph()
+		: Factor(this), Base(this)
+	{
+	}
+
+
+	//---------------------------------------------------------------------------------------------
+	ASTOpMeshMorph::~ASTOpMeshMorph()
+	{
+		// Explicit call needed to avoid recursive destruction
+		ASTOp::RemoveChildren();
+	}
+
+
+	//---------------------------------------------------------------------------------------------
+	void ASTOpMeshMorph::AddTarget(const Ptr<ASTOp>& InTarget)
+	{
+		Targets.Add( ASTChild(this,InTarget));
+	}
+
+
+	//---------------------------------------------------------------------------------------------
+	bool ASTOpMeshMorph::IsEqual(const ASTOp& otherUntyped) const
+	{
+		if (const ASTOpMeshMorph* other = dynamic_cast<const ASTOpMeshMorph*>(&otherUntyped))
+		{
+			return Factor == other->Factor && Base == other->Base && Targets == other->Targets;
+		}
+		return false;
+	}
+
+
+	//---------------------------------------------------------------------------------------------
+	mu::Ptr<ASTOp> ASTOpMeshMorph::Clone(MapChildFuncRef mapChild) const
+	{
+		Ptr<ASTOpMeshMorph> n = new ASTOpMeshMorph();
+		n->Factor = mapChild(Factor.child());
+		n->Base = mapChild(Base.child());
+		for (const ASTChild& r : Targets)
+		{
+			n->Targets.Add(ASTChild(n,mapChild(r.child())));
+		}
+		return n;
+	}
+
+
+	//---------------------------------------------------------------------------------------------
+	void ASTOpMeshMorph::ForEachChild(const TFunctionRef<void(ASTChild&)> f)
+	{
+		f(Factor);
+		f(Base);
+		for (ASTChild& T : Targets)
+		{
+			f(T);
+		}
+	}
+
+
+	//---------------------------------------------------------------------------------------------
+	uint64 ASTOpMeshMorph::Hash() const
+	{
+		uint64 res = std::hash<ASTOp*>()(Factor.child().get());
+		hash_combine(res, Base.child().get());
+		for (const ASTChild& T : Targets)
+		{
+			hash_combine(res, T.child().get());
+		}
+		return res;
+	}
+
+
+	//---------------------------------------------------------------------------------------------
+	void ASTOpMeshMorph::Link(PROGRAM& program, const FLinkerOptions*)
+	{
+		// Already linked?
+		if (!linkedAddress)
+		{
+			linkedAddress = (OP::ADDRESS)program.m_opAddress.Num();
+
+			program.m_opAddress.Add((uint32_t)program.m_byteCode.Num());
+			AppendCode(program.m_byteCode, OP_TYPE::ME_MORPH2);
+
+			OP::ADDRESS FactorAt = Factor ? Factor->linkedAddress : 0;
+			AppendCode(program.m_byteCode, FactorAt);
+
+			OP::ADDRESS BaseAt = Base ? Base->linkedAddress : 0;
+			AppendCode(program.m_byteCode, BaseAt);
+
+			AppendCode(program.m_byteCode, (uint8)Targets.Num());
+			for (const ASTChild& b : Targets)
+			{
+				OP::ADDRESS TargetAt = b ? b->linkedAddress : 0;
+				AppendCode(program.m_byteCode, TargetAt);
+			}
+		}
+	}
+
+
+	//-------------------------------------------------------------------------------------------------
+	mu::Ptr<ASTOp> ASTOpMeshMorph::OptimiseSink(const MODEL_OPTIMIZATION_OPTIONS&, OPTIMIZE_SINK_CONTEXT&) const
+	{
+		return nullptr;
+	}
+
+}
diff --git a/Engine/Plugins/Experimental/Mutable/Source/MutableTools/Private/MuT/ASTOpMeshMorph.h b/Engine/Plugins/Experimental/Mutable/Source/MutableTools/Private/MuT/ASTOpMeshMorph.h
new file mode 100644
--- /dev/null
+++ b/Engine/Plugins/Experimental/Mutable/Source/MutableTools/Private/MuT/ASTOpMeshMorph.h
@@ -0,0 +1,52 @@
+// Copyright Epic Games, Inc. All Rights Reserved.
+
+#pragma once
+
+#include "HAL/Platform.h"
+#include "MuR/MemoryPrivate.h"
+#include "MuR/Operations.h"
+#include "MuR/Ptr.h"
+#include "MuT/AST.h"
+
+
+namespace mu
+{
+	struct PROGRAM;
+
+	//---------------------------------------------------------------------------------------------
+	//! From a source mesh, remove a list of fragments with a condition.
+	//---------------------------------------------------------------------------------------------
+	class ASTOpMeshMorph : public ASTOp
+	{
+	public:
+
+		//! Factor selecting what morphs to apply and with what weight.
+		ASTChild Factor;
+
+		//! Base mesh to remove from.
+		ASTChild Base;
+
+		//! Targets to apply on the base depending on the factor
+		TArray<ASTChild> Targets;
+
+	public:
+
+		ASTOpMeshMorph();
+		ASTOpMeshMorph(const ASTOpMeshMorph&) = delete;
+		~ASTOpMeshMorph() override;
+
+		// ASTOp interface
+		OP_TYPE GetOpType() const override { return OP_TYPE::ME_MORPH2; }
+		uint64 Hash() const override;
+		void ForEachChild(const TFunctionRef<void(ASTChild&)>) override;
+		bool IsEqual(const ASTOp& otherUntyped) const override;
+		Ptr<ASTOp> Clone(MapChildFuncRef mapChild) const override;
+		void Link(PROGRAM& program, const FLinkerOptions*) override;
+		Ptr<ASTOp> OptimiseSink(const MODEL_OPTIMIZATION_OPTIONS&, OPTIMIZE_SINK_CONTEXT&) const override;
+
+		// Own interface
+		void AddTarget(const Ptr<ASTOp>&);
+	};
+
+}
+
diff --git a/Engine/Plugins/Experimental/Mutable/Source/MutableTools/Private/MuT/ASTOpMeshRemoveMask.cpp b/Engine/Plugins/Experimental/Mutable/Source/MutableTools/Private/MuT/ASTOpMeshRemoveMask.cpp
--- a/Engine/Plugins/Experimental/Mutable/Source/MutableTools/Private/MuT/ASTOpMeshRemoveMask.cpp
+++ b/Engine/Plugins/Experimental/Mutable/Source/MutableTools/Private/MuT/ASTOpMeshRemoveMask.cpp
@@ -2,6 +2,7 @@
 
 #include "MuT/ASTOpMeshRemoveMask.h"
 
+#include "MuT/ASTOpMeshMorph.h"
 #include "HAL/PlatformMath.h"
 #include "MuR/ModelPrivate.h"
 #include "MuR/RefCounted.h"
@@ -40,7 +41,7 @@ namespace mu
 	//---------------------------------------------------------------------------------------------
 	bool ASTOpMeshRemoveMask::IsEqual(const ASTOp& otherUntyped) const
 	{
-		if (auto other = dynamic_cast<const ASTOpMeshRemoveMask*>(&otherUntyped))
+		if (const ASTOpMeshRemoveMask* other = dynamic_cast<const ASTOpMeshRemoveMask*>(&otherUntyped))
 		{
 			return source == other->source && removes == other->removes;
 		}
@@ -53,7 +54,7 @@ namespace mu
 	{
 		Ptr<ASTOpMeshRemoveMask> n = new ASTOpMeshRemoveMask();
 		n->source = mapChild(source.child());
-		for (const auto& r : removes)
+		for (const TPair<ASTChild, ASTChild>& r : removes)
 		{
 			n->removes.Add({ ASTChild(n,mapChild(r.Key.child())), ASTChild(n,mapChild(r.Value.child())) });
 		}
@@ -61,18 +62,11 @@ namespace mu
 	}
 
 
-	//---------------------------------------------------------------------------------------------
-	void ASTOpMeshRemoveMask::Assert()
-	{
-		ASTOp::Assert();
-	}
-
-
 	//---------------------------------------------------------------------------------------------
 	void ASTOpMeshRemoveMask::ForEachChild(const TFunctionRef<void(ASTChild&)> f)
 	{
 		f(source);
-		for (auto& r : removes)
+		for (TPair<ASTChild, ASTChild>& r : removes)
 		{
 			f(r.Key);
 			f(r.Value);
@@ -84,7 +78,7 @@ namespace mu
 	uint64 ASTOpMeshRemoveMask::Hash() const
 	{
 		uint64 res = std::hash<ASTOp*>()(source.child().get());
-		for (const auto& r : removes)
+		for (const TPair<ASTChild, ASTChild>& r : removes)
 		{
 			hash_combine(res, r.Key.child().get());
 			hash_combine(res, r.Value.child().get());
@@ -106,7 +100,7 @@ namespace mu
 			OP::ADDRESS sourceAt = source ? source->linkedAddress : 0;
 			AppendCode(program.m_byteCode, sourceAt);
 			AppendCode(program.m_byteCode, (uint16)removes.Num());
-			for (const auto& b : removes)
+			for (const TPair<ASTChild, ASTChild>& b : removes)
 			{
 				OP::ADDRESS condition = b.Key ? b.Key->linkedAddress : 0;
 				AppendCode(program.m_byteCode, condition);
@@ -164,7 +158,7 @@ namespace mu
 				}
 
 				// Already visited?
-				auto cacheIt = m_oldToNew.Find(at);
+				Ptr<ASTOp>* cacheIt = m_oldToNew.Find(at);
 				if (cacheIt)
 				{
 					return *cacheIt;
@@ -176,8 +170,8 @@ namespace mu
 
 				case OP_TYPE::ME_MORPH2:
 				{
-					mu::Ptr<ASTOpFixed> newOp = mu::Clone<ASTOpFixed>(at);
-					newOp->SetChild(newOp->op.args.MeshMorph2.base, Visit(newOp->children[newOp->op.args.MeshMorph2.base].child()));
+					mu::Ptr<ASTOpMeshMorph> newOp = mu::Clone<ASTOpMeshMorph>(at);
+					newOp->Base = Visit(newOp->Base.child());
 					newAt = newOp;
 					break;
 				}
@@ -209,7 +203,7 @@ namespace mu
 					//
 					if (at != m_initialSource)
 					{
-						auto newOp = mu::Clone<ASTOpMeshRemoveMask>(m_root);
+						Ptr<ASTOpMeshRemoveMask> newOp = mu::Clone<ASTOpMeshRemoveMask>(m_root);
 						newOp->source = at;
 						newAt = newOp;
 					}
@@ -230,7 +224,7 @@ namespace mu
 	mu::Ptr<ASTOp> ASTOpMeshRemoveMask::OptimiseSink(const MODEL_OPTIMIZATION_OPTIONS&, OPTIMIZE_SINK_CONTEXT&) const
 	{
 		Sink_MeshRemoveMaskAST sinker;
-		auto at = sinker.Apply(this);
+		mu::Ptr<ASTOp> at = sinker.Apply(this);
 
 		return at;
 	}
diff --git a/Engine/Plugins/Experimental/Mutable/Source/MutableTools/Private/MuT/ASTOpMeshRemoveMask.h b/Engine/Plugins/Experimental/Mutable/Source/MutableTools/Private/MuT/ASTOpMeshRemoveMask.h
--- a/Engine/Plugins/Experimental/Mutable/Source/MutableTools/Private/MuT/ASTOpMeshRemoveMask.h
+++ b/Engine/Plugins/Experimental/Mutable/Source/MutableTools/Private/MuT/ASTOpMeshRemoveMask.h
@@ -38,7 +38,6 @@ namespace mu
 		void ForEachChild(const TFunctionRef<void(ASTChild&)>) override;
 		bool IsEqual(const ASTOp& otherUntyped) const override;
 		Ptr<ASTOp> Clone(MapChildFuncRef mapChild) const override;
-		void Assert() override;
 		void Link(PROGRAM& program, const FLinkerOptions*) override;
 		Ptr<ASTOp> OptimiseSink(const MODEL_OPTIMIZATION_OPTIONS&, OPTIMIZE_SINK_CONTEXT&) const override;
 
diff --git a/Engine/Plugins/Experimental/Mutable/Source/MutableTools/Private/MuT/CodeGenerator.cpp b/Engine/Plugins/Experimental/Mutable/Source/MutableTools/Private/MuT/CodeGenerator.cpp
--- a/Engine/Plugins/Experimental/Mutable/Source/MutableTools/Private/MuT/CodeGenerator.cpp
+++ b/Engine/Plugins/Experimental/Mutable/Source/MutableTools/Private/MuT/CodeGenerator.cpp
@@ -27,6 +27,7 @@
 #include "MuT/ASTOpMeshMaskClipMesh.h"
 #include "MuT/ASTOpMeshRemoveMask.h"
 #include "MuT/ASTOpMeshDifference.h"
+#include "MuT/ASTOpMeshMorph.h"
 #include "MuT/ASTOpParameter.h"
 #include "MuT/CodeGenerator_SecondPass.h"
 #include "MuT/CodeOptimiser.h"
@@ -268,7 +269,7 @@ namespace mu
         // Free caches
         m_compiled.Reset();
         m_constantMeshes.Empty();
-        m_addedLayouts.Empty();
+		m_generatedLayouts.Empty();
         m_nodeVariables.clear();
         m_generatedMeshes.Reset();
         m_generatedProjectors.Reset();
@@ -523,29 +524,26 @@ namespace mu
 		// The layout we are adding must be a source layout, without block ids yet.
 		check(SourceLayout->m_blocks.IsEmpty() || SourceLayout->m_blocks[0].m_id == -1);
 		
-		Ptr<const Layout>* it = m_addedLayouts.Find(SourceLayout.get() );
+		Ptr<const Layout>* it = m_generatedLayouts.Find(SourceLayout.get() );
 
-		Ptr<const Layout> pResult;
         if ( it )
         {
-            pResult = *it;
+            return *it;
         }
-        else
+
+		// Assign unique ids to each layout block
+		Ptr<Layout> ClonedLayout = SourceLayout->Clone();
+        for (int32 b=0;b< ClonedLayout->m_blocks.Num();++b)
         {
-			Ptr<Layout> ClonedLayout = SourceLayout->Clone();
-            for (int32 b=0;b< ClonedLayout->m_blocks.Num();++b)
-            {
-				// This is a hard limit due to layout block index data being stored in 16 bit.
-				check(m_absoluteLayoutIndex<65536);
-				ClonedLayout->m_blocks[b].m_id = m_absoluteLayoutIndex++;
-            }
-			check(SourceLayout->m_blocks.Num() == ClonedLayout->m_blocks.Num());
-			check(ClonedLayout->m_blocks.IsEmpty() || ClonedLayout->m_blocks[0].m_id != -1);
-			m_addedLayouts.Add(SourceLayout.get(), ClonedLayout);
-			pResult = ClonedLayout;
+			// This is a hard limit due to layout block index data being stored in 16 bit.
+			check(m_absoluteLayoutIndex<65536);
+			ClonedLayout->m_blocks[b].m_id = m_absoluteLayoutIndex++;
         }
+		check(SourceLayout->m_blocks.Num() == ClonedLayout->m_blocks.Num());
+		check(ClonedLayout->m_blocks.IsEmpty() || ClonedLayout->m_blocks[0].m_id != -1);
+		m_generatedLayouts.Add(SourceLayout.get(), ClonedLayout);
 
-        return pResult;
+        return ClonedLayout;
     }
 
 
@@ -765,7 +763,7 @@ namespace mu
     //---------------------------------------------------------------------------------------------
     Ptr<ASTOp> CodeGenerator::Visit( const NodeLOD::Private& node )
     {
-        // Build a series of operations to assemble the component
+        // Build a series of operations to assemble all the LOD components
         Ptr<ASTOp> lastCompOp;
 
         // Create the expression for each component in this object
@@ -847,11 +845,11 @@ namespace mu
         const auto& node = *surfaceNode->GetPrivate();
 
         // Clear the surface generation state
-//        m_addedLayouts.Empty();
+//        m_generatedLayouts.Empty();
 //        m_absoluteLayoutIndex = 0;
 //        m_currentLayoutMesh = nullptr;
 //        m_currentLayoutChannel = 0;
-//        m_generatedMeshes.clear();
+//        m_generatedMeshes.Empty();
 
         // Build a series of operations to assemble the surface
         Ptr<ASTOp> lastSurfOp;
@@ -968,8 +966,8 @@ namespace mu
 							bModifiersForBeforeOperations, node.m_errorContext);
 
                         FMeshGenerationResult::EXTRA_LAYOUTS data;
-                        data.layouts = addResults.layouts;
-                        data.condition = e.condition;
+						data.GeneratedLayouts = addResults.GeneratedLayouts;
+						data.condition = e.condition;
                         data.meshFragment = addResults.meshOp;
                         meshResults.extraMeshLayouts[editIndex] = data;
 
@@ -1107,13 +1105,12 @@ namespace mu
                     // Morph operation
                     Ptr<ASTOp> morphAd;
                     {
-                        Ptr<ASTOpFixed> op = new ASTOpFixed();
-						op->op.type = OP_TYPE::ME_MORPH2;
+                        Ptr<ASTOpMeshMorph> op = new ASTOpMeshMorph();
 
 						// Factor
 						if (e.node->m_pFactor)
 						{
-							op->SetChild(op->op.args.MeshMorph2.factor, Generate(e.node->m_pFactor));
+							op->Factor = Generate(e.node->m_pFactor);
 						}
 						else
 						{
@@ -1121,15 +1118,15 @@ namespace mu
 							auxNode->SetValue(1.0f);
 							Ptr<ASTOp> resultNode = Generate(auxNode);
 
-							op->SetChild(op->op.args.MeshMorph2.factor, resultNode);
+							op->Factor = resultNode;
 						}
 
 						// Base		
-						op->SetChild(op->op.args.MeshMorph2.base, lastMeshOp);
+						op->Base = lastMeshOp;
 
 						// Targets
-						op->SetChild(op->op.args.MeshMorph2.targets[0], diffBase);
-						op->SetChild(op->op.args.MeshMorph2.targets[1], diffAd);
+						op->AddTarget(diffBase);
+						op->AddTarget(diffAd);
                         morphAd = op;
                     }
 
@@ -1156,11 +1153,11 @@ namespace mu
 				bModifiersForBeforeOperations, node.m_errorContext);
 
             // Layouts
-            for ( int32 l=0; l<meshResults.layouts.Num(); ++l )
+            for ( int32 l=0; l<meshResults.GeneratedLayouts.Num(); ++l )
             {
                 Ptr<ASTOp> layoutOp = nullptr;
 
-				Ptr<const Layout> pLayout = meshResults.layouts[l].get();
+				Ptr<const Layout> pLayout = meshResults.GeneratedLayouts[l].get();
                 if ( pLayout )
                 {
                     // Create the layout expression
@@ -1168,12 +1165,9 @@ namespace mu
                     // Constant layout to start with
                     {
                         Ptr<ASTOpConstantResource> op = new ASTOpConstantResource();
-                        op->type = OP_TYPE::LA_CONSTANT;
-
-                        // Clone and set absolute block ids
-						Ptr<const Layout> pCloned = AddLayout( pLayout );
+                        op->type = OP_TYPE::LA_CONSTANT;                       
 
-                        op->SetValue( pCloned, m_compilerOptions->m_optimisationOptions.m_useDiskCache );
+                        op->SetValue(pLayout, m_compilerOptions->m_optimisationOptions.m_useDiskCache );
                         layoutOp = op;
                     }
 
@@ -1183,11 +1177,11 @@ namespace mu
                         if ( !data.meshFragment )
                         {
                             // No mesh to add, we assume there are no layouts to add either.
-                            check(data.layouts.IsEmpty());
+                            check(data.GeneratedLayouts.IsEmpty());
                             continue;
                         }
 
-                        if ( data.layouts.Num() != meshResults.layouts.Num() )
+                        if ( data.GeneratedLayouts.Num() != meshResults.GeneratedLayouts.Num() )
                         {
                             m_pErrorLog->GetPrivate()->Add
                                 ( "Merged layout has been ignored because the number of layouts"
@@ -1202,7 +1196,7 @@ namespace mu
                                 Ptr<ASTOpConstantResource> op = new ASTOpConstantResource();
                                 op->type = OP_TYPE::LA_CONSTANT;
 
-								Ptr<const Layout> pCloned = AddLayout( data.layouts[l] );
+								Ptr<const Layout> pCloned = data.GeneratedLayouts[l];
                                 op->SetValue( pCloned, m_compilerOptions->m_optimisationOptions.m_useDiskCache );
 
                                 layoutFragmentAd = op;
@@ -1435,7 +1429,7 @@ namespace mu
                           ==
                           CompilerOptions::TextureLayoutStrategy::Pack )
                 {
-                    if( layout >= meshResults.layouts.Num() ||
+                    if( layout >= meshResults.GeneratedLayouts.Num() ||
                         layout >= meshResults.layoutOps.Num() )
                     {
                         m_pErrorLog->GetPrivate()->Add("Missing layout in object, or its parent.",
@@ -1443,7 +1437,7 @@ namespace mu
                     }
                     else
                     {
-                        const Layout* pLayout = m_addedLayouts[ meshResults.layouts[ layout ].get() ].get();
+                        const Layout* pLayout = meshResults.GeneratedLayouts[ layout ].get();
                         check(pLayout);
 
                         Ptr<ASTOpInstanceAdd> op = new ASTOpInstanceAdd();
@@ -1583,9 +1577,9 @@ namespace mu
                                 auto pExtend = e.node->m_textures[t].m_pExtend;
                                 if (pExtend)
                                 {
-                                    if ( layout>=meshResults.extraMeshLayouts[editIndex].layouts.Num()
+                                    if ( layout>=meshResults.extraMeshLayouts[editIndex].GeneratedLayouts.Num()
                                          ||
-                                         !meshResults.extraMeshLayouts[editIndex].layouts[layout] )
+                                         !meshResults.extraMeshLayouts[editIndex].GeneratedLayouts[layout] )
                                     {
                                         char buf[256];
                                         mutable_snprintf
@@ -1599,8 +1593,7 @@ namespace mu
                                     }
                                     else
                                     {
-                                        auto pExtendLayoutSource = meshResults.extraMeshLayouts[editIndex].layouts[layout];
-                                        auto pExtendLayout = m_addedLayouts[ pExtendLayoutSource.get() ];
+                                        Ptr<const Layout> pExtendLayout = meshResults.extraMeshLayouts[editIndex].GeneratedLayouts[layout];
 
                                         // Find out the size of the image
                                         // TODO: What if the image is empty and everything is added?
@@ -1608,7 +1601,7 @@ namespace mu
                                         FImageDesc extendDesc = CalculateImageDesc( *pExtend->GetBasePrivate() );
 
                                         // Size of a layout block in pixels
-                                        FIntPoint extlayout = pExtendLayout->GetGridSize( );
+                                        FIntPoint extlayout = pExtendLayout->GetGridSize();
 
                                         Ptr<ASTOp> lastBase = imageAd;
 
@@ -2105,7 +2098,7 @@ namespace mu
         {
             if ( const NodeLOD* pLODNode = node.m_lods[t].get() )
             {
-                m_currentParents.Last().m_lod = (int)t;
+                m_currentParents.Last().m_lod = t;
 
                 Ptr<ASTOp> lodOp = Generate( pLODNode );
 
diff --git a/Engine/Plugins/Experimental/Mutable/Source/MutableTools/Private/MuT/CodeGenerator.h b/Engine/Plugins/Experimental/Mutable/Source/MutableTools/Private/MuT/CodeGenerator.h
--- a/Engine/Plugins/Experimental/Mutable/Source/MutableTools/Private/MuT/CodeGenerator.h
+++ b/Engine/Plugins/Experimental/Mutable/Source/MutableTools/Private/MuT/CodeGenerator.h
@@ -234,10 +234,6 @@ namespace mu
         //! palceholders for missing images.
         ImagePtr m_missingImage[size_t(EImageFormat::IF_COUNT)];
 
-        //! Map of layouts found in the code already generated. The map is from the source layout
-        //! pointer to the cloned layout. The cloned layout will have absolute block ids assigned.
-        TMap<Ptr<const Layout>,Ptr<const Layout>> m_addedLayouts;
-
         //! First free index for a layout block
         int32 m_absoluteLayoutIndex = 0;
 
@@ -331,7 +327,7 @@ namespace mu
         //void GetSurfacesWithTag(const string& tag, vector<FirstPassGenerator::SURFACE>& surfaces);
 
         //-----------------------------------------------------------------------------------------
-        void PrepareForLayout( LayoutPtrConst pSourceLayout,
+        void PrepareForLayout( LayoutPtrConst GeneratedLayout,
                                   MeshPtr currentLayoutMesh,
                                   size_t currentLayoutChannel,
                                   const void* errorContext );
@@ -494,15 +490,15 @@ namespace mu
 			//! Original base mesh before removes, morphs, etc.
 			Ptr<ASTOp> baseMeshOp;
 
-			/** Source node layouts to use with these extra mesh. They don't have block ids. */
-			TArray<Ptr<const Layout>> layouts;
+			/** Generated node layouts with their own block ids. */
+			TArray<Ptr<const Layout>> GeneratedLayouts;
 
 			TArray<Ptr<ASTOp>> layoutOps;
 
 			struct EXTRA_LAYOUTS
 			{
 				/** Source node layouts to use with these extra mesh. They don't have block ids. */
-				TArray<Ptr<const Layout>> layouts;
+				TArray<Ptr<const Layout>> GeneratedLayouts;
 				Ptr<ASTOp> condition;
 				Ptr<ASTOp> meshFragment;
 			};
@@ -531,6 +527,11 @@ namespace mu
         typedef TMap<FGeneratedMeshCacheKey,FMeshGenerationResult> GeneratedMeshMap;
         GeneratedMeshMap m_generatedMeshes;
 
+		//! Map of layouts found in the code already generated. The map is from the source layout
+		//! pointer of the layouts in the node meshes to the cloned and modified layout. 
+		//! The cloned layout will have absolute block ids assigned.
+		TMap<Ptr<const Layout>, Ptr<const Layout>> m_generatedLayouts;
+
         void GenerateMesh(const FMeshGenerationOptions&, FMeshGenerationResult& result, const NodeMeshPtrConst& node);
         void GenerateMesh_Constant(const FMeshGenerationOptions&, FMeshGenerationResult&, const NodeMeshConstant* );
         void GenerateMesh_Format(const FMeshGenerationOptions&, FMeshGenerationResult&, const NodeMeshFormat* );
diff --git a/Engine/Plugins/Experimental/Mutable/Source/MutableTools/Private/MuT/CodeGenerator_Mesh.cpp b/Engine/Plugins/Experimental/Mutable/Source/MutableTools/Private/MuT/CodeGenerator_Mesh.cpp
--- a/Engine/Plugins/Experimental/Mutable/Source/MutableTools/Private/MuT/CodeGenerator_Mesh.cpp
+++ b/Engine/Plugins/Experimental/Mutable/Source/MutableTools/Private/MuT/CodeGenerator_Mesh.cpp
@@ -29,6 +29,7 @@
 #include "MuT/ASTOpMeshMorphReshape.h"
 #include "MuT/ASTOpMeshTransform.h"
 #include "MuT/ASTOpMeshDifference.h"
+#include "MuT/ASTOpMeshMorph.h"
 #include "MuT/ASTOpSwitch.h"
 #include "MuT/CodeGenerator.h"
 #include "MuT/CodeGenerator_FirstPass.h"
@@ -85,18 +86,21 @@ namespace mu
 class Node;
 
     //---------------------------------------------------------------------------------------------
-    void CodeGenerator::PrepareForLayout( LayoutPtrConst pSourceLayout,
-                                             MeshPtr currentLayoutMesh,
-                                             size_t currentLayoutChannel,
-                                             const void* errorContext )
+    void CodeGenerator::PrepareForLayout(Ptr<const Layout> GeneratedLayout,
+		MeshPtr currentLayoutMesh,
+		size_t currentLayoutChannel,
+		const void* errorContext )
     {
         if (currentLayoutMesh->GetVertexCount()==0)
         {
             return;
         }
 
+		// The layout we are adding must have block ids.
+		check(GeneratedLayout->m_blocks.IsEmpty() || GeneratedLayout->m_blocks[0].m_id != -1);
+
         // 
-		Ptr<const Layout> pLayout = AddLayout( pSourceLayout );
+		Ptr<const Layout> pLayout = GeneratedLayout;
         currentLayoutMesh->AddLayout( pLayout );
 
         int buffer = -1;
@@ -343,19 +347,17 @@ class Node;
     {
         NodeMeshMorph::Private& node = *InMorphNode->GetPrivate();
 
-        Ptr<ASTOpFixed> OpMorph = new ASTOpFixed();
-        OpMorph->op.type = OP_TYPE::ME_MORPH2;
+        Ptr<ASTOpMeshMorph> OpMorph = new ASTOpMeshMorph();
 
         // Factor
         if ( node.m_pFactor )
         {
-            OpMorph->SetChild( OpMorph->op.args.MeshMorph2.factor, Generate( node.m_pFactor.get() ) );
+            OpMorph->Factor = Generate( node.m_pFactor.get() );
         }
         else
         {
             // This argument is required
-            OpMorph->SetChild( OpMorph->op.args.MeshMorph2.factor, GenerateMissingScalarCode( "Morph factor", 0.5f,
-                                                                   node.m_errorContext ) );
+            OpMorph->Factor = GenerateMissingScalarCode( "Morph factor", 0.5f, node.m_errorContext );
         }
 
         // Base
@@ -365,54 +367,38 @@ class Node;
 			FMeshGenerationOptions BaseOptions = InOptions;
 			BaseOptions.bUniqueVertexIDs = true;
             GenerateMesh(BaseOptions,BaseResult, node.m_pBase );
-            OpMorph->SetChild( OpMorph->op.args.MeshMorph2.base, BaseResult.meshOp );
+            OpMorph->Base = BaseResult.meshOp;
         }
         else
         {
             // This argument is required
             m_pErrorLog->GetPrivate()->Add( "Mesh morph base node is not set.",
                                             ELMT_ERROR, node.m_errorContext );
-        }
+        }        
 
-        // TODO: Support more than MUTABLE_OP_MAX_MORPH2_TARGETS targets
-        if ( node.m_morphs.Num()>MUTABLE_OP_MAX_MORPH2_TARGETS )
+        for ( int32 t=0 ; t<node.m_morphs.Num(); ++t )
         {
-            m_pErrorLog->GetPrivate()->Add( 
-				FString::Printf(TEXT("A morph node has more targets [%d] than currently supported [%d]."), node.m_morphs.Num(), MUTABLE_OP_MAX_MORPH2_TARGETS), 
-				ELMT_WARNING, 
-				node.m_errorContext);
-        }
+            FMeshGenerationResult TargetResult;
+			FMeshGenerationOptions TargetOptions = InOptions;
+			TargetOptions.bUniqueVertexIDs = false;
+			TargetOptions.bLayouts = false;
+			TargetOptions.OverrideLayouts.Empty();
+			TargetOptions.ActiveTags.Empty();
+            GenerateMesh(TargetOptions, TargetResult, node.m_morphs[t]);
 
-        int count = 0;
-        for ( std::size_t t=0
-            ; t<node.m_morphs.Num() && t<MUTABLE_OP_MAX_MORPH2_TARGETS
-            ; ++t )
-        {
-            if ( auto pA = node.m_morphs[t].get() )
-            {
-                FMeshGenerationResult TargetResult;
-				FMeshGenerationOptions TargetOptions = InOptions;
-				TargetOptions.bUniqueVertexIDs = false;
-				TargetOptions.bLayouts = false;
-				TargetOptions.OverrideLayouts.Empty();
-				TargetOptions.ActiveTags.Empty();
-                GenerateMesh(TargetOptions, TargetResult, pA );
-
-                // TODO: Make sure that the target is a mesh with the morph format
-                Ptr<ASTOp> target = TargetResult.meshOp;
-
-                // If the vertex indices are supposed to be relative in the targets, adjust them
-                //if (node.m_vertexIndicesAreRelative)
-                //{
-                //    Ptr<ASTOpMeshRemapIndices> remapIndices = new ASTOpMeshRemapIndices;
-                //    remapIndices->source = target;
-                //    remapIndices->reference = baseResult.baseMeshOp;
-                //    target = remapIndices;
-                //}
-
-                OpMorph->SetChild( OpMorph->op.args.MeshMorph2.targets[count], target);
-                count++;
-            }
+            // TODO: Make sure that the target is a mesh with the morph format
+            Ptr<ASTOp> target = TargetResult.meshOp;
+
+            // If the vertex indices are supposed to be relative in the targets, adjust them
+            //if (node.m_vertexIndicesAreRelative)
+            //{
+            //    Ptr<ASTOpMeshRemapIndices> remapIndices = new ASTOpMeshRemapIndices;
+            //    remapIndices->source = target;
+            //    remapIndices->reference = baseResult.baseMeshOp;
+            //    target = remapIndices;
+            //}
+
+            OpMorph->AddTarget( target );
         }
  
         const bool bReshapeEnabled = node.m_reshapeSkeleton || node.m_reshapePhysicsVolumes;
@@ -459,8 +445,8 @@ class Node;
 		}
 
         OutResult.baseMeshOp = BaseResult.baseMeshOp;
-        OutResult.layouts = BaseResult.layouts;
-    } 
+		OutResult.GeneratedLayouts = BaseResult.GeneratedLayouts;
+}
 
 
     //---------------------------------------------------------------------------------------------
@@ -513,8 +499,8 @@ class Node;
 
         OutResult.meshOp = op;
         OutResult.baseMeshOp = BaseResult.baseMeshOp;
-        OutResult.layouts = BaseResult.layouts;
-    }
+		OutResult.GeneratedLayouts = BaseResult.GeneratedLayouts;
+	}
 
     //---------------------------------------------------------------------------------------------
     void CodeGenerator::GenerateMesh_Fragment(const FMeshGenerationOptions& InOptions, FMeshGenerationResult& OutResult,
@@ -539,10 +525,9 @@ class Node;
 
                 op->source = BaseResult.meshOp;
 
-                if (BaseResult.layouts.Num()>node.m_layoutOrGroup )
+                if (BaseResult.GeneratedLayouts.Num()>node.m_layoutOrGroup )
                 {
-                    LayoutPtrConst pSourceLayout = BaseResult.layouts[ node.m_layoutOrGroup ];
-                    const Layout* pLayout = m_addedLayouts[ pSourceLayout.get() ].get();
+                    const Layout* pLayout = BaseResult.GeneratedLayouts[node.m_layoutOrGroup].get();
                     op->layout = (uint16)node.m_layoutOrGroup;
 
                     for ( int32 i=0; i<node.m_blocks.Num(); ++i )
@@ -588,8 +573,8 @@ class Node;
         }
 
         OutResult.baseMeshOp = BaseResult.baseMeshOp;
-        OutResult.layouts = BaseResult.layouts;
-    }
+		OutResult.GeneratedLayouts = BaseResult.GeneratedLayouts;
+	}
 
 
     //---------------------------------------------------------------------------------------------
@@ -639,8 +624,8 @@ class Node;
                     op->SetChild( op->op.args.MeshInterpolate.base, TargetResult.meshOp );
 
                     OutResult.baseMeshOp = TargetResult.baseMeshOp;
-                    OutResult.layouts = TargetResult.layouts;
-                }
+					OutResult.GeneratedLayouts = TargetResult.GeneratedLayouts;
+				}
                 else
                 {
                     Ptr<ASTOpMeshDifference> dop = new ASTOpMeshDifference();
@@ -707,13 +692,13 @@ class Node;
         }
 
         // Options
-        for ( std::size_t t=0; t< node.m_options.Num(); ++t )
+        for ( int32 t=0; t< node.m_options.Num(); ++t )
         {
 			FMeshGenerationOptions TargetOptions = InOptions;
 
             if (t!=0)
             {
-				TargetOptions.OverrideLayouts = OutResult.layouts;
+				TargetOptions.OverrideLayouts = OutResult.GeneratedLayouts;
             }
 
             if ( node.m_options[t] )
@@ -770,7 +755,7 @@ class Node;
 
 				if (t != 0)
 				{
-					TargetOptions.OverrideLayouts = NewResult.layouts;
+					TargetOptions.OverrideLayouts = NewResult.GeneratedLayouts;
 				}
 
 				FMeshGenerationResult BranchResults;
@@ -844,7 +829,7 @@ class Node;
 
                 if (firstOptionProcessed)
                 {
-					VariationOptions.OverrideLayouts = currentResult.layouts;
+					VariationOptions.OverrideLayouts = currentResult.GeneratedLayouts;
                 }
          
                 FMeshGenerationResult BranchResults;
@@ -882,18 +867,107 @@ class Node;
         op->type = OP_TYPE::ME_CONSTANT;
 		OutResult.baseMeshOp = op;
 		OutResult.meshOp = op;
+		OutResult.GeneratedLayouts.Empty();
 
 		bool bIsOverridingLayouts = !InOptions.OverrideLayouts.IsEmpty();
 
         MeshPtr pMesh = node.m_pValue.get();
-		if (pMesh)
+		if (!pMesh)
 		{
-			// Clone the mesh
-			MeshPtr pCloned = pMesh->Clone();
-			pCloned->EnsureSurfaceData();
+			// This data is required
+			MeshPtr pTempMesh = new Mesh();
+			op->SetValue(pTempMesh, m_compilerOptions->m_optimisationOptions.m_useDiskCache);
+			m_constantMeshes.Add(pTempMesh);
+
+			// Log an error message
+			m_pErrorLog->GetPrivate()->Add("Constant mesh not set.", ELMT_WARNING, node.m_errorContext);
+
+			return;
+		}
+
+		// Find out if we can (or have to) reuse a mesh that we have already generated.
+		MeshPtrConst DuplicateOf;
+		for (int32 i = 0; i < m_constantMeshes.Num(); ++i)
+		{
+			MeshPtrConst Candidate = m_constantMeshes[i];
+			
+			bool bCompareLayouts = InOptions.bLayouts && !bIsOverridingLayouts;
+
+			if (Candidate->IsSimilar(*pMesh, bCompareLayouts))
+			{
+				// If it is similar, and we need unique vertex IDs, check that it also has them. This was skipped in the IsSimilar.
+				if (InOptions.bUniqueVertexIDs)
+				{
+					int32 FoundBuffer = -1;
+					int32 FoundChannel = -1;
+					Candidate->GetVertexBuffers().FindChannel(MBS_VERTEXINDEX, 0, &FoundBuffer, &FoundChannel);
+					bool bHasUniqueVertexIDs = FoundBuffer >= 0 && FoundChannel >= 0;
+					if (!bHasUniqueVertexIDs)
+					{
+						continue;
+					}
+				}
+
+				// If it is similar and we are overriding the layouts, we must compare the layouts of the candidate with the ones
+				// we are using to override.
+				if (InOptions.bLayouts && bIsOverridingLayouts)
+				{
+					if (Candidate->GetLayoutCount() != InOptions.OverrideLayouts.Num())
+					{
+						continue;
+					}
+
+					bool bLayoutsAreEqual = true;
+					for (int32 l = 0; l < Candidate->GetLayoutCount(); ++l)
+					{
+						bLayoutsAreEqual = (*Candidate->GetLayout(l) == *InOptions.OverrideLayouts[l]);
+						if ( !bLayoutsAreEqual )
+						{
+							break;
+						}
+					}
+
+					if (!bLayoutsAreEqual)
+					{
+						continue;
+					}
+				}
+
+				DuplicateOf = Candidate;
+				break;
+			}
+		}
+
+		if (DuplicateOf)
+		{
+			// Make sure the source layouts of the mesh are mapped to the layouts of the duplicated mesh.
+			if (InOptions.bLayouts)
+			{
+				if (bIsOverridingLayouts)
+				{
+					for (int32 l = 0; l < DuplicateOf->GetLayoutCount(); ++l)
+					{
+						const Layout* OverridingLayout = InOptions.OverrideLayouts[l].get();
+						OutResult.GeneratedLayouts.Add(OverridingLayout);
+					}
+				}
+				else
+				{
+					for (int32 l = 0; l < DuplicateOf->GetLayoutCount(); ++l)
+					{
+						const Layout* DuplicatedLayout = DuplicateOf->GetLayout(l);
+						OutResult.GeneratedLayouts.Add(DuplicatedLayout);
+					}
+				}
+			}
 
-			// We will redefine the layouts
-			OutResult.layouts.Empty();
+			op->SetValue(DuplicateOf, m_compilerOptions->m_optimisationOptions.m_useDiskCache);
+		}
+		else
+		{
+			// We need to clone the mesh in the node because we will modify it.
+			Ptr<Mesh> pCloned = pMesh->Clone();
+			pCloned->EnsureSurfaceData();
 
 			if (InOptions.bLayouts)
 			{
@@ -907,20 +981,11 @@ class Node;
 						const NodeLayoutBlocks* TypedNode = dynamic_cast<NodeLayoutBlocks*>(pLayoutNode.get());
 						if (TypedNode)
 						{
-							if (!pCloned)
-							{
-								m_pErrorLog->GetPrivate()->Add("Generating a layout node without a parent mesh.",
-									ELMT_ERROR, TypedNode->GetPrivate()->m_errorContext);
-								return;
-							}
-
-							LayoutPtr pSourceLayout = TypedNode->GetPrivate()->m_pLayout;
-							OutResult.layouts.Add(pSourceLayout);
-
-							PrepareForLayout(pSourceLayout,
-								pCloned, LayoutIndex,
-								TypedNode->GetPrivate()->m_errorContext);
+							Ptr<const Layout> SourceLayout = TypedNode->GetPrivate()->m_pLayout;
+							Ptr<const Layout> GeneratedLayout = AddLayout( SourceLayout );
+							PrepareForLayout(GeneratedLayout, pCloned, LayoutIndex, TypedNode->GetPrivate()->m_errorContext);
 
+							OutResult.GeneratedLayouts.Add(GeneratedLayout);
 						}
 					}
 				}
@@ -929,131 +994,56 @@ class Node;
 					// We need to apply the transform of the layouts used to override
 					for (int32 LayoutIndex = 0; LayoutIndex < InOptions.OverrideLayouts.Num(); ++LayoutIndex)
 					{
-						PrepareForLayout(InOptions.OverrideLayouts[LayoutIndex],
-							pCloned, LayoutIndex,
-							node.m_errorContext);
+						Ptr<const Layout> GeneratedLayout = InOptions.OverrideLayouts[LayoutIndex];
+						PrepareForLayout(GeneratedLayout, pCloned, LayoutIndex, node.m_errorContext);
 
+						OutResult.GeneratedLayouts.Add(GeneratedLayout);
 					}
-
-					OutResult.layouts = InOptions.OverrideLayouts;
 				}
 			}
 
-            // See if we already have a mesh identical to this one, except for the internal data like vertex indices.
-			// For now we cannot handle the case were we are overriding layouts and reusing meshes at the same time.
-            bool bIsDuplicated = false;
-            MeshPtrConst Candidate;
-            for ( int32 i=0; !bIsOverridingLayouts && !bIsDuplicated && i<m_constantMeshes.Num(); ++i)
-            {
-				Candidate = m_constantMeshes[i];
-
-                if (Candidate->IsSimilar(*pCloned,InOptions.bLayouts) )
-                {
-					// If it is similar, and we need unique vertex IDs, check that it also has them. This was skipped in the IsSimilar.
-					if (InOptions.bUniqueVertexIDs)
-					{
-						int32 FoundBuffer = -1;
-						int32 FoundChannel = -1;
-						Candidate->GetVertexBuffers().FindChannel(MBS_VERTEXINDEX, 0, &FoundBuffer, &FoundChannel);
-						bool bHasUniqueVertexIDs = FoundBuffer >= 0 && FoundChannel >= 0;
-						if (!bHasUniqueVertexIDs)
-						{
-							break;
-						}
-					}
-
-					bIsDuplicated = true;
-
-					if (InOptions.bLayouts)
-					{
-						// We want to reuse the layouts from the similar mesh.					
-						// Remap layouts that we have created in the previous step from the source mesh to the ones created for the mesh we will reuse instead.
-						for (int l = 0; l < Candidate->GetLayoutCount(); ++l)
-						{
-							check(Candidate->GetLayoutCount() == pCloned->GetLayoutCount());
-
-							const Layout* ClonedMeshLayout = pCloned->GetLayout(l);
-
-							Ptr<const Layout> DestLayoutKey = nullptr;
-							for (const auto& e : m_addedLayouts)
-							{
-								if (e.Value == ClonedMeshLayout)
-								{
-									DestLayoutKey = e.Key;
-									break;
-								}
-							}
-							check(DestLayoutKey);
-
-							const Layout* GeneratedMeshLayout = Candidate->GetLayout(l);
-							check(DestLayoutKey->m_blocks.IsEmpty() || DestLayoutKey->m_blocks[0].m_id == -1);
-							check(GeneratedMeshLayout->m_blocks.IsEmpty() || GeneratedMeshLayout->m_blocks[0].m_id != -1);
-							m_addedLayouts.Add(DestLayoutKey, GeneratedMeshLayout);
-						}
-					}
-                }
-            }
-
-            if (bIsDuplicated)
-            {
-                op->SetValue(Candidate, m_compilerOptions->m_optimisationOptions.m_useDiskCache );
-            }
-            else
-            {
-				if (InOptions.bUniqueVertexIDs)
+			if (InOptions.bUniqueVertexIDs)
+			{
+				// Enumerate the vertices uniquely unless they already have indices
+				int buf = -1;
+				int chan = -1;
+				pCloned->GetVertexBuffers().FindChannel(MBS_VERTEXINDEX, 0, &buf, &chan);
+				bool hasVertexIndices = (buf >= 0 && chan >= 0);
+				if (!hasVertexIndices)
 				{
-					// Enumerate the vertices uniquely unless they already have indices
-					int buf = -1;
-					int chan = -1;
-					pCloned->GetVertexBuffers().FindChannel(MBS_VERTEXINDEX, 0, &buf, &chan);
-					bool hasVertexIndices = (buf >= 0 && chan >= 0);
-					if (!hasVertexIndices)
+					int newBuffer = pCloned->GetVertexBuffers().GetBufferCount();
+					pCloned->GetVertexBuffers().SetBufferCount(newBuffer + 1);
+					MESH_BUFFER_SEMANTIC semantic = MBS_VERTEXINDEX;
+					int semanticIndex = 0;
+					MESH_BUFFER_FORMAT format = MBF_UINT32;
+					int components = 1;
+					int offset = 0;
+					pCloned->GetVertexBuffers().SetBuffer
+					(
+						newBuffer,
+						sizeof(uint32),
+						1,
+						&semantic, &semanticIndex,
+						&format, &components,
+						&offset
+					);
+					uint32* pIdData = (uint32*)pCloned->GetVertexBuffers().GetBufferData(newBuffer);
+					for (int i = 0; i < pMesh->GetVertexCount(); ++i)
 					{
-						int newBuffer = pCloned->GetVertexBuffers().GetBufferCount();
-						pCloned->GetVertexBuffers().SetBufferCount(newBuffer + 1);
-						MESH_BUFFER_SEMANTIC semantic = MBS_VERTEXINDEX;
-						int semanticIndex = 0;
-						MESH_BUFFER_FORMAT format = MBF_UINT32;
-						int components = 1;
-						int offset = 0;
-						pCloned->GetVertexBuffers().SetBuffer
-						(
-							newBuffer,
-							sizeof(uint32),
-							1,
-							&semantic, &semanticIndex,
-							&format, &components,
-							&offset
-						);
-						uint32* pIdData = (uint32*)pCloned->GetVertexBuffers().GetBufferData(newBuffer);
-						for (int i = 0; i < pMesh->GetVertexCount(); ++i)
-						{
-							check(m_freeVertexIndex < std::numeric_limits<uint32>::max());
+						check(m_freeVertexIndex < std::numeric_limits<uint32>::max());
 
-							(*pIdData++) = m_freeVertexIndex++;
-							check(m_freeVertexIndex < std::numeric_limits<uint32>::max());
-						}
+						(*pIdData++) = m_freeVertexIndex++;
+						check(m_freeVertexIndex < std::numeric_limits<uint32>::max());
 					}
 				}
-
-				// Add the constant data
-				m_constantMeshes.Add(pCloned);
-				op->SetValue(pCloned.get(), m_compilerOptions->m_optimisationOptions.m_useDiskCache);
 			}
-        }
-        else
-        {
-			OutResult.layouts.Empty();
-
-            // This data is required
-            MeshPtr pTempMesh = new Mesh();
-            op->SetValue( pTempMesh, m_compilerOptions->m_optimisationOptions.m_useDiskCache );
-            m_constantMeshes.Add(pTempMesh);
 
-            // Log an error message
-            m_pErrorLog->GetPrivate()->Add( "Constant mesh not set.", ELMT_WARNING, node.m_errorContext );
-        }
+			// Add the constant data
+			m_constantMeshes.Add(pCloned);
+			op->SetValue(pCloned.get(), m_compilerOptions->m_optimisationOptions.m_useDiskCache);
+		}
 
+ 
 		// Apply the modifier for the pre-normal operations stage.
 		BOTTOM_UP_STATE temp = m_currentBottomUpState;
 
@@ -1111,8 +1101,8 @@ class Node;
 
             OutResult.meshOp = op;
             OutResult.baseMeshOp = baseResult.baseMeshOp;
-            OutResult.layouts = baseResult.layouts;
-        }
+			OutResult.GeneratedLayouts = baseResult.GeneratedLayouts;
+		}
         else
         {
             // Put something there
diff --git a/Engine/Plugins/Experimental/Mutable/Source/MutableTools/Private/MuT/CodeOptimiser_Parameters.cpp b/Engine/Plugins/Experimental/Mutable/Source/MutableTools/Private/MuT/CodeOptimiser_Parameters.cpp
--- a/Engine/Plugins/Experimental/Mutable/Source/MutableTools/Private/MuT/CodeOptimiser_Parameters.cpp
+++ b/Engine/Plugins/Experimental/Mutable/Source/MutableTools/Private/MuT/CodeOptimiser_Parameters.cpp
@@ -20,6 +20,7 @@
 #include "MuT/ASTOpImageCompose.h"
 #include "MuT/ASTOpImageMipmap.h"
 #include "MuT/ASTOpImagePixelFormat.h"
+#include "MuT/ASTOpMeshMorph.h"
 #include "MuT/ASTOpInstanceAdd.h"
 #include "MuT/ASTOpParameter.h"
 #include "MuT/CodeOptimiser.h"
@@ -1694,9 +1695,9 @@ class TaskManager;
                 case OP_TYPE::ME_MORPH2:
                 {
                     // Manually choose how to recurse this op
-                    auto pTyped = dynamic_cast<const ASTOpFixed*>( at.get() );
+					const ASTOpMeshMorph* pTyped = dynamic_cast<const ASTOpMeshMorph*>( at.get() );
 
-                    if ( auto& base = pTyped->children[pTyped->op.args.MeshMorph2.base] )
+                    if ( auto& base = pTyped->Base )
                     {
 						pending.Add({ base.m_child, state });
                     }
@@ -1704,14 +1705,14 @@ class TaskManager;
                     // Mesh morphs don't modify the layouts, so we can ignore the factor and morphs
                     if (!state)
                     {
-                        if ( auto& factor = pTyped->children[pTyped->op.args.MeshMorph2.factor] )
+                        if ( auto& factor = pTyped->Factor )
                         {
 							pending.Add({ factor.m_child, state });
                         }
 
-                        for (int t=0;t<MUTABLE_OP_MAX_MORPH2_TARGETS;++t)
+                        for (int32 t=0;t<pTyped->Targets.Num(); ++t)
                         {
-                            if ( auto& target = pTyped->children[pTyped->op.args.MeshMorph2.targets[t]] )
+                            if ( auto& target = pTyped->Targets[t] )
                             {
 								pending.Add({ target.m_child, state });
                             }
diff --git a/Engine/Plugins/Experimental/Mutable/Source/MutableTools/Private/MuT/CodeOptimiser_Semantic.cpp b/Engine/Plugins/Experimental/Mutable/Source/MutableTools/Private/MuT/CodeOptimiser_Semantic.cpp
--- a/Engine/Plugins/Experimental/Mutable/Source/MutableTools/Private/MuT/CodeOptimiser_Semantic.cpp
+++ b/Engine/Plugins/Experimental/Mutable/Source/MutableTools/Private/MuT/CodeOptimiser_Semantic.cpp
@@ -19,6 +19,7 @@
 #include "MuT/ASTOpImageMultiLayer.h"
 #include "MuT/ASTOpImagePatch.h"
 #include "MuT/ASTOpImagePixelFormat.h"
+#include "MuT/ASTOpMeshMorph.h"
 #include "MuT/ASTOpSwitch.h"
 #include "MuT/CodeOptimiser.h"
 #include "MuT/Table.h"
@@ -1402,7 +1403,7 @@ namespace mu
 
                 if (nop->def)
                 {
-                    auto defOp = mu::Clone<ASTOpFixed>(this);
+					Ptr<ASTOpFixed> defOp = mu::Clone<ASTOpFixed>(this);
                     defOp->SetChild( defOp->op.args.LayoutRemoveBlocks.mesh, nop->def );
                     nop->def = defOp;
                 }
@@ -1411,7 +1412,7 @@ namespace mu
                 {
                     if ( o.branch )
                     {
-                        auto bOp = mu::Clone<ASTOpFixed>(this);
+						Ptr<ASTOpFixed> bOp = mu::Clone<ASTOpFixed>(this);
                         bOp->SetChild( bOp->op.args.LayoutRemoveBlocks.mesh, o.branch );
                         o.branch = bOp;
                     }
@@ -1424,9 +1425,9 @@ namespace mu
             case OP_TYPE::ME_MORPH2:
             {
                 // Move the remove-blocks down the base
-                auto typedMeshAt = dynamic_cast<const ASTOpFixed*>(meshAt.get());
-                auto nop = mu::Clone<ASTOpFixed>(this);
-                nop->SetChild( nop->op.args.LayoutRemoveBlocks.mesh, typedMeshAt->children[ typedMeshAt->op.args.MeshMorph2.base] );
+				const ASTOpMeshMorph* typedMeshAt = dynamic_cast<const ASTOpMeshMorph*>(meshAt.get());
+                Ptr<ASTOpFixed> nop = mu::Clone<ASTOpFixed>(this);
+                nop->SetChild( nop->op.args.LayoutRemoveBlocks.mesh, typedMeshAt->Base );
                 at = nop;
                 break;
             }
@@ -1435,7 +1436,7 @@ namespace mu
             {
                 // Assume that the UVs will not be interpolated across blocks!
                 auto typedMeshAt = dynamic_cast<const ASTOpFixed*>(meshAt.get());
-                auto nop = mu::Clone<ASTOpFixed>(this);
+				Ptr<ASTOpFixed> nop = mu::Clone<ASTOpFixed>(this);
                 nop->SetChild( nop->op.args.LayoutRemoveBlocks.mesh, typedMeshAt->children[ typedMeshAt->op.args.MeshInterpolate.base] );
                 at = nop;
                 break;
@@ -1493,10 +1494,9 @@ namespace mu
 
             case OP_TYPE::ME_MORPH2:
             {
-                auto typedSource = dynamic_cast<const ASTOpFixed*>(sourceAt.get());
+				const ASTOpMeshMorph* typedSource = dynamic_cast<const ASTOpMeshMorph*>(sourceAt.get());
 				Ptr<ASTOpFixed> rasterOp = mu::Clone<ASTOpFixed>(this);
-                rasterOp->SetChild( rasterOp->op.args.ImageRasterMesh.mesh,
-                        typedSource->children[typedSource->op.args.MeshMorph2.base] );
+                rasterOp->SetChild( rasterOp->op.args.ImageRasterMesh.mesh, typedSource->Base );
                 at = rasterOp;
                 break;
             }
diff --git a/Engine/Plugins/Experimental/Mutable/Source/MutableTools/Private/MuT/DataPacker.cpp b/Engine/Plugins/Experimental/Mutable/Source/MutableTools/Private/MuT/DataPacker.cpp
--- a/Engine/Plugins/Experimental/Mutable/Source/MutableTools/Private/MuT/DataPacker.cpp
+++ b/Engine/Plugins/Experimental/Mutable/Source/MutableTools/Private/MuT/DataPacker.cpp
@@ -23,6 +23,7 @@
 #include "MuT/ASTOpMeshExtractLayoutBlocks.h"
 #include "MuT/ASTOpMeshRemoveMask.h"
 #include "MuT/ASTOpMeshDifference.h"
+#include "MuT/ASTOpMeshMorph.h"
 
 #include <cstdint>
 #include <memory>
@@ -578,16 +579,16 @@ namespace mu
             {
                 recurse = false;
 
-				const ASTOpFixed* op = dynamic_cast<const ASTOpFixed*>(node.get());
+				const ASTOpMeshMorph* op = dynamic_cast<const ASTOpMeshMorph*>(node.get());
 
                 uint64_t newState = currentSemantics;
                 newState |= (UINT64_C(1)<<MBS_VERTEXINDEX);
-                RecurseWithState( op->children[op->op.args.MeshMorph2.base].child(), newState );
-                for ( int o=0; o<MUTABLE_OP_MAX_MORPH2_TARGETS; ++o )
+                RecurseWithState( op->Base.child(), newState );
+                for ( int32 o=0; o<op->Targets.Num(); ++o)
                 {
-                    if (op->children[op->op.args.MeshMorph2.targets[o]])
+                    if (op->Targets[o])
                     {
-                        RecurseWithState( op->children[op->op.args.MeshMorph2.targets[o]].child(), newState );
+                        RecurseWithState( op->Targets[o].child(), newState );
                     }
                 }
                 break;
