[mutable] New methods to select which bones are going to be deformed.
improvement: now table nodes use the row name of the data table to get the skeletal mesh instead of the row index.

#rnx

diff --git a/Engine/Plugins/Experimental/Mutable/Source/CustomizableObject/Public/MuCO/CustomizableObjectCustomVersion.h b/Engine/Plugins/Experimental/Mutable/Source/CustomizableObject/Public/MuCO/CustomizableObjectCustomVersion.h
--- a/Engine/Plugins/Experimental/Mutable/Source/CustomizableObject/Public/MuCO/CustomizableObjectCustomVersion.h
+++ b/Engine/Plugins/Experimental/Mutable/Source/CustomizableObject/Public/MuCO/CustomizableObjectCustomVersion.h
@@ -59,6 +59,8 @@ struct CUSTOMIZABLEOBJECT_API FCustomizableObjectCustomVersion
 		DescriptorBuildParameterDecorations,
 
 		DescriptorMultilayerProjectors,
+
+		DeformSkeletonOptionsAdded,
 		
 		// -----<new versions can be added above this line>-------------------------------------------------
 		VersionPlusOne,
diff --git a/Engine/Plugins/Experimental/Mutable/Source/CustomizableObjectEditor/Classes/MuCOE/Nodes/CustomizableObjectNodeMeshMorph.h b/Engine/Plugins/Experimental/Mutable/Source/CustomizableObjectEditor/Classes/MuCOE/Nodes/CustomizableObjectNodeMeshMorph.h
--- a/Engine/Plugins/Experimental/Mutable/Source/CustomizableObjectEditor/Classes/MuCOE/Nodes/CustomizableObjectNodeMeshMorph.h
+++ b/Engine/Plugins/Experimental/Mutable/Source/CustomizableObjectEditor/Classes/MuCOE/Nodes/CustomizableObjectNodeMeshMorph.h
@@ -17,6 +17,7 @@ class UCustomizableObjectNodeRemapPins;
 class UEdGraphPin;
 class UObject;
 struct FMeshReshapeBoneReference;
+enum class EBoneDeformSelectionMethod : uint8;
 
 UCLASS()
 class CUSTOMIZABLEOBJECTEDITOR_API UCustomizableObjectNodeMeshMorph : public UCustomizableObjectNode
@@ -51,6 +52,8 @@ public:
 		return FindPin(TEXT("Factor"));
 	}
 
+	void Serialize(FArchive& Ar) override;
+
 	UPROPERTY(Category=CustomizableObject, EditAnywhere)
 	FString MorphTargetName;
 
@@ -61,13 +64,17 @@ public:
 	/** Enable the deformation of physics volumes of the base mesh */
     UPROPERTY(EditAnywhere, Category = CustomizableObject)
     bool bReshapePhysicsVolumes = false;
+
+	/** Bone Reshape Selection Method */
+	UPROPERTY(EditAnywhere, Category = ReshapeBones, meta = (EditCondition = "bReshapeSkeleton"))
+	EBoneDeformSelectionMethod SelectionMethod;
 	
 	/** Enables the deformation of all bones of the skeleton */
-	UPROPERTY(EditAnywhere, Category = ReshapeBones, meta=(EditCondition="bReshapeSkeleton"))
-	bool bDeformAllBones = false;
+	UPROPERTY()
+	bool bDeformAllBones_DEPRECATED = false;
 
-	/** Array with bones that will be deformed */
-	UPROPERTY(EditAnywhere, Category = ReshapeBones, meta=(EditCondition="bReshapeSkeleton && !bDeformAllBones"))
+	/** Array with selected bones that will be deformed */
+	UPROPERTY(EditAnywhere, Category = ReshapeBones, meta = (EditCondition = "bReshapeSkeleton && (SelectionMethod == EBoneDeformSelectionMethod::ONLY_SELECTED || SelectionMethod == EBoneDeformSelectionMethod::ALL_BUT_SELECTED)"))
 	TArray<FMeshReshapeBoneReference> BonesToDeform;
 
 	UPROPERTY(EditAnywhere, Category = ReshapePhysics, meta = (EditCondition="bReshapePhysicsVolumes"))
diff --git a/Engine/Plugins/Experimental/Mutable/Source/CustomizableObjectEditor/Classes/MuCOE/Nodes/CustomizableObjectNodeMeshReshape.h b/Engine/Plugins/Experimental/Mutable/Source/CustomizableObjectEditor/Classes/MuCOE/Nodes/CustomizableObjectNodeMeshReshape.h
--- a/Engine/Plugins/Experimental/Mutable/Source/CustomizableObjectEditor/Classes/MuCOE/Nodes/CustomizableObjectNodeMeshReshape.h
+++ b/Engine/Plugins/Experimental/Mutable/Source/CustomizableObjectEditor/Classes/MuCOE/Nodes/CustomizableObjectNodeMeshReshape.h
@@ -17,6 +17,7 @@ class UCustomizableObjectNodeRemapPins;
 class UEdGraphPin;
 class UObject;
 struct FMeshReshapeBoneReference;
+enum class EBoneDeformSelectionMethod : uint8;
 
 UCLASS()
 class CUSTOMIZABLEOBJECTEDITOR_API UCustomizableObjectNodeMeshReshape : public UCustomizableObjectNode
@@ -50,6 +51,8 @@ public:
 		return FindPin(TEXT("Target Shape"), EGPD_Input);
 	}
 
+	void Serialize(FArchive& Ar) override;
+
 	/** Enable the deformation of the skeleton of the base mesh. */
 	UPROPERTY(EditAnywhere, Category = CustomizableObject)
 	bool bReshapeSkeleton = false;	
@@ -62,12 +65,16 @@ public:
 	UPROPERTY(EditAnywhere, Category = CustomizableObject)
 	bool bEnableRigidParts = false;
 
+	/** Bone Reshape Selection Method */
+	UPROPERTY(EditAnywhere, Category = ReshapeBones, meta = (EditCondition = "bReshapeSkeleton"))
+	EBoneDeformSelectionMethod SelectionMethod;
+
 	/** Enables the deformation of all bones of the skeleton */
-	UPROPERTY(EditAnywhere, Category = ReshapeBones, meta=(EditCondition="bReshapeSkeleton"))
-	bool bDeformAllBones = false;
+	UPROPERTY()
+	bool bDeformAllBones_DEPRECATED = false;
 
-	/** Array with bones that will be deformed */
-	UPROPERTY(EditAnywhere, Category = ReshapeBones, meta=(EditCondition="bReshapeSkeleton && !bDeformAllBones"))
+	/** Array with selected bones that will be deformed */
+	UPROPERTY(EditAnywhere, Category = ReshapeBones, meta = (EditCondition = "bReshapeSkeleton && (SelectionMethod == EBoneDeformSelectionMethod::ONLY_SELECTED || SelectionMethod == EBoneDeformSelectionMethod::ALL_BUT_SELECTED)"))
 	TArray<FMeshReshapeBoneReference> BonesToDeform;
 
 	/** Enables the deformation of all physic bodies*/
@@ -75,7 +82,7 @@ public:
 	bool bDeformAllPhysicsBodies = false;
 
 	/** Array with bones with physics bodies that will be deformed */
-	UPROPERTY(EditAnywhere, Category = ReshapePhysics, meta=(EditCondition="bReshapePhysicsVolumes && !bDeformAllPhysicsBodies"))
+	UPROPERTY(EditAnywhere, Category = ReshapePhysics, meta = (EditCondition = "bReshapePhysicsVolumes && !bDeformAllPhysicsBodies"))
 	TArray<FMeshReshapeBoneReference> PhysicsBodiesToDeform;
 };
 
diff --git a/Engine/Plugins/Experimental/Mutable/Source/CustomizableObjectEditor/Classes/MuCOE/Nodes/CustomizableObjectNodeMeshReshapeCommon.h b/Engine/Plugins/Experimental/Mutable/Source/CustomizableObjectEditor/Classes/MuCOE/Nodes/CustomizableObjectNodeMeshReshapeCommon.h
--- a/Engine/Plugins/Experimental/Mutable/Source/CustomizableObjectEditor/Classes/MuCOE/Nodes/CustomizableObjectNodeMeshReshapeCommon.h
+++ b/Engine/Plugins/Experimental/Mutable/Source/CustomizableObjectEditor/Classes/MuCOE/Nodes/CustomizableObjectNodeMeshReshapeCommon.h
@@ -6,6 +6,23 @@
 
 #include "CustomizableObjectNodeMeshReshapeCommon.generated.h"
 
+UENUM()
+enum class EBoneDeformSelectionMethod : uint8
+{
+	// Only selected bones will be deform
+	ONLY_SELECTED = 0 UMETA(DisplayName = "Only Selected"),
+
+	// All bones will be deform except the selected ones
+	ALL_BUT_SELECTED = 1 UMETA(DisplayName = "All But Selected"),
+
+	// Deform only the bones of the reference skeleton
+	DEFORM_REF_SKELETON = 2 UMETA(DisplayName = "Deform Only Ref.Seleton Bones"),
+
+	// Deform only the bones that are not in the reference skeleton
+	DEFORM_NONE_REF_SKELETON = 3 UMETA(DisplayName = "Deform All But not Ref.Skeleton Bones")
+};
+
+
 USTRUCT()
 struct FMeshReshapeBoneReference
 {
diff --git a/Engine/Plugins/Experimental/Mutable/Source/CustomizableObjectEditor/Classes/MuCOE/Nodes/CustomizableObjectNodeTable.h b/Engine/Plugins/Experimental/Mutable/Source/CustomizableObjectEditor/Classes/MuCOE/Nodes/CustomizableObjectNodeTable.h
--- a/Engine/Plugins/Experimental/Mutable/Source/CustomizableObjectEditor/Classes/MuCOE/Nodes/CustomizableObjectNodeTable.h
+++ b/Engine/Plugins/Experimental/Mutable/Source/CustomizableObjectEditor/Classes/MuCOE/Nodes/CustomizableObjectNodeTable.h
@@ -234,7 +234,7 @@ public:
 	}
 
 	// We should do this in a template!
-	USkeletalMesh* GetSkeletalMeshAt(const UEdGraphPin* Pin, int32 RowIndex) const;
+	USkeletalMesh* GetSkeletalMeshAt(const UEdGraphPin* Pin, const FName& RowName) const;
 
 	// Return the name of the enabled rows in the data table.
 	// Returns the name if the row has a bool column set as false (true == disabled)
diff --git a/Engine/Plugins/Experimental/Mutable/Source/CustomizableObjectEditor/Private/MuCOE/GenerateMutableSource/GenerateMutableSourceMesh.cpp b/Engine/Plugins/Experimental/Mutable/Source/CustomizableObjectEditor/Private/MuCOE/GenerateMutableSource/GenerateMutableSourceMesh.cpp
--- a/Engine/Plugins/Experimental/Mutable/Source/CustomizableObjectEditor/Private/MuCOE/GenerateMutableSource/GenerateMutableSourceMesh.cpp
+++ b/Engine/Plugins/Experimental/Mutable/Source/CustomizableObjectEditor/Private/MuCOE/GenerateMutableSource/GenerateMutableSourceMesh.cpp
@@ -1681,7 +1681,7 @@ mu::MeshPtr GenerateMutableMesh(UObject * Mesh, int32 LOD, int32 MaterialIndex,
 }
 
 
-mu::MeshPtr BuildMorphedMutableMesh(const UEdGraphPin* BaseSourcePin, const FString& MorphTargetName, FMutableGraphGenerationContext & GenerationContext, int32 RowIndex)
+mu::MeshPtr BuildMorphedMutableMesh(const UEdGraphPin* BaseSourcePin, const FString& MorphTargetName, FMutableGraphGenerationContext & GenerationContext, const FName& RowName)
 {
 	check(BaseSourcePin);
 	SCOPED_PIN_DATA(GenerationContext, BaseSourcePin)
@@ -1710,7 +1710,7 @@ mu::MeshPtr BuildMorphedMutableMesh(const UEdGraphPin* BaseSourcePin, const FStr
 	else if (const UCustomizableObjectNodeTable* TypedNodeTable = Cast<UCustomizableObjectNodeTable>(Node))
 	{
 		TypedNodeTable->GetPinLODAndMaterial(BaseSourcePin, LODIndex, SectionIndex);
-		SkeletalMesh = TypedNodeTable->GetSkeletalMeshAt(BaseSourcePin, RowIndex);
+		SkeletalMesh = TypedNodeTable->GetSkeletalMeshAt(BaseSourcePin, RowName);
 	}
 
 	if (SkeletalMesh)
@@ -1837,6 +1837,7 @@ void GenerateMorphFactor(const UCustomizableObjectNode* Node, const UEdGraphPin&
 	}
 }
 
+
 TArray<USkeletalMesh*> GetSkeletalMeshesForReshapeSelection(
 		const UEdGraphNode* SkeletalMeshOrTableNode, const UEdGraphPin* SourceMeshPin)
 {
@@ -1858,11 +1859,9 @@ TArray<USkeletalMesh*> GetSkeletalMeshesForReshapeSelection(
 	{
 		if (TableNode->Table)
 		{
-			uint32 NumRows = TableNode->Table->GetRowNames().Num();
-
-			for (uint32 RowIndex = 0; RowIndex < NumRows; ++RowIndex)
+			for (const FName& RowName : TableNode->GetRowNames())
 			{
-				USkeletalMesh* SkeletalMesh = TableNode->GetSkeletalMeshAt(SourceMeshPin, RowIndex);
+				USkeletalMesh* SkeletalMesh = TableNode->GetSkeletalMeshAt(SourceMeshPin, RowName);
 
 				if (SkeletalMesh)
 				{
@@ -1879,46 +1878,122 @@ TArray<USkeletalMesh*> GetSkeletalMeshesForReshapeSelection(
 	return SkeletalMeshes;
 }
 
+
 bool GetAndValidateReshapeBonesToDeform(
 	TArray<FString>& OutBonesToDeform,
 	const TArray<FMeshReshapeBoneReference>& InBonesToDeform,
 	const TArray<USkeletalMesh*>& SkeletalMeshes,
 	const UCustomizableObjectNode* Node,
+	const EBoneDeformSelectionMethod SelectionMethod,
 	FMutableGraphGenerationContext& GenerationContext)
 {
 	bool bSetRefreshWarning = false;
 
-	int32 NumBonesToDeform = InBonesToDeform.Num();
-	for (int32 BoneIndex = 0; BoneIndex < NumBonesToDeform; ++BoneIndex)
+	if(SelectionMethod == EBoneDeformSelectionMethod::ONLY_SELECTED)
 	{
-		bool bMissingBone = true;
-		
-		const FName BoneName = InBonesToDeform[BoneIndex].BoneName;
+		int32 NumBonesToDeform = InBonesToDeform.Num();
+		for (int32 InBoneIndex = 0; InBoneIndex < NumBonesToDeform; ++InBoneIndex)
+		{
+			bool bMissingBone = true;
 
+			const FName BoneName = InBonesToDeform[InBoneIndex].BoneName;
+
+			for (const USkeletalMesh* SkeletalMesh : SkeletalMeshes)
+			{
+				int32 BoneIndex = SkeletalMesh->GetRefSkeleton().FindBoneIndex(BoneName);
+				if (BoneIndex != INDEX_NONE)
+				{
+					if (SkeletalMesh->GetRefSkeleton().GetParentIndex(BoneIndex) != INDEX_NONE)
+					{
+						OutBonesToDeform.AddUnique(BoneName.ToString());
+					}
+
+					bMissingBone = false;
+					break;
+				}
+			}
+
+			if (bMissingBone)
+			{
+				FString msg = FString::Printf(TEXT("Could not find the selected Bone to Deform [%s] in Skeleton"), *(BoneName.ToString()));
+
+				GenerationContext.Compiler->CompilerLog(FText::FromString(msg), Node);
+
+				bSetRefreshWarning = true;
+			}
+		}
+	}
+
+	else if (SelectionMethod == EBoneDeformSelectionMethod::ALL_BUT_SELECTED)
+	{
 		for (const USkeletalMesh* SkeletalMesh : SkeletalMeshes)
 		{
-			if (SkeletalMesh->GetRefSkeleton().FindBoneIndex(BoneName) != INDEX_NONE)
+			int32 NumBonesToDeform = SkeletalMesh->GetRefSkeleton().GetNum();
+
+			for (int32 BoneIndex = 0; BoneIndex < NumBonesToDeform; ++BoneIndex)
 			{
-				OutBonesToDeform.AddUnique(BoneName.ToString());
+				FName BoneName = SkeletalMesh->GetRefSkeleton().GetBoneName(BoneIndex);
+				bool bFound = false;
+				int32 InNumBonesToDeform = InBonesToDeform.Num();
 
-				bMissingBone = false;
-				break;
+				for (int32 InBoneIndex = 0; InBoneIndex < InNumBonesToDeform; ++InBoneIndex)
+				{
+					if (InBonesToDeform[InBoneIndex].BoneName == BoneName)
+					{
+						bFound = true;
+						break;
+					}
+				}
+
+				if (!bFound && SkeletalMesh->GetRefSkeleton().GetParentIndex(BoneIndex) != INDEX_NONE)
+				{
+					OutBonesToDeform.AddUnique(BoneName.ToString());
+				}
 			}
 		}
+	}
 
-		if (bMissingBone)
+	else if (SelectionMethod == EBoneDeformSelectionMethod::DEFORM_REF_SKELETON)
+	{
+		// Getting reference skeleton from the reference skeletal mesh of the current component
+		const FReferenceSkeleton RefSkeleton = GenerationContext.ComponentInfos[GenerationContext.CurrentMeshComponent].RefSkeletalMesh->GetRefSkeleton();
+		int32 NumBones = RefSkeleton.GetNum();
+
+		for (int32 BoneIndex = 0; BoneIndex < NumBones; ++BoneIndex)
+		{
+			if (RefSkeleton.GetParentIndex(BoneIndex) != INDEX_NONE)
+			{
+				OutBonesToDeform.AddUnique(RefSkeleton.GetBoneName(BoneIndex).ToString());
+			}
+		}
+	}
+
+	else if (SelectionMethod == EBoneDeformSelectionMethod::DEFORM_NONE_REF_SKELETON)
+	{
+		// Getting reference skeleton from the reference skeletal mesh of the current component
+		const FReferenceSkeleton RefSkeleton = GenerationContext.ComponentInfos[GenerationContext.CurrentMeshComponent].RefSkeletalMesh->GetRefSkeleton();
+
+		for (const USkeletalMesh* SkeletalMesh : SkeletalMeshes)
 		{
-			FString msg = FString::Printf(TEXT("Could not find the selected Bone to Deform [%s] in Skeleton"), *(BoneName.ToString()));
+			int32 NumBones = SkeletalMesh->GetRefSkeleton().GetNum();
 
-			GenerationContext.Compiler->CompilerLog(FText::FromString(msg), Node);
+			for (int32 BoneIndex = 0; BoneIndex < NumBones; ++BoneIndex)
+			{
+				FName BoneName = SkeletalMesh->GetRefSkeleton().GetBoneName(BoneIndex);
 
-			bSetRefreshWarning = true;
+				if (RefSkeleton.FindBoneIndex(BoneName) == INDEX_NONE 
+					&& SkeletalMesh->GetRefSkeleton().GetParentIndex(BoneIndex) != INDEX_NONE)
+				{
+					OutBonesToDeform.AddUnique(BoneName.ToString());
+				}
+			}
 		}
 	}
 
 	return bSetRefreshWarning;
 }
 
+
 bool GetAndValidateReshapePhysicsToDeform(
 	TArray<FString>& OutPhysiscsToDeform,
 	const TArray<FMeshReshapeBoneReference>& InPhysicsToDeform,
@@ -1981,6 +2056,7 @@ bool GetAndValidateReshapePhysicsToDeform(
 	return bSetRefreshWarning;
 }
 
+
 mu::NodeMeshPtr GenerateMorphMesh(const UEdGraphPin* Pin,
 	TArray<UCustomizableObjectNodeMeshMorph*> TypedNodeMorphs,
 	int32 MorphIndex,
@@ -2036,6 +2112,8 @@ mu::NodeMeshPtr GenerateMorphMesh(const UEdGraphPin* Pin,
 
 		for (int32 RowIndex = 0; RowIndex < NumRows; ++RowIndex)
 		{
+			const FName RowName = TypedNodeTable->GetRowNames()[RowIndex];
+
 			ColumnIndex = Table->FindColumn(TCHAR_TO_ANSI(*ColumnName));
 
 			if (ColumnIndex == -1)
@@ -2043,7 +2121,7 @@ mu::NodeMeshPtr GenerateMorphMesh(const UEdGraphPin* Pin,
 				ColumnIndex = Table->AddColumn(TCHAR_TO_ANSI(*ColumnName), mu::TABLE_COLUMN_TYPE::TCT_MESH);
 			}
 
-			mu::MeshPtr MorphedSourceTableMesh = BuildMorphedMutableMesh(Pin, TypedNodeMorphs[MorphIndex]->MorphTargetName, GenerationContext, RowIndex);
+			mu::MeshPtr MorphedSourceTableMesh = BuildMorphedMutableMesh(Pin, TypedNodeMorphs[MorphIndex]->MorphTargetName, GenerationContext, RowName);
 			Table->SetCell(ColumnIndex, RowIndex, MorphedSourceTableMesh.get());
 		}
 
@@ -2100,7 +2178,6 @@ mu::NodeMeshPtr GenerateMorphMesh(const UEdGraphPin* Pin,
 
 			UCustomizableObjectNodeMeshMorph* TypedMorphNode = TypedNodeMorphs[MorphIndex];
 
-			Result->SetDeformAllBones(TypedMorphNode->bDeformAllBones);
 			Result->SetDeformAllPhysics(TypedMorphNode->bDeformAllPhysicsBodies);
 			Result->SetReshapeSkeleton(TypedMorphNode->bReshapeSkeleton);
 			Result->SetReshapePhysicsVolumes(TypedMorphNode->bReshapePhysicsVolumes);	
@@ -2112,11 +2189,11 @@ mu::NodeMeshPtr GenerateMorphMesh(const UEdGraphPin* Pin,
 				TArray<USkeletalMesh*> SkeletalMeshesToDeform = GetSkeletalMeshesForReshapeSelection(SkeletalMeshNode, SourceMeshPin);
 				
 				bool bWarningFound = false;
-				if (TypedMorphNode->bReshapeSkeleton && !TypedMorphNode->bDeformAllBones)
+				if (TypedMorphNode->bReshapeSkeleton)
 				{
 					TArray<FString> BonesToDeform;
 					bWarningFound = GetAndValidateReshapeBonesToDeform(
-						BonesToDeform, TypedMorphNode->BonesToDeform, SkeletalMeshesToDeform, TypedMorphNode, GenerationContext);
+						BonesToDeform, TypedMorphNode->BonesToDeform, SkeletalMeshesToDeform, TypedMorphNode, TypedMorphNode->SelectionMethod, GenerationContext);
 					
 					for (const FString& BoneName : BonesToDeform)
 					{
@@ -2794,7 +2871,6 @@ mu::NodeMeshPtr GenerateMutableSourceMesh(const UEdGraphPin * Pin,
 			MeshNode->SetReshapeSkeleton(TypedNodeReshape->bReshapeSkeleton);
 			MeshNode->SetReshapePhysicsVolumes(TypedNodeReshape->bReshapePhysicsVolumes);
 			MeshNode->SetEnableRigidParts(TypedNodeReshape->bEnableRigidParts);
-			MeshNode->SetDeformAllBones(TypedNodeReshape->bDeformAllBones);
 			MeshNode->SetDeformAllPhysics(TypedNodeReshape->bDeformAllPhysicsBodies);
 				
 			const UEdGraphPin* ConnectedPin = FollowInputPin(*TypedNodeReshape->BaseMeshPin());
@@ -2804,11 +2880,11 @@ mu::NodeMeshPtr GenerateMutableSourceMesh(const UEdGraphPin * Pin,
 			TArray<USkeletalMesh*> SkeletalMeshesToDeform = GetSkeletalMeshesForReshapeSelection(SkeletalMeshNode, SourceMeshPin);
 
 			bool bWarningFound = false;
-			if (TypedNodeReshape->bReshapeSkeleton && !TypedNodeReshape->bDeformAllBones)
+			if (TypedNodeReshape->bReshapeSkeleton)
 			{
 				TArray<FString> BonesToDeform;
 				bWarningFound = GetAndValidateReshapeBonesToDeform(
-					BonesToDeform, TypedNodeReshape->BonesToDeform, SkeletalMeshesToDeform, TypedNodeReshape, GenerationContext);
+					BonesToDeform, TypedNodeReshape->BonesToDeform, SkeletalMeshesToDeform, TypedNodeReshape, TypedNodeReshape->SelectionMethod, GenerationContext);
 				
 				for (const FString& BoneName : BonesToDeform)
 				{
diff --git a/Engine/Plugins/Experimental/Mutable/Source/CustomizableObjectEditor/Private/MuCOE/GenerateMutableSource/GenerateMutableSourceMesh.h b/Engine/Plugins/Experimental/Mutable/Source/CustomizableObjectEditor/Private/MuCOE/GenerateMutableSource/GenerateMutableSourceMesh.h
--- a/Engine/Plugins/Experimental/Mutable/Source/CustomizableObjectEditor/Private/MuCOE/GenerateMutableSource/GenerateMutableSourceMesh.h
+++ b/Engine/Plugins/Experimental/Mutable/Source/CustomizableObjectEditor/Private/MuCOE/GenerateMutableSource/GenerateMutableSourceMesh.h
@@ -24,7 +24,7 @@ mu::MeshPtr ConvertStaticMeshToMutable(UStaticMesh* StaticMesh, int LOD, int Mat
 mu::MeshPtr GenerateMutableMesh(UObject * Mesh, int32 LOD, int32 MaterialIndex, FMutableGraphGenerationContext & GenerationContext, const UCustomizableObjectNode* CurrentNode);
 
 
-mu::MeshPtr BuildMorphedMutableMesh(const UEdGraphPin * BaseSourcePin, const FString& MorphTargetName, FMutableGraphGenerationContext & GenerationContext, int32 RowIndex = 0);
+mu::MeshPtr BuildMorphedMutableMesh(const UEdGraphPin * BaseSourcePin, const FString& MorphTargetName, FMutableGraphGenerationContext & GenerationContext, const FName& RowName = "");
 
 /** Convert a CustomizableObject Source Graph into a mutable source graph. */
 mu::NodeMeshPtr GenerateMutableSourceMesh(const UEdGraphPin* Pin, FMutableGraphGenerationContext& GenerationContext, FMutableGraphMeshGenerationData& MeshData);
diff --git a/Engine/Plugins/Experimental/Mutable/Source/CustomizableObjectEditor/Private/MuCOE/Nodes/CustomizableObjectNodeMeshMorph.cpp b/Engine/Plugins/Experimental/Mutable/Source/CustomizableObjectEditor/Private/MuCOE/Nodes/CustomizableObjectNodeMeshMorph.cpp
--- a/Engine/Plugins/Experimental/Mutable/Source/CustomizableObjectEditor/Private/MuCOE/Nodes/CustomizableObjectNodeMeshMorph.cpp
+++ b/Engine/Plugins/Experimental/Mutable/Source/CustomizableObjectEditor/Private/MuCOE/Nodes/CustomizableObjectNodeMeshMorph.cpp
@@ -143,5 +143,21 @@ FText UCustomizableObjectNodeMeshMorph::GetTooltipText() const
 }
 
 
+void UCustomizableObjectNodeMeshMorph::Serialize(FArchive& Ar)
+{
+	Super::Serialize(Ar);
+
+	Ar.UsingCustomVersion(FCustomizableObjectCustomVersion::GUID);
+
+	if (Ar.CustomVer(FCustomizableObjectCustomVersion::GUID) < FCustomizableObjectCustomVersion::DeformSkeletonOptionsAdded)
+	{
+		if (bDeformAllBones_DEPRECATED)
+		{
+			SelectionMethod = EBoneDeformSelectionMethod::ALL_BUT_SELECTED;
+			BonesToDeform.Empty();
+		}
+	}
+}
+
 
 #undef LOCTEXT_NAMESPACE
diff --git a/Engine/Plugins/Experimental/Mutable/Source/CustomizableObjectEditor/Private/MuCOE/Nodes/CustomizableObjectNodeMeshReshape.cpp b/Engine/Plugins/Experimental/Mutable/Source/CustomizableObjectEditor/Private/MuCOE/Nodes/CustomizableObjectNodeMeshReshape.cpp
--- a/Engine/Plugins/Experimental/Mutable/Source/CustomizableObjectEditor/Private/MuCOE/Nodes/CustomizableObjectNodeMeshReshape.cpp
+++ b/Engine/Plugins/Experimental/Mutable/Source/CustomizableObjectEditor/Private/MuCOE/Nodes/CustomizableObjectNodeMeshReshape.cpp
@@ -62,4 +62,20 @@ FString UCustomizableObjectNodeMeshReshape::GetRefreshMessage() const
 }
 
 
+void UCustomizableObjectNodeMeshReshape::Serialize(FArchive& Ar)
+{
+	Super::Serialize(Ar);
+
+	Ar.UsingCustomVersion(FCustomizableObjectCustomVersion::GUID);
+
+	if (Ar.CustomVer(FCustomizableObjectCustomVersion::GUID) < FCustomizableObjectCustomVersion::DeformSkeletonOptionsAdded)
+	{
+		if (bDeformAllBones_DEPRECATED)
+		{
+			SelectionMethod = EBoneDeformSelectionMethod::ALL_BUT_SELECTED;
+			BonesToDeform.Empty();
+		}
+	}
+}
+
 #undef LOCTEXT_NAMESPACE
diff --git a/Engine/Plugins/Experimental/Mutable/Source/CustomizableObjectEditor/Private/MuCOE/Nodes/CustomizableObjectNodeTable.cpp b/Engine/Plugins/Experimental/Mutable/Source/CustomizableObjectEditor/Private/MuCOE/Nodes/CustomizableObjectNodeTable.cpp
--- a/Engine/Plugins/Experimental/Mutable/Source/CustomizableObjectEditor/Private/MuCOE/Nodes/CustomizableObjectNodeTable.cpp
+++ b/Engine/Plugins/Experimental/Mutable/Source/CustomizableObjectEditor/Private/MuCOE/Nodes/CustomizableObjectNodeTable.cpp
@@ -848,9 +848,9 @@ bool UCustomizableObjectNodeTable::CheckPinUpdated(const FString& PinName, const
 }
 
 
-USkeletalMesh* UCustomizableObjectNodeTable::GetSkeletalMeshAt(const UEdGraphPin* Pin, int32 RowIndex) const
+USkeletalMesh* UCustomizableObjectNodeTable::GetSkeletalMeshAt(const UEdGraphPin* Pin, const FName& RowName) const
 {
-	if (!Table || !Table->GetRowStruct() || !Pin || !Table->GetRowNames().IsValidIndex(RowIndex))
+	if (!Table || !Table->GetRowStruct() || !Pin || !Table->GetRowNames().Contains(RowName))
 	{
 		return nullptr;
 	}
@@ -860,8 +860,6 @@ USkeletalMesh* UCustomizableObjectNodeTable::GetSkeletalMeshAt(const UEdGraphPin
 
 	FString ColumnName = GetColumnNameByPin(Pin);
 
-	FName RowName = Table->GetRowNames()[RowIndex];
-
 	for (TFieldIterator<FProperty> It(TableStruct); It; ++It)
 	{
 		FProperty* ColumnProperty = *It;
diff --git a/Engine/Plugins/Experimental/Mutable/Source/MutableTools/Private/MuT/CodeGenerator_Mesh.cpp b/Engine/Plugins/Experimental/Mutable/Source/MutableTools/Private/MuT/CodeGenerator_Mesh.cpp
--- a/Engine/Plugins/Experimental/Mutable/Source/MutableTools/Private/MuT/CodeGenerator_Mesh.cpp
+++ b/Engine/Plugins/Experimental/Mutable/Source/MutableTools/Private/MuT/CodeGenerator_Mesh.cpp
@@ -413,7 +413,7 @@ class Node;
 			// PhysicsBodies and the Skeleton.
             OpBind->m_reshapeVertices = false;
 		    OpBind->m_reshapeSkeleton = node.m_reshapeSkeleton;
-		    OpBind->m_deformAllBones = node.m_deformAllBones;
+		    OpBind->m_deformAllBones = false; //TODO: Remove deprecated var from node
 		    OpBind->m_bonesToDeform = node.m_bonesToDeform;
     	    OpBind->m_reshapePhysicsVolumes = node.m_reshapePhysicsVolumes; 
 			OpBind->m_physicsToDeform = node.m_physicsToDeform;
@@ -1401,7 +1401,7 @@ class Node;
 
 		opBind->m_reshapeSkeleton = node.m_reshapeSkeleton;
     	opBind->m_enableRigidParts = node.m_enableRigidParts;
-		opBind->m_deformAllBones = node.m_deformAllBones;
+		opBind->m_deformAllBones = false; //TODO: Remove deprecated var from node
 		opBind->m_bonesToDeform = node.m_bonesToDeform;
     	opBind->m_reshapePhysicsVolumes = node.m_reshapePhysicsVolumes;
 		opBind->m_deformAllPhysics = node.m_deformAllPhysics;
diff --git a/Engine/Plugins/Experimental/Mutable/Source/MutableTools/Private/MuT/NodeMeshMorph.cpp b/Engine/Plugins/Experimental/Mutable/Source/MutableTools/Private/MuT/NodeMeshMorph.cpp
--- a/Engine/Plugins/Experimental/Mutable/Source/MutableTools/Private/MuT/NodeMeshMorph.cpp
+++ b/Engine/Plugins/Experimental/Mutable/Source/MutableTools/Private/MuT/NodeMeshMorph.cpp
@@ -161,12 +161,6 @@ namespace mu
 		m_pD->m_reshapeSkeleton = bEnable;
 	}	
 
-	//---------------------------------------------------------------------------------------------
-	void NodeMeshMorph::SetDeformAllBones(bool bEnable)
-	{
-		m_pD->m_deformAllBones = bEnable;
-	}
-
 	//---------------------------------------------------------------------------------------------
 	void NodeMeshMorph::SetDeformAllPhysics(bool bEnable)
 	{
diff --git a/Engine/Plugins/Experimental/Mutable/Source/MutableTools/Private/MuT/NodeMeshMorphPrivate.h b/Engine/Plugins/Experimental/Mutable/Source/MutableTools/Private/MuT/NodeMeshMorphPrivate.h
--- a/Engine/Plugins/Experimental/Mutable/Source/MutableTools/Private/MuT/NodeMeshMorphPrivate.h
+++ b/Engine/Plugins/Experimental/Mutable/Source/MutableTools/Private/MuT/NodeMeshMorphPrivate.h
@@ -29,7 +29,7 @@ namespace mu
 		bool m_reshapeSkeleton = false;
 		bool m_reshapePhysicsVolumes = false;
 		
-		bool m_deformAllBones = false;
+		bool m_deformAllBones_DEPRECATED = false;
 		bool m_deformAllPhysics = false;
 		
 		TArray<string> m_bonesToDeform;
@@ -38,7 +38,7 @@ namespace mu
         //!
 		void Serialise( OutputArchive& arch ) const
 		{
-            uint32_t ver = 3;
+            uint32_t ver = 4;
 			arch << ver;
 
 			arch << m_pFactor;
@@ -50,7 +50,6 @@ namespace mu
 			arch << m_reshapePhysicsVolumes;
 			arch << m_bonesToDeform;
 
-			arch << m_deformAllBones;
 			arch << m_deformAllPhysics;
 			
 			arch << m_physicsToDeform;
@@ -62,7 +61,7 @@ namespace mu
 		{
             uint32_t ver;
 			arch >> ver;
-			check(ver <= 3);
+			check(ver <= 4);
 
 			arch >> m_pFactor;
 			arch >> m_pBase;
@@ -88,15 +87,19 @@ namespace mu
 				m_bonesToDeform.Empty();
 			}
 
+			if (ver == 3)
+			{
+				arch >> m_deformAllBones_DEPRECATED;
+			}
+
 			if (ver >= 3)
 			{
-				arch >> m_deformAllBones;
 				arch >> m_deformAllPhysics;
 				arch >> m_physicsToDeform;
 			}
 			else
 			{
-				m_deformAllBones = false;
+				m_deformAllBones_DEPRECATED = false;
 				m_deformAllPhysics = false;
 				m_physicsToDeform.Empty();
 			}
diff --git a/Engine/Plugins/Experimental/Mutable/Source/MutableTools/Private/MuT/NodeMeshReshape.cpp b/Engine/Plugins/Experimental/Mutable/Source/MutableTools/Private/MuT/NodeMeshReshape.cpp
--- a/Engine/Plugins/Experimental/Mutable/Source/MutableTools/Private/MuT/NodeMeshReshape.cpp
+++ b/Engine/Plugins/Experimental/Mutable/Source/MutableTools/Private/MuT/NodeMeshReshape.cpp
@@ -135,12 +135,6 @@ namespace mu
 		m_pD->m_enableRigidParts = bEnable;
 	}
 
-	//---------------------------------------------------------------------------------------------
-	void NodeMeshReshape::SetDeformAllBones(bool bEnable)
-	{
-		m_pD->m_deformAllBones = bEnable;
-	}
-
 	//---------------------------------------------------------------------------------------------
 	void NodeMeshReshape::SetDeformAllPhysics(bool bEnable)
 	{
diff --git a/Engine/Plugins/Experimental/Mutable/Source/MutableTools/Private/MuT/NodeMeshReshapePrivate.h b/Engine/Plugins/Experimental/Mutable/Source/MutableTools/Private/MuT/NodeMeshReshapePrivate.h
--- a/Engine/Plugins/Experimental/Mutable/Source/MutableTools/Private/MuT/NodeMeshReshapePrivate.h
+++ b/Engine/Plugins/Experimental/Mutable/Source/MutableTools/Private/MuT/NodeMeshReshapePrivate.h
@@ -26,7 +26,7 @@ namespace mu
 		Ptr<NodeMesh> m_pTargetShape;
 		bool m_reshapeSkeleton = false;
 		bool m_enableRigidParts = false;
-		bool m_deformAllBones = false;
+		bool m_deformAllBones_DEPRECATED = false;
 		bool m_deformAllPhysics = false;
 		bool m_reshapePhysicsVolumes = false;
 		
@@ -35,7 +35,7 @@ namespace mu
         //!
 		void Serialise( OutputArchive& arch ) const
 		{
-            uint32_t ver = 5;
+            uint32_t ver = 6;
 			arch << ver;
 
 			arch << m_pBaseMesh;
@@ -43,7 +43,6 @@ namespace mu
 			arch << m_pTargetShape;
 			arch << m_reshapeSkeleton;
 			arch << m_enableRigidParts;
-			arch << m_deformAllBones;
 			arch << m_bonesToDeform;
 			arch << m_reshapePhysicsVolumes;
 			arch << m_deformAllPhysics;
@@ -55,14 +54,19 @@ namespace mu
 		{
             uint32_t ver;
 			arch >> ver;
-			check(ver<=5);
+			check(ver<=6);
 
 			arch >> m_pBaseMesh;
 			arch >> m_pBaseShape;
 			arch >> m_pTargetShape;
 			arch >> m_reshapeSkeleton;
 			arch >> m_enableRigidParts;
-			arch >> m_deformAllBones;
+
+			if (ver <= 5)
+			{
+				arch >> m_deformAllBones_DEPRECATED;
+			}
+
 			arch >> m_bonesToDeform;
 			
 		  	if (ver >= 4)
diff --git a/Engine/Plugins/Experimental/Mutable/Source/MutableTools/Public/MuT/NodeMeshMorph.h b/Engine/Plugins/Experimental/Mutable/Source/MutableTools/Public/MuT/NodeMeshMorph.h
--- a/Engine/Plugins/Experimental/Mutable/Source/MutableTools/Public/MuT/NodeMeshMorph.h
+++ b/Engine/Plugins/Experimental/Mutable/Source/MutableTools/Public/MuT/NodeMeshMorph.h
@@ -80,9 +80,6 @@ namespace mu
 		/** Also deform the mesh skeleton. Disabled by default. */
 		void SetReshapeSkeleton(bool);	
 
-		/** Deform all the bones of the skeleton */
-		void SetDeformAllBones(bool);
-
 		/** Deform all the physiscs bodies */
 		void SetDeformAllPhysics(bool);
 
diff --git a/Engine/Plugins/Experimental/Mutable/Source/MutableTools/Public/MuT/NodeMeshReshape.h b/Engine/Plugins/Experimental/Mutable/Source/MutableTools/Public/MuT/NodeMeshReshape.h
--- a/Engine/Plugins/Experimental/Mutable/Source/MutableTools/Public/MuT/NodeMeshReshape.h
+++ b/Engine/Plugins/Experimental/Mutable/Source/MutableTools/Public/MuT/NodeMeshReshape.h
@@ -59,9 +59,6 @@ namespace mu
 		/** Search for clusters of rigid parts and not deform them. */
 		void SetEnableRigidParts(bool);
 
-		/** Deform all the bones of the skeleton */
-		void SetDeformAllBones(bool);
-
 		/** Deform all the physiscs bodies */
 		void SetDeformAllPhysics(bool);
 
