[mutable] Increased current supported version

#rnx

diff --git a/Engine/Plugins/Experimental/Mutable/Source/CustomizableObject/Classes/MuCO/CustomizableObject.h b/Engine/Plugins/Experimental/Mutable/Source/CustomizableObject/Classes/MuCO/CustomizableObject.h
--- a/Engine/Plugins/Experimental/Mutable/Source/CustomizableObject/Classes/MuCO/CustomizableObject.h
+++ b/Engine/Plugins/Experimental/Mutable/Source/CustomizableObject/Classes/MuCO/CustomizableObject.h
@@ -1132,7 +1132,7 @@ private:
 	// This is a manual version number for the binary blobs in this asset.
 	// Increasing it invalidates all the previously compiled models.
 	// Warning: If while merging code both versions have changed, take the highest+1.
-	static const int32 CurrentSupportedVersion = 349;
+	static const int32 CurrentSupportedVersion = 350;
 
 public:
 
