[mutable] Fix some textures being resized to zero.

#rnx

diff --git a/Engine/Plugins/Experimental/Mutable/Source/MutableTools/Private/MuT/ImageDescGenerator.cpp b/Engine/Plugins/Experimental/Mutable/Source/MutableTools/Private/MuT/ImageDescGenerator.cpp
--- a/Engine/Plugins/Experimental/Mutable/Source/MutableTools/Private/MuT/ImageDescGenerator.cpp
+++ b/Engine/Plugins/Experimental/Mutable/Source/MutableTools/Private/MuT/ImageDescGenerator.cpp
@@ -163,6 +163,9 @@ namespace mu
                 m_desc.m_size[0] = uint16( node.m_sizeX );
                 m_desc.m_size[1] = uint16( node.m_sizeY );
             }
+
+			m_desc.m_size[0] = FMath::Max(m_desc.m_size[0], (uint16)1);
+			m_desc.m_size[1] = FMath::Max(m_desc.m_size[1], (uint16)1);
         }
 
         return 0;
