#rnx

diff --git a/Engine/Plugins/Experimental/Mutable/Source/MutableTools/Private/MuT/CodeGenerator_Mesh.cpp b/Engine/Plugins/Experimental/Mutable/Source/MutableTools/Private/MuT/CodeGenerator_Mesh.cpp
--- a/Engine/Plugins/Experimental/Mutable/Source/MutableTools/Private/MuT/CodeGenerator_Mesh.cpp
+++ b/Engine/Plugins/Experimental/Mutable/Source/MutableTools/Private/MuT/CodeGenerator_Mesh.cpp
@@ -814,8 +814,10 @@ class Node;
 
             if ( tagIndex < 0 )
             {
+				const char* Aux = nullptr;
+				Aux = tag.c_str();
                 m_pErrorLog->GetPrivate()->Add( 
-					FString::Printf(TEXT("Unknown tag found in mesh variation [%s]."), tag.c_str()), 
+					FString::Printf(TEXT("Unknown tag found in mesh variation [%s]."), ANSI_TO_TCHAR(Aux)),
 					ELMT_WARNING, 
 					node.m_errorContext
 				);
