[mutable] Added hint in model data to find out parameters controlling the range size for multidimensional parameters.
More cleanup of internal math lib in favour of Unreal's.
Renamed several structures to match conding standard.
#rnx

diff --git a/Engine/Plugins/Experimental/Mutable/Source/CustomizableObject/Private/MuCO/CustomizableObjectInstanceDescriptor.cpp b/Engine/Plugins/Experimental/Mutable/Source/CustomizableObject/Private/MuCO/CustomizableObjectInstanceDescriptor.cpp
--- a/Engine/Plugins/Experimental/Mutable/Source/CustomizableObject/Private/MuCO/CustomizableObjectInstanceDescriptor.cpp
+++ b/Engine/Plugins/Experimental/Mutable/Source/CustomizableObject/Private/MuCO/CustomizableObjectInstanceDescriptor.cpp
@@ -788,10 +788,10 @@ mu::ParametersPtr FCustomizableObjectInstanceDescriptor::ReloadParametersFromObj
 						case ECustomizableObjectProjectorType::Wrapping:
 						{
 							MutableParameters->SetProjectorValue(paramIndex,
-								Value.Position[0], Value.Position[1], Value.Position[2],
-								Value.Direction[0], Value.Direction[1], Value.Direction[2],
-								Value.Up[0], Value.Up[1], Value.Up[2],
-								Value.Scale[0], Value.Scale[1], Value.Scale[2],
+								Value.Position,
+								Value.Direction,
+								Value.Up,
+								Value.Scale,
 								Value.Angle,
 								RangeIdxPtr);
 							break;
@@ -801,14 +801,14 @@ mu::ParametersPtr FCustomizableObjectInstanceDescriptor::ReloadParametersFromObj
 						{
 							// Apply strange swizzle for scales
 							// TODO: try to avoid this
-								float Radius = FMath::Abs(Value.Scale[0] / 2.0f);
-								float Height = Value.Scale[2];
+							float Radius = FMath::Abs(Value.Scale[0] / 2.0f);
+							float Height = Value.Scale[2];
 							// TODO: try to avoid this
 							MutableParameters->SetProjectorValue(paramIndex,
-									Value.Position[0], Value.Position[1], Value.Position[2],
-									-Value.Direction[0], -Value.Direction[1], -Value.Direction[2],
-									-Value.Up[0], -Value.Up[1], -Value.Up[2],
-									-Height, Radius, Radius,
+									Value.Position,
+									-Value.Direction,
+									-Value.Up,
+									FVector3f(-Height, Radius, Radius),
 									Value.Angle,
 									RangeIdxPtr);
 							break;
@@ -833,10 +833,10 @@ mu::ParametersPtr FCustomizableObjectInstanceDescriptor::ReloadParametersFromObj
 				mu::PROJECTOR_TYPE type;
 				MutableParameters->GetProjectorValue(paramIndex,
 					&type,
-					&param.Value.Position[0], &param.Value.Position[1], &param.Value.Position[2],
-					&param.Value.Direction[0], &param.Value.Direction[1], &param.Value.Direction[2],
-					&param.Value.Up[0], &param.Value.Up[1], &param.Value.Up[2],
-					&param.Value.Scale[0], &param.Value.Scale[1], &param.Value.Scale[2],
+					&param.Value.Position,
+					&param.Value.Direction,
+					&param.Value.Up,
+					&param.Value.Scale,
 					&param.Value.Angle );
 				switch (type)
 				{
@@ -878,13 +878,13 @@ mu::ParametersPtr FCustomizableObjectInstanceDescriptor::ReloadParametersFromObj
 
 						mu::PROJECTOR_TYPE type;
 						MutableParameters->GetProjectorValue(paramIndex,
-						&type,
-						&param.RangeValues[ValueIndex].Position[0], &param.RangeValues[ValueIndex].Position[1], &param.RangeValues[ValueIndex].Position[2],
-						&param.RangeValues[ValueIndex].Direction[0], &param.RangeValues[ValueIndex].Direction[1], &param.RangeValues[ValueIndex].Direction[2],
-						&param.RangeValues[ValueIndex].Up[0], &param.RangeValues[ValueIndex].Up[1], &param.RangeValues[ValueIndex].Up[2],
-						&param.RangeValues[ValueIndex].Scale[0], &param.RangeValues[ValueIndex].Scale[1], &param.RangeValues[ValueIndex].Scale[2],
-						&param.RangeValues[ValueIndex].Angle,
-						RangeValueIdxPtr);
+							&type,
+							&param.RangeValues[ValueIndex].Position,
+							&param.RangeValues[ValueIndex].Direction,
+							&param.RangeValues[ValueIndex].Up,
+							&param.RangeValues[ValueIndex].Scale,
+							&param.RangeValues[ValueIndex].Angle,
+							RangeValueIdxPtr);
 
 						switch (type)
 						{
@@ -1853,10 +1853,10 @@ FCustomizableObjectProjector FCustomizableObjectInstanceDescriptor::GetProjector
 	
 	MutableParameters->GetProjectorValue(ParamIndex,
 		&type,
-		&Projector.Position[0], &Projector.Position[1], &Projector.Position[2],
-		&Projector.Direction[0], &Projector.Direction[1], &Projector.Direction[2],
-		&Projector.Up[0], &Projector.Up[1], &Projector.Up[2],
-		&Projector.Scale[0], &Projector.Scale[1], &Projector.Scale[2],
+		&Projector.Position,
+		&Projector.Direction,
+		&Projector.Up,
+		&Projector.Scale,
 		&Projector.Angle,
 		nullptr);
 
diff --git a/Engine/Plugins/Experimental/Mutable/Source/CustomizableObjectEditor/Private/MuCOE/SMutableCodeViewer.cpp b/Engine/Plugins/Experimental/Mutable/Source/CustomizableObjectEditor/Private/MuCOE/SMutableCodeViewer.cpp
--- a/Engine/Plugins/Experimental/Mutable/Source/CustomizableObjectEditor/Private/MuCOE/SMutableCodeViewer.cpp
+++ b/Engine/Plugins/Experimental/Mutable/Source/CustomizableObjectEditor/Private/MuCOE/SMutableCodeViewer.cpp
@@ -87,8 +87,8 @@ class FReferenceCollector;
 class FUICommandList;
 class SWidget;
 namespace mu { struct Curve; }
-namespace mu { struct PROJECTOR; }
-namespace mu { struct SHAPE; }
+namespace mu { struct FProjector; }
+namespace mu { struct FShape; }
 struct FGeometry;
 struct FSlateBrush;
 
@@ -126,7 +126,7 @@ public:
 		// Primary column showing the name of the operation and tye type
 		if (ColumnName == MutableCodeTreeViewColumns::OperationsColumnID)
 		{
-			const mu::PROGRAM& Program = RowItem->MutableModel->GetPrivate()->m_program;
+			const mu::FProgram& Program = RowItem->MutableModel->GetPrivate()->m_program;
 			const mu::OP_TYPE Type = Program.GetOpType(RowItem->MutableOperation);
 			const TCHAR* OpName = mu::s_opNames[int32(Type)];
 
@@ -352,7 +352,7 @@ void SMutableCodeViewer::Construct(const FArguments& InArgs, const mu::ModelPtr&
 	const int32 StateCount = ModelPrivate->m_program.m_states.Num();
 	for ( int32 StateIndex=0; StateIndex<StateCount; ++StateIndex )
 	{
-		const mu::PROGRAM::STATE& State = ModelPrivate->m_program.m_states[StateIndex];
+		const mu::FProgram::FState& State = ModelPrivate->m_program.m_states[StateIndex];
 		FString Caption = FString::Printf( TEXT("state [%s]"), ANSI_TO_TCHAR(State.m_name.c_str()) );
 		RootNodes.Add(MakeShareable(new FMutableCodeTreeElement(MutableModel, State.m_root, Caption)));
 	}
@@ -738,7 +738,7 @@ void SMutableCodeViewer::CacheAddressesOfOperationsOfType(mu::OP_TYPE TargetedOp
 	TSet<mu::OP::ADDRESS> OperationsWithTargetedType;
 
 	// Main update procedure run for the targeted state and the targeted parameter values
-	const mu::PROGRAM& Program = MutableModel->GetPrivate()->m_program;
+	const mu::FProgram& Program = MutableModel->GetPrivate()->m_program;
 	for	(const mu::OP::ADDRESS& RootNodeAddress : RootNodeAddresses)
 	{
 		GetOperationsOfType(TargetedOperationType,RootNodeAddress,Program,OperationsWithTargetedType,ProcessedAddresses);
@@ -755,7 +755,7 @@ void SMutableCodeViewer::CacheAddressesOfOperationsOfType(mu::OP_TYPE TargetedOp
 
 }
 
-void SMutableCodeViewer::GetOperationsOfType ( const mu::OP_TYPE& TargetOperationType, const mu::OP::ADDRESS& InParentAddress,const mu::PROGRAM& InProgram,
+void SMutableCodeViewer::GetOperationsOfType ( const mu::OP_TYPE& TargetOperationType, const mu::OP::ADDRESS& InParentAddress,const mu::FProgram& InProgram,
 	TSet<mu::OP::ADDRESS>& OutAddressesOfType,
 	TSet<mu::OP::ADDRESS>& AlreadyProcessedAddresses)
 {
@@ -790,7 +790,7 @@ void SMutableCodeViewer::CacheOperationTypesPresentOnModel()
 	// Set of operation types to be filed 
 	TSet<mu::OP_TYPE> OperationTypes;
 	
-	const mu::PROGRAM& Program = MutableModel->GetPrivate()->m_program;
+	const mu::FProgram& Program = MutableModel->GetPrivate()->m_program;
 	TSet<mu::OP::ADDRESS> AlreadyVisitedAddresses;
 	for	(const mu::OP::ADDRESS& RootOperationAddress : RootNodeAddresses)
 	{
@@ -833,7 +833,7 @@ void SMutableCodeViewer::CacheOperationTypesPresentOnModel()
 
 
 void SMutableCodeViewer::GetOperationTypesPresentOnModel(
-	const mu::OP::ADDRESS& InParentAddress, const mu::PROGRAM& InProgram, TSet<mu::OP_TYPE>& OutLocatedOperations,
+	const mu::OP::ADDRESS& InParentAddress, const mu::FProgram& InProgram, TSet<mu::OP_TYPE>& OutLocatedOperations,
 	TSet<mu::OP::ADDRESS>& AlreadyProcessedAddresses)
 {
 	// Generic case for unnamed children traversal.
@@ -1316,7 +1316,7 @@ void SMutableCodeViewer::GetChildrenForInfo(TSharedPtr<FMutableCodeTreeElement>
 	//	return;
 	//}
 
-	const mu::PROGRAM& Program = MutableModel->GetPrivate()->m_program;
+	const mu::FProgram& Program = MutableModel->GetPrivate()->m_program;
 
 	mu::OP::ADDRESS ParentAddress = InInfo->MutableOperation;
 
@@ -1812,7 +1812,7 @@ void SMutableCodeViewer::CacheRootNodeAddresses()
 	const int32 StateCount = ModelPrivate->m_program.m_states.Num();
 	for ( int32 StateIndex=0; StateIndex < StateCount; ++StateIndex )
 	{
-		const mu::PROGRAM::STATE& State = ModelPrivate->m_program.m_states[StateIndex];
+		const mu::FProgram::FState& State = ModelPrivate->m_program.m_states[StateIndex];
 		FoundRootNodeAddresses.Add(State.m_root);
 	}
 
@@ -1837,7 +1837,7 @@ void SMutableCodeViewer::CacheRootNodeAddresses()
 	TSet<mu::OP::ADDRESS> ConstantResourcesAddresses; 
 
 	// Main update procedure run for the targeted state and the targeted parameter values
-	const mu::PROGRAM& Program = MutableModel->GetPrivate()->m_program;
+	const mu::FProgram& Program = MutableModel->GetPrivate()->m_program;
 	for (const mu::OP::ADDRESS& RootOperationAddress : RootNodeAddresses)
 	{
 		GetOperationsReferencingConstantResource(ConstantDataType, IndexOnConstantsArray,Program, RootOperationAddress,
@@ -1877,7 +1877,7 @@ void SMutableCodeViewer::CacheRootNodeAddresses()
 
 }
 
-void SMutableCodeViewer::GetOperationsReferencingConstantResource(const mu::DATATYPE ConstantDataType, const int32 IndexOnConstantsArray, const mu::PROGRAM& InProgram,
+void SMutableCodeViewer::GetOperationsReferencingConstantResource(const mu::DATATYPE ConstantDataType, const int32 IndexOnConstantsArray, const mu::FProgram& InProgram,
 	const mu::OP::ADDRESS& InParentAddress,
 	TSet<mu::OP::ADDRESS>& OutAddressesWithPresence,
 	TSet<mu::OP::ADDRESS>& AlreadyProcessedAddresses)
@@ -1904,7 +1904,7 @@ void SMutableCodeViewer::GetOperationsReferencingConstantResource(const mu::DATA
 }
 
 bool SMutableCodeViewer::IsConstantResourceUsedByOperation(const int32 IndexOnConstantsArray,
-	const mu::DATATYPE ConstantDataType, const mu::OP::ADDRESS OperationAddress, const mu::PROGRAM& InProgram) const
+	const mu::DATATYPE ConstantDataType, const mu::OP::ADDRESS OperationAddress, const mu::FProgram& InProgram) const
 {
 	// Cache the current operation type to know where to look and what to check
 	const mu::OP_TYPE OperationType = InProgram.GetOpType(OperationAddress);
@@ -2349,7 +2349,7 @@ void SMutableCodeViewer::PreviewMutableLayout(mu::LayoutPtrConst Layout)
 }
 
 
-void SMutableCodeViewer::PreviewMutableProjector(const mu::PROJECTOR* Projector)
+void SMutableCodeViewer::PreviewMutableProjector(const mu::FProjector* Projector)
 {
 	if (!Projector)
 	{
@@ -2400,7 +2400,7 @@ void SMutableCodeViewer::PreviewMutableMatrix(const mu::mat4f* Mat)
 }
 
 // TODO: Implement shape viewer
-void SMutableCodeViewer::PreviewMutableShape(const mu::SHAPE* Shape)
+void SMutableCodeViewer::PreviewMutableShape(const mu::FShape* Shape)
 {
 	UE_LOG(LogMutable,Warning,TEXT("Previewer for Mutable Shapes not yet implemented"))
 }
diff --git a/Engine/Plugins/Experimental/Mutable/Source/CustomizableObjectEditor/Private/MuCOE/SMutableCodeViewer.h b/Engine/Plugins/Experimental/Mutable/Source/CustomizableObjectEditor/Private/MuCOE/SMutableCodeViewer.h
--- a/Engine/Plugins/Experimental/Mutable/Source/CustomizableObjectEditor/Private/MuCOE/SMutableCodeViewer.h
+++ b/Engine/Plugins/Experimental/Mutable/Source/CustomizableObjectEditor/Private/MuCOE/SMutableCodeViewer.h
@@ -52,8 +52,8 @@ class SMutableStringViewer;
 class STextComboBox;
 class SWidget;
 namespace mu { struct Curve; }
-namespace mu { struct PROJECTOR; }
-namespace mu { struct SHAPE; }
+namespace mu { struct FProjector; }
+namespace mu { struct FShape; }
 struct FGeometry;
 
 /** This widget shows the internal Mutable Code for debugging purposes. 
@@ -281,7 +281,7 @@ private:
 	 * @param  AlreadyProcessedAddresses Used during recursive calls. Cache the already processed operations to avoid
 	 * processing them more than once.
 	 */
-	void GetOperationsOfType ( const mu::OP_TYPE& TargetOperationType, const mu::OP::ADDRESS& InParentAddress, const mu::PROGRAM& InProgram,
+	void GetOperationsOfType ( const mu::OP_TYPE& TargetOperationType, const mu::OP::ADDRESS& InParentAddress, const mu::FProgram& InProgram,
 		TSet<mu::OP::ADDRESS>& OutAddressesOfType,
 		TSet<mu::OP::ADDRESS>& AlreadyProcessedAddresses);
 
@@ -306,7 +306,7 @@ private:
 	 * @param OutLocatedOperations Output of the method : Set with all the operation types present on the model.
 	 * @param AlreadyProcessedAddresses Addresses already processed. Used to avoid processing the same operation twice.
 	 */
-	void GetOperationTypesPresentOnModel( const mu::OP::ADDRESS& InParentAddress, const mu::PROGRAM& InProgram,
+	void GetOperationTypesPresentOnModel( const mu::OP::ADDRESS& InParentAddress, const mu::FProgram& InProgram,
 		TSet<mu::OP_TYPE>& OutLocatedOperations,
 		TSet<mu::OP::ADDRESS>& AlreadyProcessedAddresses);
 
@@ -348,7 +348,7 @@ private:
 	 */
 	void GetOperationsReferencingConstantResource(const mu::DATATYPE ConstantDataType,
 	                                              const int32 IndexOnConstantsArray,
-	                                              const mu::PROGRAM& InProgram,
+	                                              const mu::FProgram& InProgram,
 	                                              const mu::OP::ADDRESS& InParentAddress,
 	                                              TSet<mu::OP::ADDRESS>& OutAddressesWithPresence,
 	                                              TSet<mu::OP::ADDRESS>& AlreadyProcessedAddresses);
@@ -367,7 +367,7 @@ private:
 	 */
 	bool IsConstantResourceUsedByOperation(const int32 IndexOnConstantsArray,
 	                                       const mu::DATATYPE ConstantDataType, const mu::OP::ADDRESS OperationAddress,
-	                                       const mu::PROGRAM& InProgram) const;
+	                                       const mu::FProgram& InProgram) const;
 	
 	
 	/*
@@ -563,9 +563,9 @@ public:
 	void PreviewMutableLayout(mu::LayoutPtrConst Layout);
 	void PreviewMutableSkeleton(mu::SkeletonPtrConst Skeleton);
 	void PreviewMutableString(const mu::string* InStringPtr);
-	void PreviewMutableProjector(const mu::PROJECTOR* Projector);
+	void PreviewMutableProjector(const mu::FProjector* Projector);
 	void PreviewMutableMatrix(const mu::mat4f* Mat);
-	void PreviewMutableShape(const mu::SHAPE* Shape);
+	void PreviewMutableShape(const mu::FShape* Shape);
 	void PreviewMutableCurve(const mu::Curve* Curve);
 	
 private:
@@ -605,8 +605,8 @@ public:
 			}
 			
 			// Iterate over all states and try to locate the operation
-			const mu::PROGRAM& MutableProgram = InModel->GetPrivate()->m_program;
-			for (const mu::PROGRAM::STATE& CurrentState : MutableProgram.m_states)
+			const mu::FProgram& MutableProgram = InModel->GetPrivate()->m_program;
+			for (const mu::FProgram::FState& CurrentState : MutableProgram.m_states)
 			{
 				// Check if it is a dynamic resource
 				for (auto& DynamicResource : CurrentState.m_dynamicResources)
diff --git a/Engine/Plugins/Experimental/Mutable/Source/CustomizableObjectEditor/Private/MuCOE/SMutableConstantsWidget.cpp b/Engine/Plugins/Experimental/Mutable/Source/CustomizableObjectEditor/Private/MuCOE/SMutableConstantsWidget.cpp
--- a/Engine/Plugins/Experimental/Mutable/Source/CustomizableObjectEditor/Private/MuCOE/SMutableConstantsWidget.cpp
+++ b/Engine/Plugins/Experimental/Mutable/Source/CustomizableObjectEditor/Private/MuCOE/SMutableConstantsWidget.cpp
@@ -430,7 +430,7 @@ public:
 
 
 
-void SMutableConstantsWidget::Construct(const FArguments& InArgs,const mu::PROGRAM* InMutableProgramPtr,  TSharedPtr<SMutableCodeViewer> InMutableCodeViewerPtr)
+void SMutableConstantsWidget::Construct(const FArguments& InArgs,const mu::FProgram* InMutableProgramPtr,  TSharedPtr<SMutableCodeViewer> InMutableCodeViewerPtr)
 {
 	// A pointer to MutableCodeViewerPtr is required in order to be able to invoke the preview of our constants
 	check (InMutableCodeViewerPtr);
@@ -717,7 +717,7 @@ void SMutableConstantsWidget::Construct(const FArguments& InArgs,const mu::PROGR
 
 
 
-void SMutableConstantsWidget::SetProgram(const mu::PROGRAM* InProgram)
+void SMutableConstantsWidget::SetProgram(const mu::FProgram* InProgram)
 {
 	// Make sure we do not process the setting of the same program object as the one already set
 	if (this->MutableProgramPtr == InProgram)
diff --git a/Engine/Plugins/Experimental/Mutable/Source/CustomizableObjectEditor/Private/MuCOE/SMutableConstantsWidget.h b/Engine/Plugins/Experimental/Mutable/Source/CustomizableObjectEditor/Private/MuCOE/SMutableConstantsWidget.h
--- a/Engine/Plugins/Experimental/Mutable/Source/CustomizableObjectEditor/Private/MuCOE/SMutableConstantsWidget.h
+++ b/Engine/Plugins/Experimental/Mutable/Source/CustomizableObjectEditor/Private/MuCOE/SMutableConstantsWidget.h
@@ -30,9 +30,9 @@ class SMutableCodeViewer;
 class STableViewBase;
 class SWidget;
 namespace mu { struct Curve; }
-namespace mu { struct PROGRAM; }
-namespace mu { struct PROJECTOR; }
-namespace mu { struct SHAPE; }
+namespace mu { struct FProgram; }
+namespace mu { struct FProjector; }
+namespace mu { struct FShape; }
 
 /**
  * Base Structure to define the different elements used by the lists on this object
@@ -88,7 +88,7 @@ struct FMutableConstantSkeletonElement : public FMutableConstantElement
 */
 struct FMutableConstantProjectorElement : public FMutableConstantElement
 {
-	const mu::PROJECTOR* Projector = nullptr;
+	const mu::FProjector* Projector = nullptr;
 };
 
 /**
@@ -104,7 +104,7 @@ struct FMutableConstantMatrixElement : public FMutableConstantElement
 */
 struct FMutableConstantShapeElement : public FMutableConstantElement
 {
-	const mu::SHAPE* Shape = nullptr;
+	const mu::FShape* Shape = nullptr;
 };
 
 /**
@@ -128,15 +128,15 @@ public:
 
 	/** Builds the widget
 	 * @param InArgs - Arguments provided when generating this slate object
-	 * @param InMutableProgramPtr - Pointer to the mu::PROGRAM object that holds the constants data.
+	 * @param InMutableProgramPtr - Pointer to the mu::FProgram object that holds the constants data.
 	 * @param  InMutableCodeViewerPtr - Pointer to the MutableCodeViewer tasked with the previewing of the constant values
 	 */
-	void Construct(const FArguments& InArgs, const mu::PROGRAM*  InMutableProgramPtr ,  TSharedPtr<SMutableCodeViewer> InMutableCodeViewerPtr);
+	void Construct(const FArguments& InArgs, const mu::FProgram*  InMutableProgramPtr ,  TSharedPtr<SMutableCodeViewer> InMutableCodeViewerPtr);
 	
 private:
 
 	/** Mutable object containing the constants data */
-	const mu::PROGRAM* MutableProgramPtr = nullptr;
+	const mu::FProgram* MutableProgramPtr = nullptr;
 
 	/** Slate capable of accessing the previewer object */
 	 TSharedPtr<SMutableCodeViewer> MutableCodeViewerPtr = nullptr;
@@ -145,7 +145,7 @@ private:
 	 *Sets the back end for the operation of this widget. Each time this is done the ui backend gets updated
 	 * @param InProgram - Mutable program object holding all the constants data
 	 */
-	void SetProgram(const mu::PROGRAM* InProgram);
+	void SetProgram(const mu::FProgram* InProgram);
 	
 	/*
 	* Data backend for the lists of constants
@@ -162,7 +162,7 @@ private:
 	TArray<TSharedPtr<FMutableConstantSkeletonElement>> ConstantSkeletonElements;
 
 	/**
-	 * Load up all the elements with the data found on the mu::PROGRAM object onto TArrays after parsing the data found.
+	 * Load up all the elements with the data found on the mu::FProgram object onto TArrays after parsing the data found.
 	 */
 	void LoadConstantElements();
 	
diff --git a/Engine/Plugins/Experimental/Mutable/Source/CustomizableObjectEditor/Private/MuCOE/SMutableParametersWidget.cpp b/Engine/Plugins/Experimental/Mutable/Source/CustomizableObjectEditor/Private/MuCOE/SMutableParametersWidget.cpp
--- a/Engine/Plugins/Experimental/Mutable/Source/CustomizableObjectEditor/Private/MuCOE/SMutableParametersWidget.cpp
+++ b/Engine/Plugins/Experimental/Mutable/Source/CustomizableObjectEditor/Private/MuCOE/SMutableParametersWidget.cpp
@@ -490,17 +490,12 @@ TOptional<FVector::FReal> SMutableParametersWidget::GetProjectorLocation(EAxis::
 	}
 
 	mu::PROJECTOR_TYPE Type;
-	float Pos[3];
-	float Dir[3];
-	float Up[3];
-	float Scale[3];
+	FVector3f Pos;
+	FVector3f Dir;
+	FVector3f Up;
+	FVector3f Scale;
 	float Angle;
-	MutableParameters->GetProjectorValue(ParamIndex, &Type, 
-		&Pos[0], &Pos[1], &Pos[2],
-		&Dir[0], &Dir[1], &Dir[2],
-		&Up[0], &Up[1], &Up[2],
-		&Scale[0], &Scale[1], &Scale[2],
-		&Angle );
+	MutableParameters->GetProjectorValue(ParamIndex, &Type, &Pos, &Dir, &Up, &Scale, &Angle );
 
 	return Pos[ Axis-EAxis::X ];
 }
@@ -516,26 +511,16 @@ void SMutableParametersWidget::SetProjectorLocation(FVector::FReal NewValue, ETe
 	}
 
 	mu::PROJECTOR_TYPE Type;
-	float Pos[3];
-	float Dir[3];
-	float Up[3];
-	float Scale[3];
+	FVector3f Pos;
+	FVector3f Dir;
+	FVector3f Up;
+	FVector3f Scale;
 	float Angle;
-	MutableParameters->GetProjectorValue(ParamIndex, &Type,
-		&Pos[0], &Pos[1], &Pos[2],
-		&Dir[0], &Dir[1], &Dir[2],
-		&Up[0], &Up[1], &Up[2],
-		&Scale[0], &Scale[1], &Scale[2],
-		&Angle);
+	MutableParameters->GetProjectorValue(ParamIndex, &Type, &Pos, &Dir, &Up, &Scale, &Angle);
 
 	Pos[Axis - EAxis::X] = NewValue;
 
-	MutableParameters->SetProjectorValue(ParamIndex,
-		Pos[0], Pos[1], Pos[2],
-		Dir[0], Dir[1], Dir[2],
-		Up[0], Up[1], Up[2],
-		Scale[0], Scale[1], Scale[2],
-		Angle);
+	MutableParameters->SetProjectorValue(ParamIndex, Pos, Dir, Up, Scale, Angle);
 
 	OnParametersValueChanged.ExecuteIfBound(ParamIndex);
 }
@@ -551,17 +536,12 @@ TOptional<FVector::FReal> SMutableParametersWidget::GetProjectorScale(EAxis::Typ
 	}
 
 	mu::PROJECTOR_TYPE Type;
-	float Pos[3];
-	float Dir[3];
-	float Up[3];
-	float Scale[3];
+	FVector3f Pos;
+	FVector3f Dir;
+	FVector3f Up;
+	FVector3f Scale;
 	float Angle;
-	MutableParameters->GetProjectorValue(ParamIndex, &Type,
-		&Pos[0], &Pos[1], &Pos[2],
-		&Dir[0], &Dir[1], &Dir[2],
-		&Up[0], &Up[1], &Up[2],
-		&Scale[0], &Scale[1], &Scale[2],
-		&Angle);
+	MutableParameters->GetProjectorValue(ParamIndex, &Type, &Pos, &Dir, &Up, &Scale, &Angle);
 
 	return Scale[Axis - EAxis::X];
 }
@@ -577,26 +557,16 @@ void SMutableParametersWidget::SetProjectorScale(FVector::FReal NewValue, ETextC
 	}
 
 	mu::PROJECTOR_TYPE Type;
-	float Pos[3];
-	float Dir[3];
-	float Up[3];
-	float Scale[3];
+	FVector3f Pos;
+	FVector3f Dir;
+	FVector3f Up;
+	FVector3f Scale;
 	float Angle;
-	MutableParameters->GetProjectorValue(ParamIndex, &Type,
-		&Pos[0], &Pos[1], &Pos[2],
-		&Dir[0], &Dir[1], &Dir[2],
-		&Up[0], &Up[1], &Up[2],
-		&Scale[0], &Scale[1], &Scale[2],
-		&Angle);
+	MutableParameters->GetProjectorValue(ParamIndex, &Type, &Pos, &Dir, &Up, &Scale, &Angle);
 
 	Scale[Axis - EAxis::X] = NewValue;
 
-	MutableParameters->SetProjectorValue(ParamIndex,
-		Pos[0], Pos[1], Pos[2],
-		Dir[0], Dir[1], Dir[2],
-		Up[0], Up[1], Up[2],
-		Scale[0], Scale[1], Scale[2],
-		Angle);
+	MutableParameters->SetProjectorValue(ParamIndex, Pos, Dir, Up, Scale, Angle);
 
 	OnParametersValueChanged.ExecuteIfBound(ParamIndex);
 }
diff --git a/Engine/Plugins/Experimental/Mutable/Source/CustomizableObjectEditor/Private/MuCOE/SMutableProjectorViewer.cpp b/Engine/Plugins/Experimental/Mutable/Source/CustomizableObjectEditor/Private/MuCOE/SMutableProjectorViewer.cpp
--- a/Engine/Plugins/Experimental/Mutable/Source/CustomizableObjectEditor/Private/MuCOE/SMutableProjectorViewer.cpp
+++ b/Engine/Plugins/Experimental/Mutable/Source/CustomizableObjectEditor/Private/MuCOE/SMutableProjectorViewer.cpp
@@ -183,7 +183,7 @@ void SMutableProjectorViewer::Construct(const FArguments& InArgs)
 }
 
 
-void SMutableProjectorViewer::SetProjector(const mu::PROJECTOR& InMutableProjector)
+void SMutableProjectorViewer::SetProjector(const mu::FProjector& InMutableProjector)
 {
 	MutableProjector = InMutableProjector;
 }
diff --git a/Engine/Plugins/Experimental/Mutable/Source/CustomizableObjectEditor/Private/MuCOE/SMutableProjectorViewer.h b/Engine/Plugins/Experimental/Mutable/Source/CustomizableObjectEditor/Private/MuCOE/SMutableProjectorViewer.h
--- a/Engine/Plugins/Experimental/Mutable/Source/CustomizableObjectEditor/Private/MuCOE/SMutableProjectorViewer.h
+++ b/Engine/Plugins/Experimental/Mutable/Source/CustomizableObjectEditor/Private/MuCOE/SMutableProjectorViewer.h
@@ -14,7 +14,7 @@
 class SMutableProjectorViewer final : public SCompoundWidget
 {
 	SLATE_BEGIN_ARGS(SMutableProjectorViewer) {}
-		SLATE_ARGUMENT(mu::PROJECTOR,MutableProjector);
+		SLATE_ARGUMENT(mu::FProjector,MutableProjector);
 	SLATE_END_ARGS()
 
 public:
@@ -29,12 +29,12 @@ public:
 	 * Sets the projector object whose data we want to inspect
 	 * @param InMutableProjector - Projector we desire to inspect. It is the origin of the data shown on the UI
 	 */
-	void SetProjector(const mu::PROJECTOR& InMutableProjector);
+	void SetProjector(const mu::FProjector& InMutableProjector);
 
 private:
 
 	/** The mutable projector witch data is being exposed on the UI */
-	mu::PROJECTOR MutableProjector;
+	mu::FProjector MutableProjector;
 	
 	/** Get the float corresponding to the position of the projector on the targeted component index
 	 * @param VectorComponentIndex - Determines if you get the X, Y or Z value from the vector (0,1 and 2)
diff --git a/Engine/Plugins/Experimental/Mutable/Source/MutableRuntime/Private/MuR/CodeRunner.cpp b/Engine/Plugins/Experimental/Mutable/Source/MutableRuntime/Private/MuR/CodeRunner.cpp
--- a/Engine/Plugins/Experimental/Mutable/Source/MutableRuntime/Private/MuR/CodeRunner.cpp
+++ b/Engine/Plugins/Experimental/Mutable/Source/MutableRuntime/Private/MuR/CodeRunner.cpp
@@ -99,7 +99,7 @@ namespace mu
 		// We will read this in the end, so make sure we keep it.
 		m_pSystem->m_memory->IncreaseHitCount(CACHE_ADDRESS(at, 0, executionOptions));
 
-		PROGRAM& program = m_pModel->GetPrivate()->m_program;
+		FProgram& program = m_pModel->GetPrivate()->m_program;
 		m_romPendingOps.SetNum(program.m_roms.Num());
 
 		// Push the root operation
@@ -802,7 +802,7 @@ namespace mu
 
             OP::ADDRESS cat = args.value;
 
-            PROGRAM& program = pModel->GetPrivate()->m_program;
+            FProgram& program = pModel->GetPrivate()->m_program;
 
             // Assume the ROM has been loaded previously
             check(program.m_constantMeshes[cat].Value )
@@ -836,7 +836,7 @@ namespace mu
             auto args = pModel->GetPrivate()->m_program.GetOpArgs<OP::ResourceConstantArgs>(item.at);
                 OP::ADDRESS cat = args.value;
 
-            const PROGRAM& program = pModel->GetPrivate()->m_program;
+            const FProgram& program = pModel->GetPrivate()->m_program;
 
 			int32 MipsToSkip = item.executionOptions;
             ImagePtrConst pSource;
@@ -1704,7 +1704,7 @@ namespace mu
                 check( args.morphShape < (uint32)pModel->GetPrivate()->m_program.m_constantShapes.Num() );
 
                 // Should be an ellipse
-                const SHAPE& morphShape = pModel->GetPrivate()->m_program.
+                const FShape& morphShape = pModel->GetPrivate()->m_program.
                     m_constantShapes[args.morphShape];
 
                 const mu::vec3f& origin = morphShape.position;
@@ -1715,7 +1715,7 @@ namespace mu
                     check( args.vertexSelectionShapeOrBone < (uint32)pModel->GetPrivate()->m_program.m_constantShapes.Num() );
 
                     // Should be None or an axis aligned box
-                    const SHAPE& selectionShape = pModel->GetPrivate()->m_program.m_constantShapes[args.vertexSelectionShapeOrBone];
+                    const FShape& selectionShape = pModel->GetPrivate()->m_program.m_constantShapes[args.vertexSelectionShapeOrBone];
                     pResult = MeshClipMorphPlane(pSource.get(), origin, normal, args.dist, args.factor, morphShape.size[0], morphShape.size[1], morphShape.size[2], selectionShape);
                 }
 
@@ -1723,16 +1723,16 @@ namespace mu
                 {
                     check( args.vertexSelectionShapeOrBone < (uint32)pModel->GetPrivate()->m_program.m_constantStrings.Num() );
 
-                    SHAPE selectionShape;
-                    selectionShape.type = (uint8)SHAPE::Type::None;
+                    FShape selectionShape;
+                    selectionShape.type = (uint8)FShape::Type::None;
                     const string& selectionBone = pModel->GetPrivate()->m_program.m_constantStrings[args.vertexSelectionShapeOrBone];
 					pResult = MeshClipMorphPlane(pSource.get(), origin, normal, args.dist, args.factor, morphShape.size[0], morphShape.size[1], morphShape.size[2], selectionShape, selectionBone, args.maxBoneRadius);
                 }
                 else
                 {
                     // No vertex selection
-                    SHAPE selectionShape;
-                    selectionShape.type = (uint8)SHAPE::Type::None;
+                    FShape selectionShape;
+                    selectionShape.type = (uint8)FShape::Type::None;
                     pResult = MeshClipMorphPlane(pSource.get(), origin, normal, args.dist, args.factor, morphShape.size[0], morphShape.size[1], morphShape.size[2], selectionShape);
                 }
 
@@ -4058,7 +4058,7 @@ namespace mu
         }
 
         const auto& program = pModel->GetPrivate()->m_program;
-        const PARAMETER_DESC& paramDesc = program.m_parameters[ parameterIndex ];
+        const FParameterDesc& paramDesc = program.m_parameters[ parameterIndex ];
         for( size_t rangeIndexInParam=0;
              rangeIndexInParam<paramDesc.m_ranges.Num();
              ++rangeIndexInParam )
@@ -5157,7 +5157,7 @@ namespace mu
 		}
 
 
-		const mu::PROGRAM& program = pModel->GetPrivate()->m_program;
+		const mu::FProgram& program = pModel->GetPrivate()->m_program;
 
 		OP_TYPE type = program.GetOpType(item.at);
 		switch (type)
diff --git a/Engine/Plugins/Experimental/Mutable/Source/MutableRuntime/Private/MuR/CodeRunnerTasked.cpp b/Engine/Plugins/Experimental/Mutable/Source/MutableRuntime/Private/MuR/CodeRunnerTasked.cpp
--- a/Engine/Plugins/Experimental/Mutable/Source/MutableRuntime/Private/MuR/CodeRunnerTasked.cpp
+++ b/Engine/Plugins/Experimental/Mutable/Source/MutableRuntime/Private/MuR/CodeRunnerTasked.cpp
@@ -63,7 +63,7 @@ namespace mu
 
 		m_pSystem->m_pStreamInterface->EndRead(o.m_streamID);
 
-		PROGRAM& program = m_pModel->GetPrivate()->m_program;
+		FProgram& program = m_pModel->GetPrivate()->m_program;
 		{
 			MUTABLE_CPUPROFILER_SCOPE(Unserialise);
 
@@ -1076,7 +1076,7 @@ namespace mu
 		// This runs in th CodeRunner thread
 		bOutFailed = false;
 
-		PROGRAM& program = InModel->GetPrivate()->m_program;
+		FProgram& program = InModel->GetPrivate()->m_program;
 
 		check(RomIndex < program.m_roms.Num());
 		bool bRomIsLoaded = program.IsRomLoaded(RomIndex);
@@ -1160,7 +1160,7 @@ namespace mu
 	//---------------------------------------------------------------------------------------------
 	bool CodeRunner::FLoadMeshRomTask::IsComplete(CodeRunner* runner)
 	{ 
-		PROGRAM& program = runner->m_pModel->GetPrivate()->m_program;
+		FProgram& program = runner->m_pModel->GetPrivate()->m_program;
 		return program.IsRomLoaded(RomIndex);
 	}
 
@@ -1175,7 +1175,7 @@ namespace mu
 		// This runs in th CodeRunner thread
 		bOutFailed = false;
 
-		PROGRAM& program = InModel->GetPrivate()->m_program;
+		FProgram& program = InModel->GetPrivate()->m_program;
 
 		for (int32 i = 0; i<LODIndexCount; ++i )
 		{
@@ -1271,7 +1271,7 @@ namespace mu
 		// Process the constant op normally, now that the rom is loaded.
 		runner->RunCode(Op, runner->m_pParams, runner->m_pModel, runner->m_lodMask);
 
-		PROGRAM& program = runner->m_pModel->GetPrivate()->m_program;
+		FProgram& program = runner->m_pModel->GetPrivate()->m_program;
 		for (int32 i = 0; i < LODIndexCount; ++i)
 		{
 			int32 CurrentIndexIndex = LODIndexIndex + i;
@@ -1298,7 +1298,7 @@ namespace mu
 	//---------------------------------------------------------------------------------------------
 	bool CodeRunner::FLoadImageRomsTask::IsComplete(CodeRunner* runner)
 	{
-		PROGRAM& program = runner->m_pModel->GetPrivate()->m_program;
+		FProgram& program = runner->m_pModel->GetPrivate()->m_program;
 		for (int32 i = 0; i < LODIndexCount; ++i)
 		{
 			int32 CurrentIndexIndex = LODIndexIndex + i;
@@ -1329,7 +1329,7 @@ namespace mu
 	{
 		TSharedPtr<FIssuedTask> Issued;
 
-		PROGRAM& program = m_pModel->GetPrivate()->m_program;
+		FProgram& program = m_pModel->GetPrivate()->m_program;
 
 		auto type = program.GetOpType(item.at);
 
diff --git a/Engine/Plugins/Experimental/Mutable/Source/MutableRuntime/Private/MuR/CodeVisitor.cpp b/Engine/Plugins/Experimental/Mutable/Source/MutableRuntime/Private/MuR/CodeVisitor.cpp
--- a/Engine/Plugins/Experimental/Mutable/Source/MutableRuntime/Private/MuR/CodeVisitor.cpp
+++ b/Engine/Plugins/Experimental/Mutable/Source/MutableRuntime/Private/MuR/CodeVisitor.cpp
@@ -13,7 +13,7 @@ namespace mu
 	//-------------------------------------------------------------------------------------------------
 	//-------------------------------------------------------------------------------------------------
 	//-------------------------------------------------------------------------------------------------
-	void SubtreeParametersVisitor::Run(OP::ADDRESS root, PROGRAM& program)
+	void SubtreeParametersVisitor::Run(OP::ADDRESS root, FProgram& program)
 	{
 		// Cached?
 		auto it = m_resultCache.Find(root);
diff --git a/Engine/Plugins/Experimental/Mutable/Source/MutableRuntime/Private/MuR/CodeVisitor.h b/Engine/Plugins/Experimental/Mutable/Source/MutableRuntime/Private/MuR/CodeVisitor.h
--- a/Engine/Plugins/Experimental/Mutable/Source/MutableRuntime/Private/MuR/CodeVisitor.h
+++ b/Engine/Plugins/Experimental/Mutable/Source/MutableRuntime/Private/MuR/CodeVisitor.h
@@ -105,7 +105,7 @@ class Parameters;
         }
 
 
-        void Traverse( OP::ADDRESS root, PROGRAM& program, bool visitDecorators = true )
+        void Traverse( OP::ADDRESS root, FProgram& program, bool visitDecorators = true )
         {
             m_pending.Reserve( program.m_opAddress.Num() );
 
@@ -127,7 +127,7 @@ class Parameters;
             }
         }
 
-        void FullTraverse( PROGRAM& program, bool visitDecorators = true )
+        void FullTraverse( FProgram& program, bool visitDecorators = true )
         {
             // Visit all the state roots
             for ( std::size_t p=0; p<program.m_states.Num(); ++p )
@@ -155,7 +155,7 @@ class Parameters;
 
         //! Do the actual work by overriding this in the derived classes.
         //! Return true if the traverse has to continue with the children of "at"
-        virtual bool Visit( OP::ADDRESS at, PROGRAM& program ) = 0;
+        virtual bool Visit( OP::ADDRESS at, FProgram& program ) = 0;
 
         //! Operations to be processed
         struct PENDING
@@ -191,7 +191,7 @@ class Parameters;
 		TArray<TArray<int>> m_visited;
 
         //! Process all the pending operations and visit all children if necessary
-        void Recurse( PROGRAM& program )
+        void Recurse( FProgram& program )
         {
 			m_visited.Empty();
 			m_visited.SetNum(program.m_opAddress.Num());
@@ -308,7 +308,7 @@ class Parameters;
         }
 
 
-        void Traverse( OP::ADDRESS root, PROGRAM& program, bool visitDecorators = true )
+        void Traverse( OP::ADDRESS root, FProgram& program, bool visitDecorators = true )
         {
             m_pending.reserve( program.m_opAddress.Num() );
 
@@ -330,7 +330,7 @@ class Parameters;
             }
         }
 
-        void FullTraverse( PROGRAM& program, bool visitDecorators = true )
+        void FullTraverse( FProgram& program, bool visitDecorators = true )
         {
             // Visit all the state roots
             for ( std::size_t p=0; p<program.m_states.Num(); ++p )
@@ -358,7 +358,7 @@ class Parameters;
 
         //! Do the actual work by overriding this in the derived classes.
         //! Return true if the traverse has to continue with the children of "at"
-        virtual bool Visit( OP::ADDRESS at, PROGRAM& program ) = 0;
+        virtual bool Visit( OP::ADDRESS at, FProgram& program ) = 0;
 
         //! Operations to be processed
         struct PENDING
@@ -391,7 +391,7 @@ class Parameters;
 
 
         //! Process all the pending operations and visit all children if necessary
-        void Recurse( PROGRAM& program )
+        void Recurse( FProgram& program )
         {
             while ( m_pending.Num() )
             {
@@ -487,7 +487,7 @@ class Parameters;
 
     protected:
 
-        virtual bool Visit( OP::ADDRESS at, PROGRAM& program )
+        virtual bool Visit( OP::ADDRESS at, FProgram& program )
         {
             bool recurse = true;
 
@@ -708,7 +708,7 @@ class Parameters;
     {
     public:
 
-        void Run( OP::ADDRESS root, PROGRAM& program );
+        void Run( OP::ADDRESS root, FProgram& program );
 
         //! After Run, list of relevant parameters.
 		TArray<int> m_params;
diff --git a/Engine/Plugins/Experimental/Mutable/Source/MutableRuntime/Private/MuR/Model.cpp b/Engine/Plugins/Experimental/Mutable/Source/MutableRuntime/Private/MuR/Model.cpp
--- a/Engine/Plugins/Experimental/Mutable/Source/MutableRuntime/Private/MuR/Model.cpp
+++ b/Engine/Plugins/Experimental/Mutable/Source/MutableRuntime/Private/MuR/Model.cpp
@@ -25,7 +25,7 @@ namespace mu
 {
 
     //---------------------------------------------------------------------------------------------
-    void PROGRAM::Check()
+    void FProgram::Check()
     {
 //    #ifdef MUTABLE_DEBUG
 //        // Process all from root
@@ -83,7 +83,7 @@ namespace mu
 
 
     //---------------------------------------------------------------------------------------------
-    void PROGRAM::LogHistogram() const
+    void FProgram::LogHistogram() const
     {
 #if 0
         uint64 countPerType[(int)OP_TYPE::COUNT];
@@ -178,7 +178,7 @@ namespace mu
     {
 		LLM_SCOPE_BYNAME(TEXT("MutableRuntime"));
 
-		mu::PROGRAM& program = p->m_pD->m_program;
+		mu::FProgram& program = p->m_pD->m_program;
 
 		TArray<TPair<int32, mu::ImagePtrConst>> InitialImages = program.m_constantImageLODs;
 		TArray<TPair<int32, mu::MeshPtrConst>> InitialMeshes = program.m_constantMeshes;
@@ -470,7 +470,7 @@ namespace mu
 
         if ( stateIndex>=0 && stateIndex<(int)m_pD->m_program.m_states.Num() )
         {
-            const PROGRAM::STATE& state = m_pD->m_program.m_states[stateIndex];
+            const FProgram::FState& state = m_pD->m_program.m_states[stateIndex];
             if ( paramIndex>=0 && paramIndex<(int)state.m_runtimeParameters.Num() )
             {
                 res = (int)state.m_runtimeParameters[paramIndex];
@@ -623,7 +623,7 @@ namespace mu
             //! 3D Projector type, defining a position, scale and orientation. Basically used for
             //! projected decals.
             case PARAMETER_TYPE::T_PROJECTOR:
-                dataSize = sizeof(PROJECTOR);
+                dataSize = sizeof(FProjector);
 
                 // \todo: padding will be random?
                 parameterValuesBlob.SetNum( pos+dataSize );
@@ -753,7 +753,7 @@ namespace mu
             m_pD->m_instanceCount = 1;
             for (size_t i=0;i<paramCount; ++i)
             {
-                const PARAMETER_DESC& param =
+                const FParameterDesc& param =
                         m_pD->m_pModel->GetPrivate()->m_program.m_parameters[ i ];
                 switch ( param.m_type )
                 {
@@ -805,12 +805,11 @@ namespace mu
 
 
     //---------------------------------------------------------------------------------------------
-    ParametersPtr ModelParametersGenerator::GetInstance( int64 index,
-                                                         float (*randomGenerator )() )
+    Ptr<Parameters> ModelParametersGenerator::GetInstance( int64 index, float (*randomGenerator )() )
     {
 		LLM_SCOPE_BYNAME(TEXT("MutableRuntime"));
 
-        ParametersPtr res = m_pD->m_pModel->NewParameters();
+		Ptr<Parameters> res = m_pD->m_pModel->NewParameters();
 
         if (m_pD->m_considerRelevancy)
         {
@@ -837,11 +836,24 @@ namespace mu
                 case PARAMETER_TYPE::T_COLOUR:
                     if ( randomGenerator )
                     {
-                        res->SetColourValue( p,
-                                             randomGenerator(), randomGenerator(), randomGenerator() );
+                        res->SetColourValue( p, randomGenerator(), randomGenerator(), randomGenerator() );
                     }
                     break;
 
+				case PARAMETER_TYPE::T_PROJECTOR:
+					if (randomGenerator)
+					{
+						// For projectors we just warp the position a little bit just to get something different
+						FVector3f Position,Direction,Up,Scale;
+						float Angle=1.0f;
+						res->GetProjectorValue(p, nullptr, &Position, &Direction, &Up, &Scale, &Angle );
+						Position.X *= 0.9 + randomGenerator() * 0.2f;
+						Position.Y *= 0.9 + randomGenerator() * 0.2f;
+						Position.Z *= 0.9 + randomGenerator() * 0.2f;
+						res->SetProjectorValue(p,Position,Direction,Up,Scale,Angle);
+					}
+					break;
+
                 default:
                     break;
                 }
@@ -849,11 +861,11 @@ namespace mu
         }
         else
         {
-            int paramCount = (int)m_pD->m_pModel->GetPrivate()->m_program.m_parameters.Num();
+            int paramCount = m_pD->m_pModel->GetPrivate()->m_program.m_parameters.Num();
             int64 currentInstance = index;
             for (int i=0;i<paramCount;++i)
             {
-                const PARAMETER_DESC& param = m_pD->m_pModel->GetPrivate()->m_program.m_parameters[ i ];
+                const FParameterDesc& param = m_pD->m_pModel->GetPrivate()->m_program.m_parameters[ i ];
                 switch ( param.m_type )
                 {
 
@@ -887,6 +899,20 @@ namespace mu
                     }
                     break;
 
+				case PARAMETER_TYPE::T_PROJECTOR:
+					if (randomGenerator)
+					{
+						// For projectors we just warp the position a little bit just to get something different
+						FVector3f Position, Direction, Up, Scale;
+						float Angle = 1.0f;
+						res->GetProjectorValue(i, nullptr, &Position, &Direction, &Up, &Scale, &Angle);
+						Position.X *= 0.9 + randomGenerator() * 0.2f;
+						Position.Y *= 0.9 + randomGenerator() * 0.2f;
+						Position.Z *= 0.9 + randomGenerator() * 0.2f;
+						res->SetProjectorValue(i, Position, Direction, Up, Scale, Angle);
+					}
+					break;
+
                 default:
                     break;
                 }
@@ -908,7 +934,7 @@ namespace mu
 
         for (int i=0;i<paramCount;++i)
         {
-            const PARAMETER_DESC& param = m_pD->m_pModel->GetPrivate()->m_program.m_parameters[ i ];
+            const FParameterDesc& param = m_pD->m_pModel->GetPrivate()->m_program.m_parameters[ i ];
             switch ( param.m_type )
             {
 
@@ -935,6 +961,19 @@ namespace mu
                 res->SetColourValue( i, randomGenerator(), randomGenerator(), randomGenerator() );
                 break;
 
+			case PARAMETER_TYPE::T_PROJECTOR:
+				{
+					// For projectors we just warp the position a little bit just to get something different
+					FVector3f Position, Direction, Up, Scale;
+					float Angle = 1.0f;
+					res->GetProjectorValue(i, nullptr, &Position, &Direction, &Up, &Scale, &Angle);
+					Position.X *= 0.9 + randomGenerator() * 0.2f;
+					Position.Y *= 0.9 + randomGenerator() * 0.2f;
+					Position.Z *= 0.9 + randomGenerator() * 0.2f;
+					res->SetProjectorValue(i, Position, Direction, Up, Scale, Angle);
+				}
+				break;
+
             default:
                 break;
             }
@@ -972,7 +1011,7 @@ namespace mu
 //        }
 
 
-//        void Visit( OP::ADDRESS at, PROGRAM& program ) override
+//        void Visit( OP::ADDRESS at, FProgram& program ) override
 //        {
 //            if ( !m_relevant && !m_visited[at])
 //            {
@@ -1063,7 +1102,7 @@ namespace mu
 
             if (relevant)
             {
-                const PARAMETER_DESC& param =
+                const FParameterDesc& param =
                         m_pModel->GetPrivate()->m_program.m_parameters[ currentParameter ];
                 switch ( param.m_type )
                 {
diff --git a/Engine/Plugins/Experimental/Mutable/Source/MutableRuntime/Private/MuR/ModelPrivate.h b/Engine/Plugins/Experimental/Mutable/Source/MutableRuntime/Private/MuR/ModelPrivate.h
--- a/Engine/Plugins/Experimental/Mutable/Source/MutableRuntime/Private/MuR/ModelPrivate.h
+++ b/Engine/Plugins/Experimental/Mutable/Source/MutableRuntime/Private/MuR/ModelPrivate.h
@@ -64,9 +64,9 @@ namespace mu
 	MUTABLE_DEFINE_POD_SERIALISABLE(FImageLODRange);
 
     //!
-    struct PROGRAM
+    struct FProgram
     {
-        PROGRAM()
+        FProgram()
         {
             // Add the null instruction at address 0.
             // TODO: Will do it in the linker
@@ -74,7 +74,7 @@ namespace mu
             m_opAddress.Add(0);
         }
 
-        struct STATE
+        struct FState
         {
             //! Name of the state
             string m_name;
@@ -82,7 +82,7 @@ namespace mu
             //! First instruction of the full build of an instance in this state
             OP::ADDRESS m_root = 0;
 
-            //! List of parameters index (to PROGRAM::m_parameters) of the runtime parameters of
+            //! List of parameters index (to FProgram::m_parameters) of the runtime parameters of
             //! this state.
 			TArray<int> m_runtimeParameters;
 
@@ -136,7 +136,7 @@ namespace mu
             {
                 bool res = false;
 
-                for ( size_t i=0; !res && i<m_updateCache.Num(); ++i )
+                for ( SIZE_T i=0; !res && i<m_updateCache.Num(); ++i )
                 {
                     if ( m_updateCache[i]==at )
                     {
@@ -164,7 +164,7 @@ namespace mu
 		TArray<uint8> m_byteCode;
 
         //!
-		TArray<STATE> m_states;
+		TArray<FState> m_states;
 
 		//! Data for every rom.
 		TArray<FRomData> m_roms;
@@ -188,13 +188,13 @@ namespace mu
 		TArray<Ptr<const Layout>> m_constantLayouts;
 
         //! Constant projectors
-		TArray<PROJECTOR> m_constantProjectors;
+		TArray<FProjector> m_constantProjectors;
 
         //! Constant matrices, usually used for transforms
 		TArray<mat4f> m_constantMatrices;
 
 		//! Constant projectors
-		TArray<SHAPE> m_constantShapes;
+		TArray<FShape> m_constantShapes;
 
         //! Constant curves
 		TArray<Curve> m_constantCurves;
@@ -207,10 +207,10 @@ namespace mu
 
         //! Parameters of the model.
         //! The value stored here is the default value.
-		TArray<PARAMETER_DESC> m_parameters;
+		TArray<FParameterDesc> m_parameters;
 
         //! Ranges for interation of the model operations.
-		TArray<RANGE_DESC> m_ranges;
+		TArray<FRangeDesc> m_ranges;
 
         //! List of parameter lists. These are used in several places, like storing the
         //! pregenerated list of parameters influencing a resource.
@@ -403,7 +403,7 @@ namespace mu
 		OP::ADDRESS AddConstant( Ptr<const Layout> pLayout )
         {
             // Ensure unique
-            for ( size_t i=0; i<m_constantLayouts.Num(); ++i)
+            for ( SIZE_T i=0; i<m_constantLayouts.Num(); ++i)
             {
                 if (m_constantLayouts[i]==pLayout)
                 {
@@ -419,7 +419,7 @@ namespace mu
         OP::ADDRESS AddConstant( Ptr<const Skeleton> pSkeleton )
         {
             // Ensure unique
-            for ( size_t i=0; i<m_constantSkeletons.Num(); ++i)
+            for ( SIZE_T i=0; i<m_constantSkeletons.Num(); ++i)
             {
                 if ( m_constantSkeletons[i]==pSkeleton
                      ||
@@ -438,7 +438,7 @@ namespace mu
         OP::ADDRESS AddConstant( Ptr<const PhysicsBody> pPhysicsBody )
         {
             // Ensure unique
-            for ( size_t i=0; i<m_constantPhysicsBodies.Num(); ++i)
+            for ( SIZE_T i=0; i<m_constantPhysicsBodies.Num(); ++i)
             {
                 if ( m_constantPhysicsBodies[i]==pPhysicsBody
                      ||
@@ -457,7 +457,7 @@ namespace mu
         OP::ADDRESS AddConstant( const string& str )
         {            
             // Ensure unique
-            for ( size_t i=0; i<m_constantStrings.Num(); ++i)
+            for ( SIZE_T i=0; i<m_constantStrings.Num(); ++i)
             {
                 if (m_constantStrings[i]==str)
                 {
@@ -473,7 +473,7 @@ namespace mu
         OP::ADDRESS AddConstant( const mat4f& m )
         {
             // Ensure unique
-            for ( size_t i=0; i<m_constantMatrices.Num(); ++i)
+            for ( SIZE_T i=0; i<m_constantMatrices.Num(); ++i)
             {
                 if (m_constantMatrices[i]==m)
                 {
@@ -486,10 +486,10 @@ namespace mu
             return index;
         }
 
-        OP::ADDRESS AddConstant( const SHAPE& m )
+        OP::ADDRESS AddConstant( const FShape& m )
         {
             // Ensure unique
-            for ( size_t i=0; i<m_constantShapes.Num(); ++i)
+            for ( SIZE_T i=0; i<m_constantShapes.Num(); ++i)
             {
                 if (m_constantShapes[i]==m)
                 {
@@ -502,10 +502,10 @@ namespace mu
             return index;
         }
 
-        OP::ADDRESS AddConstant( const PROJECTOR& m )
+        OP::ADDRESS AddConstant( const FProjector& m )
         {
             // Ensure unique
-            for ( size_t i=0; i<m_constantProjectors.Num(); ++i)
+            for ( SIZE_T i=0; i<m_constantProjectors.Num(); ++i)
             {
                 if (m_constantProjectors[i]==m)
                 {
@@ -521,7 +521,7 @@ namespace mu
         OP::ADDRESS AddConstant( const Curve& m )
         {
             // Ensure unique
-//            for ( size_t i=0; i<m_constantCurves.Num(); ++i)
+//            for ( SIZE_T i=0; i<m_constantCurves.Num(); ++i)
 //            {
 //                if (m_constantCurves[i]==m)
 //                {
@@ -653,7 +653,7 @@ namespace mu
     public:
 
         //!
-        PROGRAM m_program;
+        FProgram m_program;
 
         //! Non-persistent, streamer-specific location information
         string m_location;
diff --git a/Engine/Plugins/Experimental/Mutable/Source/MutableRuntime/Private/MuR/MutableMath.h b/Engine/Plugins/Experimental/Mutable/Source/MutableRuntime/Private/MuR/MutableMath.h
--- a/Engine/Plugins/Experimental/Mutable/Source/MutableRuntime/Private/MuR/MutableMath.h
+++ b/Engine/Plugins/Experimental/Mutable/Source/MutableRuntime/Private/MuR/MutableMath.h
@@ -586,7 +586,15 @@ namespace mu
 		}
 
 		//-----------------------------------------------------------------------------------------
-		inline vec3( const vec<SCALAR,3>& other )
+		inline vec3(const vec<SCALAR, 3>& other)
+		{
+			this->m[0] = other[0];
+			this->m[1] = other[1];
+			this->m[2] = other[2];
+		}				
+
+		//-----------------------------------------------------------------------------------------
+		inline vec3(const FVector3f& other)
 		{
 			this->m[0] = other[0];
 			this->m[1] = other[1];
@@ -723,26 +731,25 @@ namespace mu
 	//---------------------------------------------------------------------------------------------
 	//! Point to plane distance
 	//---------------------------------------------------------------------------------------------
-	template <class SCALAR, int DIM>
-	inline float point_to_plane_dist(const vec<SCALAR, DIM> &point, const vec<SCALAR, DIM> &plane_origin, const vec<SCALAR, DIM> &plane_normal)
+	inline float point_to_plane_dist(const FVector3f& point, const FVector3f& plane_origin, const FVector3f& plane_normal)
 	{
-		return dot((point - plane_origin), plane_normal);
+		return FVector3f::DotProduct((point - plane_origin), plane_normal);
 	}
 
 	//---------------------------------------------------------------------------------------------
 	//! Ray plane intersection
 	//---------------------------------------------------------------------------------------------
-	inline vec3f ray_plane_intersection(const vec3f& ray_start, const vec3f& ray_end, const vec3f& plane_origin, const vec3f &plane_normal)
+	inline FVector3f ray_plane_intersection(const FVector3f& ray_start, const FVector3f& ray_end, const FVector3f& plane_origin, const FVector3f&plane_normal)
 	{
-		return ray_start + (ray_end - ray_start) *	(dot((plane_origin - ray_start), plane_normal) / dot((ray_end - ray_start), plane_normal));
+		return ray_start + (ray_end - ray_start) *	(FVector3f::DotProduct((plane_origin - ray_start), plane_normal) / FVector3f::DotProduct((ray_end - ray_start), plane_normal));
 	}
 
     //constexpr float point_in_plane_eps = 0.0001f;
 
-    inline bool point_in_face(const vec3f& InVtx, const vec3f& normal, vec3f vertices[3], int& out_intersected_vert, int& out_intersected_edge_v0, int& out_intersected_edge_v1, float point_in_plane_eps = 0.0001f)
+    inline bool point_in_face(const FVector3f& InVtx, const FVector3f& normal, FVector3f vertices[3], int& out_intersected_vert, int& out_intersected_edge_v0, int& out_intersected_edge_v1, float point_in_plane_eps = 0.0001f)
     {
-        vec3f SidePlaneNormal;
-        vec3f Side;
+		FVector3f SidePlaneNormal;
+		FVector3f Side;
         bool intersected_edge[3] = { false, false, false };
         int intersected_edge_count = 0;
 
@@ -750,8 +757,8 @@ namespace mu
         {
             // Create plane perpendicular to both this side and the polygon's normal.
             Side = vertices[x] - vertices[(x - 1 < 0) ? 3 - 1 : x - 1];
-            SidePlaneNormal = cross<float>(Side, normalise(normal));
-            normalise(SidePlaneNormal);
+            SidePlaneNormal = FVector3f::CrossProduct(Side, normal.GetSafeNormal());
+            SidePlaneNormal.Normalize();
 
             float dist = point_to_plane_dist(InVtx, vertices[x], SidePlaneNormal);
 
@@ -811,13 +818,13 @@ namespace mu
     }
 
     //---------------------------------------------------------------------------------------------
-    inline bool rayIntersectsFace( const vec3f& ray_start, const vec3f& ray_end,
-                                   const vec3f& v0, const vec3f& v1, const vec3f& v2,
-                                   vec3f& out_intersection,
+    inline bool rayIntersectsFace( const FVector3f& ray_start, const FVector3f& ray_end,
+                                   const FVector3f& v0, const FVector3f& v1, const FVector3f& v2,
+								   FVector3f& out_intersection,
                                    int& out_intersected_vert, int& out_intersected_edge_v0,
                                    int& out_intersected_edge_v1, float point_in_face_epsilon = 0.0001f)
     {
-        vec3f normal = cross<float>(v1 - v0, v2 - v0);
+		FVector3f normal = FVector3f::CrossProduct(v1 - v0, v2 - v0);
 
         // Does the ray cross the plane?
         const float DistStart = point_to_plane_dist(ray_start, v0, normal);
@@ -838,7 +845,7 @@ namespace mu
         }
 
         // Check if the intersection point is actually on the poly.
-        vec3f vertices[3] = { v0, v1, v2 };
+		FVector3f vertices[3] = { v0, v1, v2 };
 
         return point_in_face( out_intersection, normal, vertices,
                               out_intersected_vert,
@@ -847,18 +854,18 @@ namespace mu
 
 
     //---------------------------------------------------------------------------------------------
-    inline bool ray_plane_intersection( const vec3f& ray_start, const vec3f& ray_end,
-                                        const vec3f& plane_origin, const vec3f &plane_normal,
-                                        vec3f& out_intersection, float& out_t )
+    inline bool ray_plane_intersection( const FVector3f& ray_start, const FVector3f& ray_end,
+                                        const FVector3f& plane_origin, const FVector3f&plane_normal,
+                                        FVector3f& out_intersection, float& out_t )
     {
-        float den = mu::dot<float, 3>(ray_end - ray_start, plane_normal);
+        float den = FVector3f::DotProduct(ray_end - ray_start, plane_normal);
 
         if (den == 0.f)
         {
             return false;
         }
 
-        out_t = mu::dot<float, 3>((plane_origin - ray_start), plane_normal) / den;
+        out_t = FVector3f::DotProduct((plane_origin - ray_start), plane_normal) / den;
         out_intersection = ray_start + (ray_end - ray_start) * out_t;
 
         return true;
@@ -866,14 +873,14 @@ namespace mu
 
 
     //---------------------------------------------------------------------------------------------
-    inline bool rayIntersectsFace2(const vec3f& ray_start, const vec3f& ray_end,
-                                   const vec3f& v0, const vec3f& v1, const vec3f& v2,
-                                   vec3f& out_intersection,
+    inline bool rayIntersectsFace2(const FVector3f& ray_start, const FVector3f& ray_end,
+                                   const FVector3f& v0, const FVector3f& v1, const FVector3f& v2,
+		FVector3f& out_intersection,
                                    int& out_intersected_vert,
                                    int& out_intersected_edge_v0, int& out_intersected_edge_v1,
                                    float& out_t )
     {
-        vec3f normal = cross<float>(v1 - v0, v2 - v0);
+		FVector3f normal = FVector3f::CrossProduct(v1 - v0, v2 - v0);
 
         // Does the ray cross the plane?
         const float DistStart = point_to_plane_dist(ray_start, v0, normal);
@@ -897,7 +904,7 @@ namespace mu
         }
 
         // Check if the intersection point is actually on the poly.
-        vec3f vertices[3] = { v0, v1, v2 };
+		FVector3f vertices[3] = { v0, v1, v2 };
 
         return point_in_face( out_intersection, normal, vertices, out_intersected_vert,
                               out_intersected_edge_v0, out_intersected_edge_v1 );
@@ -1098,22 +1105,40 @@ namespace mu
 			return (v-min)/size;
 		}
 
-        //-----------------------------------------------------------------------------------------
-        inline void Bound( const VECTOR& v )
-        {
-            for ( int d=0; d<VECTOR::GetDim(); ++d )
-            {
-                if (v[d]<min[d])
-                {
-                    size[d] += min[d]-v[d];
-                    min[d] = v[d];
-                }
-                else if (v[d]-min[d]>size[d])
-                {
-                    size[d] = v[d]-min[d];
-                }
-            }
-        }
+		//-----------------------------------------------------------------------------------------
+		inline void Bound(const VECTOR& v)
+		{
+			for (int d = 0; d < VECTOR::GetDim(); ++d)
+			{
+				if (v[d] < min[d])
+				{
+					size[d] += min[d] - v[d];
+					min[d] = v[d];
+				}
+				else if (v[d] - min[d] > size[d])
+				{
+					size[d] = v[d] - min[d];
+				}
+			}
+		}
+
+		//-----------------------------------------------------------------------------------------
+		inline void Bound3(const VECTOR& v)
+		{
+			for (int d = 0; d < 3; ++d)
+			{
+				if (v[d] < min[d])
+				{
+					size[d] += min[d] - v[d];
+					min[d] = v[d];
+				}
+				else if (v[d] - min[d] > size[d])
+				{
+					size[d] = v[d] - min[d];
+				}
+			}
+		}
+
 
 		//-----------------------------------------------------------------------------------------
 		inline void Bound(const box<VECTOR>& other)
@@ -1418,18 +1443,24 @@ namespace mu
             return r;
         }
 
-        //-----------------------------------------------------------------------------------------
-        V operator*( const V& v ) const
-        {
-            V r;
+		//-----------------------------------------------------------------------------------------
+		V Transform(const V& v) const
+		{
+			V r;
 
-            for ( int i=0; i<DIM; ++i )
-            {
-                r[i] = dot( m[i], v );
-            }
+			for (int i = 0; i < DIM; ++i)
+			{
+				r[i] = dot(m[i], v);
+			}
 
-            return r;
-        }
+			return r;
+		}
+
+		//-----------------------------------------------------------------------------------------
+		V operator*(const V& v) const
+		{
+			return Transform(v);
+		}
 
         //-----------------------------------------------------------------------------------------
         SCALAR GetDeterminant() const
diff --git a/Engine/Plugins/Experimental/Mutable/Source/MutableRuntime/Private/MuR/OpImageProject.cpp b/Engine/Plugins/Experimental/Mutable/Source/MutableRuntime/Private/MuR/OpImageProject.cpp
--- a/Engine/Plugins/Experimental/Mutable/Source/MutableRuntime/Private/MuR/OpImageProject.cpp
+++ b/Engine/Plugins/Experimental/Mutable/Source/MutableRuntime/Private/MuR/OpImageProject.cpp
@@ -35,7 +35,7 @@ namespace
 	{
 	public:
 
-		RaycastPixelProcessor( const PROJECTOR& projector,
+		RaycastPixelProcessor( const FProjector& projector,
 							   const Image* pSource,
                                const uint8* pTargetData,
                                const uint8* pMaskData,
@@ -148,19 +148,19 @@ namespace
 		}
 
 
-		const Image* m_pSource;
-		vec3<float> m_direction;
-		vec3<float> m_up;
-		vec3<float> m_side;
-		vec3<float> m_position;
-        vec3<float> m_scale;
+		const Image* m_pSource = nullptr;
+		FVector3f m_direction;
+		FVector3f m_up;
+		FVector3f m_side;
+		FVector3f m_position;
+		FVector3f m_scale;
 
         //! Cosine of the fading angle range
-        float m_fadeStartCos, m_fadeEndCos;
+		float m_fadeStartCos = 0.0f, m_fadeEndCos = 0.0f;
 
-        const uint8* m_pTargetData;
-        const uint8* m_pMaskData;
-		int m_pixelSize;
+        const uint8* m_pTargetData = nullptr;
+        const uint8* m_pMaskData = nullptr;
+		int m_pixelSize = 0;
 
 	};
 
@@ -171,7 +171,7 @@ namespace
 	{
 	public:
 
-		RaycastPixelProcessor_UBYTE( const PROJECTOR& projector,
+		RaycastPixelProcessor_UBYTE( const FProjector& projector,
 									 const Image* pSource,
                                      const uint8* pTargetData,
                                      const uint8* pMaskData,
@@ -284,11 +284,11 @@ namespace
         const uint8* m_pSourceData;
 		int m_sourceSizeX, m_sourceSizeY;
 
-		vec3<float> m_direction;
-		vec3<float> m_up;
-		vec3<float> m_side;
-		vec3<float> m_position;
-        vec3<float> m_scale;
+		FVector3f m_direction;
+		FVector3f m_up;
+		FVector3f m_side;
+		FVector3f m_position;
+		FVector3f m_scale;
 
         //! Cosine of the fading angle range
         float m_fadeStartCos, m_fadeEndCos;
@@ -692,7 +692,7 @@ void ImageRasterProjected_Generic( const Mesh* pMesh,
         vertices[v].x = uv[0] * sizeX;
         vertices[v].y = uv[1] * sizeY;
 
-        vec3f pos={0.0f,0.0f,0.0f};
+        FVector3f pos={0.0f,0.0f,0.0f};
         ConvertData( 0, &pos[0], MBF_FLOAT32, posIt.ptr(), posIt.GetFormat() );
         ConvertData( 1, &pos[0], MBF_FLOAT32, posIt.ptr(), posIt.GetFormat() );
         ConvertData( 2, &pos[0], MBF_FLOAT32, posIt.ptr(), posIt.GetFormat() );
@@ -815,9 +815,9 @@ void ImageRasterProjected_Generic( const Mesh* pMesh,
 //! See CreateMeshOptimisedForProjection
 struct OPTIMISED_VERTEX
 {
-    vec2<float> uv;
-    vec3<float> pos;
-    vec3<float> nor;
+	FVector2f uv;
+	FVector3f pos;
+	FVector3f nor;
 };
 
 static_assert( sizeof(OPTIMISED_VERTEX)==32, "UNEXPECTED_STRUCT_SIZE" );
@@ -828,9 +828,9 @@ static_assert( sizeof(OPTIMISED_VERTEX)==32, "UNEXPECTED_STRUCT_SIZE" );
 //! See CreateMeshOptimisedForWrappingProjection
 struct OPTIMISED_VERTEX_WRAPPING
 {
-    vec2<float> uv;
-    vec3<float> pos;
-    vec3<float> nor;
+	FVector2f uv;
+	FVector3f pos;
+	FVector3f nor;
     uint32 layoutBlock;
 };
 
@@ -1284,7 +1284,7 @@ constexpr float vert_collapse_eps = 0.0001f;
 //---------------------------------------------------------------------------------------------
 inline void MeshCreateCollapsedVertexMap( const Mesh* pMesh,
 	TArray<int>& collapsedVertexMap,
-	TArray<vec3f>& vertices,
+	TArray<FVector3f>& vertices,
 	TMultiMap<int, int>& collapsedVertsMap
 )
 {
@@ -1304,7 +1304,7 @@ inline void MeshCreateCollapsedVertexMap( const Mesh* pMesh,
 
             auto it = UntypedMeshBufferIteratorConst(pMesh->GetVertexBuffers(), sem, semIndex);
 
-			box<vec3f> meshBoundingBox;
+			box<FVector3f> meshBoundingBox;
 			Octree* pOctree = nullptr;
 
             switch (sem)
@@ -1314,7 +1314,7 @@ inline void MeshCreateCollapsedVertexMap( const Mesh* pMesh,
                 // First create a cache of the vertices in the vertices array
                 for (int v = 0; v < vcount; ++v)
                 {
-                    vec3f vertex(0.0f, 0.0f, 0.0f);
+                    FVector3f vertex(0.0f, 0.0f, 0.0f);
                     for (int i = 0; i < 3; ++i)
                     {
                         ConvertData(i, &vertex[0], MBF_FLOAT32, it.ptr(), it.GetFormat());
@@ -1327,10 +1327,10 @@ inline void MeshCreateCollapsedVertexMap( const Mesh* pMesh,
 				
 				//Initializing Octree boundingbox
 				meshBoundingBox.min = vertices[0];
-				meshBoundingBox.size = vec3f(0.0f, 0.0f, 0.0f);
+				meshBoundingBox.size = FVector3f(0.0f, 0.0f, 0.0f);
 				for (int32 i = 1; i < vertices.Num(); ++i)
 				{
-					meshBoundingBox.Bound(vertices[i]);
+					meshBoundingBox.Bound3(vertices[i]);
 				}
 
 				//Initializing and filling Octree
@@ -1416,13 +1416,13 @@ struct AdjacentFaces
 };
 
 
-void PlanarlyProjectVertex(const vec3f& unfoldedPosition, vec4f& projectedPosition,
-	const PROJECTOR& projector, const vec3f& projectorPosition, const vec3f& projectorDirection, const vec3f& s, const vec3f& u)
+void PlanarlyProjectVertex(const FVector3f& unfoldedPosition, vec4f& projectedPosition,
+	const FProjector& projector, const FVector3f& projectorPosition, const FVector3f& projectorDirection, const FVector3f& s, const FVector3f& u)
 {
-	float x = dot( unfoldedPosition - projectorPosition, s ) / projector.scale[0] + 0.5f;
-	float y = dot( unfoldedPosition - projectorPosition, u ) / projector.scale[1] + 0.5f;
+	float x = FVector3f::DotProduct( unfoldedPosition - projectorPosition, s ) / projector.scale[0] + 0.5f;
+	float y = FVector3f::DotProduct( unfoldedPosition - projectorPosition, u ) / projector.scale[1] + 0.5f;
 	y = 1.0f - y;
-	float z = dot( unfoldedPosition - projectorPosition, projectorDirection ) / projector.scale[2];
+	float z = FVector3f::DotProduct( unfoldedPosition - projectorPosition, projectorDirection ) / projector.scale[2];
 
 	bool inside = x>=0.0f && x<=1.0f && y>=0.0f && y<=1.0 && z>=0.0f && z<=1.0f;
 	projectedPosition[0] = x;
@@ -1432,57 +1432,57 @@ void PlanarlyProjectVertex(const vec3f& unfoldedPosition, vec4f& projectedPositi
 }
 
 
-vec2f ChangeBase2D(const vec2f& origPosition, const vec2f& newOrigin, const vec2f& newBaseX, const vec2f& newBaseY)
+FVector2f ChangeBase2D(const FVector2f& origPosition, const FVector2f& newOrigin, const FVector2f& newBaseX, const FVector2f& newBaseY)
 {
-	float x = dot(origPosition - newOrigin, newBaseX) / powf(length(newBaseX), 2.f) + 0.5f;
-	float y = dot(origPosition - newOrigin, newBaseY) / powf(length(newBaseY), 2.f) + 0.5f;
+	float x = FVector2f::DotProduct(origPosition - newOrigin, newBaseX) / powf(newBaseX.Length(), 2.f) + 0.5f;
+	float y = FVector2f::DotProduct(origPosition - newOrigin, newBaseY) / powf(newBaseY.Length(), 2.f) + 0.5f;
 	//y = 1.0f - y;
 
-	return vec2f(x, y);
+	return FVector2f(x, y);
 }
 
 
 // Compute barycentric coordinates (u, v, w) for
 // point p with respect to triangle (a, b, c)
-void GetBarycentricCoords(vec3f p, vec3f a, vec3f b, vec3f c, float &u, float &v, float &w)
+void GetBarycentricCoords(FVector3f p, FVector3f a, FVector3f b, FVector3f c, float &u, float &v, float &w)
 {
-    vec3f v0 = b - a, v1 = c - a, v2 = p - a;
-    float d00 = dot(v0, v0);
-    float d01 = dot(v0, v1);
-    float d11 = dot(v1, v1);
-    float d20 = dot(v2, v0);
-    float d21 = dot(v2, v1);
+    FVector3f v0 = b - a, v1 = c - a, v2 = p - a;
+    float d00 = FVector3f::DotProduct(v0, v0);
+    float d01 = FVector3f::DotProduct(v0, v1);
+    float d11 = FVector3f::DotProduct(v1, v1);
+    float d20 = FVector3f::DotProduct(v2, v0);
+    float d21 = FVector3f::DotProduct(v2, v1);
     float denom = d00 * d11 - d01 * d01;
     v = (d11 * d20 - d01 * d21) / denom;
     w = (d00 * d21 - d01 * d20) / denom;
     u = 1.0f - v - w;
 }
 
-void GetBarycentricCoords(vec2f p, vec2f a, vec2f b, vec2f c, float &u, float &v, float &w)
+void GetBarycentricCoords(FVector2f p, FVector2f a, FVector2f b, FVector2f c, float &u, float &v, float &w)
 {
-	vec2f v0 = b - a, v1 = c - a, v2 = p - a;
-    float den = v0.x() * v1.y() - v1.x() * v0.y();
+	FVector2f v0 = b - a, v1 = c - a, v2 = p - a;
+    float den = v0.X * v1.Y - v1.X * v0.Y;
 #ifdef DEBUG_PROJECTION
 	assert(den != 0.f);
 #endif
-    v = (v2.x() * v1.y() - v1.x() * v2.y()) / den;
-    w = (v0.x() * v2.y() - v2.x() * v0.y()) / den;
+    v = (v2.X * v1.Y - v1.X * v2.Y) / den;
+    w = (v0.X * v2.Y - v2.X * v0.Y) / den;
     u = 1.0f - v - w;
 }
 
 
-float getTriangleArea(vec2f& a, vec2f& b, vec2f& c)
+float getTriangleArea(FVector2f& a, FVector2f& b, FVector2f& c)
 {
-	float area = a.x() * (b.y() - c.y()) + b.x() * (c.y() - a.y()) + c.x() * (a.y() - b.y());
+	float area = a.X * (b.Y - c.Y) + b.X * (c.Y - a.Y) + c.X * (a.Y - b.Y);
 	return area / 2.f;
 }
 
 
-float getTriangleRatio(const vec2f& a, const vec2f& b, const vec2f& c)
+float getTriangleRatio(const FVector2f& a, const FVector2f& b, const FVector2f& c)
 {
-	float lenSide1 = length(b - a);
-	float lenSide2 = length(c - a);
-	float lenSide3 = length(b - c);
+	float lenSide1 = (b - a).Length();
+	float lenSide2 = (c - a).Length();
+	float lenSide3 = (b - c).Length();
 
 	float maxLen = FMath::Max3(lenSide1, lenSide2, lenSide3);
 	float minLen = FMath::Min3(lenSide1, lenSide2, lenSide3);
@@ -1518,59 +1518,59 @@ void getEdgeHorizontalLength( int edgeVert0, int edgeVert1, int opposedVert,
 {
 	const int oldVertices[2] = { edgeVert0, edgeVert1 };
 
-	const vec2f& edge0_oldUVSpace = pVertices[oldVertices[0]].uv;
-	const vec2f& edge1_oldUVSpace = pVertices[oldVertices[1]].uv;
-	const vec2f& opposedVert_oldUVSpace = pVertices[opposedVert].uv;
+	const FVector2f& edge0_oldUVSpace = pVertices[oldVertices[0]].uv;
+	const FVector2f& edge1_oldUVSpace = pVertices[oldVertices[1]].uv;
+	const FVector2f& opposedVert_oldUVSpace = pVertices[opposedVert].uv;
 
-	vec2f edgeVector_oldUVSpace = edge1_oldUVSpace - edge0_oldUVSpace;
-	vec2f sideVector_oldUVSpace = opposedVert_oldUVSpace - edge0_oldUVSpace;
+	FVector2f edgeVector_oldUVSpace = edge1_oldUVSpace - edge0_oldUVSpace;
+	FVector2f sideVector_oldUVSpace = opposedVert_oldUVSpace - edge0_oldUVSpace;
 
-	float edgeVector_oldUVSpaceLen = length(edgeVector_oldUVSpace);
-	float dotEdgeSideVectors_oldUVSpace = dot(edgeVector_oldUVSpace, sideVector_oldUVSpace);
+	float edgeVector_oldUVSpaceLen = edgeVector_oldUVSpace.Length();
+	float dotEdgeSideVectors_oldUVSpace = FVector2f::DotProduct(edgeVector_oldUVSpace, sideVector_oldUVSpace);
 	float midEdgePointFraction_oldUVSpace = (dotEdgeSideVectors_oldUVSpace / powf(edgeVector_oldUVSpaceLen, 2));
-	vec2f edgeVectorProj_oldUVSpace = edgeVector_oldUVSpace * midEdgePointFraction_oldUVSpace;
-	vec2f midEdgePoint_oldUVSpace = edge0_oldUVSpace + edgeVectorProj_oldUVSpace;
+	FVector2f edgeVectorProj_oldUVSpace = edgeVector_oldUVSpace * midEdgePointFraction_oldUVSpace;
+	FVector2f midEdgePoint_oldUVSpace = edge0_oldUVSpace + edgeVectorProj_oldUVSpace;
 
-	vec2f perpEdgeVector_oldUVSpace = opposedVert_oldUVSpace - midEdgePoint_oldUVSpace;
+	FVector2f perpEdgeVector_oldUVSpace = opposedVert_oldUVSpace - midEdgePoint_oldUVSpace;
 #ifdef DEBUG_PROJECTION
 	assert(fabs(dot(perpEdgeVector_oldUVSpace, edgeVector_oldUVSpace)) < 0.1f);
 #endif
-	float perpEdgeVector_oldUVSpace_Len = length(perpEdgeVector_oldUVSpace);
+	float perpEdgeVector_oldUVSpace_Len = perpEdgeVector_oldUVSpace.Length();
 	out_uvSpaceLen = perpEdgeVector_oldUVSpace_Len;
 	out_midEdgePointFraction = midEdgePointFraction_oldUVSpace;
 
 	// Do the same in object space to be able to extract the scale of the original uv space
-	const vec3f& edge0_objSpace = pVertices[oldVertices[0]].pos;
-	const vec3f& edge1_objSpace = pVertices[oldVertices[1]].pos;
-	const vec3f& opposedVert_objSpace = pVertices[opposedVert].pos;
+	const FVector3f& edge0_objSpace = pVertices[oldVertices[0]].pos;
+	const FVector3f& edge1_objSpace = pVertices[oldVertices[1]].pos;
+	const FVector3f& opposedVert_objSpace = pVertices[opposedVert].pos;
 
-	vec3f edgeVector_objSpace = edge1_objSpace - edge0_objSpace;
-	vec3f sideVector_objSpace = opposedVert_objSpace - edge0_objSpace;
+	FVector3f edgeVector_objSpace = edge1_objSpace - edge0_objSpace;
+	FVector3f sideVector_objSpace = opposedVert_objSpace - edge0_objSpace;
 
-	float edgeVector_objSpaceLen = length(edgeVector_objSpace);
-	float dotEdgeSideVectors_objSpace = dot(edgeVector_objSpace, sideVector_objSpace);
+	float edgeVector_objSpaceLen = edgeVector_objSpace.Length();
+	float dotEdgeSideVectors_objSpace = FVector3f::DotProduct(edgeVector_objSpace, sideVector_objSpace);
 	float midEdgePointFraction_objSpace = (dotEdgeSideVectors_objSpace / powf(edgeVector_objSpaceLen, 2));
-	vec3f edgeVectorProj_objSpace = edgeVector_objSpace * midEdgePointFraction_objSpace;
-	vec3f midEdgePoint_objSpace = edge0_objSpace + edgeVectorProj_objSpace;
+	FVector3f edgeVectorProj_objSpace = edgeVector_objSpace * midEdgePointFraction_objSpace;
+	FVector3f midEdgePoint_objSpace = edge0_objSpace + edgeVectorProj_objSpace;
 
-	vec3f perpEdgeVector_objSpace = opposedVert_objSpace - midEdgePoint_objSpace;
+	FVector3f perpEdgeVector_objSpace = opposedVert_objSpace - midEdgePoint_objSpace;
 #ifdef DEBUG_PROJECTION
 	assert(fabs(dot(perpEdgeVector_objSpace, edgeVector_objSpace)) < 0.1f);
 #endif
-	float perpEdgeVector_objSpace_Len = length(perpEdgeVector_objSpace);
+	float perpEdgeVector_objSpace_Len = perpEdgeVector_objSpace.Length();
 	out_objSpaceLen = perpEdgeVector_objSpace_Len;
 }
 
 
-bool testPointsAreInOppositeSidesOfEdge( const vec2f& pointA, const vec2f& pointB,
-                                         const vec2f& edge0, const vec2f& edge1 )
+bool testPointsAreInOppositeSidesOfEdge( const FVector2f& pointA, const FVector2f& pointB,
+                                         const FVector2f& edge0, const FVector2f& edge1 )
 {
-	vec2f edge = edge1 - edge0;
-	vec2f perp_edge = vec2f(-edge.y(), edge.x());
-	vec2f AVertexVector = pointA - edge0;
-	vec2f BVertexVector = pointB - edge0;
-	float dotAVertexVector = dot(AVertexVector, perp_edge);
-	float dotBVertexVector = dot(BVertexVector, perp_edge);
+	FVector2f edge = edge1 - edge0;
+	FVector2f perp_edge = FVector2f(-edge.Y, edge.X);
+	FVector2f AVertexVector = pointA - edge0;
+	FVector2f BVertexVector = pointB - edge0;
+	float dotAVertexVector = FVector2f::DotProduct(AVertexVector, perp_edge);
+	float dotBVertexVector = FVector2f::DotProduct(BVertexVector, perp_edge);
 	
 	return dotAVertexVector * dotBVertexVector < 0.f;
 }
@@ -1583,7 +1583,7 @@ struct PROJECTED_VERTEX
     float pos0 = 0, pos1 = 0, pos2 = 0;
     uint32 mask3 = 0;
 
-	inline vec2<float> xy() const { return vec2<float>(pos0,pos1); }
+	inline FVector2f xy() const { return FVector2f(pos0,pos1); }
 };
 #pragma pack(pop)
 static_assert( sizeof(PROJECTED_VERTEX)==16, "Unexpected struct size" );
@@ -1592,9 +1592,9 @@ static_assert( sizeof(PROJECTED_VERTEX)==16, "Unexpected struct size" );
 //-------------------------------------------------------------------------------------------------
 void MeshProject_Optimised_Planar( const OPTIMISED_VERTEX* pVertices, int vertexCount,
                                    const uint32* pIndices, int faceCount,
-                                   const vec3f& projectorPosition, const vec3f& projectorDirection,
-                                   const vec3f& projectorSide, const vec3f& projectorUp,
-                                   const vec3f& projectorScale,
+                                   const FVector3f& projectorPosition, const FVector3f& projectorDirection,
+                                   const FVector3f& projectorSide, const FVector3f& projectorUp,
+                                   const FVector3f& projectorScale,
                                    OPTIMISED_VERTEX* pResultVertices, int& currentVertex,
                                    uint32* pResultIndices, int& currentIndex  )
 {
@@ -1608,10 +1608,10 @@ void MeshProject_Optimised_Planar( const OPTIMISED_VERTEX* pVertices, int vertex
 
     for ( int v=0; v<vertexCount; ++v )
     {
-        float x = dot( pVertices[v].pos-projectorPosition, projectorSide ) / projectorScale[0] + 0.5f;
-        float y = dot( pVertices[v].pos-projectorPosition, projectorUp ) / projectorScale[1] + 0.5f;
+        float x = FVector3f::DotProduct( pVertices[v].pos-projectorPosition, projectorSide ) / projectorScale[0] + 0.5f;
+        float y = FVector3f::DotProduct( pVertices[v].pos-projectorPosition, projectorUp ) / projectorScale[1] + 0.5f;
         y = 1.0f-y;
-        float z = dot( pVertices[v].pos-projectorPosition, projectorDirection ) / projectorScale[2];
+        float z = FVector3f::DotProduct( pVertices[v].pos-projectorPosition, projectorDirection ) / projectorScale[2];
 
         // Plane mask with bits for each plane discarding the vertex
         uint32 planeMask =
@@ -1658,7 +1658,7 @@ void MeshProject_Optimised_Planar( const OPTIMISED_VERTEX* pVertices, int vertex
                     pResultVertices[currentVertex].pos[2] = projectedPositions[i].pos2;
 
                     // Normal is actually the fade factor
-                    float angleCos = dot( pVertices[i].nor, projectorDirection * -1.0f );
+                    float angleCos = FVector3f::DotProduct( pVertices[i].nor, projectorDirection * -1.0f );
                     pResultVertices[currentVertex].nor[0] = angleCos;
                     pResultVertices[currentVertex].nor[1] = angleCos;
                     pResultVertices[currentVertex].nor[2] = angleCos;
@@ -1676,9 +1676,9 @@ void MeshProject_Optimised_Planar( const OPTIMISED_VERTEX* pVertices, int vertex
 //-------------------------------------------------------------------------------------------------
 void MeshProject_Optimised_Cylindrical( const OPTIMISED_VERTEX* pVertices, int vertexCount,
                                         const uint32* pIndices, int faceCount,
-                                        const vec3f& projectorPosition, const vec3f& projectorDirection,
-                                        const vec3f& projectorSide, const vec3f& projectorUp,
-                                        const vec3f& projectorScale,
+                                        const FVector3f& projectorPosition, const FVector3f& projectorDirection,
+                                        const FVector3f& projectorSide, const FVector3f& projectorUp,
+                                        const FVector3f& projectorScale,
                                         OPTIMISED_VERTEX* pResultVertices, int& currentVertex,
                                         uint32* pResultIndices, int& currentIndex  )
 {
@@ -1693,7 +1693,7 @@ void MeshProject_Optimised_Cylindrical( const OPTIMISED_VERTEX* pVertices, int v
     // TODO: support for non uniform scale?
     float radius = projectorScale[1];
     float height = projectorScale[0];
-    mat3f worldToCylinder( projectorDirection, projectorSide, projectorUp );
+    mat3f worldToCylinder = mat3f(vec3f(projectorDirection), vec3f(projectorSide), vec3f(projectorUp) );
     //worldToCylinder = worldToCylinder.GetTransposed();
 
     for ( int v=0; v<vertexCount; ++v )
@@ -1701,8 +1701,8 @@ void MeshProject_Optimised_Cylindrical( const OPTIMISED_VERTEX* pVertices, int v
         // Cylinder is along the X axis
 
         // Project
-        vec3f relPos = pVertices[v].pos - projectorPosition;
-        vec3f vertexPos_cylinder = worldToCylinder * relPos;
+		vec3f relPos = vec3f(pVertices[v].pos - projectorPosition);
+        vec3f vertexPos_cylinder = worldToCylinder.Transform(relPos);
 
         // This final projection needs to be done per pixel
         float x = vertexPos_cylinder.x() / height;
@@ -1770,9 +1770,9 @@ void MeshProject_Optimised_Cylindrical( const OPTIMISED_VERTEX* pVertices, int v
 
 //-------------------------------------------------------------------------------------------------
 void MeshProject_Optimised_Wrapping( const Mesh* pMesh,
-                                     const vec3f& projectorPosition, const vec3f& projectorDirection,
-                                     const vec3f& projectorSide, const vec3f& projectorUp,
-                                     const vec3f& projectorScale,
+                                     const FVector3f& projectorPosition, const FVector3f& projectorDirection,
+                                     const FVector3f& projectorSide, const FVector3f& projectorUp,
+                                     const FVector3f& projectorScale,
                                      OPTIMISED_VERTEX_WRAPPING* pResultVertices, int& currentVertex,
                                      uint32* pResultIndices, int& currentIndex )
 {
@@ -1797,8 +1797,8 @@ void MeshProject_Optimised_Wrapping( const Mesh* pMesh,
     float rayLength = maxDist;
     int intersectedFace = -1;
 	TSet<int> processedVertices;
-    vec3f projectionPlaneNormal;
-    vec3f out_intersection;
+    FVector3f projectionPlaneNormal;
+    FVector3f out_intersection;
 
 	TArray<AdjacentFaces> faceConnectivity;
 	faceConnectivity.SetNum(faceCount);
@@ -1809,7 +1809,7 @@ void MeshProject_Optimised_Wrapping( const Mesh* pMesh,
 
     // Map vertices to the one they are collapsed to because they are very similar, if they aren't collapsed then they are mapped to themselves
     TArray<int> collapsedVertexMap;
-	TArray<vec3f> vertices;
+	TArray<FVector3f> vertices;
     TMultiMap<int, int> collapsedVertsMap; // Maps a collapsed vertex to all the vertices that collapse to it
     MeshCreateCollapsedVertexMap(pMesh, collapsedVertexMap, vertices, collapsedVertsMap);
 
@@ -1838,13 +1838,13 @@ void MeshProject_Optimised_Wrapping( const Mesh* pMesh,
         int i1 = pIndices[f * 3 + 1];
         int i2 = pIndices[f * 3 + 2];
 
-        vec3f rayStart = projectorPosition;
-        vec3f rayEnd = projectorPosition + projectorDirection * maxDist;
-        vec3f aux_out_intersection;
+		FVector3f rayStart = projectorPosition;
+		FVector3f rayEnd = projectorPosition + projectorDirection * maxDist;
+		FVector3f aux_out_intersection;
         int out_intersected_vert, out_intersected_edge_v0, out_intersected_edge_v1;
         float t;
 
-        rayLength = length(rayEnd - rayStart);
+        rayLength = (rayEnd - rayStart).Length();
 
         bool intersects = rayIntersectsFace2(rayStart, rayEnd, pVertices[i0].pos, pVertices[i1].pos, pVertices[i2].pos,
             aux_out_intersection, out_intersected_vert, out_intersected_edge_v0, out_intersected_edge_v1, t);
@@ -1855,10 +1855,10 @@ void MeshProject_Optimised_Wrapping( const Mesh* pMesh,
             min_t = t;
             out_intersection = aux_out_intersection;
 
-            vec3<float> v0 = pVertices[i0].pos;
-            vec3<float> v1 = pVertices[i1].pos;
-            vec3<float> v2 = pVertices[i2].pos;
-            projectionPlaneNormal = normalise(cross<float>(v1 - v0, v2 - v0));
+            FVector3f v0 = pVertices[i0].pos;
+			FVector3f v1 = pVertices[i1].pos;
+			FVector3f v2 = pVertices[i2].pos;
+            projectionPlaneNormal = FVector3f::CrossProduct(v1 - v0, v2 - v0).GetSafeNormal();
         }
 
         // Face connectivity info
@@ -1928,27 +1928,27 @@ void MeshProject_Optimised_Wrapping( const Mesh* pMesh,
     }
 
     // New projector located perpendicularly up from the hit face
-    PROJECTOR projector2;
+    FProjector projector2;
     projector2.direction[0] = projectionPlaneNormal[0];
     projector2.direction[1] = projectionPlaneNormal[1];
     projector2.direction[2] = projectionPlaneNormal[2];
-    vec3f auxPosition = out_intersection - projectionPlaneNormal * rayLength * min_t;
-    //vec3f auxPosition = vec3f( projector.position[0], projector.position[1], projector.position[2] );
+    FVector3f auxPosition = out_intersection - projectionPlaneNormal * rayLength * min_t;
+    //FVector3f auxPosition = FVector3f( projector.position[0], projector.position[1], projector.position[2] );
     projector2.position[0] = auxPosition[0];
     projector2.position[1] = auxPosition[1];
     projector2.position[2] = auxPosition[2];
 
-    vec3f auxUp = projectorUp;
-    float test = fabs(dot(auxUp, projectionPlaneNormal));
+    FVector3f auxUp = projectorUp;
+    float test = fabs(FVector3f::DotProduct(auxUp, projectionPlaneNormal));
 
     if (test < 0.9f)
     {
-        vec3f auxSide = normalise(cross(projectionPlaneNormal, auxUp));
-        auxUp = cross(auxSide, projectionPlaneNormal);
+        FVector3f auxSide = FVector3f::CrossProduct(projectionPlaneNormal, auxUp).GetSafeNormal();
+        auxUp = FVector3f::CrossProduct(auxSide, projectionPlaneNormal);
     }
     else
     {
-        auxUp = cross(projectionPlaneNormal, projectorSide);
+        auxUp = FVector3f::CrossProduct(projectionPlaneNormal, projectorSide);
     }
 
     projector2.up[0] = auxUp[0];
@@ -1958,8 +1958,8 @@ void MeshProject_Optimised_Wrapping( const Mesh* pMesh,
     projector2.scale[0] = projectorScale[0];
     projector2.scale[1] = projectorScale[1];
     projector2.scale[2] = projectorScale[2];
-    vec3f projectorPosition2, projectorDirection2, s2, u2;
-    projectorPosition2 = vec3f( projector2.position[0], projector2.position[1], projector2.position[2] );
+    FVector3f projectorDirection2, s2, u2;
+    FVector3f projectorPosition2 = projector2.position;
     projector2.GetDirectionSideUp( projectorDirection2, s2, u2 );
     float maxDistSquared = powf(projector2.scale[0], 2.f) + powf(projector2.scale[1], 2.f);
 
@@ -1971,7 +1971,7 @@ void MeshProject_Optimised_Wrapping( const Mesh* pMesh,
 		// \TODO: Bool array to speed up?
 		TSet<int> pendingFacesUnique; // Used to quickly check uniqueness in the pendingFaces queue
 
-        //vec3f hitFaceProjectedNormal;
+        //FVector3f hitFaceProjectedNormal;
         bool hitFaceHasPositiveArea = false;
 
         NeighborFace neighborFace;
@@ -2010,47 +2010,47 @@ void MeshProject_Optimised_Wrapping( const Mesh* pMesh,
                 assert(newVertexFromQueue == -1);
 #endif
 
-                vec3f outIntersectionBaricentric;
+                FVector3f outIntersectionBaricentric;
                 int i0 = pIndices[currentFace * 3 + 0];
                 int i1 = pIndices[currentFace * 3 + 1];
                 int i2 = pIndices[currentFace * 3 + 2];
-                vec3<float> v3_0 = pVertices[i0].pos;
-                vec3<float> v3_1 = pVertices[i1].pos;
-                vec3<float> v3_2 = pVertices[i2].pos;
+				FVector3f v3_0 = pVertices[i0].pos;
+				FVector3f v3_1 = pVertices[i1].pos;
+				FVector3f v3_2 = pVertices[i2].pos;
                 GetBarycentricCoords(out_intersection, v3_0, v3_1, v3_2, outIntersectionBaricentric[0], outIntersectionBaricentric[1], outIntersectionBaricentric[2]);
-                vec2f outIntersectionUV = pVertices[i0].uv * outIntersectionBaricentric[0] + pVertices[i1].uv * outIntersectionBaricentric[1] + pVertices[i2].uv * outIntersectionBaricentric[2];
+				FVector2f outIntersectionUV = pVertices[i0].uv * outIntersectionBaricentric[0] + pVertices[i1].uv * outIntersectionBaricentric[1] + pVertices[i2].uv * outIntersectionBaricentric[2];
 
                 vec4f proj4D_v0, proj4D_v1, proj4D_v2;
                 PlanarlyProjectVertex(pVertices[i0].pos, proj4D_v0, projector2, projectorPosition2, projectorDirection2, s2, u2);
                 PlanarlyProjectVertex(pVertices[i1].pos, proj4D_v1, projector2, projectorPosition2, projectorDirection2, s2, u2);
                 PlanarlyProjectVertex(pVertices[i2].pos, proj4D_v2, projector2, projectorPosition2, projectorDirection2, s2, u2);
 
-                vec2f proj_v0 = vec2f(proj4D_v0[0], proj4D_v0[1]);
-                vec2f proj_v1 = vec2f(proj4D_v1[0], proj4D_v1[1]);
-                vec2f proj_v2 = vec2f(proj4D_v2[0], proj4D_v2[1]);
+				FVector2f proj_v0 = FVector2f(proj4D_v0[0], proj4D_v0[1]);
+				FVector2f proj_v1 = FVector2f(proj4D_v1[0], proj4D_v1[1]);
+				FVector2f proj_v2 = FVector2f(proj4D_v2[0], proj4D_v2[1]);
 
                 float proj_TriangleArea = getTriangleArea(proj_v0, proj_v1, proj_v2);
 
-                vec3f BaricentricBaseU;
-                vec3f BaricentricBaseV;
-                vec3f BaricentricOrigin;
+                FVector3f BaricentricBaseU;
+                FVector3f BaricentricBaseV;
+                FVector3f BaricentricOrigin;
 
-                GetBarycentricCoords(vec2f(1.f, 0.f), proj_v0, proj_v1, proj_v2, BaricentricBaseU[0], BaricentricBaseU[1], BaricentricBaseU[2]);
-                GetBarycentricCoords(vec2f(0.f, 1.f), proj_v0, proj_v1, proj_v2, BaricentricBaseV[0], BaricentricBaseV[1], BaricentricBaseV[2]);
-                GetBarycentricCoords(vec2f(0.f, 0.f), proj_v0, proj_v1, proj_v2, BaricentricOrigin[0], BaricentricOrigin[1], BaricentricOrigin[2]);
+                GetBarycentricCoords(FVector2f(1.f, 0.f), proj_v0, proj_v1, proj_v2, BaricentricBaseU[0], BaricentricBaseU[1], BaricentricBaseU[2]);
+                GetBarycentricCoords(FVector2f(0.f, 1.f), proj_v0, proj_v1, proj_v2, BaricentricBaseV[0], BaricentricBaseV[1], BaricentricBaseV[2]);
+                GetBarycentricCoords(FVector2f(0.f, 0.f), proj_v0, proj_v1, proj_v2, BaricentricOrigin[0], BaricentricOrigin[1], BaricentricOrigin[2]);
 
-                vec2f OrigUVs_BaseU_EndPoint = pVertices[i0].uv * BaricentricBaseU[0] + pVertices[i1].uv * BaricentricBaseU[1] + pVertices[i2].uv * BaricentricBaseU[2];
-                vec2f OrigUVs_BaseV_EndPoint = pVertices[i0].uv * BaricentricBaseV[0] + pVertices[i1].uv * BaricentricBaseV[1] + pVertices[i2].uv * BaricentricBaseV[2];
-                vec2f OrigUVs_Origin = pVertices[i0].uv * BaricentricOrigin[0] + pVertices[i1].uv * BaricentricOrigin[1] + pVertices[i2].uv * BaricentricOrigin[2];
+                FVector2f OrigUVs_BaseU_EndPoint = pVertices[i0].uv * BaricentricBaseU[0] + pVertices[i1].uv * BaricentricBaseU[1] + pVertices[i2].uv * BaricentricBaseU[2];
+                FVector2f OrigUVs_BaseV_EndPoint = pVertices[i0].uv * BaricentricBaseV[0] + pVertices[i1].uv * BaricentricBaseV[1] + pVertices[i2].uv * BaricentricBaseV[2];
+                FVector2f OrigUVs_Origin = pVertices[i0].uv * BaricentricOrigin[0] + pVertices[i1].uv * BaricentricOrigin[1] + pVertices[i2].uv * BaricentricOrigin[2];
 
-                vec2f OrigUVs_BaseU = OrigUVs_BaseU_EndPoint - OrigUVs_Origin;
-                vec2f OrigUVs_BaseV_Aux = OrigUVs_BaseV_EndPoint - OrigUVs_Origin;
+                FVector2f OrigUVs_BaseU = OrigUVs_BaseU_EndPoint - OrigUVs_Origin;
+                FVector2f OrigUVs_BaseV_Aux = OrigUVs_BaseV_EndPoint - OrigUVs_Origin;
 
-                OrigUVs_BaseU = normalise(OrigUVs_BaseU);
-                OrigUVs_BaseV_Aux = normalise(OrigUVs_BaseV_Aux);
+                OrigUVs_BaseU = OrigUVs_BaseU.GetSafeNormal();
+                OrigUVs_BaseV_Aux = OrigUVs_BaseV_Aux.GetSafeNormal();
 
                 // Generate perp vector from OrigUVs_BaseU by flipping coords and one sign. Using OrigUVs_BaseV directly sometimes doesn't give an orthogonal basis
-                vec2f OrigUVs_BaseV = vec2f(copysign(OrigUVs_BaseU.y(), OrigUVs_BaseV_Aux.x()), copysign(OrigUVs_BaseU.x(), OrigUVs_BaseV_Aux.y()));
+                FVector2f OrigUVs_BaseV = FVector2f(copysign(OrigUVs_BaseU.Y, OrigUVs_BaseV_Aux.X), copysign(OrigUVs_BaseU.X, OrigUVs_BaseV_Aux.Y));
 
 #ifdef DEBUG_PROJECTION
                 assert(fabs(dot(OrigUVs_BaseU, OrigUVs_BaseV)) < 0.3f);
@@ -2071,7 +2071,7 @@ void MeshProject_Optimised_Wrapping( const Mesh* pMesh,
                     processedVertices.Add(collapsedVert);
 
                     //PlanarlyProjectVertex(pVertices[collapsedVert].pos, projectedPositions[collapsedVert], projector2, projectorPosition2, projectorDirection2, s2, u2);
-                    vec2f projectedVertex = ChangeBase2D(pVertices[v].uv, OrigUVs_Origin, OrigUVs_BaseU, OrigUVs_BaseV);
+                    FVector2f projectedVertex = ChangeBase2D(pVertices[v].uv, OrigUVs_Origin, OrigUVs_BaseU, OrigUVs_BaseV);
 
                     projectedPositions[collapsedVert].pos0 = projectedVertex[0];
                     projectedPositions[collapsedVert].pos1 = projectedVertex[1];
@@ -2079,15 +2079,15 @@ void MeshProject_Optimised_Wrapping( const Mesh* pMesh,
                     projectedPositions[collapsedVert].mask3 = 1;
                 }
 
-                vec2<float> v0 = projectedPositions[collapsedVertexMap[pIndices[currentFace * 3 + 0]]].xy();
-                vec2<float> v1 = projectedPositions[collapsedVertexMap[pIndices[currentFace * 3 + 1]]].xy();
-                vec2<float> v2 = projectedPositions[collapsedVertexMap[pIndices[currentFace * 3 + 2]]].xy();
+                FVector2f v0 = projectedPositions[collapsedVertexMap[pIndices[currentFace * 3 + 0]]].xy();
+				FVector2f v1 = projectedPositions[collapsedVertexMap[pIndices[currentFace * 3 + 1]]].xy();
+				FVector2f v2 = projectedPositions[collapsedVertexMap[pIndices[currentFace * 3 + 2]]].xy();
                 currentTriangleArea = getTriangleArea(v0, v1, v2);
                 float triangleAreaFactor = sqrt(fabs(proj_TriangleArea) / fabs(currentTriangleArea));
 
-                vec2f origin = v0 * outIntersectionBaricentric[0] + v1 * outIntersectionBaricentric[1] + v2 * outIntersectionBaricentric[2];
+                FVector2f origin = v0 * outIntersectionBaricentric[0] + v1 * outIntersectionBaricentric[1] + v2 * outIntersectionBaricentric[2];
 #ifdef DEBUG_PROJECTION
-                assert(length(origin - vec2f(0.5f, 0.5f)) < 0.001f);
+                assert(length(origin - FVector2f(0.5f, 0.5f)) < 0.001f);
 #endif
 
                 //hitFaceProjectedNormal = (cross<float>(v1 - v0, v2 - v0));
@@ -2124,7 +2124,7 @@ void MeshProject_Optimised_Wrapping( const Mesh* pMesh,
 #endif
                 origin = v0 * outIntersectionBaricentric[0] + v1 * outIntersectionBaricentric[1] + v2 * outIntersectionBaricentric[2];
 #ifdef DEBUG_PROJECTION
-                assert(length(origin - vec2f(0.5f, 0.5f)) < 0.001f);
+                assert(length(origin - FVector2f(0.5f, 0.5f)) < 0.001f);
 #endif
             }
             else
@@ -2188,14 +2188,14 @@ void MeshProject_Optimised_Wrapping( const Mesh* pMesh,
                     int pi2 = pIndices[previousFace * 3 + 2];
 
                     // Previous face uv coords
-                    vec2<float> pv0;
-                    vec2<float> pv1;
-                    vec2<float> pv2;
+					FVector2f pv0;
+					FVector2f pv1;
+					FVector2f pv2;
 
                     // Proj vertices from previous triangle
-                    vec2<float> pv0_proj;
-                    vec2<float> pv1_proj;
-                    vec2<float> pv2_proj;
+					FVector2f pv0_proj;
+					FVector2f pv1_proj;
+					FVector2f pv2_proj;
 
                     int oldVertices_previousFace[2] = { -1, -1 };
                     int oldVertex = -1;
@@ -2238,18 +2238,18 @@ void MeshProject_Optimised_Wrapping( const Mesh* pMesh,
                         float objSpaceToCurrentUVSpaceConversion = currentFace_uvSpaceLen / currentFace_objSpaceLen;
                         float previousFaceWidthInCurrentUVspace = previousFace_objSpaceLen * objSpaceToCurrentUVSpaceConversion;
 
-                        const vec2f& edge0_currentFace_oldUVSpace = pVertices[oldVertices[0]].uv;
-                        const vec2f& edge1_currentFace_oldUVSpace = pVertices[oldVertices[1]].uv;
-                        const vec2f& newVertex_currentFace_oldUVSpace = pVertices[newVertex].uv;
+                        const FVector2f& edge0_currentFace_oldUVSpace = pVertices[oldVertices[0]].uv;
+                        const FVector2f& edge1_currentFace_oldUVSpace = pVertices[oldVertices[1]].uv;
+                        const FVector2f& newVertex_currentFace_oldUVSpace = pVertices[newVertex].uv;
 
-                        vec2f edge_currentFace_oldUVSpace = edge1_currentFace_oldUVSpace - edge0_currentFace_oldUVSpace;
-                        vec2f midpoint_currentFace_oldUVSpace = edge0_currentFace_oldUVSpace + edge_currentFace_oldUVSpace * midEdgePointFraction_currentFace;
-                        vec2f edgePerpVectorNorm = normalise(midpoint_currentFace_oldUVSpace - newVertex_currentFace_oldUVSpace);
-                        vec2f test_edge_currentFace_oldUVSpaceNorm = normalise(edge_currentFace_oldUVSpace);
+                        FVector2f edge_currentFace_oldUVSpace = edge1_currentFace_oldUVSpace - edge0_currentFace_oldUVSpace;
+                        FVector2f midpoint_currentFace_oldUVSpace = edge0_currentFace_oldUVSpace + edge_currentFace_oldUVSpace * midEdgePointFraction_currentFace;
+                        FVector2f edgePerpVectorNorm = (midpoint_currentFace_oldUVSpace - newVertex_currentFace_oldUVSpace).GetSafeNormal();
+                        FVector2f test_edge_currentFace_oldUVSpaceNorm = edge_currentFace_oldUVSpace.GetSafeNormal();
                         (void)test_edge_currentFace_oldUVSpaceNorm;
 
 #ifdef DEBUG_PROJECTION
-                        vec2f sideVector_oldUVSpace = newVertex_currentFace_oldUVSpace - edge0_currentFace_oldUVSpace;
+                        FVector2f sideVector_oldUVSpace = newVertex_currentFace_oldUVSpace - edge0_currentFace_oldUVSpace;
                         float edgeVector_oldUVSpaceLen = length(edge_currentFace_oldUVSpace);
                         float dotEdgeSideVectors_oldUVSpace = dot(edge_currentFace_oldUVSpace, sideVector_oldUVSpace);
                         float midEdgePointFraction_oldUVSpace = (dotEdgeSideVectors_oldUVSpace / powf(edgeVector_oldUVSpaceLen, 2));
@@ -2257,8 +2257,8 @@ void MeshProject_Optimised_Wrapping( const Mesh* pMesh,
                         assert(fabs(dot(edgePerpVectorNorm, test_edge_currentFace_oldUVSpaceNorm)) < 0.1f);
 #endif
 
-                        vec2f midpoint_previousFace_oldUVSpace = edge0_currentFace_oldUVSpace + edge_currentFace_oldUVSpace * midEdgePointFraction_previousFace;
-                        vec2f oldVertex_currentFace_oldUVSpace = midpoint_previousFace_oldUVSpace + edgePerpVectorNorm * previousFaceWidthInCurrentUVspace;
+                        FVector2f midpoint_previousFace_oldUVSpace = edge0_currentFace_oldUVSpace + edge_currentFace_oldUVSpace * midEdgePointFraction_previousFace;
+                        FVector2f oldVertex_currentFace_oldUVSpace = midpoint_previousFace_oldUVSpace + edgePerpVectorNorm * previousFaceWidthInCurrentUVspace;
 
                         pv0 = edge0_currentFace_oldUVSpace;
                         pv1 = edge1_currentFace_oldUVSpace;
@@ -2284,14 +2284,14 @@ void MeshProject_Optimised_Wrapping( const Mesh* pMesh,
                         pv2_proj = projectedPositions[pi2].xy();
                     }
 
-                    vec2<float> v2 = pVertices[newVertex].uv;
+					FVector2f v2 = pVertices[newVertex].uv;
 
                     // New vertex baricentric coords in respect to old triangle
                     float a, b, c;
                     GetBarycentricCoords(v2, pv0, pv1, pv2, a, b, c);
-                    vec2f newVertex_proj = pv0_proj * a + pv1_proj * b + pv2_proj * c;
+                    FVector2f newVertex_proj = pv0_proj * a + pv1_proj * b + pv2_proj * c;
 #ifdef DEBUG_PROJECTION
-                    vec2f v2Test = pv0 * a + pv1 * b + pv2 * c;
+                    FVector2f v2Test = pv0 * a + pv1 * b + pv2 * c;
                     float v2TestDist = length(v2 - v2Test);
                     assert(v2TestDist < 0.001f);
 
@@ -2343,12 +2343,12 @@ void MeshProject_Optimised_Wrapping( const Mesh* pMesh,
                 int i1 = collapsedVertexMap[pIndices[currentFace * 3 + 1]];
                 int i2 = collapsedVertexMap[pIndices[currentFace * 3 + 2]];
 
-                vec2<float> v0 = projectedPositions[i0].xy();
-                vec2<float> v1 = projectedPositions[i1].xy();
-                vec2<float> v2 = projectedPositions[i2].xy();
+                FVector2f v0 = projectedPositions[i0].xy();
+                FVector2f v1 = projectedPositions[i1].xy();
+                FVector2f v2 = projectedPositions[i2].xy();
                 currentTriangleArea = getTriangleArea(v0, v1, v2);
                 bool currentTriangleHasPositiveArea = currentTriangleArea >= 0.f;
-                //vec3f currentFaceNormal = (cross<float>(v1 - v0, v2 - v0));
+                //FVector3f currentFaceNormal = (cross<float>(v1 - v0, v2 - v0));
 
                 //if (hitFaceProjectedNormal.z() * currentFaceNormal.z() < 0)
                 if(hitFaceHasPositiveArea != currentTriangleHasPositiveArea) // Is the current face wound in the opposite direction?
@@ -2375,8 +2375,8 @@ void MeshProject_Optimised_Wrapping( const Mesh* pMesh,
             for (int i = 0; i < 3; ++i)
             {
                 int v = pIndices[currentFace * 3 + i];
-                vec3f r = pVertices[v].pos - out_intersection;
-                float squaredDist = dot(r, r);
+                FVector3f r = pVertices[v].pos - out_intersection;
+                float squaredDist = FVector3f::DotProduct(r, r);
 
                 if (squaredDist <= maxDistSquared / 4.f) // Is it inside?
                 {
@@ -2529,9 +2529,9 @@ void MeshProject_Optimised_Wrapping( const Mesh* pMesh,
 
             outfile << "f " << i0 << "/" << i0 << " " << i1 << "/" << i1 << " " << i2 << "/" << i2 << "\n";
 
-            vec2f v0 = pResultVertices[i0].pos.xy();
-            vec2f v1 = pResultVertices[i1].pos.xy();
-            vec2f v2 = pResultVertices[i2].pos.xy();
+            FVector2f v0 = pResultVertices[i0].pos.xy();
+            FVector2f v1 = pResultVertices[i1].pos.xy();
+            FVector2f v2 = pResultVertices[i2].pos.xy();
 
             outfile << "# Triangle " << i << ", area = " << getTriangleArea(v0, v1, v2) << "\n";
         }
@@ -2547,14 +2547,14 @@ void MeshProject_Optimised_Wrapping( const Mesh* pMesh,
 
 //-------------------------------------------------------------------------------------------------
 MeshPtr MeshProject_Optimised( const Mesh* pMesh,
-                               const PROJECTOR& projector )
+                               const FProjector& projector )
 {
 	MUTABLE_CPUPROFILER_SCOPE(MeshProject_Optimised);
 
-    vec3f projectorPosition,projectorDirection,projectorSide,projectorUp;
-    projectorPosition = vec3f( projector.position[0], projector.position[1], projector.position[2] );
+    FVector3f projectorPosition,projectorDirection,projectorSide,projectorUp;
+    projectorPosition = FVector3f( projector.position[0], projector.position[1], projector.position[2] );
     projector.GetDirectionSideUp( projectorDirection, projectorSide, projectorUp );
-    vec3f projectorScale( projector.scale[0], projector.scale[1], projector.scale[2] );
+    FVector3f projectorScale( projector.scale[0], projector.scale[1], projector.scale[2] );
 
     // Create with worse case, shrink later.
     int vertexCount = pMesh->GetVertexCount();
@@ -2652,7 +2652,7 @@ PRAGMA_ENABLE_OPTIMIZATION
 
 //-------------------------------------------------------------------------------------------------
 MeshPtr MeshProject( const Mesh* pMesh,
-                         const PROJECTOR& projector )
+                         const FProjector& projector )
 {
 	MUTABLE_CPUPROFILER_SCOPE(MeshProject);
 
diff --git a/Engine/Plugins/Experimental/Mutable/Source/MutableRuntime/Private/MuR/OpImageProject.h b/Engine/Plugins/Experimental/Mutable/Source/MutableRuntime/Private/MuR/OpImageProject.h
--- a/Engine/Plugins/Experimental/Mutable/Source/MutableRuntime/Private/MuR/OpImageProject.h
+++ b/Engine/Plugins/Experimental/Mutable/Source/MutableRuntime/Private/MuR/OpImageProject.h
@@ -9,7 +9,7 @@
 namespace mu
 {
 	class Image;
-	struct PROJECTOR;
+	struct FProjector;
 	template <int NUM_INTERPOLATORS> class RasterVertex;
 
     struct SCRATCH_IMAGE_PROJECT
@@ -23,7 +23,7 @@ namespace mu
 //                              Image* pTargetImage,
 //                              const Image* pSource,
 //                              const Image* pMask,
-//                              const PROJECTOR& projector,
+//                              const FProjector& projector,
 //                              float fadeStart,
 //                              float fadeEnd,
 //                              int layout,
@@ -61,7 +61,7 @@ namespace mu
                                       SCRATCH_IMAGE_PROJECT* scratch );
 
     extern MeshPtr MeshProject( const Mesh* pMesh,
-                                const PROJECTOR& projector );
+                                const FProjector& projector );
 
 	MUTABLERUNTIME_API extern MeshPtr CreateMeshOptimisedForProjection( int layout );
 	MUTABLERUNTIME_API extern MeshPtr CreateMeshOptimisedForWrappingProjection( int layout );
diff --git a/Engine/Plugins/Experimental/Mutable/Source/MutableRuntime/Private/MuR/OpMeshClipMorphPlane.h b/Engine/Plugins/Experimental/Mutable/Source/MutableRuntime/Private/MuR/OpMeshClipMorphPlane.h
--- a/Engine/Plugins/Experimental/Mutable/Source/MutableRuntime/Private/MuR/OpMeshClipMorphPlane.h
+++ b/Engine/Plugins/Experimental/Mutable/Source/MutableRuntime/Private/MuR/OpMeshClipMorphPlane.h
@@ -9,7 +9,7 @@
 
 namespace mu
 {
-	inline bool PointInBoundingBox(const vec3f& point, const SHAPE& selectionShape)
+	inline bool PointInBoundingBox(const vec3f& point, const FShape& selectionShape)
 	{
 		vec3f v = point - selectionShape.position;
 
@@ -69,7 +69,7 @@ namespace mu
 	//! Reference version
 	//---------------------------------------------------------------------------------------------
 	inline MeshPtr MeshClipMorphPlane(const Mesh* pBase, const vec3f& origin, const vec3f& normal, float dist, float factor, float radius,
-		float radius2, float angle, const SHAPE& selectionShape, const mu::string& boneName = mu::string(), float vertexSelectionBoneMaxRadius = -1.f)
+		float radius2, float angle, const FShape& selectionShape, const mu::string& boneName = mu::string(), float vertexSelectionBoneMaxRadius = -1.f)
 	{
 		//float radius = 8.f;
 		//float radius2 = 4.f;
@@ -301,8 +301,8 @@ namespace mu
 
 						if (
 							(  !boneName.empty() && VertexIsAffectedByBone(v, bone_is_affected, vertex_info) && VertexIsInMaxRadius(vertex, origin, vertexSelectionBoneMaxRadius))
-                            || (boneName.empty() && selectionShape.type == (uint8_t)SHAPE::Type::None)
-                            || (selectionShape.type == (uint8_t)SHAPE::Type::AABox && PointInBoundingBox(vertex, selectionShape))
+                            || (boneName.empty() && selectionShape.type == (uint8_t)FShape::Type::None)
+                            || (selectionShape.type == (uint8_t)FShape::Type::AABox && PointInBoundingBox(vertex, selectionShape))
 							)
 						{
 							vec3f morph_plane_center = origin;					// MORPH PLANE POS relative to root of the selected bone
diff --git a/Engine/Plugins/Experimental/Mutable/Source/MutableRuntime/Private/MuR/OpMeshClipWithMesh.h b/Engine/Plugins/Experimental/Mutable/Source/MutableRuntime/Private/MuR/OpMeshClipWithMesh.h
--- a/Engine/Plugins/Experimental/Mutable/Source/MutableRuntime/Private/MuR/OpMeshClipWithMesh.h
+++ b/Engine/Plugins/Experimental/Mutable/Source/MutableRuntime/Private/MuR/OpMeshClipWithMesh.h
@@ -254,9 +254,9 @@ namespace mu
     //---------------------------------------------------------------------------------------------
     //! Reference version
     //---------------------------------------------------------------------------------------------
-    inline int get_num_intersections( const vec3f& vertex,
-                                      const vec3f& ray,
-                                      const TArray<vec3f>& vertices,
+    inline int get_num_intersections( const FVector3f& vertex,
+                                      const FVector3f& ray,
+                                      const TArray<FVector3f>& vertices,
                                       const TArray<uint32>& faces,
                                       TArray<int>& collapsedVertexMap,
 									  TArray<uint8>& vertex_already_intersected,
@@ -268,7 +268,7 @@ namespace mu
         MUTABLE_CPUPROFILER_SCOPE(get_num_intersections);
     	
         int num_intersections = 0;
-        vec3f intersection;
+		FVector3f intersection;
 
 		FMemory::Memzero( vertex_already_intersected.GetData(), vertex_already_intersected.Num());
 
@@ -283,9 +283,9 @@ namespace mu
         {
             uint32 vertex_indexs[3] = { faces[3 * face], faces[3 * face + 1], faces[3 * face + 2] };
 
-            vec3f v0 = get_collapsed_vertex(vertex_indexs[0]);
-            vec3f v1 = get_collapsed_vertex(vertex_indexs[1]);
-            vec3f v2 = get_collapsed_vertex(vertex_indexs[2]);
+			FVector3f v0 = get_collapsed_vertex(vertex_indexs[0]);
+			FVector3f v1 = get_collapsed_vertex(vertex_indexs[1]);
+			FVector3f v2 = get_collapsed_vertex(vertex_indexs[2]);
 
             int out_intersected_vert, out_intersected_edge_v0, out_intersected_edge_v1;
 
diff --git a/Engine/Plugins/Experimental/Mutable/Source/MutableRuntime/Private/MuR/Operations.cpp b/Engine/Plugins/Experimental/Mutable/Source/MutableRuntime/Private/MuR/Operations.cpp
--- a/Engine/Plugins/Experimental/Mutable/Source/MutableRuntime/Private/MuR/Operations.cpp
+++ b/Engine/Plugins/Experimental/Mutable/Source/MutableRuntime/Private/MuR/Operations.cpp
@@ -447,7 +447,7 @@ namespace mu
 
 
     //---------------------------------------------------------------------------------------------
-    void ForEachReference( const PROGRAM& program, OP::ADDRESS at, const TFunctionRef<void(OP::ADDRESS)> f )
+    void ForEachReference( const FProgram& program, OP::ADDRESS at, const TFunctionRef<void(OP::ADDRESS)> f )
     {
         OP_TYPE type = program.GetOpType(at);
         switch ( type )
diff --git a/Engine/Plugins/Experimental/Mutable/Source/MutableRuntime/Private/MuR/Operations.h b/Engine/Plugins/Experimental/Mutable/Source/MutableRuntime/Private/MuR/Operations.h
--- a/Engine/Plugins/Experimental/Mutable/Source/MutableRuntime/Private/MuR/Operations.h
+++ b/Engine/Plugins/Experimental/Mutable/Source/MutableRuntime/Private/MuR/Operations.h
@@ -956,7 +956,7 @@ namespace mu
             uint32 externalId;
             ADDRESS name;
 
-            // Index in the PROGRAM::m_parameterLists with the parameters that are relevant
+            // Index in the FProgram::m_parameterLists with the parameters that are relevant
             // for whatever is added by this operation. This is used only for resources like
             // images or meshes.
             ADDRESS relevantParametersListIndex;
@@ -1176,7 +1176,7 @@ namespace mu
 	MUTABLERUNTIME_API extern void ForEachReference( OP& op, const TFunctionRef<void(OP::ADDRESS*)> );
 
     //! Utility function to apply a function to all operation references to other operations.
-	MUTABLERUNTIME_API extern void ForEachReference( const struct PROGRAM& program, OP::ADDRESS at, const TFunctionRef<void(OP::ADDRESS)> );
+	MUTABLERUNTIME_API extern void ForEachReference( const struct FProgram& program, OP::ADDRESS at, const TFunctionRef<void(OP::ADDRESS)> );
 
 	//!
 	MUTABLERUNTIME_API inline OP_TYPE GetSwitchForType( DATATYPE d )
diff --git a/Engine/Plugins/Experimental/Mutable/Source/MutableRuntime/Private/MuR/Parameters.cpp b/Engine/Plugins/Experimental/Mutable/Source/MutableRuntime/Private/MuR/Parameters.cpp
--- a/Engine/Plugins/Experimental/Mutable/Source/MutableRuntime/Private/MuR/Parameters.cpp
+++ b/Engine/Plugins/Experimental/Mutable/Source/MutableRuntime/Private/MuR/Parameters.cpp
@@ -71,7 +71,7 @@ namespace mu
 
         for (size_t p=0; p<valueCount; ++p)
         {
-            const PARAMETER_DESC& desc = params->GetPrivate()->m_pModel->GetPrivate()->m_program.m_parameters[p];
+            const FParameterDesc& desc = params->GetPrivate()->m_pModel->GetPrivate()->m_program.m_parameters[p];
             arch << desc.m_name;
             arch << desc.m_uid;
             arch << desc.m_type;
@@ -117,7 +117,7 @@ namespace mu
             {
                 for (size_t mp=0; mp<modelParameters; ++mp)
                 {
-                    const PARAMETER_DESC& desc = pModel->GetPrivate()->m_program.m_parameters[mp];
+                    const FParameterDesc& desc = pModel->GetPrivate()->m_program.m_parameters[mp];
                     if (desc.m_uid==uid)
                     {
                         modelParam = int(mp);
@@ -131,7 +131,7 @@ namespace mu
             {
                 for (size_t mp=0; mp<modelParameters; ++mp)
                 {
-                    const PARAMETER_DESC& desc = pModel->GetPrivate()->m_program.m_parameters[mp];
+                    const FParameterDesc& desc = pModel->GetPrivate()->m_program.m_parameters[mp];
                     if (desc.m_name==name && desc.m_type==type)
                     {
                         modelParam = int(mp);
@@ -145,7 +145,7 @@ namespace mu
             {
                 for (size_t mp=0; mp<modelParameters; ++mp)
                 {
-                    const PARAMETER_DESC& desc = pModel->GetPrivate()->m_program.m_parameters[mp];
+                    const FParameterDesc& desc = pModel->GetPrivate()->m_program.m_parameters[mp];
                     if (desc.m_name==name)
                     {
                         modelParam = int(mp);
@@ -197,7 +197,7 @@ namespace mu
 	//---------------------------------------------------------------------------------------------
 	const char* Parameters::GetName( int index ) const
 	{
-		const PROGRAM& program = m_pD->m_pModel->GetPrivate()->m_program;
+		const FProgram& program = m_pD->m_pModel->GetPrivate()->m_program;
 		check( index>=0 && index<(int)program.m_parameters.Num() );
 
 		return program.m_parameters[index].m_name.c_str();
@@ -207,7 +207,7 @@ namespace mu
 	//---------------------------------------------------------------------------------------------
 	const char* Parameters::GetUid( int index ) const
 	{
-		const PROGRAM& program = m_pD->m_pModel->GetPrivate()->m_program;
+		const FProgram& program = m_pD->m_pModel->GetPrivate()->m_program;
 		check( index>=0 && index<(int)program.m_parameters.Num() );
 
 		return program.m_parameters[index].m_uid.c_str();
@@ -224,7 +224,7 @@ namespace mu
 	//---------------------------------------------------------------------------------------------
 	PARAMETER_TYPE Parameters::GetType( int index ) const
 	{
-		const PROGRAM& program = m_pD->m_pModel->GetPrivate()->m_program;
+		const FProgram& program = m_pD->m_pModel->GetPrivate()->m_program;
 		check( index>=0 && index<(int)program.m_parameters.Num() );
 
 		return program.m_parameters[index].m_type;
@@ -234,7 +234,7 @@ namespace mu
 	//---------------------------------------------------------------------------------------------
 	PARAMETER_DETAILED_TYPE Parameters::GetDetailedType( int index ) const
 	{
-		const PROGRAM& program = m_pD->m_pModel->GetPrivate()->m_program;
+		const FProgram& program = m_pD->m_pModel->GetPrivate()->m_program;
 		check( index>=0 && index<(int)program.m_parameters.Num() );
 
 		return program.m_parameters[index].m_detailedType;
@@ -246,7 +246,7 @@ namespace mu
     {
 		LLM_SCOPE_BYNAME(TEXT("MutableRuntime"));
 
-        const PROGRAM& program = m_pD->m_pModel->GetPrivate()->m_program;
+        const FProgram& program = m_pD->m_pModel->GetPrivate()->m_program;
         check( paramIndex>=0 && paramIndex<int(program.m_parameters.Num()) );
         Ptr<RangeIndex> range;
 
@@ -319,7 +319,7 @@ namespace mu
 	//---------------------------------------------------------------------------------------------
 	int Parameters::GetAdditionalImageCount( int index ) const
 	{
-		const PROGRAM& program = m_pD->m_pModel->GetPrivate()->m_program;
+		const FProgram& program = m_pD->m_pModel->GetPrivate()->m_program;
 		check( index>=0 && index<(int)program.m_parameters.Num() );
 		return (int)program.m_parameters[index].m_descImages.Num();
 	}
@@ -416,7 +416,7 @@ namespace mu
 	//---------------------------------------------------------------------------------------------
 	int Parameters::GetIntPossibleValueCount( int paramIndex ) const
 	{
-		const PROGRAM& program = m_pD->m_pModel->GetPrivate()->m_program;
+		const FProgram& program = m_pD->m_pModel->GetPrivate()->m_program;
 		check( paramIndex>=0 && paramIndex<(int)program.m_parameters.Num() );
 		return (int)program.m_parameters[paramIndex].m_possibleValues.Num();
 	}
@@ -425,7 +425,7 @@ namespace mu
 	//---------------------------------------------------------------------------------------------
 	int Parameters::GetIntPossibleValue( int paramIndex, int valueIndex ) const
 	{
-		const PROGRAM& program = m_pD->m_pModel->GetPrivate()->m_program;
+		const FProgram& program = m_pD->m_pModel->GetPrivate()->m_program;
 		check( paramIndex>=0
 				&& paramIndex<(int)program.m_parameters.Num() );
 		check( valueIndex>=0
@@ -438,7 +438,7 @@ namespace mu
 	//---------------------------------------------------------------------------------------------
 	int Parameters::GetIntValueIndex(int paramIndex, const char* valueName) const
 	{
-		const PROGRAM& program = m_pD->m_pModel->GetPrivate()->m_program;
+		const FProgram& program = m_pD->m_pModel->GetPrivate()->m_program;
 		check(paramIndex >= 0
 			&& paramIndex < (int)program.m_parameters.Num());
 
@@ -458,7 +458,7 @@ namespace mu
 	//---------------------------------------------------------------------------------------------
 	int Parameters::GetIntValueIndex(int paramIndex, int32 Value) const
 	{
-		const PROGRAM& program = m_pD->m_pModel->GetPrivate()->m_program;
+		const FProgram& program = m_pD->m_pModel->GetPrivate()->m_program;
 		check(paramIndex >= 0
 			&& paramIndex < (int)program.m_parameters.Num());
 
@@ -476,7 +476,7 @@ namespace mu
 	//---------------------------------------------------------------------------------------------
 	const char* Parameters::GetIntPossibleValueName( int paramIndex, int valueIndex ) const
 	{
-		const PROGRAM& program = m_pD->m_pModel->GetPrivate()->m_program;
+		const FProgram& program = m_pD->m_pModel->GetPrivate()->m_program;
 		check( paramIndex>=0
 				&& paramIndex<(int)program.m_parameters.Num() );
 		check( valueIndex>=0
@@ -910,11 +910,11 @@ namespace mu
 
 
     //---------------------------------------------------------------------------------------------
-    PROJECTOR Parameters::Private::GetProjectorValue( int index,
+    FProjector Parameters::Private::GetProjectorValue( int index,
                                                       const Ptr<const RangeIndex>& pos ) const
     {
 
-        const PROGRAM& program = m_pModel->GetPrivate()->m_program;
+        const FProgram& program = m_pModel->GetPrivate()->m_program;
 
         // Early out in case of invalid parameters
         if ( index < 0
@@ -926,10 +926,10 @@ namespace mu
              program.m_parameters[index].m_type != PARAMETER_TYPE::T_PROJECTOR )
         {
 			check(false);
-            return PROJECTOR();
+            return FProjector();
         }
 
-        const PROJECTOR* result = nullptr;
+        const FProjector* result = nullptr;
 
         // Single value case
         if (!pos)
@@ -964,101 +964,90 @@ namespace mu
 
 
 	//---------------------------------------------------------------------------------------------
-	void Parameters::GetProjectorValue( int index,
-                                        PROJECTOR_TYPE* pType,
-										float* pPosX, float* pPosY, float* pPosZ,
-										float* pDirX, float* pDirY, float* pDirZ,
-										float* pUpX, float* pUpY, float* pUpZ,
-                                        float* pScaleX, float* pScaleY, float* pScaleZ,
-                                        float* pProjectionAngle,
-                                        const Ptr<const RangeIndex>& pos ) const
+	void Parameters::GetProjectorValue( int ParameterIndex,
+		PROJECTOR_TYPE* OutType,
+		FVector3f* OutPos,
+		FVector3f* OutDir,
+		FVector3f* OutUp,
+		FVector3f* OutScale,
+		float* OutProjectionAngle,
+		const Ptr<const RangeIndex>& pos ) const
 	{
-		check( index>=0 && index<(int)m_pD->m_values.Num() );
-		check( GetType(index)== PARAMETER_TYPE::T_PROJECTOR );
+		check(ParameterIndex >=0 && ParameterIndex <m_pD->m_values.Num() );
+		check( GetType(ParameterIndex)== PARAMETER_TYPE::T_PROJECTOR );
 
         // Early out in case of invalid parameters
-        if ( index < 0
+        if (ParameterIndex < 0
              ||
-             index >= (int)m_pD->m_values.Num()
+			ParameterIndex >= m_pD->m_values.Num()
              ||
-             GetType(index) != PARAMETER_TYPE::T_PROJECTOR )
+             GetType(ParameterIndex) != PARAMETER_TYPE::T_PROJECTOR )
         {
             return;
         }
 
-        PROJECTOR result = m_pD->GetProjectorValue( index, pos );
+        FProjector result = m_pD->GetProjectorValue(ParameterIndex, pos );
 
         // Copy results
-        if (pType) *pType = result.type;
-
-        if (pPosX) *pPosX = result.position[0];
-        if (pPosY) *pPosY = result.position[1];
-        if (pPosZ) *pPosZ = result.position[2];
-
-        if (pDirX) *pDirX = result.direction[0];
-        if (pDirY) *pDirY = result.direction[1];
-        if (pDirZ) *pDirZ = result.direction[2];
+        if (OutType) *OutType = result.type;
 
-        if (pUpX) *pUpX = result.up[0];
-        if (pUpY) *pUpY = result.up[1];
-        if (pUpZ) *pUpZ = result.up[2];
+        if (OutPos) *OutPos = result.position;
+        if (OutDir) *OutDir = result.direction;
+        if (OutUp) *OutUp = result.up;
+        if (OutScale) *OutScale = result.scale;
 
-        if (pScaleX) *pScaleX = result.scale[0];
-        if (pScaleY) *pScaleY = result.scale[1];
-        if (pScaleZ) *pScaleZ = result.scale[2];
-
-        if (pProjectionAngle) *pProjectionAngle = result.projectionAngle;
+        if (OutProjectionAngle) *OutProjectionAngle = result.projectionAngle;
     }
 
 
 	//---------------------------------------------------------------------------------------------
-	void Parameters::SetProjectorValue( int index,
-										float posX, float posY, float posZ,
-										float dirX, float dirY, float dirZ,
-										float upX, float upY, float upZ,
-                                        float scaleX, float scaleY, float scaleZ,
-                                        float projectionAngle,
-                                        const Ptr<const RangeIndex>& pos )
+	void Parameters::SetProjectorValue( int ParameterIndex,
+		const FVector3f& pos,
+		const FVector3f& dir,
+		const FVector3f& up,
+		const FVector3f& scale,
+        float projectionAngle,
+        const Ptr<const RangeIndex>& RangePosition)
 	{
-		check( index>=0 && index<(int)m_pD->m_values.Num() );
-		check( GetType(index)== PARAMETER_TYPE::T_PROJECTOR );
+		check(ParameterIndex >=0 && ParameterIndex <m_pD->m_values.Num() );
+		check( GetType(ParameterIndex)== PARAMETER_TYPE::T_PROJECTOR );
 
         // Early out in case of invalid parameters
-        if ( index < 0
+        if (ParameterIndex < 0
              ||
-             index >= (int)m_pD->m_values.Num()
+			ParameterIndex >= m_pD->m_values.Num()
              ||
-             GetType(index) != PARAMETER_TYPE::T_PROJECTOR )
+             GetType(ParameterIndex) != PARAMETER_TYPE::T_PROJECTOR )
         {
             return;
         }
 
-        PROJECTOR* result = nullptr;
+        FProjector* result = nullptr;
 
         // Single value case
-        if (!pos)
+        if (!RangePosition)
         {
             // Clear multivalue, if set.
-            if (index<int(m_pD->m_multiValues.Num()))
+            if (ParameterIndex<m_pD->m_multiValues.Num())
             {
-                m_pD->m_multiValues[index].Empty();
+                m_pD->m_multiValues[ParameterIndex].Empty();
             }
 
-            result = &m_pD->m_values[index].m_projector;
+            result = &m_pD->m_values[ParameterIndex].m_projector;
         }
 
         // Multivalue case
         else
         {
-            check( pos->m_pD->m_parameter==index );
+            check(RangePosition->m_pD->m_parameter== ParameterIndex);
 
-            if ( index>=int(m_pD->m_multiValues.Num()))
+            if (ParameterIndex >=m_pD->m_multiValues.Num())
             {
-                m_pD->m_multiValues.SetNum(index+1);
+                m_pD->m_multiValues.SetNum(ParameterIndex+1);
             }
 
-			TMap< TArray<int32_t>, PARAMETER_VALUE >& m = m_pD->m_multiValues[index];
-			PARAMETER_VALUE& it = m.FindOrAdd(pos->m_pD->m_values);
+			TMap< TArray<int32_t>, PARAMETER_VALUE >& m = m_pD->m_multiValues[ParameterIndex];
+			PARAMETER_VALUE& it = m.FindOrAdd(RangePosition->m_pD->m_values);
             result = &it.m_projector;
         }
 
@@ -1068,25 +1057,14 @@ namespace mu
         result->type = PROJECTOR_TYPE::COUNT;
         if (m_pD->m_pModel)
         {
-            const PROGRAM& program = m_pD->m_pModel->GetPrivate()->m_program;
-            result->type = program.m_parameters[index].m_defaultValue.m_projector.type;
+            const FProgram& program = m_pD->m_pModel->GetPrivate()->m_program;
+            result->type = program.m_parameters[ParameterIndex].m_defaultValue.m_projector.type;
         }
 
-        result->position[0] = posX;
-        result->position[1] = posY;
-        result->position[2] = posZ;
-
-        result->direction[0] = dirX;
-        result->direction[1] = dirY;
-        result->direction[2] = dirZ;
-
-        result->up[0] = upX;
-        result->up[1] = upY;
-        result->up[2] = upZ;
-
-        result->scale[0] = scaleX;
-        result->scale[1] = scaleY;
-        result->scale[2] = scaleZ;
+        result->position = pos;
+        result->direction = dir;
+        result->up = up;        
+        result->scale = scale;
 
         result->projectionAngle = projectionAngle;
     }
@@ -1142,13 +1120,13 @@ namespace mu
 	//---------------------------------------------------------------------------------------------
 	int Parameters::Private::Find( const char* strName ) const
 	{
-		const PROGRAM& program = m_pModel->GetPrivate()->m_program;
+		const FProgram& program = m_pModel->GetPrivate()->m_program;
 
 		int result = -1;
 
 		for( int i=0; result<0 && i<(int)program.m_parameters.Num(); ++i )
 		{
-			const PARAMETER_DESC& p = program.m_parameters[i];
+			const FParameterDesc& p = program.m_parameters[i];
 
 			if ( p.m_name == strName )
 			{
diff --git a/Engine/Plugins/Experimental/Mutable/Source/MutableRuntime/Private/MuR/ParametersPrivate.h b/Engine/Plugins/Experimental/Mutable/Source/MutableRuntime/Private/MuR/ParametersPrivate.h
--- a/Engine/Plugins/Experimental/Mutable/Source/MutableRuntime/Private/MuR/ParametersPrivate.h
+++ b/Engine/Plugins/Experimental/Mutable/Source/MutableRuntime/Private/MuR/ParametersPrivate.h
@@ -23,24 +23,22 @@ namespace mu
     //---------------------------------------------------------------------------------------------
     //! Information to project an image on a mesh.
     //---------------------------------------------------------------------------------------------
-    struct PROJECTOR
+    struct FProjector
     {
         PROJECTOR_TYPE type = PROJECTOR_TYPE::PLANAR;
-        float position[3] = {0,0,0};
-        float direction[3] = {0,0,0};
-        float up[3] = {0,0,0};
-        float scale[3] = {0,0,0};
+        FVector3f position = {0,0,0};
+		FVector3f direction = {0,0,0};
+		FVector3f up = {0,0,0};
+		FVector3f scale = {0,0,0};
         float projectionAngle = 0.0f;
 
         //!
-        inline void GetDirectionSideUp( vec3f& d,
-                                        vec3f& s,
-                                        vec3f& u ) const
+        inline void GetDirectionSideUp(FVector3f& d, FVector3f& s, FVector3f& u ) const
         {
-            d = vec3f( direction[0], direction[1], direction[2] );
-            u = vec3f( up[0], up[1], up[2] );
-            s = cross( u, d );
-            normalise(s);
+            d = direction;
+            u = up;
+            s = FVector3f::CrossProduct( up, direction );
+            s.Normalize();
         }
 
 
@@ -82,21 +80,13 @@ namespace mu
             arch >> projectionAngle;
         }
 
-        bool operator==( const PROJECTOR& o ) const
+        bool operator==( const FProjector& o ) const
         {
             return     type==o.type
-                    && position[0]==o.position[0]
-                    && position[1]==o.position[1]
-                    && position[2]==o.position[2]
-                    && up[0]==o.up[0]
-                    && up[1]==o.up[1]
-                    && up[2]==o.up[2]
-                    && direction[0]==o.direction[0]
-                    && direction[1]==o.direction[1]
-                    && direction[2]==o.direction[2]
-                    && scale[0]==o.scale[0]
-                    && scale[1]==o.scale[1]
-                    && scale[2]==o.scale[2]
+                    && position==o.position
+                    && up==o.up
+                    && direction==o.direction
+                    && scale==o.scale
                     && projectionAngle==o.projectionAngle;
         }
 
@@ -105,7 +95,7 @@ namespace mu
 	//---------------------------------------------------------------------------------------------
 	//! Information about a generic shape in space.
 	//---------------------------------------------------------------------------------------------
-	struct SHAPE
+	struct FShape
 	{
 		// Transform
 		vec3f position;
@@ -115,13 +105,13 @@ namespace mu
 		vec3f size;
 
 		// 
-		enum class Type : uint8_t
+		enum class Type : uint8
 		{
 			None = 0,
 			Ellipse,
 			AABox
 		};
-        uint8_t type = 0;
+        uint8 type = 0;
 		
 
 		//!
@@ -144,7 +134,7 @@ namespace mu
 			arch >> type;
 		}
 
-        bool operator==( const SHAPE& o ) const
+        bool operator==( const FShape& o ) const
         {
             return type==o.type
                     && position==o.position
@@ -167,7 +157,7 @@ namespace mu
 		float out_tangent;
 		float out_tangent_weight;
 
-		enum class InterpMode : uint8_t
+		enum class InterpMode : uint8
 		{
 			Linear = 0,
 			Constant,
@@ -177,10 +167,10 @@ namespace mu
 			//! Utility enumeration value, not really a Interp mode.
 			Count
 		};
-        uint8_t interp_mode;
+        uint8 interp_mode;
 
 
-		enum class TangentMode : uint8_t
+		enum class TangentMode : uint8
 		{
 			Auto = 0,
 			Manual,
@@ -190,10 +180,10 @@ namespace mu
 			//! Utility enumeration value, not really a Tangent mode.
 			Count
 		};
-        uint8_t tangent_mode;
+        uint8 tangent_mode;
 
 
-		enum class TangentWeightMode : uint8_t
+		enum class TangentWeightMode : uint8
 		{
 			Auto = 0,
 			Manual,
@@ -203,7 +193,7 @@ namespace mu
 			//! Utility enumeration value, not really a Tangent weight mode.
 			Count
 		};
-        uint8_t tangent_weight_mode;
+        uint8 tangent_weight_mode;
 
         //!
         bool operator==( const CurveKeyFrame& other ) const
@@ -294,7 +284,7 @@ namespace mu
 
         float m_colour[3];
 
-        PROJECTOR m_projector;
+        FProjector m_projector;
 
         EXTERNAL_IMAGE_ID m_image;
 
@@ -311,7 +301,7 @@ namespace mu
     MUTABLE_DEFINE_POD_SERIALISABLE( PARAMETER_VALUE )
 
 
-    struct PARAMETER_DESC
+    struct FParameterDesc
     {
         string m_name;
 
@@ -326,7 +316,7 @@ namespace mu
 
         //! Ranges, if the parameter is multi-dimensional. The indices refer to the Model's program
         //! vector of range descriptors.
-		TArray<uint32_t> m_ranges;
+		TArray<uint32> m_ranges;
 
         //! Additional parameter description information
 		TArray<OP::ADDRESS> m_descImages;
@@ -334,7 +324,7 @@ namespace mu
         //! Possible values of the parameter in case of being an integer, and its names
         struct INT_VALUE_DESC
         {
-            int16_t m_value;
+            int16 m_value;
             string m_name;
 
             //!
@@ -364,7 +354,7 @@ namespace mu
 		TArray<INT_VALUE_DESC> m_possibleValues;
 
         //!
-        bool operator==( const PARAMETER_DESC& other ) const
+        bool operator==( const FParameterDesc& other ) const
         {
             return m_name == other.m_name && m_uid == other.m_uid && m_type == other.m_type &&
                    m_defaultValue == other.m_defaultValue &&
@@ -375,7 +365,7 @@ namespace mu
         //!
         void Serialise( OutputArchive& arch ) const
         {
-            const int32_t ver = 5;
+            const int32 ver = 5;
             arch << ver;
 
 			arch << m_name;
@@ -390,7 +380,7 @@ namespace mu
         //!
         void Unserialise( InputArchive& arch )
         {
-            int32_t ver;
+            int32 ver;
             arch >> ver;
             check( ver == 5 );
 
@@ -405,38 +395,49 @@ namespace mu
     };
 
 
-    struct RANGE_DESC
+    struct FRangeDesc
     {
         string m_name;
         string m_uid;
 
+		/** Parameter that controls the size of this range, if any. */
+		int32 m_dimensionParameter = -1;
+
         //!
-        bool operator==( const RANGE_DESC& other ) const
+        bool operator==( const FRangeDesc& other ) const
         {
             return m_name==other.m_name
-                    &&
-                    m_uid==other.m_uid;
+				&&
+				m_uid == other.m_uid 
+				&&
+				m_dimensionParameter == other.m_dimensionParameter;
         }
 
         //!
         void Serialise( OutputArchive& arch ) const
         {
-            const int32_t ver = 1;
+            const uint32 ver = 2;
             arch << ver;
 
             arch << m_name;
             arch << m_uid;
+			arch << m_dimensionParameter;
         }
 
         //!
         void Unserialise( InputArchive& arch )
         {
-            int32_t ver;
+            uint32 ver;
             arch >> ver;
-            check( ver == 1 );
+            check( ver <= 2 );
 
             arch >> m_name;
             arch >> m_uid;
+
+			if (ver >= 2)
+			{
+				arch >> m_dimensionParameter;
+			}
         }
     };
 
@@ -470,13 +471,13 @@ namespace mu
 
         //! If the parameter is multidemensional, the values are stored here.
         //! The key of the map is the vector of values stored in a RangeIndex
-		TArray< TMap< TArray<int32_t>, PARAMETER_VALUE > > m_multiValues;
+		TArray< TMap< TArray<int32>, PARAMETER_VALUE > > m_multiValues;
 
 
         //!
         void Serialise( OutputArchive& arch ) const
         {
-            const uint32_t ver = 1;
+            const uint32 ver = 1;
             arch << ver;
 
             arch << m_values;
@@ -486,7 +487,7 @@ namespace mu
         //!
         void Unserialise( InputArchive& arch )
         {
-            uint32_t ver;
+            uint32 ver;
             arch >> ver;
             check( ver == 1 );
 
@@ -498,7 +499,7 @@ namespace mu
         int Find( const char* strName ) const;
 
         //!
-        PROJECTOR GetProjectorValue( int index, const Ptr<const RangeIndex>& pos ) const;
+        FProjector GetProjectorValue( int index, const Ptr<const RangeIndex>& pos ) const;
     };
 
 }
diff --git a/Engine/Plugins/Experimental/Mutable/Source/MutableRuntime/Private/MuR/System.cpp b/Engine/Plugins/Experimental/Mutable/Source/MutableRuntime/Private/MuR/System.cpp
--- a/Engine/Plugins/Experimental/Mutable/Source/MutableRuntime/Private/MuR/System.cpp
+++ b/Engine/Plugins/Experimental/Mutable/Source/MutableRuntime/Private/MuR/System.cpp
@@ -237,7 +237,7 @@ namespace mu
 
 		m_pD->m_memory = pLiveInstance->m_memory;
 
-		PROGRAM& program = pLiveInstance->m_pModel->GetPrivate()->m_program;
+		FProgram& program = pLiveInstance->m_pModel->GetPrivate()->m_program;
 
 		bool validState = stateIndex >= 0 && stateIndex < (int)program.m_states.Num();
 		if (!validState)
@@ -357,7 +357,7 @@ namespace mu
 			if (res.m_id == imageId)
 			{
 				const mu::Model* Model = pLiveInstance->m_pModel.get();
-				const mu::PROGRAM& program = Model->GetPrivate()->m_program;
+				const mu::FProgram& program = Model->GetPrivate()->m_program;
 
 				int32 VarValue = CVarClearImageDescCache.GetValueOnAnyThread();
 				if (VarValue != 0)
@@ -473,7 +473,7 @@ namespace mu
 			&&
 			parameter < (int)pModel->GetPrivate()->m_program.m_parameters.Num());
 
-		const PARAMETER_DESC& param =
+		const FParameterDesc& param =
 			pModel->GetPrivate()->m_program.m_parameters[parameter];
 		check(image >= 0 && image < (int)param.m_descImages.Num());
 
@@ -520,7 +520,7 @@ namespace mu
         }
 
 
-        bool Visit( OP::ADDRESS at, PROGRAM& program ) override
+        bool Visit( OP::ADDRESS at, FProgram& program ) override
         {
             switch ( program.GetOpType(at) )
             {
@@ -580,7 +580,7 @@ namespace mu
 
         // check what parameters have changed
         updatedParameters = 0;
-        const PROGRAM& program = pLiveInstance->m_pModel->GetPrivate()->m_program;
+        const FProgram& program = pLiveInstance->m_pModel->GetPrivate()->m_program;
         const TArray<int>& runtimeParams = program.m_states[ pLiveInstance->m_state ].m_runtimeParameters;
 
         check( pParams->GetCount() == (int)program.m_parameters.Num() );
@@ -833,7 +833,7 @@ namespace mu
 	{
 		MUTABLE_CPUPROFILER_SCOPE(PrepareCache);
 
-		PROGRAM& program = pModel->GetPrivate()->m_program;
+		FProgram& program = pModel->GetPrivate()->m_program;
 		size_t opCount = program.m_opAddress.Num();
 		m_memory->m_opHitCount.clear();
 		m_memory->Init(opCount);
@@ -841,7 +841,7 @@ namespace mu
 		// Mark the resources that have to be cached to update the instance in this state
 		if (state >= 0)
 		{
-			const PROGRAM::STATE& s = program.m_states[state];
+			const FProgram::FState& s = program.m_states[state];
 			for (uint32 a : s.m_updateCache)
 			{
 				m_memory->SetForceCached(a);
@@ -893,7 +893,7 @@ namespace mu
 			mu::Ptr<mu::Model> pCacheModel = m.m_pModel.Pin();
             if (pCacheModel)
             {
-				mu::PROGRAM& program = pCacheModel->GetPrivate()->m_program;
+				mu::FProgram& program = pCacheModel->GetPrivate()->m_program;
 				for (int32 RomIndex = 0; RomIndex < program.m_roms.Num(); ++RomIndex)
 				{
 					if (program.IsRomLoaded(RomIndex))
@@ -919,7 +919,7 @@ namespace mu
 					mu::Ptr<mu::Model> pCacheModel = modelCache.m_pModel.Pin();
                     if (pCacheModel)
                     {
-						mu::PROGRAM& program = pCacheModel->GetPrivate()->m_program;
+						mu::FProgram& program = pCacheModel->GetPrivate()->m_program;
                         check( modelCache.m_romWeight.Num() == program.m_roms.Num());
 
                         for (int32 RomIndex=0; RomIndex <program.m_roms.Num(); ++RomIndex)
@@ -975,7 +975,7 @@ namespace mu
     {
         check(pModel);
 
-        PROGRAM& program = pModel->GetPrivate()->m_program;
+        FProgram& program = pModel->GetPrivate()->m_program;
 
         // If budget is zero, we don't unload anything here, and we assume it is managed
         // somewhere else.
@@ -1007,7 +1007,7 @@ namespace mu
     {
 		MarkRomUsed( romIndex, pModel);
 
-		PROGRAM& program = pModel->GetPrivate()->m_program;
+		FProgram& program = pModel->GetPrivate()->m_program;
         EnsureCacheBelowBudget(program.m_roms[romIndex].Size, isRomLockedFunc);
     }
 
diff --git a/Engine/Plugins/Experimental/Mutable/Source/MutableRuntime/Private/MuR/SystemPrivate.h b/Engine/Plugins/Experimental/Mutable/Source/MutableRuntime/Private/MuR/SystemPrivate.h
--- a/Engine/Plugins/Experimental/Mutable/Source/MutableRuntime/Private/MuR/SystemPrivate.h
+++ b/Engine/Plugins/Experimental/Mutable/Source/MutableRuntime/Private/MuR/SystemPrivate.h
@@ -60,7 +60,7 @@ namespace mu
 	class Projector : public RefCounted
 	{
 	public:
-		PROJECTOR m_value;
+		FProjector m_value;
 	};
 	typedef Ptr<Projector> ProjectorPtr;
 
diff --git a/Engine/Plugins/Experimental/Mutable/Source/MutableRuntime/Public/MuR/Parameters.h b/Engine/Plugins/Experimental/Mutable/Source/MutableRuntime/Public/MuR/Parameters.h
--- a/Engine/Plugins/Experimental/Mutable/Source/MutableRuntime/Public/MuR/Parameters.h
+++ b/Engine/Plugins/Experimental/Mutable/Source/MutableRuntime/Public/MuR/Parameters.h
@@ -320,41 +320,39 @@ namespace mu
 		//! Return the value of a projector parameter, as a 4x4 matrix. The matrix is supposed to be
 		//! a linear transform in column-major.
 		//! \pre The parameter specified by index is a T_PROJECTOR.
-        //! \param index Index of the parameter from 0 to GetCount()-1
-        //! \param pPosX,pPosY,pPosZ Pointers to where the object-space position coordinates of the
-        //!         projector will be stored.
-        //! \param pDirX,pDirY,pDirZ Pointers to where the object-space direction vector of the
-        //!         projector will be stored.
-        //! \param pUpX,pUpY,pUpZ Pointers to where the object-space vertically up direction vector
+        //! \param ParameterIndex Index of the parameter from 0 to GetCount()-1
+        //! \param OutPosX Pointer to where the object-space position coordinates of the projector will be stored.
+        //! \param OutDir Pointer to where the object-space direction vector of the projector will be stored.
+        //! \param OutUp Pointer to where the object-space vertically up direction vector
         //!         of the projector will be stored. This controls the "roll" angle of the
         //!         projector.
-        //! \param pScaleX, pScaleY Pointers to the projector-space scaling of the projector.
-        //! \param pos Only for multidimensional parametres: relevant position to get in the ranges
-        void GetProjectorValue( int index,
-                                PROJECTOR_TYPE* pProjectionType,
-								float* pPosX, float* pPosY, float* pPosZ,
-								float* pDirX, float* pDirY, float* pDirZ,
-								float* pUpX, float* pUpY, float* pUpZ,
-                                float* pScaleX, float* pScaleY, float* pScaleZ,
-                                float* pProjectionAngle,
-                                const Ptr<const RangeIndex>& pos=nullptr ) const;
+        //! \param OutScale Pointer to the projector-space scaling of the projector.
+        //! \param RangePosition Only for multidimensional parametres: relevant position to get in the ranges
+        void GetProjectorValue( int ParameterIndex,
+                                PROJECTOR_TYPE* OutProjectionType,
+								FVector3f* OutPos,
+								FVector3f* OutDir,
+								FVector3f* OutUp,
+								FVector3f* OutScale,
+                                float* OutProjectionAngle,
+                                const Ptr<const RangeIndex>& RangePosition=nullptr ) const;
 
 		//! If the parameter is of the projector type, set its value.
-		//! \param index Index of the parameter from 0 to GetCount()-1
+		//! \param ParameterIndex Index of the parameter from 0 to GetCount()-1
         //! \param posX,posY,posZ Object-space position coordinates of the projector.
         //! \param dirX,dirY,dirZ Object-space direction vector of the projector.
         //! \param upX,upY,upZ Object-space vertically up direction vector of the projector.
         //! \param scaleX, scaleY Projector-space scaling of the projector.
         //! \param projectionAngle [only for Cylindrical projectors], the angle in radians of the
         //! projection area on the cylinder surface.
-        //! \param pos Only for multidimensional parametres: relevant position to set in the ranges
-        void SetProjectorValue( int index,
-                                float posX, float posY, float posZ,
-								float dirX, float dirY, float dirZ,
-								float upX, float upY, float upZ,
-                                float scaleX, float scaleY, float scaleZ,
+        //! \param RangePosition Only for multidimensional parametres: relevant position to set in the ranges
+        void SetProjectorValue( int ParameterIndex,
+								const FVector3f& pos,
+								const FVector3f& dir,
+								const FVector3f& up,
+								const FVector3f& scale,
                                 float projectionAngle,
-                                const Ptr<const RangeIndex>& pos=nullptr );
+                                const Ptr<const RangeIndex>& RangePosition=nullptr );
 
         //! Return the value of an image parameter.
         //! \pre The parameter specified by index is a T_IMAGE.
diff --git a/Engine/Plugins/Experimental/Mutable/Source/MutableTools/Private/MuT/AST.cpp b/Engine/Plugins/Experimental/Mutable/Source/MutableTools/Private/MuT/AST.cpp
--- a/Engine/Plugins/Experimental/Mutable/Source/MutableTools/Private/MuT/AST.cpp
+++ b/Engine/Plugins/Experimental/Mutable/Source/MutableTools/Private/MuT/AST.cpp
@@ -11,6 +11,7 @@
 #include "MuR/Platform.h"
 #include "MuT/ASTOpConstantResource.h"
 #include "MuT/ASTOpMeshRemoveMask.h"
+#include "MuT/ASTOpParameter.h"
 #include "MuT/StreamsPrivate.h"
 #include "MuT/Platform.h"
 #include "Trace/Detail/Channel.h"
@@ -240,7 +241,7 @@ mu::Ptr<ASTOp> ASTOp::DeepClone( const Ptr<ASTOp>& root )
 
 
 //-------------------------------------------------------------------------------------------------
-void ASTOp::FullLink( Ptr<ASTOp>& root, PROGRAM& program, const FLinkerOptions* Options )
+void ASTOp::FullLink( Ptr<ASTOp>& root, FProgram& program, const FLinkerOptions* Options )
 {
     MUTABLE_CPUPROFILER_SCOPE(AST_FullLink);
     Traverse_BottomUp_Unique( root,
@@ -919,11 +920,51 @@ bool ASTOp::GetNonBlackRect( FImageRect& ) const
 }
 
 
+void ASTOp::LinkRange(FProgram& program,
+	const FRangeData& range,
+	OP::ADDRESS& rangeSize,
+	uint16& rangeId)
+{
+	if (range.rangeSize)
+	{
+		if (range.rangeSize->linkedRange < 0)
+		{
+			check(program.m_ranges.Num() < 255);
+			range.rangeSize->linkedRange = int8(program.m_ranges.Num());
+			FRangeDesc rangeData;
+			rangeData.m_name = range.rangeName;
+			rangeData.m_uid = range.rangeUID;
+
+			// Try to see if a parameter directly controls de size of the range. This is used
+			// to store hint data for instance generation in tools or randomizers that want to
+			// support multilayer, but it is not critical otherwise.
+			int32 EstimatedSizeParameter = -1;
+			if ( range.rangeSize->GetOpType() == OP_TYPE::SC_PARAMETER
+				||
+				range.rangeSize->GetOpType() == OP_TYPE::NU_PARAMETER )
+			{
+				const ASTOpParameter* ParamOp = dynamic_cast<const ASTOpParameter*>(range.rangeSize.child().get());
+
+				EstimatedSizeParameter = ParamOp->LinkedParameterIndex;
+			}
+			rangeData.m_dimensionParameter = EstimatedSizeParameter;
+
+			program.m_ranges.Add(rangeData);
+		}
+		rangeSize = range.rangeSize->linkedAddress;
+		rangeId = uint16(range.rangeSize->linkedRange);
+	}
+}
+
+
 //-------------------------------------------------------------------------------------------------
 //-------------------------------------------------------------------------------------------------
 //-------------------------------------------------------------------------------------------------
 ASTOpFixed::ASTOpFixed()
-    : children({{ASTChild(this,nullptr),ASTChild(this,nullptr),ASTChild(this,nullptr),ASTChild(this,nullptr),ASTChild(this,nullptr),ASTChild(this,nullptr),ASTChild(this,nullptr),ASTChild(this,nullptr)}})
+    : children({{
+			ASTChild(this,nullptr),ASTChild(this,nullptr),ASTChild(this,nullptr),ASTChild(this,nullptr),
+			ASTChild(this,nullptr),ASTChild(this,nullptr),ASTChild(this,nullptr),ASTChild(this,nullptr)
+		}})
 {
 }
 
@@ -949,7 +990,7 @@ void ASTOpFixed::ForEachChild( const TFunctionRef<void(ASTChild&)> f )
 }
 
 
-void ASTOpFixed::Link( PROGRAM& program, const FLinkerOptions* )
+void ASTOpFixed::Link( FProgram& program, const FLinkerOptions* )
 {
     if (!linkedAddress)
     {
@@ -1335,9 +1376,9 @@ void ASTOpFixed::GetLayoutBlockSize( int* pBlockX, int* pBlockY )
 }
 
 
-ASTOp::BOOL_EVAL_RESULT ASTOpFixed::EvaluateBool( ASTOpList& facts, EVALUATE_BOOL_CACHE* cache ) const
+ASTOp::FBoolEvalResult ASTOpFixed::EvaluateBool( ASTOpList& facts, FEvaluateBoolCache* cache ) const
 {
-    EVALUATE_BOOL_CACHE localCache;
+    FEvaluateBoolCache localCache;
     if (!cache)
     {
         cache = &localCache;
@@ -1352,14 +1393,14 @@ ASTOp::BOOL_EVAL_RESULT ASTOpFixed::EvaluateBool( ASTOpList& facts, EVALUATE_BOO
         }
     }
 
-    BOOL_EVAL_RESULT result = BET_UNKNOWN;
+    FBoolEvalResult result = BET_UNKNOWN;
     switch(op.type)
     {
     case OP_TYPE::BO_NOT:
     {
         if (children[op.args.BoolNot.source])
         {
-            BOOL_EVAL_RESULT child = children[op.args.BoolNot.source]->EvaluateBool(facts,cache);
+            FBoolEvalResult child = children[op.args.BoolNot.source]->EvaluateBool(facts,cache);
             if (child==BET_TRUE)
             {
                 result = BET_FALSE;
@@ -1397,8 +1438,8 @@ ASTOp::BOOL_EVAL_RESULT ASTOpFixed::EvaluateBool( ASTOpList& facts, EVALUATE_BOO
     {
         const auto& a = children[op.args.BoolBinary.a].child();
         const auto& b = children[op.args.BoolBinary.b].child();
-        BOOL_EVAL_RESULT resultA = BET_UNKNOWN;
-        BOOL_EVAL_RESULT resultB = BET_UNKNOWN;
+        FBoolEvalResult resultA = BET_UNKNOWN;
+        FBoolEvalResult resultB = BET_UNKNOWN;
         for ( size_t f=0; f<facts.Num(); ++f )
         {
             if ( a && resultA == BET_UNKNOWN )
@@ -1439,8 +1480,8 @@ ASTOp::BOOL_EVAL_RESULT ASTOpFixed::EvaluateBool( ASTOpList& facts, EVALUATE_BOO
     {
         const auto& a = children[op.args.BoolBinary.a].child();
         const auto& b = children[op.args.BoolBinary.b].child();
-        BOOL_EVAL_RESULT resultA = BET_UNKNOWN;
-        BOOL_EVAL_RESULT resultB = BET_UNKNOWN;
+        FBoolEvalResult resultA = BET_UNKNOWN;
+        FBoolEvalResult resultB = BET_UNKNOWN;
         for ( size_t f=0; f<facts.Num(); ++f )
         {
             if ( a && resultA == BET_UNKNOWN )
diff --git a/Engine/Plugins/Experimental/Mutable/Source/MutableTools/Private/MuT/AST.h b/Engine/Plugins/Experimental/Mutable/Source/MutableTools/Private/MuT/AST.h
--- a/Engine/Plugins/Experimental/Mutable/Source/MutableTools/Private/MuT/AST.h
+++ b/Engine/Plugins/Experimental/Mutable/Source/MutableTools/Private/MuT/AST.h
@@ -30,9 +30,6 @@
 #include <utility>
 #include <set>
 
-namespace mu { class ASTOp; }
-namespace mu { class ASTOpFixed; }
-
 namespace std
 {
 
@@ -50,6 +47,9 @@ namespace std
 
 namespace mu
 {
+	class ASTOp;
+	class ASTOpFixed;
+
 	template<typename T>
 	inline uint32 GetTypeHash(const Ptr<const T>& p)
 	{
@@ -223,7 +223,7 @@ namespace mu
 		const ASTOp* m_root = nullptr;
 		Ptr<ASTOp> m_initialSource;
 		//! For each operation we sink, the map from old instructions to new instructions.
-		TMap<Ptr<const class ASTOpFixed>, std::unordered_map<Ptr<ASTOp>, Ptr<ASTOp>>> m_oldToNew;
+		TMap<Ptr<const ASTOpFixed>, std::unordered_map<Ptr<ASTOp>, Ptr<ASTOp>>> m_oldToNew;
 		TArray<Ptr<ASTOp>> m_newOps;
 
 		Ptr<ASTOp> Visit(Ptr<ASTOp> at, const ASTOpFixed* currentCropOp);
@@ -448,7 +448,7 @@ namespace mu
         //---------------------------------------------------------------------------------------------
 
         //!
-        static void FullLink( Ptr<ASTOp>& root, PROGRAM&, const FLinkerOptions*);
+        static void FullLink( Ptr<ASTOp>& root, FProgram&, const FLinkerOptions*);
 
         //!
         static void ClearLinkData( Ptr<ASTOp>& root );
@@ -459,14 +459,14 @@ namespace mu
     private:
 
         //!
-        virtual void Link( PROGRAM& program, const FLinkerOptions* Options ) = 0;
+        virtual void Link( FProgram& program, const FLinkerOptions* Options ) = 0;
 
     protected:
 
         //---------------------------------------------------------------------------------------------
         //!
         //---------------------------------------------------------------------------------------------
-        struct RANGE_DATA
+        struct FRangeData
         {
             //!
             ASTChild rangeSize;
@@ -478,30 +478,30 @@ namespace mu
             string rangeUID;
 
             //!
-            RANGE_DATA( ASTOp* parentOp, Ptr<ASTOp> childOp, const string& name, const string& uid )
+            FRangeData( ASTOp* parentOp, Ptr<ASTOp> childOp, const string& name, const string& uid )
                 : rangeSize( parentOp, childOp )
                 , rangeName(name)
                 , rangeUID(uid)
             {
             }
 
-            RANGE_DATA(const RANGE_DATA&) = delete;
-            RANGE_DATA& operator=(const RANGE_DATA&) = delete;
-            RANGE_DATA& operator=(RANGE_DATA&&) = delete;
+            FRangeData(const FRangeData&) = delete;
+            FRangeData& operator=(const FRangeData&) = delete;
+            FRangeData& operator=(FRangeData&&) = delete;
 
 
             // move constructor
-            RANGE_DATA(RANGE_DATA&& rhs)
-                 : rangeSize(std::move(rhs.rangeSize))
-                 , rangeName(std::move(rhs.rangeName))
-                 , rangeUID(std::move(rhs.rangeUID))
+            FRangeData(FRangeData&& rhs)
+                 : rangeSize(MoveTemp(rhs.rangeSize))
+                 , rangeName(MoveTemp(rhs.rangeName))
+                 , rangeUID(MoveTemp(rhs.rangeUID))
             {
             }
 
             //!
 
             //!
-            bool operator==(const RANGE_DATA& o) const
+            bool operator==(const FRangeData& o) const
             {
                 return rangeSize==o.rangeSize
                         &&
@@ -512,26 +512,10 @@ namespace mu
         };
 
         //!
-        static void LinkRange( PROGRAM& program,
-                               const RANGE_DATA& range,
-                               OP::ADDRESS& rangeSize,
-                               uint16& rangeId )
-        {
-            if (range.rangeSize)
-            {
-                if (range.rangeSize->linkedRange<0)
-                {
-                    check( program.m_ranges.Num()<255 );
-                    range.rangeSize->linkedRange = int8( program.m_ranges.Num() );
-                    RANGE_DESC rangeData;
-                    rangeData.m_name = range.rangeName;
-                    rangeData.m_uid = range.rangeUID;
-                    program.m_ranges.Add( rangeData );
-                }
-                rangeSize = range.rangeSize->linkedAddress;
-                rangeId = uint16(range.rangeSize->linkedRange);
-            }
-        }
+		static void LinkRange(FProgram& program,
+			const FRangeData& range,
+			OP::ADDRESS& rangeSize,
+			uint16& rangeId);
 
     public:
 
@@ -712,12 +696,12 @@ namespace mu
             BET_UNKNOWN,
             BET_TRUE,
             BET_FALSE,
-        } BOOL_EVAL_RESULT;
+        } FBoolEvalResult;
 
 
-        using EVALUATE_BOOL_CACHE=std::unordered_map<const ASTOp*,BOOL_EVAL_RESULT>;
+        using FEvaluateBoolCache = std::unordered_map<const ASTOp*,FBoolEvalResult>;
 
-        virtual BOOL_EVAL_RESULT EvaluateBool( ASTOpList& /*facts*/, EVALUATE_BOOL_CACHE* = nullptr ) const
+        virtual FBoolEvalResult EvaluateBool( ASTOpList& /*facts*/, FEvaluateBoolCache* = nullptr ) const
         {
             check(false);
             return BET_UNKNOWN;
@@ -982,7 +966,7 @@ namespace mu
 
         Ptr<ASTOp> Clone( MapChildFuncRef mapChild ) const override;
         void ForEachChild( const TFunctionRef<void(ASTChild&)> f ) override;
-        void Link( PROGRAM& program, const FLinkerOptions* Options) override;
+        void Link( FProgram& program, const FLinkerOptions* Options) override;
         bool IsEqual(const ASTOp& otherUntyped) const override;
 		uint64 Hash() const override;
 		FImageDesc GetImageDesc( bool returnBestOption, GetImageDescContext* context ) override;
@@ -992,7 +976,7 @@ namespace mu
         Ptr<ASTOp> OptimiseSize() const override;
         Ptr<ASTOp> OptimiseSemantic(const MODEL_OPTIMIZATION_OPTIONS&) const override;
         Ptr<ASTOp> OptimiseSink(const MODEL_OPTIMIZATION_OPTIONS&, OPTIMIZE_SINK_CONTEXT&) const override;
-        BOOL_EVAL_RESULT EvaluateBool( ASTOpList& facts, EVALUATE_BOOL_CACHE* cache ) const override;
+        FBoolEvalResult EvaluateBool( ASTOpList& facts, FEvaluateBoolCache* cache ) const override;
         int EvaluateInt( ASTOpList& facts, bool &unknown ) const override;
         bool IsImagePlainConstant( vec4<float>& colour ) const override;
         bool IsColourConstant( vec4<float>& colour ) const override;
diff --git a/Engine/Plugins/Experimental/Mutable/Source/MutableTools/Private/MuT/ASTOpAddLOD.cpp b/Engine/Plugins/Experimental/Mutable/Source/MutableTools/Private/MuT/ASTOpAddLOD.cpp
--- a/Engine/Plugins/Experimental/Mutable/Source/MutableTools/Private/MuT/ASTOpAddLOD.cpp
+++ b/Engine/Plugins/Experimental/Mutable/Source/MutableTools/Private/MuT/ASTOpAddLOD.cpp
@@ -71,7 +71,7 @@ namespace mu
 	}
 
 
-	void ASTOpAddLOD::Link(PROGRAM& program, const FLinkerOptions*)
+	void ASTOpAddLOD::Link(FProgram& program, const FLinkerOptions*)
 	{
 		// Already linked?
 		if (!linkedAddress)
diff --git a/Engine/Plugins/Experimental/Mutable/Source/MutableTools/Private/MuT/ASTOpAddLOD.h b/Engine/Plugins/Experimental/Mutable/Source/MutableTools/Private/MuT/ASTOpAddLOD.h
--- a/Engine/Plugins/Experimental/Mutable/Source/MutableTools/Private/MuT/ASTOpAddLOD.h
+++ b/Engine/Plugins/Experimental/Mutable/Source/MutableTools/Private/MuT/ASTOpAddLOD.h
@@ -11,7 +11,7 @@
 
 namespace mu
 {
-struct PROGRAM;
+struct FProgram;
 
 
 	//---------------------------------------------------------------------------------------------
@@ -34,7 +34,7 @@ struct PROGRAM;
 		void ForEachChild(const TFunctionRef<void(ASTChild&)>) override;
 		bool IsEqual(const ASTOp& otherUntyped) const override;
 		Ptr<ASTOp> Clone(MapChildFuncRef mapChild) const override;
-		void Link(PROGRAM& program, const FLinkerOptions* Options) override;
+		void Link(FProgram& program, const FLinkerOptions* Options) override;
 	};
 
 
diff --git a/Engine/Plugins/Experimental/Mutable/Source/MutableTools/Private/MuT/ASTOpConditional.cpp b/Engine/Plugins/Experimental/Mutable/Source/MutableTools/Private/MuT/ASTOpConditional.cpp
--- a/Engine/Plugins/Experimental/Mutable/Source/MutableTools/Private/MuT/ASTOpConditional.cpp
+++ b/Engine/Plugins/Experimental/Mutable/Source/MutableTools/Private/MuT/ASTOpConditional.cpp
@@ -104,7 +104,7 @@ namespace mu
 
 
 	//-------------------------------------------------------------------------------------------------
-	void ASTOpConditional::Link(PROGRAM& program, const FLinkerOptions*)
+	void ASTOpConditional::Link(FProgram& program, const FLinkerOptions*)
 	{
 		// Already linked?
 		if (!linkedAddress)
diff --git a/Engine/Plugins/Experimental/Mutable/Source/MutableTools/Private/MuT/ASTOpConditional.h b/Engine/Plugins/Experimental/Mutable/Source/MutableTools/Private/MuT/ASTOpConditional.h
--- a/Engine/Plugins/Experimental/Mutable/Source/MutableTools/Private/MuT/ASTOpConditional.h
+++ b/Engine/Plugins/Experimental/Mutable/Source/MutableTools/Private/MuT/ASTOpConditional.h
@@ -11,7 +11,7 @@
 
 namespace mu
 {
-struct PROGRAM;
+struct FProgram;
 
 	//---------------------------------------------------------------------------------------------
 	//! Conditional operation.
@@ -43,7 +43,7 @@ struct PROGRAM;
 		uint64 Hash() const override;
 		void Assert() override;
 		void ForEachChild(const TFunctionRef<void(ASTChild&)> f) override;
-		void Link(PROGRAM& program, const FLinkerOptions* Options) override;
+		void Link(FProgram& program, const FLinkerOptions* Options) override;
 		FImageDesc GetImageDesc(bool returnBestOption, class GetImageDescContext* context) override;
 		void GetLayoutBlockSize(int* pBlockX, int* pBlockY) override;
 		void GetBlockLayoutSize(int blockIndex, int* pBlockX, int* pBlockY,
diff --git a/Engine/Plugins/Experimental/Mutable/Source/MutableTools/Private/MuT/ASTOpConstantBool.cpp b/Engine/Plugins/Experimental/Mutable/Source/MutableTools/Private/MuT/ASTOpConstantBool.cpp
--- a/Engine/Plugins/Experimental/Mutable/Source/MutableTools/Private/MuT/ASTOpConstantBool.cpp
+++ b/Engine/Plugins/Experimental/Mutable/Source/MutableTools/Private/MuT/ASTOpConstantBool.cpp
@@ -54,7 +54,7 @@ namespace mu
 	}
 
 
-	void ASTOpConstantBool::Link(PROGRAM& program, const FLinkerOptions*)
+	void ASTOpConstantBool::Link(FProgram& program, const FLinkerOptions*)
 	{
 		if (!linkedAddress)
 		{
@@ -71,7 +71,7 @@ namespace mu
 	}
 
 
-	ASTOp::BOOL_EVAL_RESULT ASTOpConstantBool::EvaluateBool(ASTOpList&, EVALUATE_BOOL_CACHE*) const
+	ASTOp::FBoolEvalResult ASTOpConstantBool::EvaluateBool(ASTOpList&, FEvaluateBoolCache*) const
 	{
 		return value ? BET_TRUE : BET_FALSE;
 	}
diff --git a/Engine/Plugins/Experimental/Mutable/Source/MutableTools/Private/MuT/ASTOpConstantBool.h b/Engine/Plugins/Experimental/Mutable/Source/MutableTools/Private/MuT/ASTOpConstantBool.h
--- a/Engine/Plugins/Experimental/Mutable/Source/MutableTools/Private/MuT/ASTOpConstantBool.h
+++ b/Engine/Plugins/Experimental/Mutable/Source/MutableTools/Private/MuT/ASTOpConstantBool.h
@@ -10,7 +10,7 @@
 
 namespace mu
 {
-struct PROGRAM;
+struct FProgram;
 
 
 	//---------------------------------------------------------------------------------------------
@@ -32,8 +32,8 @@ struct PROGRAM;
 		void ForEachChild(const TFunctionRef<void(ASTChild&)>) override;
 		bool IsEqual(const ASTOp& otherUntyped) const override;
 		Ptr<ASTOp> Clone(MapChildFuncRef mapChild) const override;
-		void Link(PROGRAM& program, const FLinkerOptions* Options) override;
-		BOOL_EVAL_RESULT EvaluateBool(ASTOpList& facts, EVALUATE_BOOL_CACHE* cache) const override;
+		void Link(FProgram& program, const FLinkerOptions* Options) override;
+		FBoolEvalResult EvaluateBool(ASTOpList& facts, FEvaluateBoolCache* cache) const override;
 	};
 
 
diff --git a/Engine/Plugins/Experimental/Mutable/Source/MutableTools/Private/MuT/ASTOpConstantProjector.cpp b/Engine/Plugins/Experimental/Mutable/Source/MutableTools/Private/MuT/ASTOpConstantProjector.cpp
--- a/Engine/Plugins/Experimental/Mutable/Source/MutableTools/Private/MuT/ASTOpConstantProjector.cpp
+++ b/Engine/Plugins/Experimental/Mutable/Source/MutableTools/Private/MuT/ASTOpConstantProjector.cpp
@@ -47,7 +47,7 @@ namespace mu
 	}
 
 
-	void ASTOpConstantProjector::Link(PROGRAM& program, const FLinkerOptions*)
+	void ASTOpConstantProjector::Link(FProgram& program, const FLinkerOptions*)
 	{
 		if (!linkedAddress)
 		{
diff --git a/Engine/Plugins/Experimental/Mutable/Source/MutableTools/Private/MuT/ASTOpConstantProjector.h b/Engine/Plugins/Experimental/Mutable/Source/MutableTools/Private/MuT/ASTOpConstantProjector.h
--- a/Engine/Plugins/Experimental/Mutable/Source/MutableTools/Private/MuT/ASTOpConstantProjector.h
+++ b/Engine/Plugins/Experimental/Mutable/Source/MutableTools/Private/MuT/ASTOpConstantProjector.h
@@ -11,7 +11,7 @@
 
 namespace mu
 {
-struct PROGRAM;
+struct FProgram;
 
 
 	//---------------------------------------------------------------------------------------------
@@ -22,7 +22,7 @@ struct PROGRAM;
 	public:
 
 		//!
-		PROJECTOR value;
+		FProjector value;
 
 	public:
 
@@ -31,7 +31,7 @@ struct PROGRAM;
 		void ForEachChild(const TFunctionRef<void(ASTChild&)>) override;
 		bool IsEqual(const ASTOp& otherUntyped) const override;
 		Ptr<ASTOp> Clone(MapChildFuncRef mapChild) const override;
-		void Link(PROGRAM& program, const FLinkerOptions*) override;
+		void Link(FProgram& program, const FLinkerOptions*) override;
 	};
 
 
diff --git a/Engine/Plugins/Experimental/Mutable/Source/MutableTools/Private/MuT/ASTOpConstantResource.cpp b/Engine/Plugins/Experimental/Mutable/Source/MutableTools/Private/MuT/ASTOpConstantResource.cpp
--- a/Engine/Plugins/Experimental/Mutable/Source/MutableTools/Private/MuT/ASTOpConstantResource.cpp
+++ b/Engine/Plugins/Experimental/Mutable/Source/MutableTools/Private/MuT/ASTOpConstantResource.cpp
@@ -66,7 +66,7 @@ namespace mu
 
 
 	//-------------------------------------------------------------------------------------------------
-	void ASTOpConstantResource::Link(PROGRAM& program, const FLinkerOptions* Options)
+	void ASTOpConstantResource::Link(FProgram& program, const FLinkerOptions* Options)
 	{
 		if (!linkedAddress && !bLinkedAndNull)
 		{
diff --git a/Engine/Plugins/Experimental/Mutable/Source/MutableTools/Private/MuT/ASTOpConstantResource.h b/Engine/Plugins/Experimental/Mutable/Source/MutableTools/Private/MuT/ASTOpConstantResource.h
--- a/Engine/Plugins/Experimental/Mutable/Source/MutableTools/Private/MuT/ASTOpConstantResource.h
+++ b/Engine/Plugins/Experimental/Mutable/Source/MutableTools/Private/MuT/ASTOpConstantResource.h
@@ -12,7 +12,7 @@
 
 namespace mu
 {
-struct PROGRAM;
+struct FProgram;
 template <class SCALAR> class vec4;
 
 
@@ -64,7 +64,7 @@ template <class SCALAR> class vec4;
 		bool IsEqual(const ASTOp& otherUntyped) const override;
 		Ptr<ASTOp> Clone(MapChildFuncRef mapChild) const override;
 		uint64 Hash() const override;
-		void Link(PROGRAM& program, const FLinkerOptions*) override;
+		void Link(FProgram& program, const FLinkerOptions*) override;
 		FImageDesc GetImageDesc(bool, class GetImageDescContext*) override;
 		void GetBlockLayoutSize(int blockIndex, int* pBlockX, int* pBlockY,
 			BLOCK_LAYOUT_SIZE_CACHE* cache) override;
diff --git a/Engine/Plugins/Experimental/Mutable/Source/MutableTools/Private/MuT/ASTOpConstantString.cpp b/Engine/Plugins/Experimental/Mutable/Source/MutableTools/Private/MuT/ASTOpConstantString.cpp
--- a/Engine/Plugins/Experimental/Mutable/Source/MutableTools/Private/MuT/ASTOpConstantString.cpp
+++ b/Engine/Plugins/Experimental/Mutable/Source/MutableTools/Private/MuT/ASTOpConstantString.cpp
@@ -47,7 +47,7 @@ namespace mu
 	}
 
 
-	void ASTOpConstantString::Link(PROGRAM& program, const FLinkerOptions*)
+	void ASTOpConstantString::Link(FProgram& program, const FLinkerOptions*)
 	{
 		if (!linkedAddress)
 		{
diff --git a/Engine/Plugins/Experimental/Mutable/Source/MutableTools/Private/MuT/ASTOpConstantString.h b/Engine/Plugins/Experimental/Mutable/Source/MutableTools/Private/MuT/ASTOpConstantString.h
--- a/Engine/Plugins/Experimental/Mutable/Source/MutableTools/Private/MuT/ASTOpConstantString.h
+++ b/Engine/Plugins/Experimental/Mutable/Source/MutableTools/Private/MuT/ASTOpConstantString.h
@@ -11,7 +11,7 @@
 
 namespace mu
 {
-struct PROGRAM;
+struct FProgram;
 
 
 	//---------------------------------------------------------------------------------------------
@@ -29,7 +29,7 @@ struct PROGRAM;
 		void ForEachChild(const TFunctionRef<void(ASTChild&)>) override;
 		bool IsEqual(const ASTOp& otherUntyped) const override;
 		Ptr<ASTOp> Clone(MapChildFuncRef mapChild) const override;
-		void Link(PROGRAM& program, const FLinkerOptions* Options) override;
+		void Link(FProgram& program, const FLinkerOptions* Options) override;
 	};
 
 
diff --git a/Engine/Plugins/Experimental/Mutable/Source/MutableTools/Private/MuT/ASTOpImageCompose.cpp b/Engine/Plugins/Experimental/Mutable/Source/MutableTools/Private/MuT/ASTOpImageCompose.cpp
--- a/Engine/Plugins/Experimental/Mutable/Source/MutableTools/Private/MuT/ASTOpImageCompose.cpp
+++ b/Engine/Plugins/Experimental/Mutable/Source/MutableTools/Private/MuT/ASTOpImageCompose.cpp
@@ -96,7 +96,7 @@ void ASTOpImageCompose::ForEachChild(const TFunctionRef<void(ASTChild&)> f )
 
 
 //-------------------------------------------------------------------------------------------------
-void ASTOpImageCompose::Link( PROGRAM& program, const FLinkerOptions*)
+void ASTOpImageCompose::Link( FProgram& program, const FLinkerOptions*)
 {
     // Already linked?
     if (!linkedAddress)
diff --git a/Engine/Plugins/Experimental/Mutable/Source/MutableTools/Private/MuT/ASTOpImageCompose.h b/Engine/Plugins/Experimental/Mutable/Source/MutableTools/Private/MuT/ASTOpImageCompose.h
--- a/Engine/Plugins/Experimental/Mutable/Source/MutableTools/Private/MuT/ASTOpImageCompose.h
+++ b/Engine/Plugins/Experimental/Mutable/Source/MutableTools/Private/MuT/ASTOpImageCompose.h
@@ -11,7 +11,7 @@
 
 namespace mu
 {
-struct PROGRAM;
+struct FProgram;
 template <class SCALAR> class vec4;
 
 	class ASTOpImageCompose : public ASTOp
@@ -35,7 +35,7 @@ template <class SCALAR> class vec4;
 		bool IsEqual(const ASTOp& otherUntyped) const override;
 		Ptr<ASTOp> Clone(MapChildFuncRef mapChild) const override;
 		void ForEachChild(const TFunctionRef<void(ASTChild&)>) override;
-		void Link(PROGRAM& program, const FLinkerOptions* Options) override;
+		void Link(FProgram& program, const FLinkerOptions* Options) override;
 		//TODO: Ptr<ASTOp> OptimiseSink(const MODEL_OPTIMIZATION_OPTIONS& options, OPTIMIZE_SINK_CONTEXT& context) const override;
 		Ptr<ASTOp> OptimiseSemantic(const MODEL_OPTIMIZATION_OPTIONS& options) const;
 		FImageDesc GetImageDesc(bool returnBestOption, GetImageDescContext* context) override;
diff --git a/Engine/Plugins/Experimental/Mutable/Source/MutableTools/Private/MuT/ASTOpImageMipmap.cpp b/Engine/Plugins/Experimental/Mutable/Source/MutableTools/Private/MuT/ASTOpImageMipmap.cpp
--- a/Engine/Plugins/Experimental/Mutable/Source/MutableTools/Private/MuT/ASTOpImageMipmap.cpp
+++ b/Engine/Plugins/Experimental/Mutable/Source/MutableTools/Private/MuT/ASTOpImageMipmap.cpp
@@ -86,7 +86,7 @@ void ASTOpImageMipmap::ForEachChild(const TFunctionRef<void(ASTChild&)> f )
 }
 
 
-void ASTOpImageMipmap::Link( PROGRAM& program, const FLinkerOptions* )
+void ASTOpImageMipmap::Link( FProgram& program, const FLinkerOptions* )
 {
     // Already linked?
     if (!linkedAddress)
diff --git a/Engine/Plugins/Experimental/Mutable/Source/MutableTools/Private/MuT/ASTOpImageMipmap.h b/Engine/Plugins/Experimental/Mutable/Source/MutableTools/Private/MuT/ASTOpImageMipmap.h
--- a/Engine/Plugins/Experimental/Mutable/Source/MutableTools/Private/MuT/ASTOpImageMipmap.h
+++ b/Engine/Plugins/Experimental/Mutable/Source/MutableTools/Private/MuT/ASTOpImageMipmap.h
@@ -11,7 +11,7 @@
 
 namespace mu
 {
-struct PROGRAM;
+struct FProgram;
 template <class SCALAR> class vec4;
 
 	class ASTOpImageMipmap : public ASTOp
@@ -46,7 +46,7 @@ template <class SCALAR> class vec4;
 		bool IsEqual(const ASTOp& otherUntyped) const override;
 		Ptr<ASTOp> Clone(MapChildFuncRef mapChild) const override;
 		void ForEachChild(const TFunctionRef<void(ASTChild&)>) override;
-		void Link(PROGRAM& program, const FLinkerOptions* Options) override;
+		void Link(FProgram& program, const FLinkerOptions* Options) override;
 		Ptr<ASTOp> OptimiseSink(const MODEL_OPTIMIZATION_OPTIONS& options, OPTIMIZE_SINK_CONTEXT& context) const override;
 		FImageDesc GetImageDesc(bool returnBestOption, GetImageDescContext* context) override;
 		void GetLayoutBlockSize(int* pBlockX, int* pBlockY) override;
diff --git a/Engine/Plugins/Experimental/Mutable/Source/MutableTools/Private/MuT/ASTOpImageMultiLayer.cpp b/Engine/Plugins/Experimental/Mutable/Source/MutableTools/Private/MuT/ASTOpImageMultiLayer.cpp
--- a/Engine/Plugins/Experimental/Mutable/Source/MutableTools/Private/MuT/ASTOpImageMultiLayer.cpp
+++ b/Engine/Plugins/Experimental/Mutable/Source/MutableTools/Private/MuT/ASTOpImageMultiLayer.cpp
@@ -88,7 +88,7 @@ namespace mu
 
 
 	//-------------------------------------------------------------------------------------------------
-	void ASTOpImageMultiLayer::Link(PROGRAM& program, const FLinkerOptions*)
+	void ASTOpImageMultiLayer::Link(FProgram& program, const FLinkerOptions*)
 	{
 		// Already linked?
 		if (!linkedAddress)
diff --git a/Engine/Plugins/Experimental/Mutable/Source/MutableTools/Private/MuT/ASTOpImageMultiLayer.h b/Engine/Plugins/Experimental/Mutable/Source/MutableTools/Private/MuT/ASTOpImageMultiLayer.h
--- a/Engine/Plugins/Experimental/Mutable/Source/MutableTools/Private/MuT/ASTOpImageMultiLayer.h
+++ b/Engine/Plugins/Experimental/Mutable/Source/MutableTools/Private/MuT/ASTOpImageMultiLayer.h
@@ -11,7 +11,7 @@
 
 namespace mu
 {
-struct PROGRAM;
+struct FProgram;
 
 
 	//---------------------------------------------------------------------------------------------
@@ -24,7 +24,7 @@ struct PROGRAM;
 		ASTChild base;
 		ASTChild blend;
 		ASTChild mask;
-		RANGE_DATA range;
+		FRangeData range;
 		EBlendType blendType;
 
 	public:
@@ -38,7 +38,7 @@ struct PROGRAM;
 		bool IsEqual(const ASTOp& otherUntyped) const override;
 		Ptr<ASTOp> Clone(MapChildFuncRef mapChild) const override;
 		void ForEachChild(const TFunctionRef<void(ASTChild&)>) override;
-		void Link(PROGRAM& program, const FLinkerOptions* Options) override;
+		void Link(FProgram& program, const FLinkerOptions* Options) override;
 		FImageDesc GetImageDesc(bool returnBestOption, GetImageDescContext* context) override;
 		void GetLayoutBlockSize(int* pBlockX, int* pBlockY) override;
 		Ptr<ImageSizeExpression> GetImageSizeExpression() const override;
diff --git a/Engine/Plugins/Experimental/Mutable/Source/MutableTools/Private/MuT/ASTOpImageNormalComposite.cpp b/Engine/Plugins/Experimental/Mutable/Source/MutableTools/Private/MuT/ASTOpImageNormalComposite.cpp
--- a/Engine/Plugins/Experimental/Mutable/Source/MutableTools/Private/MuT/ASTOpImageNormalComposite.cpp
+++ b/Engine/Plugins/Experimental/Mutable/Source/MutableTools/Private/MuT/ASTOpImageNormalComposite.cpp
@@ -71,7 +71,7 @@ namespace mu
 	}
 
 
-	void ASTOpImageNormalComposite::Link(PROGRAM& program, const FLinkerOptions*)
+	void ASTOpImageNormalComposite::Link(FProgram& program, const FLinkerOptions*)
 	{
 		// Already linked?
 		if (!linkedAddress)
diff --git a/Engine/Plugins/Experimental/Mutable/Source/MutableTools/Private/MuT/ASTOpImageNormalComposite.h b/Engine/Plugins/Experimental/Mutable/Source/MutableTools/Private/MuT/ASTOpImageNormalComposite.h
--- a/Engine/Plugins/Experimental/Mutable/Source/MutableTools/Private/MuT/ASTOpImageNormalComposite.h
+++ b/Engine/Plugins/Experimental/Mutable/Source/MutableTools/Private/MuT/ASTOpImageNormalComposite.h
@@ -11,7 +11,7 @@
 
 namespace mu
 {
-struct PROGRAM;
+struct FProgram;
 
 	class ASTOpImageNormalComposite : public ASTOp
 	{
@@ -32,7 +32,7 @@ struct PROGRAM;
 		bool IsEqual(const ASTOp& otherUntyped) const override;
 		Ptr<ASTOp> Clone(MapChildFuncRef mapChild) const override;
 		void ForEachChild(const TFunctionRef<void(ASTChild&)>) override;
-		void Link(PROGRAM& program, const FLinkerOptions* Options) override;
+		void Link(FProgram& program, const FLinkerOptions* Options) override;
 		FImageDesc GetImageDesc(bool returnBestOption, GetImageDescContext* context) override;
 		void GetLayoutBlockSize(int* pBlockX, int* pBlockY) override;
 		Ptr<ImageSizeExpression> GetImageSizeExpression() const override;
diff --git a/Engine/Plugins/Experimental/Mutable/Source/MutableTools/Private/MuT/ASTOpImagePatch.cpp b/Engine/Plugins/Experimental/Mutable/Source/MutableTools/Private/MuT/ASTOpImagePatch.cpp
--- a/Engine/Plugins/Experimental/Mutable/Source/MutableTools/Private/MuT/ASTOpImagePatch.cpp
+++ b/Engine/Plugins/Experimental/Mutable/Source/MutableTools/Private/MuT/ASTOpImagePatch.cpp
@@ -71,7 +71,7 @@ namespace mu
 
 
 	//-------------------------------------------------------------------------------------------------
-	void ASTOpImagePatch::Link(PROGRAM& program, const FLinkerOptions*)
+	void ASTOpImagePatch::Link(FProgram& program, const FLinkerOptions*)
 	{
 		// Already linked?
 		if (!linkedAddress)
diff --git a/Engine/Plugins/Experimental/Mutable/Source/MutableTools/Private/MuT/ASTOpImagePatch.h b/Engine/Plugins/Experimental/Mutable/Source/MutableTools/Private/MuT/ASTOpImagePatch.h
--- a/Engine/Plugins/Experimental/Mutable/Source/MutableTools/Private/MuT/ASTOpImagePatch.h
+++ b/Engine/Plugins/Experimental/Mutable/Source/MutableTools/Private/MuT/ASTOpImagePatch.h
@@ -12,7 +12,7 @@
 
 namespace mu
 {
-struct PROGRAM;
+struct FProgram;
 
 	class ASTOpImagePatch : public ASTOp
 	{
@@ -33,7 +33,7 @@ struct PROGRAM;
 		bool IsEqual(const ASTOp& otherUntyped) const override;
 		Ptr<ASTOp> Clone(MapChildFuncRef mapChild) const override;
 		void ForEachChild(const TFunctionRef<void(ASTChild&)>) override;
-		void Link(PROGRAM& program, const FLinkerOptions* Options) override;
+		void Link(FProgram& program, const FLinkerOptions* Options) override;
 		FImageDesc GetImageDesc(bool returnBestOption, GetImageDescContext* context) override;
 		Ptr<ImageSizeExpression> GetImageSizeExpression() const override;
 		//TODO: void GetLayoutBlockSize(int* pBlockX, int* pBlockY) override;
diff --git a/Engine/Plugins/Experimental/Mutable/Source/MutableTools/Private/MuT/ASTOpImagePixelFormat.cpp b/Engine/Plugins/Experimental/Mutable/Source/MutableTools/Private/MuT/ASTOpImagePixelFormat.cpp
--- a/Engine/Plugins/Experimental/Mutable/Source/MutableTools/Private/MuT/ASTOpImagePixelFormat.cpp
+++ b/Engine/Plugins/Experimental/Mutable/Source/MutableTools/Private/MuT/ASTOpImagePixelFormat.cpp
@@ -76,7 +76,7 @@ void ASTOpImagePixelFormat::ForEachChild(const TFunctionRef<void(ASTChild&)> f )
 }
 
 
-void ASTOpImagePixelFormat::Link( PROGRAM& program, const FLinkerOptions* )
+void ASTOpImagePixelFormat::Link( FProgram& program, const FLinkerOptions* )
 {
     // Already linked?
     if (!linkedAddress)
diff --git a/Engine/Plugins/Experimental/Mutable/Source/MutableTools/Private/MuT/ASTOpImagePixelFormat.h b/Engine/Plugins/Experimental/Mutable/Source/MutableTools/Private/MuT/ASTOpImagePixelFormat.h
--- a/Engine/Plugins/Experimental/Mutable/Source/MutableTools/Private/MuT/ASTOpImagePixelFormat.h
+++ b/Engine/Plugins/Experimental/Mutable/Source/MutableTools/Private/MuT/ASTOpImagePixelFormat.h
@@ -11,7 +11,7 @@
 
 namespace mu
 {
-struct PROGRAM;
+struct FProgram;
 template <class SCALAR> class vec4;
 
 	class ASTOpImagePixelFormat : public ASTOp
@@ -33,7 +33,7 @@ template <class SCALAR> class vec4;
 		bool IsEqual(const ASTOp& otherUntyped) const override;
 		Ptr<ASTOp> Clone(MapChildFuncRef mapChild) const override;
 		void ForEachChild(const TFunctionRef<void(ASTChild&)>) override;
-		void Link(PROGRAM& program, const FLinkerOptions* Options) override;
+		void Link(FProgram& program, const FLinkerOptions* Options) override;
 		Ptr<ASTOp> OptimiseSink(const MODEL_OPTIMIZATION_OPTIONS& options, OPTIMIZE_SINK_CONTEXT& context) const override;
 		FImageDesc GetImageDesc(bool returnBestOption, GetImageDescContext* context) override;
 		void GetLayoutBlockSize(int* pBlockX, int* pBlockY) override;
diff --git a/Engine/Plugins/Experimental/Mutable/Source/MutableTools/Private/MuT/ASTOpImageTransform.cpp b/Engine/Plugins/Experimental/Mutable/Source/MutableTools/Private/MuT/ASTOpImageTransform.cpp
--- a/Engine/Plugins/Experimental/Mutable/Source/MutableTools/Private/MuT/ASTOpImageTransform.cpp
+++ b/Engine/Plugins/Experimental/Mutable/Source/MutableTools/Private/MuT/ASTOpImageTransform.cpp
@@ -94,7 +94,7 @@ namespace mu
 
 
 	//-------------------------------------------------------------------------------------------------
-	void ASTOpImageTransform::Link(PROGRAM& program, const FLinkerOptions*)
+	void ASTOpImageTransform::Link(FProgram& program, const FLinkerOptions*)
 	{
 		// Already linked?
 		if (!linkedAddress)
diff --git a/Engine/Plugins/Experimental/Mutable/Source/MutableTools/Private/MuT/ASTOpImageTransform.h b/Engine/Plugins/Experimental/Mutable/Source/MutableTools/Private/MuT/ASTOpImageTransform.h
--- a/Engine/Plugins/Experimental/Mutable/Source/MutableTools/Private/MuT/ASTOpImageTransform.h
+++ b/Engine/Plugins/Experimental/Mutable/Source/MutableTools/Private/MuT/ASTOpImageTransform.h
@@ -11,7 +11,7 @@
 
 namespace mu
 {
-struct PROGRAM;
+struct FProgram;
 
 
 	//---------------------------------------------------------------------------------------------
@@ -39,7 +39,7 @@ struct PROGRAM;
 		bool IsEqual(const ASTOp& otherUntyped) const override;
 		Ptr<ASTOp> Clone(MapChildFuncRef mapChild) const override;
 		void ForEachChild(const TFunctionRef<void(ASTChild&)>) override;
-		void Link(PROGRAM& program, const FLinkerOptions* Options) override;
+		void Link(FProgram& program, const FLinkerOptions* Options) override;
 		FImageDesc GetImageDesc(bool returnBestOption, GetImageDescContext* context) override;
 		void GetLayoutBlockSize(int* pBlockX, int* pBlockY) override;
 		Ptr<ImageSizeExpression> GetImageSizeExpression() const override;
diff --git a/Engine/Plugins/Experimental/Mutable/Source/MutableTools/Private/MuT/ASTOpInstanceAdd.cpp b/Engine/Plugins/Experimental/Mutable/Source/MutableTools/Private/MuT/ASTOpInstanceAdd.cpp
--- a/Engine/Plugins/Experimental/Mutable/Source/MutableTools/Private/MuT/ASTOpInstanceAdd.cpp
+++ b/Engine/Plugins/Experimental/Mutable/Source/MutableTools/Private/MuT/ASTOpInstanceAdd.cpp
@@ -108,7 +108,7 @@ namespace mu
 
 
 	//-------------------------------------------------------------------------------------------------
-	void ASTOpInstanceAdd::Link(PROGRAM& program, const FLinkerOptions*)
+	void ASTOpInstanceAdd::Link(FProgram& program, const FLinkerOptions*)
 	{
 		// Already linked?
 		if (!linkedAddress)
diff --git a/Engine/Plugins/Experimental/Mutable/Source/MutableTools/Private/MuT/ASTOpInstanceAdd.h b/Engine/Plugins/Experimental/Mutable/Source/MutableTools/Private/MuT/ASTOpInstanceAdd.h
--- a/Engine/Plugins/Experimental/Mutable/Source/MutableTools/Private/MuT/ASTOpInstanceAdd.h
+++ b/Engine/Plugins/Experimental/Mutable/Source/MutableTools/Private/MuT/ASTOpInstanceAdd.h
@@ -11,7 +11,7 @@
 
 namespace mu
 {
-struct PROGRAM;
+struct FProgram;
 
 
 	//---------------------------------------------------------------------------------------------
@@ -43,7 +43,7 @@ struct PROGRAM;
 		Ptr<ASTOp> Clone(MapChildFuncRef mapChild) const override;
 		uint64 Hash() const override;
 		void Assert() override;
-		void Link(PROGRAM& program, const FLinkerOptions* Options) override;
+		void Link(FProgram& program, const FLinkerOptions* Options) override;
 	};
 
 
diff --git a/Engine/Plugins/Experimental/Mutable/Source/MutableTools/Private/MuT/ASTOpMeshApplyPose.cpp b/Engine/Plugins/Experimental/Mutable/Source/MutableTools/Private/MuT/ASTOpMeshApplyPose.cpp
--- a/Engine/Plugins/Experimental/Mutable/Source/MutableTools/Private/MuT/ASTOpMeshApplyPose.cpp
+++ b/Engine/Plugins/Experimental/Mutable/Source/MutableTools/Private/MuT/ASTOpMeshApplyPose.cpp
@@ -67,7 +67,7 @@ namespace mu
 	}
 
 
-	void ASTOpMeshApplyPose::Link(PROGRAM& program, const FLinkerOptions*)
+	void ASTOpMeshApplyPose::Link(FProgram& program, const FLinkerOptions*)
 	{
 		// Already linked?
 		if (!linkedAddress)
diff --git a/Engine/Plugins/Experimental/Mutable/Source/MutableTools/Private/MuT/ASTOpMeshApplyPose.h b/Engine/Plugins/Experimental/Mutable/Source/MutableTools/Private/MuT/ASTOpMeshApplyPose.h
--- a/Engine/Plugins/Experimental/Mutable/Source/MutableTools/Private/MuT/ASTOpMeshApplyPose.h
+++ b/Engine/Plugins/Experimental/Mutable/Source/MutableTools/Private/MuT/ASTOpMeshApplyPose.h
@@ -10,7 +10,7 @@
 
 namespace mu
 {
-struct PROGRAM;
+struct FProgram;
 
 
 	//---------------------------------------------------------------------------------------------
@@ -34,7 +34,7 @@ struct PROGRAM;
 		bool IsEqual(const ASTOp& otherUntyped) const override;
 		Ptr<ASTOp> Clone(MapChildFuncRef mapChild) const override;
 		void ForEachChild(const TFunctionRef<void(ASTChild&)>) override;
-		void Link(PROGRAM& program, const FLinkerOptions* Options) override;
+		void Link(FProgram& program, const FLinkerOptions* Options) override;
 	};
 
 
diff --git a/Engine/Plugins/Experimental/Mutable/Source/MutableTools/Private/MuT/ASTOpMeshApplyShape.cpp b/Engine/Plugins/Experimental/Mutable/Source/MutableTools/Private/MuT/ASTOpMeshApplyShape.cpp
--- a/Engine/Plugins/Experimental/Mutable/Source/MutableTools/Private/MuT/ASTOpMeshApplyShape.cpp
+++ b/Engine/Plugins/Experimental/Mutable/Source/MutableTools/Private/MuT/ASTOpMeshApplyShape.cpp
@@ -76,7 +76,7 @@ namespace mu
 	}
 
 
-	void ASTOpMeshApplyShape::Link(PROGRAM& program, const FLinkerOptions*)
+	void ASTOpMeshApplyShape::Link(FProgram& program, const FLinkerOptions*)
 	{
 		// Already linked?
 		if (!linkedAddress)
diff --git a/Engine/Plugins/Experimental/Mutable/Source/MutableTools/Private/MuT/ASTOpMeshApplyShape.h b/Engine/Plugins/Experimental/Mutable/Source/MutableTools/Private/MuT/ASTOpMeshApplyShape.h
--- a/Engine/Plugins/Experimental/Mutable/Source/MutableTools/Private/MuT/ASTOpMeshApplyShape.h
+++ b/Engine/Plugins/Experimental/Mutable/Source/MutableTools/Private/MuT/ASTOpMeshApplyShape.h
@@ -10,7 +10,7 @@
 
 namespace mu
 {
-struct PROGRAM;
+struct FProgram;
 
 	//---------------------------------------------------------------------------------------------
 	//!
@@ -36,7 +36,7 @@ struct PROGRAM;
 		bool IsEqual(const ASTOp& otherUntyped) const override;
 		Ptr<ASTOp> Clone(MapChildFuncRef mapChild) const override;
 		void ForEachChild(const TFunctionRef<void(ASTChild&)>) override;
-		void Link(PROGRAM& program, const FLinkerOptions* Options) override;
+		void Link(FProgram& program, const FLinkerOptions* Options) override;
 	};
 
 }
diff --git a/Engine/Plugins/Experimental/Mutable/Source/MutableTools/Private/MuT/ASTOpMeshBindShape.cpp b/Engine/Plugins/Experimental/Mutable/Source/MutableTools/Private/MuT/ASTOpMeshBindShape.cpp
--- a/Engine/Plugins/Experimental/Mutable/Source/MutableTools/Private/MuT/ASTOpMeshBindShape.cpp
+++ b/Engine/Plugins/Experimental/Mutable/Source/MutableTools/Private/MuT/ASTOpMeshBindShape.cpp
@@ -112,7 +112,7 @@ namespace mu
 	}
 
 
-	void ASTOpMeshBindShape::Link(PROGRAM& program, const FLinkerOptions*)
+	void ASTOpMeshBindShape::Link(FProgram& program, const FLinkerOptions*)
 	{
 		// Already linked?
 		if (!linkedAddress)
diff --git a/Engine/Plugins/Experimental/Mutable/Source/MutableTools/Private/MuT/ASTOpMeshBindShape.h b/Engine/Plugins/Experimental/Mutable/Source/MutableTools/Private/MuT/ASTOpMeshBindShape.h
--- a/Engine/Plugins/Experimental/Mutable/Source/MutableTools/Private/MuT/ASTOpMeshBindShape.h
+++ b/Engine/Plugins/Experimental/Mutable/Source/MutableTools/Private/MuT/ASTOpMeshBindShape.h
@@ -12,7 +12,7 @@
 
 namespace mu
 {
-struct PROGRAM;
+struct FProgram;
 
 
 	//---------------------------------------------------------------------------------------------
@@ -50,7 +50,7 @@ struct PROGRAM;
 		bool IsEqual(const ASTOp& otherUntyped) const override;
 		Ptr<ASTOp> Clone(MapChildFuncRef mapChild) const override;
 		void ForEachChild(const TFunctionRef<void(ASTChild&)>) override;
-		void Link(PROGRAM& program, const FLinkerOptions* Options) override;
+		void Link(FProgram& program, const FLinkerOptions* Options) override;
 		Ptr<ASTOp> OptimiseSink(const MODEL_OPTIMIZATION_OPTIONS&, OPTIMIZE_SINK_CONTEXT&) const override;
 
 	};
diff --git a/Engine/Plugins/Experimental/Mutable/Source/MutableTools/Private/MuT/ASTOpMeshClipDeform.cpp b/Engine/Plugins/Experimental/Mutable/Source/MutableTools/Private/MuT/ASTOpMeshClipDeform.cpp
--- a/Engine/Plugins/Experimental/Mutable/Source/MutableTools/Private/MuT/ASTOpMeshClipDeform.cpp
+++ b/Engine/Plugins/Experimental/Mutable/Source/MutableTools/Private/MuT/ASTOpMeshClipDeform.cpp
@@ -66,7 +66,7 @@ namespace mu
 	}
 
 
-	void ASTOpMeshClipDeform::Link(PROGRAM& program, const FLinkerOptions*)
+	void ASTOpMeshClipDeform::Link(FProgram& program, const FLinkerOptions*)
 	{
 		// Already linked?
 		if (!linkedAddress)
diff --git a/Engine/Plugins/Experimental/Mutable/Source/MutableTools/Private/MuT/ASTOpMeshClipDeform.h b/Engine/Plugins/Experimental/Mutable/Source/MutableTools/Private/MuT/ASTOpMeshClipDeform.h
--- a/Engine/Plugins/Experimental/Mutable/Source/MutableTools/Private/MuT/ASTOpMeshClipDeform.h
+++ b/Engine/Plugins/Experimental/Mutable/Source/MutableTools/Private/MuT/ASTOpMeshClipDeform.h
@@ -10,7 +10,7 @@
 
 namespace mu
 {
-struct PROGRAM;
+struct FProgram;
 
 
 	//---------------------------------------------------------------------------------------------
@@ -34,7 +34,7 @@ struct PROGRAM;
 		bool IsEqual(const ASTOp& otherUntyped) const override;
 		Ptr<ASTOp> Clone(MapChildFuncRef mapChild) const override;
 		void ForEachChild(const TFunctionRef<void(ASTChild&)>) override;
-		void Link(PROGRAM& program, const FLinkerOptions* Options) override;
+		void Link(FProgram& program, const FLinkerOptions* Options) override;
 	};
 
 
diff --git a/Engine/Plugins/Experimental/Mutable/Source/MutableTools/Private/MuT/ASTOpMeshClipMorphPlane.cpp b/Engine/Plugins/Experimental/Mutable/Source/MutableTools/Private/MuT/ASTOpMeshClipMorphPlane.cpp
--- a/Engine/Plugins/Experimental/Mutable/Source/MutableTools/Private/MuT/ASTOpMeshClipMorphPlane.cpp
+++ b/Engine/Plugins/Experimental/Mutable/Source/MutableTools/Private/MuT/ASTOpMeshClipMorphPlane.cpp
@@ -77,7 +77,7 @@ namespace mu
 	}
 
 
-	void ASTOpMeshClipMorphPlane::Link(PROGRAM& program, const FLinkerOptions*)
+	void ASTOpMeshClipMorphPlane::Link(FProgram& program, const FLinkerOptions*)
 	{
 		// Already linked?
 		if (!linkedAddress)
diff --git a/Engine/Plugins/Experimental/Mutable/Source/MutableTools/Private/MuT/ASTOpMeshClipMorphPlane.h b/Engine/Plugins/Experimental/Mutable/Source/MutableTools/Private/MuT/ASTOpMeshClipMorphPlane.h
--- a/Engine/Plugins/Experimental/Mutable/Source/MutableTools/Private/MuT/ASTOpMeshClipMorphPlane.h
+++ b/Engine/Plugins/Experimental/Mutable/Source/MutableTools/Private/MuT/ASTOpMeshClipMorphPlane.h
@@ -12,7 +12,7 @@
 
 namespace mu
 {
-struct PROGRAM;
+struct FProgram;
 
 	//---------------------------------------------------------------------------------------------
 	//!
@@ -23,8 +23,8 @@ struct PROGRAM;
 
 		ASTChild source;
 
-		SHAPE morphShape;
-		SHAPE selectionShape;
+		FShape morphShape;
+		FShape selectionShape;
 		string vertexSelectionBone;
 
 		OP::MeshClipMorphPlaneArgs::VERTEX_SELECTION_TYPE vertexSelectionType = OP::MeshClipMorphPlaneArgs::VS_NONE;
@@ -42,7 +42,7 @@ struct PROGRAM;
 		bool IsEqual(const ASTOp& otherUntyped) const override;
 		Ptr<ASTOp> Clone(MapChildFuncRef mapChild) const override;
 		void ForEachChild(const TFunctionRef<void(ASTChild&)>) override;
-		void Link(PROGRAM& program, const FLinkerOptions* Options) override;
+		void Link(FProgram& program, const FLinkerOptions* Options) override;
 	};
 
 
diff --git a/Engine/Plugins/Experimental/Mutable/Source/MutableTools/Private/MuT/ASTOpMeshDifference.cpp b/Engine/Plugins/Experimental/Mutable/Source/MutableTools/Private/MuT/ASTOpMeshDifference.cpp
--- a/Engine/Plugins/Experimental/Mutable/Source/MutableTools/Private/MuT/ASTOpMeshDifference.cpp
+++ b/Engine/Plugins/Experimental/Mutable/Source/MutableTools/Private/MuT/ASTOpMeshDifference.cpp
@@ -75,7 +75,7 @@ namespace mu
 
 
 	//---------------------------------------------------------------------------------------------
-	void ASTOpMeshDifference::Link(PROGRAM& program, const FLinkerOptions*)
+	void ASTOpMeshDifference::Link(FProgram& program, const FLinkerOptions*)
 	{
 		// Already linked?
 		if (!linkedAddress)
diff --git a/Engine/Plugins/Experimental/Mutable/Source/MutableTools/Private/MuT/ASTOpMeshDifference.h b/Engine/Plugins/Experimental/Mutable/Source/MutableTools/Private/MuT/ASTOpMeshDifference.h
--- a/Engine/Plugins/Experimental/Mutable/Source/MutableTools/Private/MuT/ASTOpMeshDifference.h
+++ b/Engine/Plugins/Experimental/Mutable/Source/MutableTools/Private/MuT/ASTOpMeshDifference.h
@@ -11,7 +11,7 @@
 
 namespace mu
 {
-	struct PROGRAM;
+	struct FProgram;
 
 	/** Calculate the difference between two meshes (usually to calculate a morph delta). */
 	class ASTOpMeshDifference : public ASTOp
@@ -52,7 +52,7 @@ namespace mu
 		void ForEachChild(const TFunctionRef<void(ASTChild&)>) override;
 		bool IsEqual(const ASTOp& otherUntyped) const override;
 		Ptr<ASTOp> Clone(MapChildFuncRef mapChild) const override;
-		void Link(PROGRAM& program, const FLinkerOptions*) override;
+		void Link(FProgram& program, const FLinkerOptions*) override;
 		Ptr<ASTOp> OptimiseSink(const MODEL_OPTIMIZATION_OPTIONS&, OPTIMIZE_SINK_CONTEXT&) const override;
 
 	};
diff --git a/Engine/Plugins/Experimental/Mutable/Source/MutableTools/Private/MuT/ASTOpMeshExtractLayoutBlocks.cpp b/Engine/Plugins/Experimental/Mutable/Source/MutableTools/Private/MuT/ASTOpMeshExtractLayoutBlocks.cpp
--- a/Engine/Plugins/Experimental/Mutable/Source/MutableTools/Private/MuT/ASTOpMeshExtractLayoutBlocks.cpp
+++ b/Engine/Plugins/Experimental/Mutable/Source/MutableTools/Private/MuT/ASTOpMeshExtractLayoutBlocks.cpp
@@ -69,7 +69,7 @@ namespace mu
 	}
 
 
-	void ASTOpMeshExtractLayoutBlocks::Link(PROGRAM& program, const FLinkerOptions*)
+	void ASTOpMeshExtractLayoutBlocks::Link(FProgram& program, const FLinkerOptions*)
 	{
 		// Already linked?
 		if (!linkedAddress)
diff --git a/Engine/Plugins/Experimental/Mutable/Source/MutableTools/Private/MuT/ASTOpMeshExtractLayoutBlocks.h b/Engine/Plugins/Experimental/Mutable/Source/MutableTools/Private/MuT/ASTOpMeshExtractLayoutBlocks.h
--- a/Engine/Plugins/Experimental/Mutable/Source/MutableTools/Private/MuT/ASTOpMeshExtractLayoutBlocks.h
+++ b/Engine/Plugins/Experimental/Mutable/Source/MutableTools/Private/MuT/ASTOpMeshExtractLayoutBlocks.h
@@ -11,7 +11,7 @@
 
 namespace mu
 {
-struct PROGRAM;
+struct FProgram;
 
 
 	//---------------------------------------------------------------------------------------------
@@ -43,7 +43,7 @@ struct PROGRAM;
 		bool IsEqual(const ASTOp& otherUntyped) const override;
 		Ptr<ASTOp> Clone(MapChildFuncRef mapChild) const override;
 		void Assert() override;
-		void Link(PROGRAM& program, const FLinkerOptions* Options) override;
+		void Link(FProgram& program, const FLinkerOptions* Options) override;
 
 	};
 
diff --git a/Engine/Plugins/Experimental/Mutable/Source/MutableTools/Private/MuT/ASTOpMeshFormat.cpp b/Engine/Plugins/Experimental/Mutable/Source/MutableTools/Private/MuT/ASTOpMeshFormat.cpp
--- a/Engine/Plugins/Experimental/Mutable/Source/MutableTools/Private/MuT/ASTOpMeshFormat.cpp
+++ b/Engine/Plugins/Experimental/Mutable/Source/MutableTools/Private/MuT/ASTOpMeshFormat.cpp
@@ -79,7 +79,7 @@ void ASTOpMeshFormat::ForEachChild(const TFunctionRef<void(ASTChild&)> f )
 }
 
 
-void ASTOpMeshFormat::Link( PROGRAM& program, const FLinkerOptions* )
+void ASTOpMeshFormat::Link( FProgram& program, const FLinkerOptions* )
 {
     // Already linked?
     if (!linkedAddress)
diff --git a/Engine/Plugins/Experimental/Mutable/Source/MutableTools/Private/MuT/ASTOpMeshFormat.h b/Engine/Plugins/Experimental/Mutable/Source/MutableTools/Private/MuT/ASTOpMeshFormat.h
--- a/Engine/Plugins/Experimental/Mutable/Source/MutableTools/Private/MuT/ASTOpMeshFormat.h
+++ b/Engine/Plugins/Experimental/Mutable/Source/MutableTools/Private/MuT/ASTOpMeshFormat.h
@@ -10,7 +10,7 @@
 
 namespace mu
 {
-struct PROGRAM;
+struct FProgram;
 	
 	class ASTOpMeshFormat : public ASTOp
 	{
@@ -31,7 +31,7 @@ struct PROGRAM;
 		bool IsEqual(const ASTOp& otherUntyped) const override;
 		Ptr<ASTOp> Clone(MapChildFuncRef mapChild) const override;
 		void ForEachChild(const TFunctionRef<void(ASTChild&)>) override;
-		void Link(PROGRAM& program, const FLinkerOptions* Options) override;
+		void Link(FProgram& program, const FLinkerOptions* Options) override;
 		Ptr<ASTOp> OptimiseSink(const MODEL_OPTIMIZATION_OPTIONS& options, OPTIMIZE_SINK_CONTEXT& context) const override;
 
 	};
diff --git a/Engine/Plugins/Experimental/Mutable/Source/MutableTools/Private/MuT/ASTOpMeshGeometryOperation.cpp b/Engine/Plugins/Experimental/Mutable/Source/MutableTools/Private/MuT/ASTOpMeshGeometryOperation.cpp
--- a/Engine/Plugins/Experimental/Mutable/Source/MutableTools/Private/MuT/ASTOpMeshGeometryOperation.cpp
+++ b/Engine/Plugins/Experimental/Mutable/Source/MutableTools/Private/MuT/ASTOpMeshGeometryOperation.cpp
@@ -71,7 +71,7 @@ namespace mu
 	}
 
 
-	void ASTOpMeshGeometryOperation::Link(PROGRAM& program, const FLinkerOptions*)
+	void ASTOpMeshGeometryOperation::Link(FProgram& program, const FLinkerOptions*)
 	{
 		// Already linked?
 		if (!linkedAddress)
diff --git a/Engine/Plugins/Experimental/Mutable/Source/MutableTools/Private/MuT/ASTOpMeshGeometryOperation.h b/Engine/Plugins/Experimental/Mutable/Source/MutableTools/Private/MuT/ASTOpMeshGeometryOperation.h
--- a/Engine/Plugins/Experimental/Mutable/Source/MutableTools/Private/MuT/ASTOpMeshGeometryOperation.h
+++ b/Engine/Plugins/Experimental/Mutable/Source/MutableTools/Private/MuT/ASTOpMeshGeometryOperation.h
@@ -10,7 +10,7 @@
 
 namespace mu
 {
-struct PROGRAM;
+struct FProgram;
 
 
 	//---------------------------------------------------------------------------------------------
@@ -36,7 +36,7 @@ struct PROGRAM;
 		bool IsEqual(const ASTOp& otherUntyped) const override;
 		Ptr<ASTOp> Clone(MapChildFuncRef mapChild) const override;
 		void ForEachChild(const TFunctionRef<void(ASTChild&)>) override;
-		void Link(PROGRAM& program, const FLinkerOptions* Options) override;
+		void Link(FProgram& program, const FLinkerOptions* Options) override;
 	};
 
 
diff --git a/Engine/Plugins/Experimental/Mutable/Source/MutableTools/Private/MuT/ASTOpMeshMaskClipMesh.cpp b/Engine/Plugins/Experimental/Mutable/Source/MutableTools/Private/MuT/ASTOpMeshMaskClipMesh.cpp
--- a/Engine/Plugins/Experimental/Mutable/Source/MutableTools/Private/MuT/ASTOpMeshMaskClipMesh.cpp
+++ b/Engine/Plugins/Experimental/Mutable/Source/MutableTools/Private/MuT/ASTOpMeshMaskClipMesh.cpp
@@ -74,7 +74,7 @@ namespace mu
 
 
 	//-------------------------------------------------------------------------------------------------
-	void ASTOpMeshMaskClipMesh::Link(PROGRAM& program, const FLinkerOptions*)
+	void ASTOpMeshMaskClipMesh::Link(FProgram& program, const FLinkerOptions*)
 	{
 		// Already linked?
 		if (!linkedAddress)
diff --git a/Engine/Plugins/Experimental/Mutable/Source/MutableTools/Private/MuT/ASTOpMeshMaskClipMesh.h b/Engine/Plugins/Experimental/Mutable/Source/MutableTools/Private/MuT/ASTOpMeshMaskClipMesh.h
--- a/Engine/Plugins/Experimental/Mutable/Source/MutableTools/Private/MuT/ASTOpMeshMaskClipMesh.h
+++ b/Engine/Plugins/Experimental/Mutable/Source/MutableTools/Private/MuT/ASTOpMeshMaskClipMesh.h
@@ -10,7 +10,7 @@
 
 namespace mu
 {
-struct PROGRAM;
+struct FProgram;
 
 
 	//---------------------------------------------------------------------------------------------
@@ -34,7 +34,7 @@ struct PROGRAM;
 		bool IsEqual(const ASTOp& otherUntyped) const override;
 		Ptr<ASTOp> Clone(MapChildFuncRef mapChild) const override;
 		void ForEachChild(const TFunctionRef<void(ASTChild&)>) override;
-		void Link(PROGRAM& program, const FLinkerOptions* Options) override;
+		void Link(FProgram& program, const FLinkerOptions* Options) override;
 		Ptr<ASTOp> OptimiseSink(const MODEL_OPTIMIZATION_OPTIONS&, OPTIMIZE_SINK_CONTEXT&) const override;
 
 	};
diff --git a/Engine/Plugins/Experimental/Mutable/Source/MutableTools/Private/MuT/ASTOpMeshMorph.cpp b/Engine/Plugins/Experimental/Mutable/Source/MutableTools/Private/MuT/ASTOpMeshMorph.cpp
--- a/Engine/Plugins/Experimental/Mutable/Source/MutableTools/Private/MuT/ASTOpMeshMorph.cpp
+++ b/Engine/Plugins/Experimental/Mutable/Source/MutableTools/Private/MuT/ASTOpMeshMorph.cpp
@@ -86,7 +86,7 @@ namespace mu
 
 
 	//---------------------------------------------------------------------------------------------
-	void ASTOpMeshMorph::Link(PROGRAM& program, const FLinkerOptions*)
+	void ASTOpMeshMorph::Link(FProgram& program, const FLinkerOptions*)
 	{
 		// Already linked?
 		if (!linkedAddress)
diff --git a/Engine/Plugins/Experimental/Mutable/Source/MutableTools/Private/MuT/ASTOpMeshMorph.h b/Engine/Plugins/Experimental/Mutable/Source/MutableTools/Private/MuT/ASTOpMeshMorph.h
--- a/Engine/Plugins/Experimental/Mutable/Source/MutableTools/Private/MuT/ASTOpMeshMorph.h
+++ b/Engine/Plugins/Experimental/Mutable/Source/MutableTools/Private/MuT/ASTOpMeshMorph.h
@@ -11,7 +11,7 @@
 
 namespace mu
 {
-	struct PROGRAM;
+	struct FProgram;
 
 	//---------------------------------------------------------------------------------------------
 	//! From a source mesh, remove a list of fragments with a condition.
@@ -41,7 +41,7 @@ namespace mu
 		void ForEachChild(const TFunctionRef<void(ASTChild&)>) override;
 		bool IsEqual(const ASTOp& otherUntyped) const override;
 		Ptr<ASTOp> Clone(MapChildFuncRef mapChild) const override;
-		void Link(PROGRAM& program, const FLinkerOptions*) override;
+		void Link(FProgram& program, const FLinkerOptions*) override;
 		Ptr<ASTOp> OptimiseSink(const MODEL_OPTIMIZATION_OPTIONS&, OPTIMIZE_SINK_CONTEXT&) const override;
 
 		// Own interface
diff --git a/Engine/Plugins/Experimental/Mutable/Source/MutableTools/Private/MuT/ASTOpMeshMorphReshape.cpp b/Engine/Plugins/Experimental/Mutable/Source/MutableTools/Private/MuT/ASTOpMeshMorphReshape.cpp
--- a/Engine/Plugins/Experimental/Mutable/Source/MutableTools/Private/MuT/ASTOpMeshMorphReshape.cpp
+++ b/Engine/Plugins/Experimental/Mutable/Source/MutableTools/Private/MuT/ASTOpMeshMorphReshape.cpp
@@ -67,7 +67,7 @@ namespace mu
 	}
 
 
-	void ASTOpMeshMorphReshape::Link(PROGRAM& program, const FLinkerOptions*)
+	void ASTOpMeshMorphReshape::Link(FProgram& program, const FLinkerOptions*)
 	{
 		// Already linked?
 		if (!linkedAddress)
diff --git a/Engine/Plugins/Experimental/Mutable/Source/MutableTools/Private/MuT/ASTOpMeshMorphReshape.h b/Engine/Plugins/Experimental/Mutable/Source/MutableTools/Private/MuT/ASTOpMeshMorphReshape.h
--- a/Engine/Plugins/Experimental/Mutable/Source/MutableTools/Private/MuT/ASTOpMeshMorphReshape.h
+++ b/Engine/Plugins/Experimental/Mutable/Source/MutableTools/Private/MuT/ASTOpMeshMorphReshape.h
@@ -10,7 +10,7 @@
 
 namespace mu
 {
-struct PROGRAM;
+struct FProgram;
 
 	//---------------------------------------------------------------------------------------------
 	//!
@@ -33,7 +33,7 @@ struct PROGRAM;
 		bool IsEqual(const ASTOp& otherUntyped) const override;
 		Ptr<ASTOp> Clone(MapChildFuncRef mapChild) const override;
 		void ForEachChild(const TFunctionRef<void(ASTChild&)>) override;
-		void Link(PROGRAM& program, const FLinkerOptions* Options) override;
+		void Link(FProgram& program, const FLinkerOptions* Options) override;
 	};
 
 
diff --git a/Engine/Plugins/Experimental/Mutable/Source/MutableTools/Private/MuT/ASTOpMeshRemapIndices.cpp b/Engine/Plugins/Experimental/Mutable/Source/MutableTools/Private/MuT/ASTOpMeshRemapIndices.cpp
--- a/Engine/Plugins/Experimental/Mutable/Source/MutableTools/Private/MuT/ASTOpMeshRemapIndices.cpp
+++ b/Engine/Plugins/Experimental/Mutable/Source/MutableTools/Private/MuT/ASTOpMeshRemapIndices.cpp
@@ -74,7 +74,7 @@ namespace mu
 
 
 	//-------------------------------------------------------------------------------------------------
-	void ASTOpMeshRemapIndices::Link(PROGRAM& program, const FLinkerOptions*)
+	void ASTOpMeshRemapIndices::Link(FProgram& program, const FLinkerOptions*)
 	{
 		// Already linked?
 		if (!linkedAddress)
diff --git a/Engine/Plugins/Experimental/Mutable/Source/MutableTools/Private/MuT/ASTOpMeshRemapIndices.h b/Engine/Plugins/Experimental/Mutable/Source/MutableTools/Private/MuT/ASTOpMeshRemapIndices.h
--- a/Engine/Plugins/Experimental/Mutable/Source/MutableTools/Private/MuT/ASTOpMeshRemapIndices.h
+++ b/Engine/Plugins/Experimental/Mutable/Source/MutableTools/Private/MuT/ASTOpMeshRemapIndices.h
@@ -10,7 +10,7 @@
 
 namespace mu
 {
-struct PROGRAM;
+struct FProgram;
 
 
 	//---------------------------------------------------------------------------------------------
@@ -38,7 +38,7 @@ struct PROGRAM;
 		void ForEachChild(const TFunctionRef<void(ASTChild&)>) override;
 		bool IsEqual(const ASTOp& otherUntyped) const override;
 		Ptr<ASTOp> Clone(MapChildFuncRef mapChild) const override;
-		void Link(PROGRAM& program, const FLinkerOptions*) override;
+		void Link(FProgram& program, const FLinkerOptions*) override;
 	};
 
 
diff --git a/Engine/Plugins/Experimental/Mutable/Source/MutableTools/Private/MuT/ASTOpMeshRemoveMask.cpp b/Engine/Plugins/Experimental/Mutable/Source/MutableTools/Private/MuT/ASTOpMeshRemoveMask.cpp
--- a/Engine/Plugins/Experimental/Mutable/Source/MutableTools/Private/MuT/ASTOpMeshRemoveMask.cpp
+++ b/Engine/Plugins/Experimental/Mutable/Source/MutableTools/Private/MuT/ASTOpMeshRemoveMask.cpp
@@ -88,7 +88,7 @@ namespace mu
 
 
 	//---------------------------------------------------------------------------------------------
-	void ASTOpMeshRemoveMask::Link(PROGRAM& program, const FLinkerOptions*)
+	void ASTOpMeshRemoveMask::Link(FProgram& program, const FLinkerOptions*)
 	{
 		// Already linked?
 		if (!linkedAddress)
diff --git a/Engine/Plugins/Experimental/Mutable/Source/MutableTools/Private/MuT/ASTOpMeshRemoveMask.h b/Engine/Plugins/Experimental/Mutable/Source/MutableTools/Private/MuT/ASTOpMeshRemoveMask.h
--- a/Engine/Plugins/Experimental/Mutable/Source/MutableTools/Private/MuT/ASTOpMeshRemoveMask.h
+++ b/Engine/Plugins/Experimental/Mutable/Source/MutableTools/Private/MuT/ASTOpMeshRemoveMask.h
@@ -11,7 +11,7 @@
 
 namespace mu
 {
-	struct PROGRAM;
+	struct FProgram;
 
 	//---------------------------------------------------------------------------------------------
 	//! From a source mesh, remove a list of fragments with a condition.
@@ -38,7 +38,7 @@ namespace mu
 		void ForEachChild(const TFunctionRef<void(ASTChild&)>) override;
 		bool IsEqual(const ASTOp& otherUntyped) const override;
 		Ptr<ASTOp> Clone(MapChildFuncRef mapChild) const override;
-		void Link(PROGRAM& program, const FLinkerOptions*) override;
+		void Link(FProgram& program, const FLinkerOptions*) override;
 		Ptr<ASTOp> OptimiseSink(const MODEL_OPTIMIZATION_OPTIONS&, OPTIMIZE_SINK_CONTEXT&) const override;
 
 		// Own interface
diff --git a/Engine/Plugins/Experimental/Mutable/Source/MutableTools/Private/MuT/ASTOpMeshTransform.cpp b/Engine/Plugins/Experimental/Mutable/Source/MutableTools/Private/MuT/ASTOpMeshTransform.cpp
--- a/Engine/Plugins/Experimental/Mutable/Source/MutableTools/Private/MuT/ASTOpMeshTransform.cpp
+++ b/Engine/Plugins/Experimental/Mutable/Source/MutableTools/Private/MuT/ASTOpMeshTransform.cpp
@@ -63,7 +63,7 @@ namespace mu
 	}
 
 
-	void ASTOpMeshTransform::Link(PROGRAM& program, const FLinkerOptions*)
+	void ASTOpMeshTransform::Link(FProgram& program, const FLinkerOptions*)
 	{
 		// Already linked?
 		if (!linkedAddress)
diff --git a/Engine/Plugins/Experimental/Mutable/Source/MutableTools/Private/MuT/ASTOpMeshTransform.h b/Engine/Plugins/Experimental/Mutable/Source/MutableTools/Private/MuT/ASTOpMeshTransform.h
--- a/Engine/Plugins/Experimental/Mutable/Source/MutableTools/Private/MuT/ASTOpMeshTransform.h
+++ b/Engine/Plugins/Experimental/Mutable/Source/MutableTools/Private/MuT/ASTOpMeshTransform.h
@@ -11,7 +11,7 @@
 
 namespace mu
 {
-struct PROGRAM;
+struct FProgram;
 
 
 	//---------------------------------------------------------------------------------------------
@@ -36,7 +36,7 @@ struct PROGRAM;
 		void ForEachChild(const TFunctionRef<void(ASTChild&)>) override;
 		bool IsEqual(const ASTOp& otherUntyped) const override;
 		Ptr<ASTOp> Clone(MapChildFuncRef mapChild) const override;
-		void Link(PROGRAM& program, const FLinkerOptions* Options) override;
+		void Link(FProgram& program, const FLinkerOptions* Options) override;
 
 	};
 
diff --git a/Engine/Plugins/Experimental/Mutable/Source/MutableTools/Private/MuT/ASTOpParameter.cpp b/Engine/Plugins/Experimental/Mutable/Source/MutableTools/Private/MuT/ASTOpParameter.cpp
--- a/Engine/Plugins/Experimental/Mutable/Source/MutableTools/Private/MuT/ASTOpParameter.cpp
+++ b/Engine/Plugins/Experimental/Mutable/Source/MutableTools/Private/MuT/ASTOpParameter.cpp
@@ -83,15 +83,16 @@ namespace mu
 	}
 
 
-	void ASTOpParameter::Link(PROGRAM& program, const FLinkerOptions*)
+	void ASTOpParameter::Link(FProgram& program, const FLinkerOptions*)
 	{
 		// Already linked?
 		if (!linkedAddress)
 		{
 			OP::ParameterArgs args;
-			memset(&args, 0, sizeof(args));
+			FMemory::Memzero(&args, sizeof(args));
 
-			args.variable = (OP::ADDRESS)program.m_parameters.Num();
+			LinkedParameterIndex = program.m_parameters.Num();
+			args.variable = (OP::ADDRESS)LinkedParameterIndex;
 			program.m_parameters.Add(parameter);
 
 			for (const auto& d : ranges)
@@ -146,8 +147,8 @@ namespace mu
 	}
 
 
-	ASTOp::BOOL_EVAL_RESULT ASTOpParameter::EvaluateBool(ASTOpList& /*facts*/,
-		EVALUATE_BOOL_CACHE*) const
+	ASTOp::FBoolEvalResult ASTOpParameter::EvaluateBool(ASTOpList& /*facts*/,
+		FEvaluateBoolCache*) const
 	{
 		check(type == OP_TYPE::BO_PARAMETER);
 		return BET_UNKNOWN;
diff --git a/Engine/Plugins/Experimental/Mutable/Source/MutableTools/Private/MuT/ASTOpParameter.h b/Engine/Plugins/Experimental/Mutable/Source/MutableTools/Private/MuT/ASTOpParameter.h
--- a/Engine/Plugins/Experimental/Mutable/Source/MutableTools/Private/MuT/ASTOpParameter.h
+++ b/Engine/Plugins/Experimental/Mutable/Source/MutableTools/Private/MuT/ASTOpParameter.h
@@ -13,7 +13,7 @@
 
 namespace mu
 {
-struct PROGRAM;
+	struct FProgram;
 
 
 	//---------------------------------------------------------------------------------------------
@@ -27,14 +27,17 @@ struct PROGRAM;
 		OP_TYPE type;
 
 		//!
-		PARAMETER_DESC parameter;
+		FParameterDesc parameter;
 
-		//! Ranges adding dimensions to this parameter
-		TArray<RANGE_DATA> ranges;
+		//** Ranges adding dimensions to this parameter. */
+		TArray<FRangeData> ranges;
 
 		//! Additional images attached to the parameter
 		TArray<ASTChild> additionalImages;
 
+		/** Index of the parameter in the program parameter list. Generated ar link time. */
+		int32 LinkedParameterIndex = -1;
+
 	public:
 
 		~ASTOpParameter() override;
@@ -45,9 +48,9 @@ struct PROGRAM;
 		bool IsEqual(const ASTOp& otherUntyped) const override;
 		Ptr<ASTOp> Clone(MapChildFuncRef mapChild) const override;
 		void Assert() override;
-		void Link(PROGRAM& program, const FLinkerOptions*) override;
+		void Link(FProgram& program, const FLinkerOptions*) override;
 		int EvaluateInt(ASTOpList& facts, bool& unknown) const override;
-		BOOL_EVAL_RESULT EvaluateBool(ASTOpList& /*facts*/, EVALUATE_BOOL_CACHE* = nullptr) const override;
+		FBoolEvalResult EvaluateBool(ASTOpList& /*facts*/, FEvaluateBoolCache* = nullptr) const override;
 		FImageDesc GetImageDesc(bool, GetImageDescContext*) override;
 
 	};
diff --git a/Engine/Plugins/Experimental/Mutable/Source/MutableTools/Private/MuT/ASTOpScalarCurve.cpp b/Engine/Plugins/Experimental/Mutable/Source/MutableTools/Private/MuT/ASTOpScalarCurve.cpp
--- a/Engine/Plugins/Experimental/Mutable/Source/MutableTools/Private/MuT/ASTOpScalarCurve.cpp
+++ b/Engine/Plugins/Experimental/Mutable/Source/MutableTools/Private/MuT/ASTOpScalarCurve.cpp
@@ -63,7 +63,7 @@ namespace mu
 	}
 
 
-	void ASTOpScalarCurve::Link(PROGRAM& program, const FLinkerOptions*)
+	void ASTOpScalarCurve::Link(FProgram& program, const FLinkerOptions*)
 	{
 		if (!linkedAddress)
 		{
diff --git a/Engine/Plugins/Experimental/Mutable/Source/MutableTools/Private/MuT/ASTOpScalarCurve.h b/Engine/Plugins/Experimental/Mutable/Source/MutableTools/Private/MuT/ASTOpScalarCurve.h
--- a/Engine/Plugins/Experimental/Mutable/Source/MutableTools/Private/MuT/ASTOpScalarCurve.h
+++ b/Engine/Plugins/Experimental/Mutable/Source/MutableTools/Private/MuT/ASTOpScalarCurve.h
@@ -11,7 +11,7 @@
 
 namespace mu
 {
-struct PROGRAM;
+struct FProgram;
 
 
 	//---------------------------------------------------------------------------------------------
@@ -37,7 +37,7 @@ struct PROGRAM;
 		void ForEachChild(const TFunctionRef<void(ASTChild&)>) override;
 		bool IsEqual(const ASTOp& otherUntyped) const override;
 		Ptr<ASTOp> Clone(MapChildFuncRef mapChild) const override;
-		void Link(PROGRAM& program, const FLinkerOptions* Options) override;
+		void Link(FProgram& program, const FLinkerOptions* Options) override;
 	};
 
 
diff --git a/Engine/Plugins/Experimental/Mutable/Source/MutableTools/Private/MuT/ASTOpSwitch.cpp b/Engine/Plugins/Experimental/Mutable/Source/MutableTools/Private/MuT/ASTOpSwitch.cpp
--- a/Engine/Plugins/Experimental/Mutable/Source/MutableTools/Private/MuT/ASTOpSwitch.cpp
+++ b/Engine/Plugins/Experimental/Mutable/Source/MutableTools/Private/MuT/ASTOpSwitch.cpp
@@ -167,7 +167,7 @@ namespace mu
 
 
 	//-------------------------------------------------------------------------------------------------
-	void ASTOpSwitch::Link(PROGRAM& program, const FLinkerOptions*)
+	void ASTOpSwitch::Link(FProgram& program, const FLinkerOptions*)
 	{
 		// Already linked?
 		if (!linkedAddress)
diff --git a/Engine/Plugins/Experimental/Mutable/Source/MutableTools/Private/MuT/ASTOpSwitch.h b/Engine/Plugins/Experimental/Mutable/Source/MutableTools/Private/MuT/ASTOpSwitch.h
--- a/Engine/Plugins/Experimental/Mutable/Source/MutableTools/Private/MuT/ASTOpSwitch.h
+++ b/Engine/Plugins/Experimental/Mutable/Source/MutableTools/Private/MuT/ASTOpSwitch.h
@@ -12,7 +12,7 @@
 
 namespace mu
 {
-struct PROGRAM;
+struct FProgram;
 template <class SCALAR> class vec4;
 
 
@@ -66,7 +66,7 @@ template <class SCALAR> class vec4;
 		bool IsEqual(const ASTOp& otherUntyped) const override;
 		Ptr<ASTOp> Clone(MapChildFuncRef mapChild) const override;
 		void Assert() override;
-		void Link(PROGRAM& program, const FLinkerOptions* Options) override;
+		void Link(FProgram& program, const FLinkerOptions* Options) override;
 		FImageDesc GetImageDesc(bool returnBestOption,
 			class GetImageDescContext* context) override;
 		void GetBlockLayoutSize(int blockIndex, int* pBlockX, int* pBlockY,
diff --git a/Engine/Plugins/Experimental/Mutable/Source/MutableTools/Private/MuT/CodeGenerator.cpp b/Engine/Plugins/Experimental/Mutable/Source/MutableTools/Private/MuT/CodeGenerator.cpp
--- a/Engine/Plugins/Experimental/Mutable/Source/MutableTools/Private/MuT/CodeGenerator.cpp
+++ b/Engine/Plugins/Experimental/Mutable/Source/MutableTools/Private/MuT/CodeGenerator.cpp
@@ -221,7 +221,7 @@ namespace mu
         m_pErrorLog = new ErrorLog;
 
         // Add the parent at the top of the hierarchy
-        m_currentParents.Add( PARENT_KEY() );
+        m_currentParents.Add( FParentKey() );
     }
 
     //---------------------------------------------------------------------------------------------
@@ -316,21 +316,21 @@ namespace mu
         // Temp by-passes while we remove the visitor pattern
 		if (auto ScalarNode = dynamic_cast<const NodeScalar*>(pNode.get()))
 		{
-			SCALAR_GENERATION_RESULT ScalarResult;
+			FScalarGenerationResult ScalarResult;
 			GenerateScalar(ScalarResult, ScalarNode);
 			return ScalarResult.op;
 		}
 
 		if (auto ColorNode = dynamic_cast<const NodeColour*>(pNode.get()))
 		{
-			COLOR_GENERATION_RESULT Result;
+			FColorGenerationResult Result;
 			GenerateColor(Result, ColorNode);
 			return Result.op;
 		}
 
 		if (auto ImageNode = dynamic_cast<const NodeImage*>(pNode.get()))
 		{
-			IMAGE_GENERATION_RESULT Result;
+			FImageGenerationResult Result;
 			GenerateImage(Result, ImageNode);
 			return Result.op;
 		}
@@ -350,7 +350,7 @@ namespace mu
 
 		if ( auto projNode = dynamic_cast<const NodeProjector*>(pNode.get()) )
         {
-            PROJECTOR_GENERATION_RESULT ProjResult;
+            FProjectorGenerationResult ProjResult;
             GenerateProjector( ProjResult, projNode );
             return ProjResult.op;
         }
@@ -358,7 +358,7 @@ namespace mu
         if ( auto surfNode = dynamic_cast<const NodeSurfaceNew*>(pNode.get()) )
         {
             // This happens only if we generate a node graph that has a NodeSurfaceNew at the root.
-            SURFACE_GENERATION_RESULT surfResult;
+            FSurfaceGenerationResult surfResult;
             const TArray<FirstPassGenerator::SURFACE::EDIT> edits;
             GenerateSurface( surfResult, surfNode, edits );
             return surfResult.surfaceOp;
@@ -386,7 +386,7 @@ namespace mu
         Ptr<ASTOp> result;
 
         // See if it was already generated
-		VISITED_MAP_KEY key = GetCurrentCacheKey(pNode);
+		FVisitedKeyMap key = GetCurrentCacheKey(pNode);
         VisitedMap::ValueType* it = m_compiled.Find( key );
         if ( it )
         {
@@ -413,41 +413,40 @@ namespace mu
 
 
     //---------------------------------------------------------------------------------------------
-    void CodeGenerator::GenerateRange( RANGE_GENERATION_RESULT& result,
-                                      NodeRangePtrConst untyped )
+    void CodeGenerator::GenerateRange( FRangeGenerationResult& Result,
+                                      NodeRangePtrConst Untyped)
     {
-        if (!untyped)
+        if (!Untyped)
         {
-            result = RANGE_GENERATION_RESULT();
+			Result = FRangeGenerationResult();
             return;
         }
 
         // See if it was already generated
-		VISITED_MAP_KEY key = GetCurrentCacheKey(untyped);
-		GeneratedRangeMap::ValueType* it = m_generatedRanges.Find( key );
+		FVisitedKeyMap Key = GetCurrentCacheKey(Untyped);
+		GeneratedRangeMap::ValueType* it = m_generatedRanges.Find(Key);
         if ( it  )
         {
-            result = *it;
+			Result = *it;
             return;
         }
 
 
         // Generate for each different type of node
-        if ( auto fromScalar = dynamic_cast<const NodeRangeFromScalar*>(untyped.get()) )
+        if (const NodeRangeFromScalar* FromScalar = dynamic_cast<const NodeRangeFromScalar*>(Untyped.get()) )
         {
-            result = RANGE_GENERATION_RESULT();
-            result.rangeName = fromScalar->GetName();
-            result.sizeOp = Generate( fromScalar->GetSize() );
+			Result = FRangeGenerationResult();
+			Result.rangeName = FromScalar->GetName();
+			Result.sizeOp = Generate(FromScalar->GetSize());
         }
         else
         {
             check(false);
-            mu::Halt();
         }
 
 
         // Cache the result
-        m_generatedRanges.Add( key, result);
+        m_generatedRanges.Add( Key, Result);
     }
 
 	
@@ -456,7 +455,7 @@ namespace mu
     {
         Ptr<ASTOp> result;
 
-        PARAMETER_DESC param;
+        FParameterDesc param;
         param.m_name = strName;
         if ( param.m_name.size()==0 )
         {
@@ -481,7 +480,7 @@ namespace mu
 
 			if (pTable->GetPrivate()->m_NoneOption)
 			{
-				PARAMETER_DESC::INT_VALUE_DESC nullValue;
+				FParameterDesc::INT_VALUE_DESC nullValue;
 				nullValue.m_value = -1;
 				nullValue.m_name = "None";
 				param.m_possibleValues.Add(nullValue);
@@ -492,7 +491,7 @@ namespace mu
             int32 rows = pTable->GetPrivate()->m_rows.Num();
 			for (size_t i = 0; i<rows; ++i)
             {
-                PARAMETER_DESC::INT_VALUE_DESC value;
+                FParameterDesc::INT_VALUE_DESC value;
                 value.m_value = (int16_t)pTable->GetPrivate()->m_rows[i].m_id;
 
                 if (nameCol>-1)
@@ -836,7 +835,7 @@ namespace mu
 
 
     //---------------------------------------------------------------------------------------------
-    void CodeGenerator::GenerateSurface( SURFACE_GENERATION_RESULT& result,
+    void CodeGenerator::GenerateSurface( FSurfaceGenerationResult& result,
                                          NodeSurfaceNewPtrConst surfaceNode,
                                          const TArray<FirstPassGenerator::SURFACE::EDIT>& edits )
     {
@@ -965,7 +964,7 @@ namespace mu
 						lastMeshOp = ApplyMeshModifiers(lastMeshOp, e.node->m_tags,
 							bModifiersForBeforeOperations, node.m_errorContext);
 
-                        FMeshGenerationResult::EXTRA_LAYOUTS data;
+                        FMeshGenerationResult::FExtraLayouts data;
 						data.GeneratedLayouts = addResults.GeneratedLayouts;
 						data.condition = e.condition;
                         data.meshFragment = addResults.meshOp;
@@ -1915,7 +1914,7 @@ namespace mu
                 sop->name = its.node->GetPrivate()->m_name;
                 sop->instance = lastCompOp;
 
-				SURFACE_GENERATION_RESULT surfaceGenerationResult;
+				FSurfaceGenerationResult surfaceGenerationResult;
                 GenerateSurface( surfaceGenerationResult, its.node, its.edits );
                 sop->value = surfaceGenerationResult.surfaceOp;
 
@@ -2055,7 +2054,7 @@ namespace mu
     {
         MUTABLE_CPUPROFILER_SCOPE(NodeObjectNew);
 
-        m_currentParents.Add( PARENT_KEY() );
+        m_currentParents.Add( FParentKey() );
         m_currentParents.Last().m_pObject = &node;
 
         // Parse the child objects first, which will accumulate operations in the patching lists
@@ -2328,8 +2327,8 @@ namespace mu
 
 				// Morph to an ellipse
 				{
-					SHAPE morphShape;
-					morphShape.type = (uint8_t)SHAPE::Type::Ellipse;
+					FShape morphShape;
+					morphShape.type = (uint8_t)FShape::Type::Ellipse;
 					morphShape.position = TypedNode->m_origin;
 					morphShape.up = TypedNode->m_normal;
 					// TODO: Move rotation to ellipse rotation reference base instead of passing it directly
@@ -2355,8 +2354,8 @@ namespace mu
 				if (TypedNode->m_vertexSelectionType == NodeModifierMeshClipMorphPlane::Private::VS_SHAPE)
 				{
 					op->vertexSelectionType = OP::MeshClipMorphPlaneArgs::VS_SHAPE;
-					SHAPE selectionShape;
-					selectionShape.type = (uint8_t)SHAPE::Type::AABox;
+					FShape selectionShape;
+					selectionShape.type = (uint8_t)FShape::Type::AABox;
 					selectionShape.position = TypedNode->m_selectionBoxOrigin;
 					selectionShape.size = TypedNode->m_selectionBoxRadius;
 					op->selectionShape = selectionShape;
diff --git a/Engine/Plugins/Experimental/Mutable/Source/MutableTools/Private/MuT/CodeGenerator.h b/Engine/Plugins/Experimental/Mutable/Source/MutableTools/Private/MuT/CodeGenerator.h
--- a/Engine/Plugins/Experimental/Mutable/Source/MutableTools/Private/MuT/CodeGenerator.h
+++ b/Engine/Plugins/Experimental/Mutable/Source/MutableTools/Private/MuT/CodeGenerator.h
@@ -111,7 +111,7 @@ namespace mu
 	class NodeStringParameter;
 	class TaskManager;
 	struct OBJECT_STATE;
-	struct PROGRAM;
+	struct FProgram;
 
 
 
@@ -159,13 +159,13 @@ namespace mu
 		FirstPassGenerator m_firstPass;
 
 
-        struct VISITED_MAP_KEY
+        struct FVisitedKeyMap
         {
-            VISITED_MAP_KEY()
+            FVisitedKeyMap()
             {
             }
 
-			friend FORCEINLINE uint32 GetTypeHash(const VISITED_MAP_KEY& InKey)
+			friend FORCEINLINE uint32 GetTypeHash(const FVisitedKeyMap& InKey)
 			{
 				uint32 KeyHash = 0;
 				KeyHash = HashCombine(KeyHash, ::GetTypeHash(InKey.pNode.get()));
@@ -180,7 +180,7 @@ namespace mu
 				return KeyHash;
 			}
 
-			bool operator==(const VISITED_MAP_KEY& InKey) const
+			bool operator==(const FVisitedKeyMap& InKey) const
 			{
 				if (pNode != InKey.pNode) return false;
 				if (state != InKey.state) return false;
@@ -211,7 +211,7 @@ namespace mu
         };
         BOTTOM_UP_STATE m_currentBottomUpState;
 
-        typedef TMap<VISITED_MAP_KEY,BOTTOM_UP_STATE> VisitedMap;
+        typedef TMap<FVisitedKeyMap,BOTTOM_UP_STATE> VisitedMap;
         VisitedMap m_compiled;
 
         //!
@@ -254,9 +254,9 @@ namespace mu
 		// (top-down) Tags that are active when generating nodes.
 		TArray< TArray<mu::string> > m_activeTags;
 
-        struct PARENT_KEY
+        struct FParentKey
         {
-            PARENT_KEY()
+            FParentKey()
             {
                 m_pObject = nullptr;
                 m_state = -1;
@@ -276,7 +276,7 @@ namespace mu
             int m_block;
         };
 
-		TArray< PARENT_KEY > m_currentParents;
+		TArray< FParentKey > m_currentParents;
 
         // List of additional components to add to an object that come from child objects.
         // The index is the object and lod that should receive the components.
@@ -355,9 +355,9 @@ namespace mu
         std::shared_timed_mutex m_codeAccessMutex;
 
 		//! Generate the key with all the relevant state that is used in generation of operations for a node.
-		VISITED_MAP_KEY GetCurrentCacheKey(const NodePtrConst& InNode) const
+		FVisitedKeyMap GetCurrentCacheKey(const NodePtrConst& InNode) const
 		{
-			VISITED_MAP_KEY key;
+			FVisitedKeyMap key;
 			key.pNode = InNode;
 			key.state = m_currentStateIndex;
 			if (!m_imageState.IsEmpty())
@@ -375,42 +375,42 @@ namespace mu
 
 		//-----------------------------------------------------------------------------------------
 		// Images
-		struct IMAGE_GENERATION_RESULT
+		struct FImageGenerationResult
 		{
 			Ptr<ASTOp> op;
 		};
 
-		typedef TMap<VISITED_MAP_KEY, IMAGE_GENERATION_RESULT> GeneratedImagesMap;
+		typedef TMap<FVisitedKeyMap, FImageGenerationResult> GeneratedImagesMap;
 		GeneratedImagesMap m_generatedImages;
 
-		void GenerateImage(IMAGE_GENERATION_RESULT& result, const NodeImagePtrConst& node);
-		void GenerateImage_Constant(IMAGE_GENERATION_RESULT&, const NodeImageConstant*);
-		void GenerateImage_Difference(IMAGE_GENERATION_RESULT&, const NodeImageDifference*);
-		void GenerateImage_Interpolate(IMAGE_GENERATION_RESULT&, const NodeImageInterpolate*);
-		void GenerateImage_Saturate(IMAGE_GENERATION_RESULT&, const NodeImageSaturate*);
-		void GenerateImage_Table(IMAGE_GENERATION_RESULT&, const NodeImageTable*);
-		void GenerateImage_Swizzle(IMAGE_GENERATION_RESULT&, const NodeImageSwizzle*);
-		void GenerateImage_SelectColour(IMAGE_GENERATION_RESULT&, const NodeImageSelectColour*);
-		void GenerateImage_ColourMap(IMAGE_GENERATION_RESULT&, const NodeImageColourMap*);
-		void GenerateImage_Gradient(IMAGE_GENERATION_RESULT&, const NodeImageGradient*);
-		void GenerateImage_Binarise(IMAGE_GENERATION_RESULT&, const NodeImageBinarise*);
-		void GenerateImage_Luminance(IMAGE_GENERATION_RESULT&, const NodeImageLuminance*);
-		void GenerateImage_Layer(IMAGE_GENERATION_RESULT&, const NodeImageLayer*);
-		void GenerateImage_LayerColour(IMAGE_GENERATION_RESULT&, const NodeImageLayerColour*);
-		void GenerateImage_Resize(IMAGE_GENERATION_RESULT&, const NodeImageResize*);
-		void GenerateImage_PlainColour(IMAGE_GENERATION_RESULT&, const NodeImagePlainColour*);
-		void GenerateImage_Interpolate3(IMAGE_GENERATION_RESULT&, const NodeImageInterpolate3*);
-		void GenerateImage_Project(IMAGE_GENERATION_RESULT&, const NodeImageProject*);
-		void GenerateImage_Mipmap(IMAGE_GENERATION_RESULT&, const NodeImageMipmap*);
-		void GenerateImage_Switch(IMAGE_GENERATION_RESULT&, const NodeImageSwitch*);
-		void GenerateImage_Conditional(IMAGE_GENERATION_RESULT&, const NodeImageConditional*);
-		void GenerateImage_Format(IMAGE_GENERATION_RESULT&, const NodeImageFormat*);
-		void GenerateImage_Parameter(IMAGE_GENERATION_RESULT&, const NodeImageParameter*);
-		void GenerateImage_MultiLayer(IMAGE_GENERATION_RESULT&, const NodeImageMultiLayer*);
-		void GenerateImage_Invert(IMAGE_GENERATION_RESULT&, const NodeImageInvert*);
-		void GenerateImage_Variation(IMAGE_GENERATION_RESULT&, const NodeImageVariation*);
-		void GenerateImage_NormalComposite(IMAGE_GENERATION_RESULT&, const NodeImageNormalComposite*);
-		void GenerateImage_Transform(IMAGE_GENERATION_RESULT&, const NodeImageTransform*);
+		void GenerateImage(FImageGenerationResult& result, const NodeImagePtrConst& node);
+		void GenerateImage_Constant(FImageGenerationResult&, const NodeImageConstant*);
+		void GenerateImage_Difference(FImageGenerationResult&, const NodeImageDifference*);
+		void GenerateImage_Interpolate(FImageGenerationResult&, const NodeImageInterpolate*);
+		void GenerateImage_Saturate(FImageGenerationResult&, const NodeImageSaturate*);
+		void GenerateImage_Table(FImageGenerationResult&, const NodeImageTable*);
+		void GenerateImage_Swizzle(FImageGenerationResult&, const NodeImageSwizzle*);
+		void GenerateImage_SelectColour(FImageGenerationResult&, const NodeImageSelectColour*);
+		void GenerateImage_ColourMap(FImageGenerationResult&, const NodeImageColourMap*);
+		void GenerateImage_Gradient(FImageGenerationResult&, const NodeImageGradient*);
+		void GenerateImage_Binarise(FImageGenerationResult&, const NodeImageBinarise*);
+		void GenerateImage_Luminance(FImageGenerationResult&, const NodeImageLuminance*);
+		void GenerateImage_Layer(FImageGenerationResult&, const NodeImageLayer*);
+		void GenerateImage_LayerColour(FImageGenerationResult&, const NodeImageLayerColour*);
+		void GenerateImage_Resize(FImageGenerationResult&, const NodeImageResize*);
+		void GenerateImage_PlainColour(FImageGenerationResult&, const NodeImagePlainColour*);
+		void GenerateImage_Interpolate3(FImageGenerationResult&, const NodeImageInterpolate3*);
+		void GenerateImage_Project(FImageGenerationResult&, const NodeImageProject*);
+		void GenerateImage_Mipmap(FImageGenerationResult&, const NodeImageMipmap*);
+		void GenerateImage_Switch(FImageGenerationResult&, const NodeImageSwitch*);
+		void GenerateImage_Conditional(FImageGenerationResult&, const NodeImageConditional*);
+		void GenerateImage_Format(FImageGenerationResult&, const NodeImageFormat*);
+		void GenerateImage_Parameter(FImageGenerationResult&, const NodeImageParameter*);
+		void GenerateImage_MultiLayer(FImageGenerationResult&, const NodeImageMultiLayer*);
+		void GenerateImage_Invert(FImageGenerationResult&, const NodeImageInvert*);
+		void GenerateImage_Variation(FImageGenerationResult&, const NodeImageVariation*);
+		void GenerateImage_NormalComposite(FImageGenerationResult&, const NodeImageNormalComposite*);
+		void GenerateImage_Transform(FImageGenerationResult&, const NodeImageTransform*);
 
 		//!
 		ImagePtr GenerateMissingImage(EImageFormat);
@@ -495,14 +495,14 @@ namespace mu
 
 			TArray<Ptr<ASTOp>> layoutOps;
 
-			struct EXTRA_LAYOUTS
+			struct FExtraLayouts
 			{
 				/** Source node layouts to use with these extra mesh. They don't have block ids. */
 				TArray<Ptr<const Layout>> GeneratedLayouts;
 				Ptr<ASTOp> condition;
 				Ptr<ASTOp> meshFragment;
 			};
-			TArray< EXTRA_LAYOUTS > extraMeshLayouts;
+			TArray< FExtraLayouts > extraMeshLayouts;
 		};
 		
 		struct FGeneratedMeshCacheKey
@@ -556,96 +556,96 @@ namespace mu
 
         //-----------------------------------------------------------------------------------------
         // Projectors
-        struct PROJECTOR_GENERATION_RESULT
+        struct FProjectorGenerationResult
         {
             Ptr<ASTOp> op;
             PROJECTOR_TYPE type;
         };
 
-        typedef TMap<VISITED_MAP_KEY,PROJECTOR_GENERATION_RESULT> GeneratedProjectorsMap;
+        typedef TMap<FVisitedKeyMap,FProjectorGenerationResult> GeneratedProjectorsMap;
         GeneratedProjectorsMap m_generatedProjectors;
 
-        void GenerateProjector( PROJECTOR_GENERATION_RESULT&, const NodeProjectorPtrConst& );
-        void GenerateProjector_Constant( PROJECTOR_GENERATION_RESULT&, const Ptr<const NodeProjectorConstant>& );
-        void GenerateProjector_Parameter( PROJECTOR_GENERATION_RESULT&, const Ptr<const NodeProjectorParameter>& );
-        void GenerateMissingProjectorCode( PROJECTOR_GENERATION_RESULT&, const void* errorContext );
+        void GenerateProjector( FProjectorGenerationResult&, const NodeProjectorPtrConst& );
+        void GenerateProjector_Constant( FProjectorGenerationResult&, const Ptr<const NodeProjectorConstant>& );
+        void GenerateProjector_Parameter( FProjectorGenerationResult&, const Ptr<const NodeProjectorParameter>& );
+        void GenerateMissingProjectorCode( FProjectorGenerationResult&, const void* errorContext );
 
 		//-----------------------------------------------------------------------------------------
 		// Bools
-		struct BOOL_GENERATION_RESULT
+		struct FBoolGenerationResult
 		{
 			Ptr<ASTOp> op;
 		};
 
-		typedef TMap<VISITED_MAP_KEY, BOOL_GENERATION_RESULT> GeneratedBoolsMap;
+		typedef TMap<FVisitedKeyMap, FBoolGenerationResult> GeneratedBoolsMap;
 		GeneratedBoolsMap m_generatedBools;
 
-		void GenerateBool(BOOL_GENERATION_RESULT&, const NodeBoolPtrConst&);
-		void GenerateBool_Constant(BOOL_GENERATION_RESULT&, const Ptr<const NodeBoolConstant>&);
-		void GenerateBool_Parameter(BOOL_GENERATION_RESULT&, const Ptr<const NodeBoolParameter>&);
-		void GenerateBool_IsNull(BOOL_GENERATION_RESULT&, const Ptr<const NodeBoolIsNull>&);
-		void GenerateBool_Not(BOOL_GENERATION_RESULT&, const Ptr<const NodeBoolNot>&);
-		void GenerateBool_And(BOOL_GENERATION_RESULT&, const Ptr<const NodeBoolAnd>&);
+		void GenerateBool(FBoolGenerationResult&, const NodeBoolPtrConst&);
+		void GenerateBool_Constant(FBoolGenerationResult&, const Ptr<const NodeBoolConstant>&);
+		void GenerateBool_Parameter(FBoolGenerationResult&, const Ptr<const NodeBoolParameter>&);
+		void GenerateBool_IsNull(FBoolGenerationResult&, const Ptr<const NodeBoolIsNull>&);
+		void GenerateBool_Not(FBoolGenerationResult&, const Ptr<const NodeBoolNot>&);
+		void GenerateBool_And(FBoolGenerationResult&, const Ptr<const NodeBoolAnd>&);
 
 		//-----------------------------------------------------------------------------------------
 		// Scalars
-		struct SCALAR_GENERATION_RESULT
+		struct FScalarGenerationResult
 		{
 			Ptr<ASTOp> op;
 		};
 
-		typedef TMap<VISITED_MAP_KEY, SCALAR_GENERATION_RESULT> GeneratedScalarsMap;
+		typedef TMap<FVisitedKeyMap, FScalarGenerationResult> GeneratedScalarsMap;
 		GeneratedScalarsMap m_generatedScalars;
 
-		void GenerateScalar(SCALAR_GENERATION_RESULT&, const NodeScalarPtrConst&);
-		void GenerateScalar_Constant(SCALAR_GENERATION_RESULT&, const Ptr<const NodeScalarConstant>&);
-		void GenerateScalar_Parameter(SCALAR_GENERATION_RESULT&, const Ptr<const NodeScalarParameter>&);
-		void GenerateScalar_Switch(SCALAR_GENERATION_RESULT&, const Ptr<const NodeScalarSwitch>&);
-		void GenerateScalar_EnumParameter(SCALAR_GENERATION_RESULT&, const Ptr<const NodeScalarEnumParameter>&);
-		void GenerateScalar_Curve(SCALAR_GENERATION_RESULT&, const Ptr<const NodeScalarCurve>&);
-		void GenerateScalar_Arithmetic(SCALAR_GENERATION_RESULT&, const Ptr<const NodeScalarArithmeticOperation>&);
-		void GenerateScalar_Variation(SCALAR_GENERATION_RESULT&, const Ptr<const NodeScalarVariation>&);
-		void GenerateScalar_Table(SCALAR_GENERATION_RESULT&, const Ptr<const NodeScalarTable>&);
+		void GenerateScalar(FScalarGenerationResult&, const NodeScalarPtrConst&);
+		void GenerateScalar_Constant(FScalarGenerationResult&, const Ptr<const NodeScalarConstant>&);
+		void GenerateScalar_Parameter(FScalarGenerationResult&, const Ptr<const NodeScalarParameter>&);
+		void GenerateScalar_Switch(FScalarGenerationResult&, const Ptr<const NodeScalarSwitch>&);
+		void GenerateScalar_EnumParameter(FScalarGenerationResult&, const Ptr<const NodeScalarEnumParameter>&);
+		void GenerateScalar_Curve(FScalarGenerationResult&, const Ptr<const NodeScalarCurve>&);
+		void GenerateScalar_Arithmetic(FScalarGenerationResult&, const Ptr<const NodeScalarArithmeticOperation>&);
+		void GenerateScalar_Variation(FScalarGenerationResult&, const Ptr<const NodeScalarVariation>&);
+		void GenerateScalar_Table(FScalarGenerationResult&, const Ptr<const NodeScalarTable>&);
 		Ptr<ASTOp> GenerateMissingScalarCode(const char* strWhere, float value, const void* errorContext);
 
 		//-----------------------------------------------------------------------------------------
 		// Colors
-		struct COLOR_GENERATION_RESULT
+		struct FColorGenerationResult
 		{
 			Ptr<ASTOp> op;
 		};
 
-		typedef TMap<VISITED_MAP_KEY, COLOR_GENERATION_RESULT> GeneratedColorsMap;
+		typedef TMap<FVisitedKeyMap, FColorGenerationResult> GeneratedColorsMap;
 		GeneratedColorsMap m_generatedColors;
 
-		void GenerateColor(COLOR_GENERATION_RESULT&, const NodeColourPtrConst&);
-		void GenerateColor_Constant(COLOR_GENERATION_RESULT&, const Ptr<const NodeColourConstant>&);
-		void GenerateColor_Parameter(COLOR_GENERATION_RESULT&, const Ptr<const NodeColourParameter>&);
-		void GenerateColor_Switch(COLOR_GENERATION_RESULT&, const Ptr<const NodeColourSwitch>&);
-		void GenerateColor_SampleImage(COLOR_GENERATION_RESULT&, const Ptr<const NodeColourSampleImage>&);
-		void GenerateColor_FromScalars(COLOR_GENERATION_RESULT&, const Ptr<const NodeColourFromScalars>&);
-		void GenerateColor_Arithmetic(COLOR_GENERATION_RESULT&, const Ptr<const NodeColourArithmeticOperation>&);
-		void GenerateColor_Variation(COLOR_GENERATION_RESULT&, const Ptr<const NodeColourVariation>&);
-		void GenerateColor_Table(COLOR_GENERATION_RESULT&, const Ptr<const NodeColourTable>&);
+		void GenerateColor(FColorGenerationResult&, const NodeColourPtrConst&);
+		void GenerateColor_Constant(FColorGenerationResult&, const Ptr<const NodeColourConstant>&);
+		void GenerateColor_Parameter(FColorGenerationResult&, const Ptr<const NodeColourParameter>&);
+		void GenerateColor_Switch(FColorGenerationResult&, const Ptr<const NodeColourSwitch>&);
+		void GenerateColor_SampleImage(FColorGenerationResult&, const Ptr<const NodeColourSampleImage>&);
+		void GenerateColor_FromScalars(FColorGenerationResult&, const Ptr<const NodeColourFromScalars>&);
+		void GenerateColor_Arithmetic(FColorGenerationResult&, const Ptr<const NodeColourArithmeticOperation>&);
+		void GenerateColor_Variation(FColorGenerationResult&, const Ptr<const NodeColourVariation>&);
+		void GenerateColor_Table(FColorGenerationResult&, const Ptr<const NodeColourTable>&);
 		Ptr<ASTOp> GenerateMissingColourCode(const char* strWhere, const void* errorContext);
 
 		//-----------------------------------------------------------------------------------------
 		// Strings
-		struct STRING_GENERATION_RESULT
+		struct FStringGenerationResult
 		{
 			Ptr<ASTOp> op;
 		};
 
-		typedef TMap<VISITED_MAP_KEY, STRING_GENERATION_RESULT> GeneratedStringsMap;
+		typedef TMap<FVisitedKeyMap, FStringGenerationResult> GeneratedStringsMap;
 		GeneratedStringsMap m_generatedStrings;
 
-		void GenerateString(STRING_GENERATION_RESULT&, const NodeStringPtrConst&);
-		void GenerateString_Constant(STRING_GENERATION_RESULT&, const Ptr<const NodeStringConstant>&);
-		void GenerateString_Parameter(STRING_GENERATION_RESULT&, const Ptr<const NodeStringParameter>&);
+		void GenerateString(FStringGenerationResult&, const NodeStringPtrConst&);
+		void GenerateString_Constant(FStringGenerationResult&, const Ptr<const NodeStringConstant>&);
+		void GenerateString_Parameter(FStringGenerationResult&, const Ptr<const NodeStringParameter>&);
 
         //-----------------------------------------------------------------------------------------
         // Ranges
-        struct RANGE_GENERATION_RESULT
+        struct FRangeGenerationResult
         {
             //
             Ptr<ASTOp> sizeOp;
@@ -657,19 +657,19 @@ namespace mu
             string rangeUID;
         };
 
-        typedef TMap<VISITED_MAP_KEY,RANGE_GENERATION_RESULT> GeneratedRangeMap;
+        typedef TMap<FVisitedKeyMap,FRangeGenerationResult> GeneratedRangeMap;
         GeneratedRangeMap m_generatedRanges;
 
-        void GenerateRange( RANGE_GENERATION_RESULT& result, Ptr<const NodeRange> node);
+        void GenerateRange( FRangeGenerationResult& result, Ptr<const NodeRange> node);
 
 
         //-----------------------------------------------------------------------------------------
-        struct SURFACE_GENERATION_RESULT
+        struct FSurfaceGenerationResult
         {
             Ptr<ASTOp> surfaceOp;
         };
 
-        void GenerateSurface( SURFACE_GENERATION_RESULT& result,
+        void GenerateSurface( FSurfaceGenerationResult& result,
                               NodeSurfaceNewPtrConst node,
                               const TArray<FirstPassGenerator::SURFACE::EDIT>& edits );
 
@@ -682,7 +682,7 @@ namespace mu
     //! address.
     //! \param returnBestOption If true, try to resolve ambiguities returning some value.
     //---------------------------------------------------------------------------------------------
-    extern FImageDesc GetImageDesc( const PROGRAM& program, OP::ADDRESS at,
+    extern FImageDesc GetImageDesc( const FProgram& program, OP::ADDRESS at,
                                     bool returnBestOption = false,
                                     class GetImageDescContext* context=nullptr );
 
diff --git a/Engine/Plugins/Experimental/Mutable/Source/MutableTools/Private/MuT/CodeGenerator_Bool.cpp b/Engine/Plugins/Experimental/Mutable/Source/MutableTools/Private/MuT/CodeGenerator_Bool.cpp
--- a/Engine/Plugins/Experimental/Mutable/Source/MutableTools/Private/MuT/CodeGenerator_Bool.cpp
+++ b/Engine/Plugins/Experimental/Mutable/Source/MutableTools/Private/MuT/CodeGenerator_Bool.cpp
@@ -27,16 +27,16 @@ namespace mu
 
 
 	//-------------------------------------------------------------------------------------------------
-	void CodeGenerator::GenerateBool(BOOL_GENERATION_RESULT& result, const NodeBoolPtrConst& untyped)
+	void CodeGenerator::GenerateBool(FBoolGenerationResult& result, const NodeBoolPtrConst& untyped)
 	{
 		if (!untyped)
 		{
-			result = BOOL_GENERATION_RESULT();
+			result = FBoolGenerationResult();
 			return;
 		}
 
 		// See if it was already generated
-		VISITED_MAP_KEY key = GetCurrentCacheKey(untyped);
+		FVisitedKeyMap key = GetCurrentCacheKey(untyped);
 		GeneratedBoolsMap::ValueType* it = m_generatedBools.Find(key);
 		if (it)
 		{
@@ -76,7 +76,7 @@ namespace mu
 
 
 	//-------------------------------------------------------------------------------------------------
-	void CodeGenerator::GenerateBool_Constant(BOOL_GENERATION_RESULT& result, const Ptr<const NodeBoolConstant>& Typed)
+	void CodeGenerator::GenerateBool_Constant(FBoolGenerationResult& result, const Ptr<const NodeBoolConstant>& Typed)
 	{
 		const NodeBoolConstant::Private& node = *Typed->GetPrivate();
 
@@ -88,7 +88,7 @@ namespace mu
 
 
 	//-------------------------------------------------------------------------------------------------
-	void CodeGenerator::GenerateBool_Parameter(BOOL_GENERATION_RESULT& result, const Ptr<const NodeBoolParameter>& Typed)
+	void CodeGenerator::GenerateBool_Parameter(FBoolGenerationResult& result, const Ptr<const NodeBoolParameter>& Typed)
 	{
 		const NodeBoolParameter::Private& node = *Typed->GetPrivate();
 
@@ -97,7 +97,7 @@ namespace mu
 		auto it = m_nodeVariables.find(node.m_pNode);
 		if (it == m_nodeVariables.end())
 		{
-			PARAMETER_DESC param;
+			FParameterDesc param;
 			param.m_name = node.m_name;
 			param.m_uid = node.m_uid;
 			param.m_type = PARAMETER_TYPE::T_BOOL;
@@ -110,7 +110,7 @@ namespace mu
 			// Generate the code for the ranges
 			for (int32 a = 0; a < node.m_ranges.Num(); ++a)
 			{
-				RANGE_GENERATION_RESULT rangeResult;
+				FRangeGenerationResult rangeResult;
 				GenerateRange(rangeResult, node.m_ranges[a]);
 				op->ranges.Emplace(op.get(), rangeResult.sizeOp, rangeResult.rangeName, rangeResult.rangeUID);
 			}
@@ -127,7 +127,7 @@ namespace mu
 
 
 	//-------------------------------------------------------------------------------------------------
-	void CodeGenerator::GenerateBool_IsNull(BOOL_GENERATION_RESULT& result, const Ptr<const NodeBoolIsNull>& Typed)
+	void CodeGenerator::GenerateBool_IsNull(FBoolGenerationResult& result, const Ptr<const NodeBoolIsNull>& Typed)
 	{
 		const NodeBoolIsNull::Private& node = *Typed->GetPrivate();
 
@@ -140,7 +140,7 @@ namespace mu
 
 
 	//-------------------------------------------------------------------------------------------------
-	void CodeGenerator::GenerateBool_Not(BOOL_GENERATION_RESULT& result, const Ptr<const NodeBoolNot>& Typed)
+	void CodeGenerator::GenerateBool_Not(FBoolGenerationResult& result, const Ptr<const NodeBoolNot>& Typed)
 	{
 		const NodeBoolNot::Private& node = *Typed->GetPrivate();
 
@@ -153,7 +153,7 @@ namespace mu
 
 
 	//-------------------------------------------------------------------------------------------------
-	void CodeGenerator::GenerateBool_And(BOOL_GENERATION_RESULT& result, const Ptr<const NodeBoolAnd>& Typed)
+	void CodeGenerator::GenerateBool_And(FBoolGenerationResult& result, const Ptr<const NodeBoolAnd>& Typed)
 	{
 		const NodeBoolAnd::Private& node = *Typed->GetPrivate();
 
diff --git a/Engine/Plugins/Experimental/Mutable/Source/MutableTools/Private/MuT/CodeGenerator_Color.cpp b/Engine/Plugins/Experimental/Mutable/Source/MutableTools/Private/MuT/CodeGenerator_Color.cpp
--- a/Engine/Plugins/Experimental/Mutable/Source/MutableTools/Private/MuT/CodeGenerator_Color.cpp
+++ b/Engine/Plugins/Experimental/Mutable/Source/MutableTools/Private/MuT/CodeGenerator_Color.cpp
@@ -58,16 +58,16 @@ namespace mu
 
 
 	//-------------------------------------------------------------------------------------------------
-	void CodeGenerator::GenerateColor(COLOR_GENERATION_RESULT& result, const NodeColourPtrConst& untyped)
+	void CodeGenerator::GenerateColor(FColorGenerationResult& result, const NodeColourPtrConst& untyped)
 	{
 		if (!untyped)
 		{
-			result = COLOR_GENERATION_RESULT();
+			result = FColorGenerationResult();
 			return;
 		}
 
 		// See if it was already generated
-		VISITED_MAP_KEY key = GetCurrentCacheKey(untyped);
+		FVisitedKeyMap key = GetCurrentCacheKey(untyped);
 		GeneratedColorsMap::ValueType* it = m_generatedColors.Find(key);
 		if (it)
 		{
@@ -119,7 +119,7 @@ namespace mu
 
 
 	//-------------------------------------------------------------------------------------------------
-	void CodeGenerator::GenerateColor_Constant(COLOR_GENERATION_RESULT& result, const Ptr<const NodeColourConstant>& Typed)
+	void CodeGenerator::GenerateColor_Constant(FColorGenerationResult& result, const Ptr<const NodeColourConstant>& Typed)
 	{
 		const NodeColourConstant::Private& node = *Typed->GetPrivate();
 
@@ -134,7 +134,7 @@ namespace mu
 
 
 	//-------------------------------------------------------------------------------------------------
-	void CodeGenerator::GenerateColor_Parameter(COLOR_GENERATION_RESULT& result, const Ptr<const NodeColourParameter>& Typed)
+	void CodeGenerator::GenerateColor_Parameter(FColorGenerationResult& result, const Ptr<const NodeColourParameter>& Typed)
 	{
 		const NodeColourParameter::Private& node = *Typed->GetPrivate();
 
@@ -144,7 +144,7 @@ namespace mu
 
 		if (it == m_nodeVariables.end())
 		{
-			PARAMETER_DESC param;
+			FParameterDesc param;
 			param.m_name = node.m_name;
 			param.m_uid = node.m_uid;
 			param.m_type = PARAMETER_TYPE::T_COLOUR;
@@ -159,7 +159,7 @@ namespace mu
 			// Generate the code for the ranges
 			for (int32 a = 0; a < node.m_ranges.Num(); ++a)
 			{
-				RANGE_GENERATION_RESULT rangeResult;
+				FRangeGenerationResult rangeResult;
 				GenerateRange(rangeResult, node.m_ranges[a]);
 				op->ranges.Emplace(op.get(), rangeResult.sizeOp, rangeResult.rangeName, rangeResult.rangeUID);
 			}
@@ -176,7 +176,7 @@ namespace mu
 
 
 	//-------------------------------------------------------------------------------------------------
-	void CodeGenerator::GenerateColor_Switch(COLOR_GENERATION_RESULT& result, const Ptr<const NodeColourSwitch>& Typed)
+	void CodeGenerator::GenerateColor_Switch(FColorGenerationResult& result, const Ptr<const NodeColourSwitch>& Typed)
 	{
 		const NodeColourSwitch::Private& node = *Typed->GetPrivate();
 
@@ -227,7 +227,7 @@ namespace mu
 
 
 	//-------------------------------------------------------------------------------------------------
-	void CodeGenerator::GenerateColor_Variation(COLOR_GENERATION_RESULT& result, const Ptr<const NodeColourVariation>& Typed)
+	void CodeGenerator::GenerateColor_Variation(FColorGenerationResult& result, const Ptr<const NodeColourVariation>& Typed)
 	{
 		const NodeColourVariation::Private& node = *Typed->GetPrivate();
 
@@ -294,7 +294,7 @@ namespace mu
 
 
 	//-------------------------------------------------------------------------------------------------
-	void CodeGenerator::GenerateColor_SampleImage(COLOR_GENERATION_RESULT& result, const Ptr<const NodeColourSampleImage>& Typed)
+	void CodeGenerator::GenerateColor_SampleImage(FColorGenerationResult& result, const Ptr<const NodeColourSampleImage>& Typed)
 	{
 		const NodeColourSampleImage::Private& node = *Typed->GetPrivate();
 
@@ -366,7 +366,7 @@ namespace mu
 
 
 	//-------------------------------------------------------------------------------------------------
-	void CodeGenerator::GenerateColor_FromScalars(COLOR_GENERATION_RESULT& result, const Ptr<const NodeColourFromScalars>& Typed)
+	void CodeGenerator::GenerateColor_FromScalars(FColorGenerationResult& result, const Ptr<const NodeColourFromScalars>& Typed)
 	{
 		const NodeColourFromScalars::Private& node = *Typed->GetPrivate();
 
@@ -426,7 +426,7 @@ namespace mu
 
 
 	//-------------------------------------------------------------------------------------------------
-	void CodeGenerator::GenerateColor_Arithmetic(COLOR_GENERATION_RESULT& result, const Ptr<const NodeColourArithmeticOperation>& Typed)
+	void CodeGenerator::GenerateColor_Arithmetic(FColorGenerationResult& result, const Ptr<const NodeColourArithmeticOperation>& Typed)
 	{
 		const NodeColourArithmeticOperation::Private& node = *Typed->GetPrivate();
 
@@ -480,7 +480,7 @@ namespace mu
 
 
 	//-------------------------------------------------------------------------------------------------
-	void CodeGenerator::GenerateColor_Table(COLOR_GENERATION_RESULT& result, const Ptr<const NodeColourTable>& Typed)
+	void CodeGenerator::GenerateColor_Table(FColorGenerationResult& result, const Ptr<const NodeColourTable>& Typed)
 	{
 		const NodeColourTable::Private& node = *Typed->GetPrivate();
 
@@ -507,7 +507,7 @@ namespace mu
 		NodeColourConstantPtr pNode = new NodeColourConstant();
 		pNode->SetValue(1, 1, 0);
 
-		COLOR_GENERATION_RESULT Result;
+		FColorGenerationResult Result;
 		GenerateColor(Result, pNode);
 
 		return Result.op;
diff --git a/Engine/Plugins/Experimental/Mutable/Source/MutableTools/Private/MuT/CodeGenerator_FirstPass.cpp b/Engine/Plugins/Experimental/Mutable/Source/MutableTools/Private/MuT/CodeGenerator_FirstPass.cpp
--- a/Engine/Plugins/Experimental/Mutable/Source/MutableTools/Private/MuT/CodeGenerator_FirstPass.cpp
+++ b/Engine/Plugins/Experimental/Mutable/Source/MutableTools/Private/MuT/CodeGenerator_FirstPass.cpp
@@ -536,7 +536,7 @@ namespace mu
 
             if ( node.m_type==NodeObjectGroup::CS_ONE_OR_NONE )
             {
-                PARAMETER_DESC::INT_VALUE_DESC nullValue;
+                FParameterDesc::INT_VALUE_DESC nullValue;
                 nullValue.m_value = -1;
                 nullValue.m_name = "None";
                 op->parameter.m_possibleValues.Add( nullValue );
@@ -585,7 +585,7 @@ namespace mu
                     case NodeObjectGroup::CS_ALWAYS_ONE:
                     {
                         // Add the option to the enumeration parameter
-                        PARAMETER_DESC::INT_VALUE_DESC value;
+                        FParameterDesc::INT_VALUE_DESC value;
                         value.m_value = (int16_t)t;
                         value.m_name = pChildNode->GetName();
                         enumOp->parameter.m_possibleValues.Add( value );
diff --git a/Engine/Plugins/Experimental/Mutable/Source/MutableTools/Private/MuT/CodeGenerator_Image.cpp b/Engine/Plugins/Experimental/Mutable/Source/MutableTools/Private/MuT/CodeGenerator_Image.cpp
--- a/Engine/Plugins/Experimental/Mutable/Source/MutableTools/Private/MuT/CodeGenerator_Image.cpp
+++ b/Engine/Plugins/Experimental/Mutable/Source/MutableTools/Private/MuT/CodeGenerator_Image.cpp
@@ -115,11 +115,11 @@ namespace mu
 
 
 	//-------------------------------------------------------------------------------------------------
-	void CodeGenerator::GenerateImage(IMAGE_GENERATION_RESULT& result, const NodeImagePtrConst& Untyped)
+	void CodeGenerator::GenerateImage(FImageGenerationResult& result, const NodeImagePtrConst& Untyped)
 	{
 		if (!Untyped)
 		{
-			result = IMAGE_GENERATION_RESULT();
+			result = FImageGenerationResult();
 			return;
 		}
 
@@ -141,7 +141,7 @@ namespace mu
 		}
 
 		// See if it was already generated
-		VISITED_MAP_KEY key = GetCurrentCacheKey(Untyped);
+		FVisitedKeyMap key = GetCurrentCacheKey(Untyped);
 		GeneratedImagesMap::ValueType* it = m_generatedImages.Find(key);
 		if (it)
 		{
@@ -198,7 +198,7 @@ namespace mu
 
 
     //---------------------------------------------------------------------------------------------
-    void CodeGenerator::GenerateImage_Constant(IMAGE_GENERATION_RESULT& result, const NodeImageConstant* InNode)
+    void CodeGenerator::GenerateImage_Constant(FImageGenerationResult& result, const NodeImageConstant* InNode)
     {
 		const NodeImageConstant::Private& node = *InNode->GetPrivate();
 		
@@ -255,7 +255,7 @@ namespace mu
 
 
     //---------------------------------------------------------------------------------------------
-	void CodeGenerator::GenerateImage_Parameter(IMAGE_GENERATION_RESULT& result, const NodeImageParameter* InNode )
+	void CodeGenerator::GenerateImage_Parameter(FImageGenerationResult& result, const NodeImageParameter* InNode )
     {
 		const NodeImageParameter::Private& node = *InNode->GetPrivate();
 
@@ -274,7 +274,7 @@ namespace mu
 			// Generate the code for the ranges
 			for (int32 a = 0; a < node.m_ranges.Num(); ++a)
 			{
-				RANGE_GENERATION_RESULT rangeResult;
+				FRangeGenerationResult rangeResult;
 				GenerateRange(rangeResult, node.m_ranges[a]);
 				op->ranges.Emplace(op.get(), rangeResult.sizeOp, rangeResult.rangeName, rangeResult.rangeUID);
 			}
@@ -291,7 +291,7 @@ namespace mu
 
 
     //---------------------------------------------------------------------------------------------
-	void CodeGenerator::GenerateImage_Layer(IMAGE_GENERATION_RESULT& result, const NodeImageLayer* InNode)
+	void CodeGenerator::GenerateImage_Layer(FImageGenerationResult& result, const NodeImageLayer* InNode)
 	{
 		const NodeImageLayer::Private& node = *InNode->GetPrivate();
 		
@@ -357,7 +357,7 @@ namespace mu
 
 
     //---------------------------------------------------------------------------------------------
-	void CodeGenerator::GenerateImage_LayerColour(IMAGE_GENERATION_RESULT& result, const NodeImageLayerColour* InNode)
+	void CodeGenerator::GenerateImage_LayerColour(FImageGenerationResult& result, const NodeImageLayerColour* InNode)
 	{
 		const NodeImageLayerColour::Private& node = *InNode->GetPrivate();
 		
@@ -412,7 +412,7 @@ namespace mu
 
 
     //---------------------------------------------------------------------------------------------
-	void CodeGenerator::GenerateImage_MultiLayer(IMAGE_GENERATION_RESULT& result, const NodeImageMultiLayer* InNode)
+	void CodeGenerator::GenerateImage_MultiLayer(FImageGenerationResult& result, const NodeImageMultiLayer* InNode)
 	{
 		const NodeImageMultiLayer::Private& node = *InNode->GetPrivate();
 
@@ -475,7 +475,7 @@ namespace mu
         if ( node.m_pRange )
         {
             //op->range = Generate( pRange );
-            RANGE_GENERATION_RESULT rangeResult;
+            FRangeGenerationResult rangeResult;
             GenerateRange( rangeResult, node.m_pRange );
 
             op->range.rangeSize = rangeResult.sizeOp;
@@ -488,7 +488,7 @@ namespace mu
 
 
 	//---------------------------------------------------------------------------------------------
-	void CodeGenerator::GenerateImage_NormalComposite(IMAGE_GENERATION_RESULT& result, const NodeImageNormalComposite* InNode)
+	void CodeGenerator::GenerateImage_NormalComposite(FImageGenerationResult& result, const NodeImageNormalComposite* InNode)
 	{
 		const NodeImageNormalComposite::Private& node = *InNode->GetPrivate();
 
@@ -539,7 +539,7 @@ namespace mu
 		result.op = op;
     }
 
-	void CodeGenerator::GenerateImage_Transform(IMAGE_GENERATION_RESULT& result, const NodeImageTransform* InNode)
+	void CodeGenerator::GenerateImage_Transform(FImageGenerationResult& result, const NodeImageTransform* InNode)
     {
 		const NodeImageTransform::Private& node = *InNode->GetPrivate();
 
@@ -607,7 +607,7 @@ namespace mu
     }
 
     //---------------------------------------------------------------------------------------------
-	void CodeGenerator::GenerateImage_Interpolate(IMAGE_GENERATION_RESULT& result, const NodeImageInterpolate* InNode)
+	void CodeGenerator::GenerateImage_Interpolate(FImageGenerationResult& result, const NodeImageInterpolate* InNode)
 	{
 		const NodeImageInterpolate::Private& node = *InNode->GetPrivate();
 
@@ -664,7 +664,7 @@ namespace mu
 
 
     //---------------------------------------------------------------------------------------------
-	void CodeGenerator::GenerateImage_Interpolate3(IMAGE_GENERATION_RESULT& result, const NodeImageInterpolate3* InNode)
+	void CodeGenerator::GenerateImage_Interpolate3(FImageGenerationResult& result, const NodeImageInterpolate3* InNode)
 	{
 		const NodeImageInterpolate3::Private& node = *InNode->GetPrivate();
 
@@ -753,7 +753,7 @@ namespace mu
 
 
 	//---------------------------------------------------------------------------------------------
-	void CodeGenerator::GenerateImage_Swizzle(IMAGE_GENERATION_RESULT& result, const NodeImageSwizzle* InNode)
+	void CodeGenerator::GenerateImage_Swizzle(FImageGenerationResult& result, const NodeImageSwizzle* InNode)
 	{
 		const NodeImageSwizzle::Private& node = *InNode->GetPrivate();
 
@@ -866,7 +866,7 @@ namespace mu
 
 
 	//---------------------------------------------------------------------------------------------
-	void CodeGenerator::GenerateImage_Format(IMAGE_GENERATION_RESULT& result, const NodeImageFormat* InNode)
+	void CodeGenerator::GenerateImage_Format(FImageGenerationResult& result, const NodeImageFormat* InNode)
 	{
 		const NodeImageFormat::Private& node = *InNode->GetPrivate();
 
@@ -893,7 +893,7 @@ namespace mu
 
 
     //---------------------------------------------------------------------------------------------
-	void CodeGenerator::GenerateImage_Saturate(IMAGE_GENERATION_RESULT& result, const NodeImageSaturate* InNode)
+	void CodeGenerator::GenerateImage_Saturate(FImageGenerationResult& result, const NodeImageSaturate* InNode)
 	{
 		const NodeImageSaturate::Private& node = *InNode->GetPrivate();
 
@@ -935,7 +935,7 @@ namespace mu
 
 
     //---------------------------------------------------------------------------------------------
-	void CodeGenerator::GenerateImage_Luminance(IMAGE_GENERATION_RESULT& result, const NodeImageLuminance* InNode)
+	void CodeGenerator::GenerateImage_Luminance(FImageGenerationResult& result, const NodeImageLuminance* InNode)
 	{
 		const NodeImageLuminance::Private& node = *InNode->GetPrivate();
 
@@ -962,7 +962,7 @@ namespace mu
 
 
     //---------------------------------------------------------------------------------------------
-	void CodeGenerator::GenerateImage_SelectColour(IMAGE_GENERATION_RESULT& result, const NodeImageSelectColour* InNode)
+	void CodeGenerator::GenerateImage_SelectColour(FImageGenerationResult& result, const NodeImageSelectColour* InNode)
 	{
 		const NodeImageSelectColour::Private& node = *InNode->GetPrivate();
 
@@ -1004,7 +1004,7 @@ namespace mu
 
 
     //---------------------------------------------------------------------------------------------
-	void CodeGenerator::GenerateImage_Difference(IMAGE_GENERATION_RESULT& result, const NodeImageDifference* InNode)
+	void CodeGenerator::GenerateImage_Difference(FImageGenerationResult& result, const NodeImageDifference* InNode)
 	{
 		const NodeImageDifference::Private& node = *InNode->GetPrivate();
 
@@ -1049,7 +1049,7 @@ namespace mu
 
 
     //---------------------------------------------------------------------------------------------
-	void CodeGenerator::GenerateImage_ColourMap(IMAGE_GENERATION_RESULT& result, const NodeImageColourMap* InNode)
+	void CodeGenerator::GenerateImage_ColourMap(FImageGenerationResult& result, const NodeImageColourMap* InNode)
 	{
 		const NodeImageColourMap::Private& node = *InNode->GetPrivate();
 
@@ -1113,7 +1113,7 @@ namespace mu
 
 
     //---------------------------------------------------------------------------------------------
-	void CodeGenerator::GenerateImage_Gradient(IMAGE_GENERATION_RESULT& result, const NodeImageGradient* InNode)
+	void CodeGenerator::GenerateImage_Gradient(FImageGenerationResult& result, const NodeImageGradient* InNode)
 	{
 		const NodeImageGradient::Private& node = *InNode->GetPrivate();
 
@@ -1154,7 +1154,7 @@ namespace mu
 
 
     //---------------------------------------------------------------------------------------------
-	void CodeGenerator::GenerateImage_Binarise(IMAGE_GENERATION_RESULT& result, const NodeImageBinarise* InNode)
+	void CodeGenerator::GenerateImage_Binarise(FImageGenerationResult& result, const NodeImageBinarise* InNode)
 	{
 		const NodeImageBinarise::Private& node = *InNode->GetPrivate();
 
@@ -1197,7 +1197,7 @@ namespace mu
 
 
     //---------------------------------------------------------------------------------------------
-	void CodeGenerator::GenerateImage_Resize(IMAGE_GENERATION_RESULT& result, const NodeImageResize* InNode)
+	void CodeGenerator::GenerateImage_Resize(FImageGenerationResult& result, const NodeImageResize* InNode)
 	{
 		const NodeImageResize::Private& node = *InNode->GetPrivate();
 
@@ -1257,7 +1257,7 @@ namespace mu
 
 
     //---------------------------------------------------------------------------------------------
-	void CodeGenerator::GenerateImage_PlainColour(IMAGE_GENERATION_RESULT& result, const NodeImagePlainColour* InNode)
+	void CodeGenerator::GenerateImage_PlainColour(FImageGenerationResult& result, const NodeImagePlainColour* InNode)
 	{
 		const NodeImagePlainColour::Private& node = *InNode->GetPrivate();
 
@@ -1299,7 +1299,7 @@ namespace mu
 
 
     //---------------------------------------------------------------------------------------------
-	void CodeGenerator::GenerateImage_Switch(IMAGE_GENERATION_RESULT& result, const NodeImageSwitch* InNode)
+	void CodeGenerator::GenerateImage_Switch(FImageGenerationResult& result, const NodeImageSwitch* InNode)
 	{
 		const NodeImageSwitch::Private& node = *InNode->GetPrivate();
 
@@ -1385,7 +1385,7 @@ namespace mu
 
 
     //---------------------------------------------------------------------------------------------
-	void CodeGenerator::GenerateImage_Conditional(IMAGE_GENERATION_RESULT& result, const NodeImageConditional* InNode)
+	void CodeGenerator::GenerateImage_Conditional(FImageGenerationResult& result, const NodeImageConditional* InNode)
 	{
 		const NodeImageConditional::Private& node = *InNode->GetPrivate();
 
@@ -1412,7 +1412,7 @@ namespace mu
 
 
     //---------------------------------------------------------------------------------------------
-	void CodeGenerator::GenerateImage_Project(IMAGE_GENERATION_RESULT& result, const NodeImageProject* InNode)
+	void CodeGenerator::GenerateImage_Project(FImageGenerationResult& result, const NodeImageProject* InNode)
 	{
 		const NodeImageProject::Private& node = *InNode->GetPrivate();
 
@@ -1426,7 +1426,7 @@ namespace mu
         Ptr<ASTOp> lastMeshOp = pop;
 
         // Projector
-        PROJECTOR_GENERATION_RESULT projectorResult;
+        FProjectorGenerationResult projectorResult;
         if ( node.m_pProjector )
         {
             GenerateProjector( projectorResult, node.m_pProjector );
@@ -1613,7 +1613,7 @@ namespace mu
 
 
     //---------------------------------------------------------------------------------------------
-	void CodeGenerator::GenerateImage_Mipmap(IMAGE_GENERATION_RESULT& result, const NodeImageMipmap* InNode)
+	void CodeGenerator::GenerateImage_Mipmap(FImageGenerationResult& result, const NodeImageMipmap* InNode)
 	{
 		const NodeImageMipmap::Private& node = *InNode->GetPrivate();
 
@@ -1678,7 +1678,7 @@ namespace mu
 
 
 	//---------------------------------------------------------------------------------------------
-	void CodeGenerator::GenerateImage_Invert(IMAGE_GENERATION_RESULT& result, const NodeImageInvert* InNode)
+	void CodeGenerator::GenerateImage_Invert(FImageGenerationResult& result, const NodeImageInvert* InNode)
 	{
 		const NodeImageInvert::Private& node = *InNode->GetPrivate();
 
@@ -1708,7 +1708,7 @@ namespace mu
 
 
     //---------------------------------------------------------------------------------------------
-	void CodeGenerator::GenerateImage_Variation(IMAGE_GENERATION_RESULT& result, const NodeImageVariation* InNode)
+	void CodeGenerator::GenerateImage_Variation(FImageGenerationResult& result, const NodeImageVariation* InNode)
 	{
 		const NodeImageVariation::Private& node = *InNode->GetPrivate();
 
@@ -1717,7 +1717,7 @@ namespace mu
         // Default case
         if ( node.m_defaultImage )
         {
-            IMAGE_GENERATION_RESULT branchResults;
+            FImageGenerationResult branchResults;
 			GenerateImage(branchResults, node.m_defaultImage);
 			currentOp = branchResults.op;
         }
@@ -1802,7 +1802,7 @@ namespace mu
 
 
 	//---------------------------------------------------------------------------------------------
-	void CodeGenerator::GenerateImage_Table(IMAGE_GENERATION_RESULT& result, const NodeImageTable* InNode)
+	void CodeGenerator::GenerateImage_Table(FImageGenerationResult& result, const NodeImageTable* InNode)
 	{
 		const NodeImageTable::Private& node = *InNode->GetPrivate();
 
diff --git a/Engine/Plugins/Experimental/Mutable/Source/MutableTools/Private/MuT/CodeGenerator_Mesh.cpp b/Engine/Plugins/Experimental/Mutable/Source/MutableTools/Private/MuT/CodeGenerator_Mesh.cpp
--- a/Engine/Plugins/Experimental/Mutable/Source/MutableTools/Private/MuT/CodeGenerator_Mesh.cpp
+++ b/Engine/Plugins/Experimental/Mutable/Source/MutableTools/Private/MuT/CodeGenerator_Mesh.cpp
@@ -1172,7 +1172,7 @@ class Node;
 
         // Morph to an ellipse
         {
-            op->morphShape.type = (uint8_t)SHAPE::Type::Ellipse;
+            op->morphShape.type = (uint8_t)FShape::Type::Ellipse;
             op->morphShape.position = node.m_origin;
             op->morphShape.up = node.m_normal;
             op->morphShape.size = vec3f(node.m_radius1, node.m_radius2, node.m_rotation); // TODO: Move rotation to ellipse rotation reference base instead of passing it directly
@@ -1196,7 +1196,7 @@ class Node;
         if (node.m_vertexSelectionType== NodeMeshClipMorphPlane::Private::VS_SHAPE)
         {
             op->vertexSelectionType = OP::MeshClipMorphPlaneArgs::VS_SHAPE;
-            op->selectionShape.type = (uint8_t)SHAPE::Type::AABox;
+            op->selectionShape.type = (uint8_t)FShape::Type::AABox;
             op->selectionShape.position = node.m_selectionBoxOrigin;
             op->selectionShape.size = node.m_selectionBoxRadius;
         }
diff --git a/Engine/Plugins/Experimental/Mutable/Source/MutableTools/Private/MuT/CodeGenerator_Projector.cpp b/Engine/Plugins/Experimental/Mutable/Source/MutableTools/Private/MuT/CodeGenerator_Projector.cpp
--- a/Engine/Plugins/Experimental/Mutable/Source/MutableTools/Private/MuT/CodeGenerator_Projector.cpp
+++ b/Engine/Plugins/Experimental/Mutable/Source/MutableTools/Private/MuT/CodeGenerator_Projector.cpp
@@ -29,9 +29,9 @@ namespace mu
 {
 
 	//-------------------------------------------------------------------------------------------------
-	PROJECTOR ProjectorFromNode(const NodeProjectorConstant::Private& node)
+	FProjector ProjectorFromNode(const NodeProjectorConstant::Private& node)
 	{
-		PROJECTOR p;
+		FProjector p;
 		p.type = node.m_type;
 		p.position[0] = node.m_position[0];
 		p.position[1] = node.m_position[1];
@@ -52,16 +52,16 @@ namespace mu
 
 
 	//-------------------------------------------------------------------------------------------------
-	void CodeGenerator::GenerateProjector(PROJECTOR_GENERATION_RESULT& result, const NodeProjectorPtrConst& untyped)
+	void CodeGenerator::GenerateProjector(FProjectorGenerationResult& result, const NodeProjectorPtrConst& untyped)
 	{
 		if (!untyped)
 		{
-			result = PROJECTOR_GENERATION_RESULT();
+			result = FProjectorGenerationResult();
 			return;
 		}
 
 		// See if it was already generated
-		VISITED_MAP_KEY key = GetCurrentCacheKey(untyped);
+		FVisitedKeyMap key = GetCurrentCacheKey(untyped);
 		GeneratedProjectorsMap::ValueType* it = m_generatedProjectors.Find(key);
 		if (it)
 		{
@@ -90,7 +90,7 @@ namespace mu
 
 
 	//-------------------------------------------------------------------------------------------------
-	void CodeGenerator::GenerateProjector_Constant(PROJECTOR_GENERATION_RESULT& result,
+	void CodeGenerator::GenerateProjector_Constant(FProjectorGenerationResult& result,
 		const Ptr<const NodeProjectorConstant>& constant)
 	{
 		const NodeProjectorConstant::Private& node = *constant->GetPrivate();
@@ -104,7 +104,7 @@ namespace mu
 
 
 	//-------------------------------------------------------------------------------------------------
-	void CodeGenerator::GenerateProjector_Parameter(PROJECTOR_GENERATION_RESULT& result,
+	void CodeGenerator::GenerateProjector_Parameter(FProjectorGenerationResult& result,
 		const Ptr<const NodeProjectorParameter>& paramn)
 	{
 		const NodeProjectorParameter::Private& node = *paramn->GetPrivate();
@@ -114,12 +114,12 @@ namespace mu
 		auto it = m_nodeVariables.find(node.m_pNode);
 		if (it == m_nodeVariables.end())
 		{
-			PARAMETER_DESC param;
+			FParameterDesc param;
 			param.m_name = node.m_name;
 			param.m_uid = node.m_uid;
 			param.m_type = PARAMETER_TYPE::T_PROJECTOR;
 
-			PROJECTOR p = ProjectorFromNode(node);
+			FProjector p = ProjectorFromNode(node);
 			param.m_defaultValue.m_projector = p;
 
 			op = new ASTOpParameter();
@@ -129,7 +129,7 @@ namespace mu
 			// Generate the code for the ranges
 			for (int32 a = 0; a < node.m_ranges.Num(); ++a)
 			{
-				RANGE_GENERATION_RESULT rangeResult;
+				FRangeGenerationResult rangeResult;
 				GenerateRange(rangeResult, node.m_ranges[a]);
 				op->ranges.Emplace(op.get(), rangeResult.sizeOp, rangeResult.rangeName, rangeResult.rangeUID);
 			}
@@ -147,14 +147,14 @@ namespace mu
 
 
 	//-------------------------------------------------------------------------------------------------
-	void CodeGenerator::GenerateMissingProjectorCode(PROJECTOR_GENERATION_RESULT& result,
+	void CodeGenerator::GenerateMissingProjectorCode(FProjectorGenerationResult& result,
 		const void* errorContext)
 	{
 		// Log an error message
 		m_pErrorLog->GetPrivate()->Add("Required projector connection not found.",
 			ELMT_ERROR, errorContext);
 
-		PROJECTOR p;
+		FProjector p;
 		p.type = PROJECTOR_TYPE::PLANAR;
 
 		p.direction[0] = 1;
diff --git a/Engine/Plugins/Experimental/Mutable/Source/MutableTools/Private/MuT/CodeGenerator_Scalar.cpp b/Engine/Plugins/Experimental/Mutable/Source/MutableTools/Private/MuT/CodeGenerator_Scalar.cpp
--- a/Engine/Plugins/Experimental/Mutable/Source/MutableTools/Private/MuT/CodeGenerator_Scalar.cpp
+++ b/Engine/Plugins/Experimental/Mutable/Source/MutableTools/Private/MuT/CodeGenerator_Scalar.cpp
@@ -55,16 +55,16 @@ namespace mu
 
 
 	//-------------------------------------------------------------------------------------------------
-	void CodeGenerator::GenerateScalar(SCALAR_GENERATION_RESULT& result, const NodeScalarPtrConst& untyped)
+	void CodeGenerator::GenerateScalar(FScalarGenerationResult& result, const NodeScalarPtrConst& untyped)
 	{
 		if (!untyped)
 		{
-			result = SCALAR_GENERATION_RESULT();
+			result = FScalarGenerationResult();
 			return;
 		}
 
 		// See if it was already generated
-		VISITED_MAP_KEY key = GetCurrentCacheKey(untyped);
+		FVisitedKeyMap key = GetCurrentCacheKey(untyped);
 		GeneratedScalarsMap::ValueType* it = m_generatedScalars.Find(key);
 		if (it)
 		{
@@ -118,7 +118,7 @@ namespace mu
 
 
 	//-------------------------------------------------------------------------------------------------
-	void CodeGenerator::GenerateScalar_Constant(SCALAR_GENERATION_RESULT& result, const Ptr<const NodeScalarConstant>& Typed)
+	void CodeGenerator::GenerateScalar_Constant(FScalarGenerationResult& result, const Ptr<const NodeScalarConstant>& Typed)
 	{
 		const NodeScalarConstant::Private& node = *Typed->GetPrivate();
 
@@ -131,7 +131,7 @@ namespace mu
 
 
 	//-------------------------------------------------------------------------------------------------
-	void CodeGenerator::GenerateScalar_Parameter(SCALAR_GENERATION_RESULT& result, const Ptr<const NodeScalarParameter>& Typed)
+	void CodeGenerator::GenerateScalar_Parameter(FScalarGenerationResult& result, const Ptr<const NodeScalarParameter>& Typed)
 	{
 		const NodeScalarParameter::Private& node = *Typed->GetPrivate();
 
@@ -140,7 +140,7 @@ namespace mu
 		auto it = m_nodeVariables.find(node.m_pNode);
 		if (it == m_nodeVariables.end())
 		{
-			PARAMETER_DESC param;
+			FParameterDesc param;
 			param.m_name = node.m_name;
 			param.m_uid = node.m_uid;
 			param.m_type = PARAMETER_TYPE::T_FLOAT;
@@ -154,7 +154,7 @@ namespace mu
 			// Generate the code for the ranges
 			for (int32 a = 0; a < node.m_ranges.Num(); ++a)
 			{
-				RANGE_GENERATION_RESULT rangeResult;
+				FRangeGenerationResult rangeResult;
 				GenerateRange(rangeResult, node.m_ranges[a]);
 				op->ranges.Emplace(op.get(), rangeResult.sizeOp, rangeResult.rangeName, rangeResult.rangeUID);
 			}
@@ -197,7 +197,7 @@ namespace mu
 
 
 	//-------------------------------------------------------------------------------------------------
-	void CodeGenerator::GenerateScalar_EnumParameter(SCALAR_GENERATION_RESULT& result, const Ptr<const NodeScalarEnumParameter>& Typed)
+	void CodeGenerator::GenerateScalar_EnumParameter(FScalarGenerationResult& result, const Ptr<const NodeScalarEnumParameter>& Typed)
 	{
 		const NodeScalarEnumParameter::Private& node = *Typed->GetPrivate();
 
@@ -206,7 +206,7 @@ namespace mu
 		auto it = m_nodeVariables.find(node.m_pNode);
 		if (it == m_nodeVariables.end())
 		{
-			PARAMETER_DESC param;
+			FParameterDesc param;
 			param.m_name = node.m_name;
 			param.m_uid = node.m_uid;
 			param.m_type = PARAMETER_TYPE::T_INT;
@@ -227,7 +227,7 @@ namespace mu
 			// Generate the code for the ranges
 			for (int32 a = 0; a < node.m_ranges.Num(); ++a)
 			{
-				RANGE_GENERATION_RESULT rangeResult;
+				FRangeGenerationResult rangeResult;
 				GenerateRange(rangeResult, node.m_ranges[a]);
 				op->ranges.Emplace(op.get(), rangeResult.sizeOp, rangeResult.rangeName, rangeResult.rangeUID);
 			}
@@ -245,7 +245,7 @@ namespace mu
 
 
 	//-------------------------------------------------------------------------------------------------
-	void CodeGenerator::GenerateScalar_Switch(SCALAR_GENERATION_RESULT& result, const Ptr<const NodeScalarSwitch>& Typed)
+	void CodeGenerator::GenerateScalar_Switch(FScalarGenerationResult& result, const Ptr<const NodeScalarSwitch>& Typed)
 	{
 		const NodeScalarSwitch::Private& node = *Typed->GetPrivate();
 
@@ -265,7 +265,7 @@ namespace mu
 		// Variable value
 		if (node.m_pParameter)
 		{
-			SCALAR_GENERATION_RESULT ChildResult;
+			FScalarGenerationResult ChildResult;
 			GenerateScalar(ChildResult, node.m_pParameter.get());
 			op->variable = ChildResult.op;
 		}
@@ -281,7 +281,7 @@ namespace mu
 			Ptr<ASTOp> branch;
 			if (node.m_options[t])
 			{
-				SCALAR_GENERATION_RESULT ChildResult;
+				FScalarGenerationResult ChildResult;
 				GenerateScalar(ChildResult, node.m_options[t].get());
 				branch = ChildResult.op;
 			}
@@ -298,7 +298,7 @@ namespace mu
 
 
 	//-------------------------------------------------------------------------------------------------
-	void CodeGenerator::GenerateScalar_Variation(SCALAR_GENERATION_RESULT& result, const Ptr<const NodeScalarVariation>& Typed)
+	void CodeGenerator::GenerateScalar_Variation(FScalarGenerationResult& result, const Ptr<const NodeScalarVariation>& Typed)
 	{
 		const NodeScalarVariation::Private& node = *Typed->GetPrivate();
 
@@ -309,7 +309,7 @@ namespace mu
 		{
 			FMeshGenerationResult branchResults;
 
-			SCALAR_GENERATION_RESULT ChildResult;
+			FScalarGenerationResult ChildResult;
 			GenerateScalar(ChildResult, node.m_defaultScalar);
 			op = ChildResult.op;
 		}
@@ -345,7 +345,7 @@ namespace mu
 			Ptr<ASTOp> variationOp;
 			if (node.m_variations[t].m_scalar)
 			{
-				SCALAR_GENERATION_RESULT ChildResult;
+				FScalarGenerationResult ChildResult;
 				GenerateScalar(ChildResult, node.m_variations[t].m_scalar);
 				variationOp = ChildResult.op;
 			}
@@ -371,7 +371,7 @@ namespace mu
 
 
 	//-------------------------------------------------------------------------------------------------
-	void CodeGenerator::GenerateScalar_Curve(SCALAR_GENERATION_RESULT& result, const Ptr<const NodeScalarCurve>& Typed)
+	void CodeGenerator::GenerateScalar_Curve(FScalarGenerationResult& result, const Ptr<const NodeScalarCurve>& Typed)
 	{
 		const NodeScalarCurve::Private& node = *Typed->GetPrivate();
 
@@ -394,7 +394,7 @@ namespace mu
 
 
 	//-------------------------------------------------------------------------------------------------
-	void CodeGenerator::GenerateScalar_Arithmetic(SCALAR_GENERATION_RESULT& result, const Ptr<const NodeScalarArithmeticOperation>& Typed)
+	void CodeGenerator::GenerateScalar_Arithmetic(FScalarGenerationResult& result, const Ptr<const NodeScalarArithmeticOperation>& Typed)
 	{
 		const NodeScalarArithmeticOperation::Private& node = *Typed->GetPrivate();
 
@@ -452,7 +452,7 @@ namespace mu
 
 
 	//-------------------------------------------------------------------------------------------------
-	void CodeGenerator::GenerateScalar_Table(SCALAR_GENERATION_RESULT& result, const Ptr<const NodeScalarTable>& Typed)
+	void CodeGenerator::GenerateScalar_Table(FScalarGenerationResult& result, const Ptr<const NodeScalarTable>& Typed)
 	{
 		const NodeScalarTable::Private& node = *Typed->GetPrivate();
 
diff --git a/Engine/Plugins/Experimental/Mutable/Source/MutableTools/Private/MuT/CodeGenerator_String.cpp b/Engine/Plugins/Experimental/Mutable/Source/MutableTools/Private/MuT/CodeGenerator_String.cpp
--- a/Engine/Plugins/Experimental/Mutable/Source/MutableTools/Private/MuT/CodeGenerator_String.cpp
+++ b/Engine/Plugins/Experimental/Mutable/Source/MutableTools/Private/MuT/CodeGenerator_String.cpp
@@ -31,16 +31,16 @@ namespace mu
 
 
 	//-------------------------------------------------------------------------------------------------
-	void CodeGenerator::GenerateString(STRING_GENERATION_RESULT& result, const NodeStringPtrConst& untyped)
+	void CodeGenerator::GenerateString(FStringGenerationResult& result, const NodeStringPtrConst& untyped)
 	{
 		if (!untyped)
 		{
-			result = STRING_GENERATION_RESULT();
+			result = FStringGenerationResult();
 			return;
 		}
 
 		// See if it was already generated
-		VISITED_MAP_KEY key = GetCurrentCacheKey(untyped);
+		FVisitedKeyMap key = GetCurrentCacheKey(untyped);
 		GeneratedStringsMap::ValueType* it = m_generatedStrings.Find(key);
 		if (it)
 		{
@@ -68,7 +68,7 @@ namespace mu
 
 
 	//-------------------------------------------------------------------------------------------------
-	void CodeGenerator::GenerateString_Constant(STRING_GENERATION_RESULT& result, const Ptr<const NodeStringConstant>& Typed)
+	void CodeGenerator::GenerateString_Constant(FStringGenerationResult& result, const Ptr<const NodeStringConstant>& Typed)
 	{
 		const NodeStringConstant::Private& node = *Typed->GetPrivate();
 
@@ -80,7 +80,7 @@ namespace mu
 
 
 	//-------------------------------------------------------------------------------------------------
-	void CodeGenerator::GenerateString_Parameter(STRING_GENERATION_RESULT& result, const Ptr<const NodeStringParameter>& Typed)
+	void CodeGenerator::GenerateString_Parameter(FStringGenerationResult& result, const Ptr<const NodeStringParameter>& Typed)
 	{
 		const NodeStringParameter::Private& node = *Typed->GetPrivate();
 
@@ -89,7 +89,7 @@ namespace mu
 		auto it = m_nodeVariables.find(node.m_pNode);
 		if (it == m_nodeVariables.end())
 		{
-			PARAMETER_DESC param;
+			FParameterDesc param;
 			param.m_name = node.m_name;
 			param.m_uid = node.m_uid;
 			param.m_type = PARAMETER_TYPE::T_FLOAT;
@@ -107,7 +107,7 @@ namespace mu
 			// Generate the code for the ranges
 			for (int32 a = 0; a < node.m_ranges.Num(); ++a)
 			{
-				RANGE_GENERATION_RESULT rangeResult;
+				FRangeGenerationResult rangeResult;
 				GenerateRange(rangeResult, node.m_ranges[a]);
 				op->ranges.Emplace(op.get(), rangeResult.sizeOp, rangeResult.rangeName, rangeResult.rangeUID);
 			}
diff --git a/Engine/Plugins/Experimental/Mutable/Source/MutableTools/Private/MuT/CodeOptimiser.cpp b/Engine/Plugins/Experimental/Mutable/Source/MutableTools/Private/MuT/CodeOptimiser.cpp
--- a/Engine/Plugins/Experimental/Mutable/Source/MutableTools/Private/MuT/CodeOptimiser.cpp
+++ b/Engine/Plugins/Experimental/Mutable/Source/MutableTools/Private/MuT/CodeOptimiser.cpp
@@ -330,7 +330,7 @@ namespace mu
 	//---------------------------------------------------------------------------------------------
 	//---------------------------------------------------------------------------------------------
 	//---------------------------------------------------------------------------------------------
-	bool FindOpTypeVisitor::Apply( PROGRAM& program, OP::ADDRESS rootAt )
+	bool FindOpTypeVisitor::Apply( FProgram& program, OP::ADDRESS rootAt )
 	{
 		bool found = false;
 
@@ -418,7 +418,7 @@ namespace mu
 
 
 	//---------------------------------------------------------------------------------------------
-	bool IsConstantVisitor::Apply( PROGRAM& program, OP::ADDRESS at )
+	bool IsConstantVisitor::Apply( FProgram& program, OP::ADDRESS at )
 	{
 		return !m_findOpTypeVisitor->Apply(program,at);
 	}
@@ -470,7 +470,7 @@ namespace mu
 			ASTOp::FullLink( cloned, model->GetPrivate()->m_program, &LinkerOptions);
 			OP::ADDRESS at = cloned->linkedAddress;
 
-			PROGRAM::STATE state;
+			FProgram::FState state;
 			state.m_root = at;
 			model->GetPrivate()->m_program.m_states.Add(state);
 
diff --git a/Engine/Plugins/Experimental/Mutable/Source/MutableTools/Private/MuT/CodeOptimiser.h b/Engine/Plugins/Experimental/Mutable/Source/MutableTools/Private/MuT/CodeOptimiser.h
--- a/Engine/Plugins/Experimental/Mutable/Source/MutableTools/Private/MuT/CodeOptimiser.h
+++ b/Engine/Plugins/Experimental/Mutable/Source/MutableTools/Private/MuT/CodeOptimiser.h
@@ -63,7 +63,7 @@ namespace mu
         {
         }
 
-        bool Apply( PROGRAM& program, OP::ADDRESS at );
+        bool Apply( FProgram& program, OP::ADDRESS at );
 
 
     private:
@@ -90,7 +90,7 @@ namespace mu
 
         IsConstantVisitor();
 
-        bool Apply( PROGRAM& program, OP::ADDRESS at );
+        bool Apply( FProgram& program, OP::ADDRESS at );
 
     private:
 
@@ -272,7 +272,7 @@ namespace mu
     //---------------------------------------------------------------------------------------------
     //! Return true if two non-zero pixels of the masks overlap.
     //---------------------------------------------------------------------------------------------
-    extern bool AreMasksOverlapping( const PROGRAM& program, OP::ADDRESS a, OP::ADDRESS b );
+    extern bool AreMasksOverlapping( const FProgram& program, OP::ADDRESS a, OP::ADDRESS b );
 
 
     //---------------------------------------------------------------------------------------------
diff --git a/Engine/Plugins/Experimental/Mutable/Source/MutableTools/Private/MuT/CodeOptimiser_Logic.cpp b/Engine/Plugins/Experimental/Mutable/Source/MutableTools/Private/MuT/CodeOptimiser_Logic.cpp
--- a/Engine/Plugins/Experimental/Mutable/Source/MutableTools/Private/MuT/CodeOptimiser_Logic.cpp
+++ b/Engine/Plugins/Experimental/Mutable/Source/MutableTools/Private/MuT/CodeOptimiser_Logic.cpp
@@ -24,7 +24,7 @@ namespace mu
     //---------------------------------------------------------------------------------------------
     //---------------------------------------------------------------------------------------------
     //---------------------------------------------------------------------------------------------
-//    IntExpressionPtr GetIntExpression( OP::ADDRESS at, const PROGRAM& program )
+//    IntExpressionPtr GetIntExpression( OP::ADDRESS at, const FProgram& program )
 //    {
 //        IntExpressionPtr pRes = new IntExpression;
 
@@ -56,7 +56,7 @@ namespace mu
     //---------------------------------------------------------------------------------------------
     //---------------------------------------------------------------------------------------------
 //    BoolExpressionPtr GetBoolExpression( OP::ADDRESS at,
-//                                         const PROGRAM& program,
+//                                         const FProgram& program,
 //                                         BOOL_EXPRESSION_CACHE* cache )
 //    {
 //        if (cache && (size_t)at<=cache->m_values.size() && cache->m_values[at])
@@ -126,7 +126,7 @@ namespace mu
     //---------------------------------------------------------------------------------------------
     //---------------------------------------------------------------------------------------------
     //---------------------------------------------------------------------------------------------
-//    void LogicCallstackGather::Apply( PROGRAM& program, int state )
+//    void LogicCallstackGather::Apply( FProgram& program, int state )
 //    {
 //        MUTABLE_CPUPROFILER_SCOPE(LogicCallstackGather);
 //        Traverse( program.m_states[state].m_root, program );
@@ -134,7 +134,7 @@ namespace mu
 
 
 //    //---------------------------------------------------------------------------------------------
-//    void LogicCallstackGather::Visit( OP::ADDRESS at, PROGRAM& program )
+//    void LogicCallstackGather::Visit( OP::ADDRESS at, FProgram& program )
 //    {
 //        // We traverse top down
 //        OP_TYPE type = (OP_TYPE)program.m_code[at].type;
@@ -244,7 +244,7 @@ namespace mu
     //---------------------------------------------------------------------------------------------
     //---------------------------------------------------------------------------------------------
     //---------------------------------------------------------------------------------------------
-//    bool ClonelessLogicOptimiser::Apply( PROGRAM& program, int state )
+//    bool ClonelessLogicOptimiser::Apply( FProgram& program, int state )
 //    {
 //        MUTABLE_CPUPROFILER_SCOPE(ClonelessLogicOptimiser);
 
@@ -259,7 +259,7 @@ namespace mu
 
 
 //    //---------------------------------------------------------------------------------------------
-//    OP::ADDRESS ClonelessLogicOptimiser::Visit( OP::ADDRESS at, PROGRAM& program )
+//    OP::ADDRESS ClonelessLogicOptimiser::Visit( OP::ADDRESS at, FProgram& program )
 //    {
 //        // We traverse top down
 
@@ -400,7 +400,7 @@ namespace mu
 //    //---------------------------------------------------------------------------------------------
 //    //---------------------------------------------------------------------------------------------
 //    //---------------------------------------------------------------------------------------------
-//    bool LogicOptimiser::Apply( PROGRAM& program, int state )
+//    bool LogicOptimiser::Apply( FProgram& program, int state )
 //    {
 //        m_initialCodeSize = program.m_opAddress.Num();
 //        m_modified = false;
@@ -411,7 +411,7 @@ namespace mu
 
 
 //    //---------------------------------------------------------------------------------------------
-//    OP::ADDRESS LogicOptimiser::Visit( OP::ADDRESS at, PROGRAM& program )
+//    OP::ADDRESS LogicOptimiser::Visit( OP::ADDRESS at, FProgram& program )
 //    {
 //        // Sanity stop
 //        const int sanityFactor = 4;
@@ -579,7 +579,7 @@ namespace mu
 
 
 //    //---------------------------------------------------------------------------------------------
-//    OP::ADDRESS IntExpression::GenerateCode( PROGRAM& program, PROGRAM::OP_CACHE* cache )
+//    OP::ADDRESS IntExpression::GenerateCode( FProgram& program, FProgram::OP_CACHE* cache )
 //    {
 //        OP::ADDRESS at = 0;
 
@@ -750,7 +750,7 @@ namespace mu
 
 
 //    //---------------------------------------------------------------------------------------------
-//    OP::ADDRESS BoolExpression::GenerateCode( PROGRAM& program, PROGRAM::OP_CACHE* cache )
+//    OP::ADDRESS BoolExpression::GenerateCode( FProgram& program, FProgram::OP_CACHE* cache )
 //    {
 //        if (m_generated)
 //        {
@@ -1045,7 +1045,7 @@ namespace mu
 
                                 // Check if the child condition has a value with the current facts
                                 Ptr<ASTOp> pChildCond = bottomConditional->condition.child();
-                                ASTOp::BOOL_EVAL_RESULT result;
+                                ASTOp::FBoolEvalResult result;
                                 {
                                     //MUTABLE_CPUPROFILER_SCOPE(EvaluateBool);
                                     result = pChildCond->EvaluateBool( facts );
@@ -1104,7 +1104,7 @@ namespace mu
 
                                 // Check if the child condition has a value with the current facts
 								Ptr<ASTOp> pChildCond = bottomConditional->condition.child();
-                                ASTOp::BOOL_EVAL_RESULT result = pChildCond->EvaluateBool( facts );
+                                ASTOp::FBoolEvalResult result = pChildCond->EvaluateBool( facts );
                                 if ( result==ASTOp::BET_FALSE )
                                 {
                                     conditionaAreExclusive = true;
@@ -1156,7 +1156,7 @@ namespace mu
 
 //                                // Check if the child condition has a value with the current facts
 //                                auto pChildCond = bottomConditional->condition.child();
-//                                ASTOp::BOOL_EVAL_RESULT result = pChildCond->EvaluateBool( facts );
+//                                ASTOp::FBoolEvalResult result = pChildCond->EvaluateBool( facts );
 //                                if ( result==ASTOp::BET_FALSE )
 //                                {
 //                                    conditionaAreExclusive = true;
diff --git a/Engine/Plugins/Experimental/Mutable/Source/MutableTools/Private/MuT/Compiler.cpp b/Engine/Plugins/Experimental/Mutable/Source/MutableTools/Private/MuT/Compiler.cpp
--- a/Engine/Plugins/Experimental/Mutable/Source/MutableTools/Private/MuT/Compiler.cpp
+++ b/Engine/Plugins/Experimental/Mutable/Source/MutableTools/Private/MuT/Compiler.cpp
@@ -374,7 +374,7 @@ namespace mu
 	//---------------------------------------------------------------------------------------------
 	namespace
 	{
-		inline void EnsureUniqueRomId(uint32& RomId, const mu::PROGRAM& program)
+		inline void EnsureUniqueRomId(uint32& RomId, const mu::FProgram& program)
 		{
 			while (true)
 			{
@@ -407,7 +407,7 @@ namespace mu
 		LLM_SCOPE_BYNAME(TEXT("MutableRuntime"));
 		MUTABLE_CPUPROFILER_SCOPE(Mutable_GenerateRoms);
 
-		mu::PROGRAM& program = p->GetPrivate()->m_program;
+		mu::FProgram& program = p->GetPrivate()->m_program;
 
 		int32 NumRoms = 0;
 		int32 NumEmbedded = 0;
diff --git a/Engine/Plugins/Experimental/Mutable/Source/MutableTools/Private/MuT/CompilerPrivate.h b/Engine/Plugins/Experimental/Mutable/Source/MutableTools/Private/MuT/CompilerPrivate.h
--- a/Engine/Plugins/Experimental/Mutable/Source/MutableTools/Private/MuT/CompilerPrivate.h
+++ b/Engine/Plugins/Experimental/Mutable/Source/MutableTools/Private/MuT/CompilerPrivate.h
@@ -125,7 +125,7 @@ namespace mu
     {
         OBJECT_STATE nodeState;
         Ptr<ASTOp> root;
-        PROGRAM::STATE state;
+        FProgram::FState state;
         STATE_OPTIMIZATION_OPTIONS optimisationFlags;
 
         //! List of instructions that need to be cached to efficiently update this state
diff --git a/Engine/Plugins/Experimental/Mutable/Source/MutableTools/Private/MuT/DataPacker.cpp b/Engine/Plugins/Experimental/Mutable/Source/MutableTools/Private/MuT/DataPacker.cpp
--- a/Engine/Plugins/Experimental/Mutable/Source/MutableTools/Private/MuT/DataPacker.cpp
+++ b/Engine/Plugins/Experimental/Mutable/Source/MutableTools/Private/MuT/DataPacker.cpp
@@ -36,7 +36,7 @@ namespace mu
     //---------------------------------------------------------------------------------------------
     //---------------------------------------------------------------------------------------------
     //---------------------------------------------------------------------------------------------
-    SubtreeSearchConstantVisitor::SubtreeSearchConstantVisitor( PROGRAM& program,
+    SubtreeSearchConstantVisitor::SubtreeSearchConstantVisitor( FProgram& program,
                                                                 OP::ADDRESS constant,
                                                                 OP_TYPE optype )
         : m_program( program ),
@@ -134,7 +134,7 @@ namespace mu
     {
     public:
 
-        GatherParametersVisitor( PROGRAM& program, OP::ADDRESS constant, OP_TYPE opType )
+        GatherParametersVisitor( FProgram& program, OP::ADDRESS constant, OP_TYPE opType )
             : m_constSearch( program, constant, opType )
         {
             MUTABLE_CPUPROFILER_SCOPE(GatherParametersVisitor);
@@ -160,7 +160,7 @@ namespace mu
         }
 
 
-        virtual bool Visit( OP::ADDRESS at, PROGRAM& program )
+        virtual bool Visit( OP::ADDRESS at, FProgram& program )
         {
             bool recurse = true;
 
diff --git a/Engine/Plugins/Experimental/Mutable/Source/MutableTools/Private/MuT/DataPacker.h b/Engine/Plugins/Experimental/Mutable/Source/MutableTools/Private/MuT/DataPacker.h
--- a/Engine/Plugins/Experimental/Mutable/Source/MutableTools/Private/MuT/DataPacker.h
+++ b/Engine/Plugins/Experimental/Mutable/Source/MutableTools/Private/MuT/DataPacker.h
@@ -8,7 +8,7 @@
 
 #include <stdint.h>
 
-namespace mu { struct PROGRAM; }
+namespace mu { struct FProgram; }
 
 
 #define MUTABLE_HASH_SEED					((uint32_t)0xcadababa)
@@ -30,14 +30,14 @@ namespace mu
     {
     public:
 
-        SubtreeSearchConstantVisitor( PROGRAM& program, OP::ADDRESS constant, OP_TYPE optype );
+        SubtreeSearchConstantVisitor( FProgram& program, OP::ADDRESS constant, OP_TYPE optype );
 
         bool Run( OP::ADDRESS root );
 
 
     private:
 
-        PROGRAM& m_program;
+        FProgram& m_program;
         OP::ADDRESS m_constant;
         OP_TYPE m_opType;
 
diff --git a/Engine/Plugins/Experimental/Mutable/Source/MutableTools/Private/MuT/ErrorLog.cpp b/Engine/Plugins/Experimental/Mutable/Source/MutableTools/Private/MuT/ErrorLog.cpp
--- a/Engine/Plugins/Experimental/Mutable/Source/MutableTools/Private/MuT/ErrorLog.cpp
+++ b/Engine/Plugins/Experimental/Mutable/Source/MutableTools/Private/MuT/ErrorLog.cpp
@@ -328,7 +328,7 @@ namespace mu
 
 
 	//---------------------------------------------------------------------------------------------
-    extern FString GetOpDesc( const PROGRAM& program, OP::ADDRESS at )
+    extern FString GetOpDesc( const FProgram& program, OP::ADDRESS at )
 	{
 		char temp[1024];
 		FMemory::Memzero( temp, 1024 );
diff --git a/Engine/Plugins/Experimental/Mutable/Source/MutableTools/Private/MuT/ErrorLogPrivate.h b/Engine/Plugins/Experimental/Mutable/Source/MutableTools/Private/MuT/ErrorLogPrivate.h
--- a/Engine/Plugins/Experimental/Mutable/Source/MutableTools/Private/MuT/ErrorLogPrivate.h
+++ b/Engine/Plugins/Experimental/Mutable/Source/MutableTools/Private/MuT/ErrorLogPrivate.h
@@ -65,7 +65,7 @@ namespace mu
 	typedef Ptr<const Model> ModelPtrConst;
 
     extern const TCHAR* GetOpName( OP_TYPE type );
-    extern FString GetOpDesc( const PROGRAM& program, OP::ADDRESS at );
+    extern FString GetOpDesc( const FProgram& program, OP::ADDRESS at );
 
 }
 
