#rnx

diff --git a/Engine/Plugins/Experimental/Mutable/Source/CustomizableObjectEditor/Private/MuCOE/CustomizableObjectEditor.cpp b/Engine/Plugins/Experimental/Mutable/Source/CustomizableObjectEditor/Private/MuCOE/CustomizableObjectEditor.cpp
--- a/Engine/Plugins/Experimental/Mutable/Source/CustomizableObjectEditor/Private/MuCOE/CustomizableObjectEditor.cpp
+++ b/Engine/Plugins/Experimental/Mutable/Source/CustomizableObjectEditor/Private/MuCOE/CustomizableObjectEditor.cpp
@@ -447,32 +447,6 @@ void FCustomizableObjectEditor::SelectNode(const UCustomizableObjectNode* Node)
 	GraphEditor->JumpToNode(Node);
 }
 
-void FCustomizableObjectEditor::SetPoseAsset(class UPoseAsset* PoseAssetParameter)
-{
-	PoseAsset = PoseAssetParameter;
-
-	if (PoseAsset != nullptr)
-	{
-		Viewport->SetAnimation(nullptr, EAnimationMode::AnimationBlueprint);
-
-		for (UDebugSkelMeshComponent* PreviewSkeletalMeshComponent : PreviewSkeletalMeshComponents)
-		{
-			PreviewSkeletalMeshComponent->SetAnimationMode(EAnimationMode::AnimationBlueprint);
-			PreviewSkeletalMeshComponent->InitAnim(false);
-			PreviewSkeletalMeshComponent->SetAnimation(PoseAsset);
-
-			UAnimSingleNodeInstance* SingleNodeInstance = Cast<UAnimSingleNodeInstance>(PreviewSkeletalMeshComponent->GetAnimInstance());
-			if (SingleNodeInstance)
-			{
-				TArray<FSmartName> ArrayPoseSmartNames = PoseAsset->GetPoseNames();
-				for (int32 i = 0; i < ArrayPoseSmartNames.Num(); ++i)
-				{
-					SingleNodeInstance->SetPreviewCurveOverride(ArrayPoseSmartNames[i].DisplayName, 1.0f, false);
-				}
-			}
-		}
-	}
-}
 
 bool FCustomizableObjectEditor::CreatePreviewComponent(int32 ComponentIndex)
 {
diff --git a/Engine/Plugins/Experimental/Mutable/Source/CustomizableObjectEditor/Private/MuCOE/CustomizableObjectEditor.h b/Engine/Plugins/Experimental/Mutable/Source/CustomizableObjectEditor/Private/MuCOE/CustomizableObjectEditor.h
--- a/Engine/Plugins/Experimental/Mutable/Source/CustomizableObjectEditor/Private/MuCOE/CustomizableObjectEditor.h
+++ b/Engine/Plugins/Experimental/Mutable/Source/CustomizableObjectEditor/Private/MuCOE/CustomizableObjectEditor.h
@@ -209,7 +209,6 @@ public:
 	virtual bool CanPasteNodes() const override;
 	virtual void PasteNodesHere(const FVector2D& Location) override;
 	virtual void SelectNode(const UCustomizableObjectNode* Node) override;
-	virtual void SetPoseAsset(class UPoseAsset* PoseAssetParameter) override;
 	
 	/** Called to undo the last action */
 	void UndoGraphAction();
@@ -500,9 +499,6 @@ private:
 	* which will perform the task */
 	float PendingTimeRefreshMaterialInAllChildren = 2.0;
 
-	/** Pose asset when doing drag and drop of an UPoseAsset to the viewport */
-	UPoseAsset* PoseAsset;
-
 	/** Texture Analyzer table widget which shows the information of the transient textures used in the customizable object instance */
 	TSharedPtr<class SCustomizableObjecEditorTextureAnalyzer> TextureAnalyzer;
 
diff --git a/Engine/Plugins/Experimental/Mutable/Source/CustomizableObjectEditor/Private/MuCOE/CustomizableObjectEditorViewportClient.cpp b/Engine/Plugins/Experimental/Mutable/Source/CustomizableObjectEditor/Private/MuCOE/CustomizableObjectEditorViewportClient.cpp
--- a/Engine/Plugins/Experimental/Mutable/Source/CustomizableObjectEditor/Private/MuCOE/CustomizableObjectEditorViewportClient.cpp
+++ b/Engine/Plugins/Experimental/Mutable/Source/CustomizableObjectEditor/Private/MuCOE/CustomizableObjectEditorViewportClient.cpp
@@ -6,6 +6,7 @@
 #include "Animation/AnimationAsset.h"
 #include "Animation/DebugSkelMeshComponent.h"
 #include "Animation/PoseAsset.h"
+#include "Animation/AnimSingleNodeInstance.h"
 #include "AssetViewerSettings.h"
 #include "BatchedElements.h"
 #include "CanvasTypes.h"
@@ -2092,20 +2093,41 @@ bool FCustomizableObjectEditorViewportClient::CanSetWidgetMode(UE::Widget::EWidg
 void FCustomizableObjectEditorViewportClient::SetAnimation(UAnimationAsset* Animation, EAnimationMode::Type AnimationType)
 {
 	bool bFoundComponent = false;
+
 	for (TWeakObjectPtr<USkeletalMeshComponent> SkeletalMeshComponent : SkeletalMeshComponents)
 	{
 		if (SkeletalMeshComponent.IsValid() && Animation != nullptr
 			&& SkeletalMeshComponent->GetSkinnedAsset() != nullptr
-			&& SkeletalMeshComponent->GetSkinnedAsset()->GetSkeleton() == Animation->GetSkeleton()
+			&& SkeletalMeshComponent->GetSkinnedAsset()->GetSkeleton()->IsCompatible(Animation->GetSkeleton())
 			)
 		{
-			SkeletalMeshComponent->SetAnimationMode(AnimationType);
-			SkeletalMeshComponent->PlayAnimation(Animation, true);
-			SkeletalMeshComponent->SetPlayRate(1.f);
 			SetRealtime(true);
 			IsPlayingAnimation = true;
 			AnimationBeingPlayed = Animation;
 
+			if (UPoseAsset* PoseAsset = Cast<UPoseAsset>(Animation))
+			{
+				SkeletalMeshComponent->SetAnimationMode(EAnimationMode::AnimationBlueprint);
+				SkeletalMeshComponent->InitAnim(false);
+				SkeletalMeshComponent->SetAnimation(PoseAsset);
+
+				UAnimSingleNodeInstance* SingleNodeInstance = Cast<UAnimSingleNodeInstance>(SkeletalMeshComponent->GetAnimInstance());
+				if (SingleNodeInstance)
+				{
+					TArray<FSmartName> ArrayPoseSmartNames = PoseAsset->GetPoseNames();
+					for (int32 i = 0; i < ArrayPoseSmartNames.Num(); ++i)
+					{
+						SingleNodeInstance->SetPreviewCurveOverride(ArrayPoseSmartNames[i].DisplayName, 1.0f, false);
+					}
+				}
+			}
+			else
+			{
+				SkeletalMeshComponent->SetAnimationMode(AnimationType);
+				SkeletalMeshComponent->PlayAnimation(Animation, true);
+				SkeletalMeshComponent->SetPlayRate(1.f);
+			}
+
 			bFoundComponent = true;
 		}
 	}
@@ -3030,6 +3052,7 @@ void FCustomizableObjectEditorViewportClient::AddReferencedObjects(FReferenceCol
 {
 	Collector.AddReferencedObject(ClipMorphMaterial);
 	Collector.AddReferencedObject(TransparentPlaneMaterialXY);
+	Collector.AddReferencedObject(AnimationBeingPlayed);
 }
 
 
diff --git a/Engine/Plugins/Experimental/Mutable/Source/CustomizableObjectEditor/Private/MuCOE/SCustomizableInstanceProperties.cpp b/Engine/Plugins/Experimental/Mutable/Source/CustomizableObjectEditor/Private/MuCOE/SCustomizableInstanceProperties.cpp
--- a/Engine/Plugins/Experimental/Mutable/Source/CustomizableObjectEditor/Private/MuCOE/SCustomizableInstanceProperties.cpp
+++ b/Engine/Plugins/Experimental/Mutable/Source/CustomizableObjectEditor/Private/MuCOE/SCustomizableInstanceProperties.cpp
@@ -1057,6 +1057,7 @@ void SCustomizableInstanceProperties::AddParameter(int32 ParamIndexInObject)
 						.MustMatchPossibleSuggestions(TAttribute<bool>(true))
 						.SuggestionListPlacement(EMenuPlacement::MenuPlacement_ComboBox)
 						.OnTextCommitted(this, &SCustomizableInstanceProperties::OnProjectorTextureParameterComboBoxChanged, PoseSwitchEnumParamNameWithRange)
+						.ToolTipText(LOCTEXT("Pose selector tooltip", "Select the skeletal mesh pose used for projection. This does not control the actual visual mesh pose in the viewport (or during gameplay for that matter). It has to be manually set. You can drag&drop a pose onto the preview viewport."))
 					]
 					
 					+ SHorizontalBox::Slot()
diff --git a/Engine/Plugins/Experimental/Mutable/Source/CustomizableObjectEditor/Private/MuCOE/SCustomizableObjectEditorViewport.cpp b/Engine/Plugins/Experimental/Mutable/Source/CustomizableObjectEditor/Private/MuCOE/SCustomizableObjectEditorViewport.cpp
--- a/Engine/Plugins/Experimental/Mutable/Source/CustomizableObjectEditor/Private/MuCOE/SCustomizableObjectEditorViewport.cpp
+++ b/Engine/Plugins/Experimental/Mutable/Source/CustomizableObjectEditor/Private/MuCOE/SCustomizableObjectEditorViewport.cpp
@@ -1072,13 +1072,7 @@ FReply SCustomizableObjectEditorViewportTabBody::OnDrop(const FGeometry& MyGeome
 	{
 		if (Helper_GetAssets(DragDropOp).Num())
 		{
-			UPoseAsset* PoseAsset = Cast<UPoseAsset>(Helper_GetAssets(DragDropOp)[0].GetAsset());
-			if (PoseAsset != nullptr)
-			{
-				CustomizableObjectEditorPtr.Pin()->SetPoseAsset(PoseAsset);
-				return FReply::Handled();
-			}
-
+			// This cast also includes UPoseAsset assets.
 			UAnimationAsset* AnimationAsset = Cast<UAnimationAsset>(Helper_GetAssets(DragDropOp)[0].GetAsset());
 			if (AnimationAsset)
 			{
