[Mutable] Add support for SkinWeightsProfiles

#rnx

diff --git a/Engine/Plugins/Experimental/Mutable/Source/CustomizableObject/Classes/MuCO/CustomizableObject.h b/Engine/Plugins/Experimental/Mutable/Source/CustomizableObject/Classes/MuCO/CustomizableObject.h
--- a/Engine/Plugins/Experimental/Mutable/Source/CustomizableObject/Classes/MuCO/CustomizableObject.h
+++ b/Engine/Plugins/Experimental/Mutable/Source/CustomizableObject/Classes/MuCO/CustomizableObject.h
@@ -200,6 +200,9 @@ struct FCompilationOptions
 	// Used to enable 16 bit bone weights
 	bool b16BitBoneWeightsEnabled = false;
 
+	// Used to enable skin weight profiles.
+	bool bSkinWeightProfilesEnabled = true;
+
 	// Used to reduce the number of notifications when compiling objects
 	bool bSilentCompilation = true;
 
@@ -600,6 +603,42 @@ struct FCustomizableObjectMeshToMeshVertData
 };
 template<> struct TCanBulkSerialize<FCustomizableObjectMeshToMeshVertData> { enum { Value = true }; };
 
+
+USTRUCT()
+struct FMutableSkinWeightProfileInfo
+{
+	GENERATED_USTRUCT_BODY()
+
+	FMutableSkinWeightProfileInfo() {};
+
+	FMutableSkinWeightProfileInfo(FName InName, bool InDefaultProfile, int8 InDefaultProfileFromLODIndex) : Name(InName),
+		DefaultProfile(InDefaultProfile), DefaultProfileFromLODIndex(InDefaultProfileFromLODIndex) {};
+
+	UPROPERTY()
+	FName Name;
+
+	UPROPERTY()
+	bool DefaultProfile = false;
+
+	UPROPERTY(meta = (ClampMin = 0))
+	int8 DefaultProfileFromLODIndex = 0;
+
+	friend FArchive& operator<<(FArchive& Ar, FMutableSkinWeightProfileInfo& Info)
+	{
+		Ar << Info.Name;
+		Ar << Info.DefaultProfile;
+		Ar << Info.DefaultProfileFromLODIndex;
+
+		return Ar;
+	}
+
+	bool operator==(const FMutableSkinWeightProfileInfo& Other) const
+	{
+		return Name == Other.Name;
+	}
+};
+
+
 UCLASS()
 class CUSTOMIZABLEOBJECT_API UMutableMaskOutCache : public UObject
 {
@@ -996,6 +1035,9 @@ public:
 	UPROPERTY(Transient)
 	TArray<FCustomizableObjectMeshToMeshVertData> ClothMeshToMeshVertData;
 
+	UPROPERTY()
+	TArray<FMutableSkinWeightProfileInfo> SkinWeightProfilesInfo;
+
 #if WITH_EDITORONLY_DATA
 
 	// Hide this property because it is not used yet.
@@ -1023,6 +1065,10 @@ public:
 	UPROPERTY(VisibleAnywhere, Category = CompileOptions)
 	bool bEnable16BitBoneWeights = false;
 
+	//
+	UPROPERTY(EditAnywhere, Category = CompileOptions)
+	bool bEnableAltSkinWeightProfiles = false;
+
 	// Options when compiling this customizable object (see EMutableCompileMeshType declaration for info)
 	UPROPERTY(EditAnywhere, Category = CompileOptions)
 	EMutableCompileMeshType MeshCompileType = EMutableCompileMeshType::LocalAndChildren;
@@ -1132,7 +1178,7 @@ private:
 	// This is a manual version number for the binary blobs in this asset.
 	// Increasing it invalidates all the previously compiled models.
 	// Warning: If while merging code both versions have changed, take the highest+1.
-	static const int32 CurrentSupportedVersion = 350;
+	static const int32 CurrentSupportedVersion = 351;
 
 public:
 
diff --git a/Engine/Plugins/Experimental/Mutable/Source/CustomizableObject/Private/MuCO/CustomizableObject.cpp b/Engine/Plugins/Experimental/Mutable/Source/CustomizableObject/Private/MuCO/CustomizableObject.cpp
--- a/Engine/Plugins/Experimental/Mutable/Source/CustomizableObject/Private/MuCO/CustomizableObject.cpp
+++ b/Engine/Plugins/Experimental/Mutable/Source/CustomizableObject/Private/MuCO/CustomizableObject.cpp
@@ -342,6 +342,7 @@ void UCustomizableObject::ClearCompiledData()
 	ClothMeshToMeshVertData.Empty();
 	ContributingClothingAssetsData.Empty();
 	ClothSharedConfigsData.Empty();
+	SkinWeightProfilesInfo.Empty();
 
 #if WITH_EDITORONLY_DATA
 	CustomizableObjectPathMap.Empty();
@@ -437,6 +438,8 @@ void UCustomizableObject::SaveCompiledData(FArchive& MemoryWriter, bool bSkipEdi
 	MemoryWriter << ContributingClothingAssetsData;
 	MemoryWriter << ClothSharedConfigsData; 
 	
+	MemoryWriter << SkinWeightProfilesInfo;
+
 	MemoryWriter << HashToStreamableBlock;
 
 	// All Editor Only data must be serialized here
@@ -525,6 +528,8 @@ void UCustomizableObject::LoadCompiledData(FArchive& MemoryReader, bool bSkipEdi
 		MemoryReader << ContributingClothingAssetsData;
 		MemoryReader << ClothSharedConfigsData;
 
+		MemoryReader << SkinWeightProfilesInfo;
+
 		MemoryReader << HashToStreamableBlock;
 
 		// All Editor Only data must be loaded here
diff --git a/Engine/Plugins/Experimental/Mutable/Source/CustomizableObject/Private/MuCO/CustomizableObjectInstance.cpp b/Engine/Plugins/Experimental/Mutable/Source/CustomizableObject/Private/MuCO/CustomizableObjectInstance.cpp
--- a/Engine/Plugins/Experimental/Mutable/Source/CustomizableObject/Private/MuCO/CustomizableObjectInstance.cpp
+++ b/Engine/Plugins/Experimental/Mutable/Source/CustomizableObject/Private/MuCO/CustomizableObjectInstance.cpp
@@ -4013,6 +4013,44 @@ bool UCustomizableInstancePrivateData::BuildSkeletalMeshRenderData(const TShared
 		LODModel.ActiveBoneIndices.Append(Component.ActiveBones);
 		LODModel.RequiredBones.Append(Component.ActiveBones);
 
+		const TArray<FMutableSkinWeightProfileInfo>& SkinWeightProfilesInfo = Public->GetCustomizableObject()->SkinWeightProfilesInfo;
+		if (SkinWeightProfilesInfo.Num())
+		{
+			bool bHasSkinWeightProfiles = false;
+
+			const mu::FMeshBufferSet& MutableMeshVertexBuffers = Component.Mesh->GetVertexBuffers();
+
+			for (int32 ProfileIndex = 0; ProfileIndex < SkinWeightProfilesInfo.Num(); ++ProfileIndex)
+			{
+				const int32 ProfileSemanticsIndex = ProfileIndex + 10;
+				int32 BoneIndicesBufferIndex, BoneIndicesBufferChannelIndex;
+				MutableMeshVertexBuffers.FindChannel(mu::MBS_BONEINDICES, ProfileSemanticsIndex, &BoneIndicesBufferIndex, &BoneIndicesBufferChannelIndex);
+
+				int32 BoneWeightsBufferIndex, BoneWeightsBufferChannelIndex;
+				MutableMeshVertexBuffers.FindChannel(mu::MBS_BONEWEIGHTS, ProfileSemanticsIndex, &BoneWeightsBufferIndex, &BoneWeightsBufferChannelIndex);
+
+				if (BoneIndicesBufferIndex < 0 || BoneIndicesBufferIndex != BoneWeightsBufferIndex)
+				{
+					continue;
+				}
+
+				if (!bHasSkinWeightProfiles)
+				{
+					LODModel.SkinWeightProfilesData.Init(&LODModel.SkinWeightVertexBuffer);
+					bHasSkinWeightProfiles = true;
+				}
+
+				const FMutableSkinWeightProfileInfo& Profile = SkinWeightProfilesInfo[ProfileIndex];
+				SkeletalMesh->AddSkinWeightProfile({ Profile.Name, Profile.DefaultProfile, Profile.DefaultProfileFromLODIndex });
+
+				UnrealConversionUtils::CopyMutableSkinWeightProfilesBuffers(
+					LODModel,
+					Profile.Name,
+					MutableMeshVertexBuffers,
+					BoneIndicesBufferIndex);
+			}
+		}
+
 		// Copy indices.
 		if (!UnrealConversionUtils::CopyMutableIndexBuffers(Component.Mesh, LODModel))
 		{
diff --git a/Engine/Plugins/Experimental/Mutable/Source/CustomizableObject/Private/MuCO/MutableMeshBufferUtils.cpp b/Engine/Plugins/Experimental/Mutable/Source/CustomizableObject/Private/MuCO/MutableMeshBufferUtils.cpp
--- a/Engine/Plugins/Experimental/Mutable/Source/CustomizableObject/Private/MuCO/MutableMeshBufferUtils.cpp
+++ b/Engine/Plugins/Experimental/Mutable/Source/CustomizableObject/Private/MuCO/MutableMeshBufferUtils.cpp
@@ -186,4 +186,64 @@ void MutableMeshBufferUtils::SetupIndexBuffer(mu::FMeshBufferSet& OutTargetIndex
 
 	OutTargetIndexBuffers.SetBuffer(0, ElementSize, ChannelCount, Semantics, SemanticIndices, Formats, Components,
 	                                 Offsets);
-}
\ No newline at end of file
+}
+
+
+void MutableMeshBufferUtils::SetupSkinWeightProfileBuffer(const int32& InCurrentVertexBuffer,
+	const int32& MaxBoneIndexTypeSizeBytes,
+	const int32& MaxBoneWeightTypeSizeBytes,
+	const int32& MaxNumBonesPerVertex,
+	const int32 SemanticsIndex,
+	mu::FMeshBufferSet& OutTargetVertexBuffers)
+{
+	using namespace mu;
+	const int32 ElementSize = sizeof(int32) + (MaxBoneIndexTypeSizeBytes + 1) * MaxNumBonesPerVertex;
+	constexpr int32 ChannelCount = 3;
+	const MESH_BUFFER_SEMANTIC Semantics[ChannelCount] = { MBS_OTHER, MBS_BONEINDICES, MBS_BONEWEIGHTS };
+	const int32 SemanticIndices[ChannelCount] = { SemanticsIndex, SemanticsIndex, SemanticsIndex };
+
+	MESH_BUFFER_FORMAT Formats[ChannelCount] = { MBF_INT32, MBF_UINT8, MBF_NUINT8 };
+	switch (MaxBoneIndexTypeSizeBytes)
+	{
+	case 1: Formats[1] = mu::MBF_UINT8;
+		break;
+	case 2: Formats[1] = mu::MBF_UINT16;
+		break;
+	case 4: Formats[1] = mu::MBF_UINT32;
+		break;
+	default:
+		// unsupported bone index type
+		check(false);
+		Formats[1] = mu::MBF_NONE;
+		break;
+	}
+
+	switch (MaxBoneWeightTypeSizeBytes)
+	{
+	case 1: Formats[2] = mu::MBF_NUINT8;
+		break;
+	case 2:
+	{
+		unimplemented()
+		Formats[2] = mu::MBF_NUINT16;
+		break;
+	}
+	case 4: Formats[2] = mu::MBF_NUINT32;
+		break;
+	default:
+		// unsupported bone weight type
+		check(false);
+		Formats[2] = mu::MBF_NONE;
+		break;
+	}
+
+	int32 Components[ChannelCount] = { 1, MaxNumBonesPerVertex, MaxNumBonesPerVertex };
+
+	int32 Offsets[ChannelCount];
+	Offsets[0] = 0;
+	Offsets[1] = sizeof(int32);
+	Offsets[2] = sizeof(int32) + MaxBoneIndexTypeSizeBytes * MaxNumBonesPerVertex;
+
+	OutTargetVertexBuffers.SetBuffer(InCurrentVertexBuffer, ElementSize, ChannelCount, Semantics, SemanticIndices,
+		Formats, Components, Offsets);
+}
diff --git a/Engine/Plugins/Experimental/Mutable/Source/CustomizableObject/Private/MuCO/UnrealConversionUtils.cpp b/Engine/Plugins/Experimental/Mutable/Source/CustomizableObject/Private/MuCO/UnrealConversionUtils.cpp
--- a/Engine/Plugins/Experimental/Mutable/Source/CustomizableObject/Private/MuCO/UnrealConversionUtils.cpp
+++ b/Engine/Plugins/Experimental/Mutable/Source/CustomizableObject/Private/MuCO/UnrealConversionUtils.cpp
@@ -31,6 +31,7 @@
 #include "StaticMeshResources.h"
 #include "Trace/Detail/Channel.h"
 #include "UObject/NameTypes.h"
+#include "Animation/SkinWeightProfile.h"
 
 class USkeleton;
 
@@ -372,6 +373,109 @@ namespace UnrealConversionUtils
 		return true;
 	}
 
+
+	struct FAuxMutableSkinWeightVertexKey
+	{
+		FAuxMutableSkinWeightVertexKey(const uint8* InKey, uint8 InKeySize, uint32 InHash)
+			: Key(InKey), KeySize(InKeySize), Hash(InHash)
+		{
+		};
+
+		const uint8* Key;
+		uint8 KeySize;
+		uint32 Hash;
+
+		friend uint32 GetTypeHash(const FAuxMutableSkinWeightVertexKey& InKey)
+		{
+			return InKey.Hash;
+		}
+
+		bool operator==(const FAuxMutableSkinWeightVertexKey& o) const
+		{
+			return FMemory::Memcmp(o.Key, Key, KeySize) == 0;
+		}
+	};
+
+	
+	void CopyMutableSkinWeightProfilesBuffers(
+		FSkeletalMeshLODRenderData& LODModel,
+		const FName ProfileName,
+		const mu::FMeshBufferSet& MutableMeshVertexBuffers,
+		const int32 BoneIndexBuffer)
+	{
+		MUTABLE_CPUPROFILER_SCOPE(CopyMutableSkinWeightProfilesBuffers);
+
+		// Basic Buffer override settings
+		FRuntimeSkinWeightProfileData& Override = LODModel.SkinWeightProfilesData.AddOverrideData(ProfileName);
+
+		const uint8 NumInfluences = LODModel.GetVertexBufferMaxBoneInfluences();
+		Override.NumWeightsPerVertex = NumInfluences;
+
+		const bool b16BitBoneIndices = LODModel.DoesVertexBufferUse16BitBoneIndex();
+		Override.b16BitBoneIndices = b16BitBoneIndices;
+		const int32 BoneIndexSize = b16BitBoneIndices ? 2 : 1;
+
+		// BoneWeights channel info
+		mu::MESH_BUFFER_FORMAT Format;
+		int32 MutableNumInfluences;
+		int32 Offset;
+		MutableMeshVertexBuffers.GetChannel(BoneIndexBuffer, 2, nullptr, nullptr, &Format, &MutableNumInfluences, &Offset);
+
+		const uint8 MutBoneIndexSize = Format == mu::MBF_UINT16 ? 2 : 1;
+		const uint8 MutBoneIndicesSize = Offset - sizeof(int32);
+		const uint8 MutBoneWeightsSize = MutableNumInfluences;
+		const uint8 MutBoneWeightVertexSize = MutBoneIndicesSize + MutBoneWeightsSize;
+
+		check(BoneIndexSize == MutBoneIndexSize);
+
+		const int32 ElementCount = MutableMeshVertexBuffers.GetElementCount();
+		Override.BoneIDs.Reserve(ElementCount * MutBoneIndicesSize);
+		Override.BoneWeights.Reserve(ElementCount * NumInfluences);
+		Override.VertexIndexToInfluenceOffset.Reserve(ElementCount);
+
+		TMap<FAuxMutableSkinWeightVertexKey, int32> HashToUniqueWeightIndexMap;
+		HashToUniqueWeightIndexMap.Reserve(ElementCount);
+		int32 UniqueWeightsCount = 0;
+
+		const uint8* SkinWeightsBuffer = reinterpret_cast<const uint8*>(MutableMeshVertexBuffers.GetBufferData(BoneIndexBuffer));
+		for (int32 ElementIndex = 0; ElementIndex < ElementCount; ++ElementIndex)
+		{
+			uint32 ElementHash;
+			FMemory::Memcpy(&ElementHash, SkinWeightsBuffer, sizeof(int32));
+			SkinWeightsBuffer += sizeof(int32);
+
+			if (ElementHash > 0)
+			{
+				if (int32* OverrideIndex = HashToUniqueWeightIndexMap.Find({ SkinWeightsBuffer, MutBoneWeightVertexSize, ElementHash }))
+				{
+					Override.VertexIndexToInfluenceOffset.Add(ElementIndex, *OverrideIndex);
+					SkinWeightsBuffer += MutBoneWeightVertexSize;
+				}
+				else
+				{
+					Override.VertexIndexToInfluenceOffset.Add(ElementIndex, UniqueWeightsCount);
+					HashToUniqueWeightIndexMap.Add({ SkinWeightsBuffer, MutBoneWeightVertexSize, ElementHash }, UniqueWeightsCount);
+
+					Override.BoneIDs.SetNumUninitialized((UniqueWeightsCount + 1) * MutBoneIndicesSize, false);
+					FMemory::Memcpy(&Override.BoneIDs[UniqueWeightsCount * MutBoneIndicesSize], SkinWeightsBuffer, MutBoneIndicesSize);
+					SkinWeightsBuffer += MutBoneIndicesSize;
+
+					Override.BoneWeights.SetNumUninitialized((UniqueWeightsCount + 1) * MutBoneWeightsSize, false);
+					FMemory::Memcpy(&Override.BoneWeights[UniqueWeightsCount * MutBoneWeightsSize], SkinWeightsBuffer, MutBoneWeightsSize);
+					SkinWeightsBuffer += MutBoneWeightsSize;
+					++UniqueWeightsCount;
+				}
+			}
+			else
+			{
+				SkinWeightsBuffer += MutBoneWeightVertexSize;
+			}
+		}
+
+		Override.BoneIDs.Shrink();
+		Override.BoneWeights.Shrink();
+		Override.VertexIndexToInfluenceOffset.Shrink();
+	} 
 }
 
 	
diff --git a/Engine/Plugins/Experimental/Mutable/Source/CustomizableObject/Public/MuCO/MutableMeshBufferUtils.h b/Engine/Plugins/Experimental/Mutable/Source/CustomizableObject/Public/MuCO/MutableMeshBufferUtils.h
--- a/Engine/Plugins/Experimental/Mutable/Source/CustomizableObject/Public/MuCO/MutableMeshBufferUtils.h
+++ b/Engine/Plugins/Experimental/Mutable/Source/CustomizableObject/Public/MuCO/MutableMeshBufferUtils.h
@@ -63,5 +63,17 @@ namespace  MutableMeshBufferUtils
 	 * @param OutTargetIndexBuffers Index buffer to be set up. 
 	 */
 	CUSTOMIZABLEOBJECT_API void SetupIndexBuffer(mu::FMeshBufferSet& OutTargetIndexBuffers);
+
+	/**
+	 * Sets up all the channels that ought to be part of the Vertex SkinWeightProfile Buffer of the Mutable mesh owning the provided MeshBufferSet
+	 * @param InCurrentVertexBuffer Index of the Vertex Buffer being set up inside the OutTargetVertexBuffers.
+	 * @param MaxBoneIndexTypeSizeBytes Maximum size of the vertex bone index type.
+	 * @param MaxBoneWeightTypeSizeBytes Maximum size of the vertex bone weight type.
+	 * @param MaxNumBonesPerVertex Amount of bone influences a vertex can have.
+	 * @param SemanticsIndex External Id that identifies a SkinWeightProfile.
+	 * @param OutTargetVertexBuffers Mesh buffer to be set up.
+	 */
+	CUSTOMIZABLEOBJECT_API void SetupSkinWeightProfileBuffer(const int32& InCurrentVertexBuffer, const int32& MaxBoneIndexTypeSizeBytes,const int32& MaxBoneWeightTypeSizeBytes, const int32& MaxNumBonesPerVertex,
+		const int32 SemanticsIndex, mu::FMeshBufferSet& OutTargetVertexBuffers);
 	
 }
diff --git a/Engine/Plugins/Experimental/Mutable/Source/CustomizableObject/Public/MuCO/UnrealConversionUtils.h b/Engine/Plugins/Experimental/Mutable/Source/CustomizableObject/Public/MuCO/UnrealConversionUtils.h
--- a/Engine/Plugins/Experimental/Mutable/Source/CustomizableObject/Public/MuCO/UnrealConversionUtils.h
+++ b/Engine/Plugins/Experimental/Mutable/Source/CustomizableObject/Public/MuCO/UnrealConversionUtils.h
@@ -85,4 +85,18 @@ namespace UnrealConversionUtils
 	CUSTOMIZABLEOBJECT_API bool CopyMutableIndexBuffers(mu::MeshPtrConst InMutableMesh,
 	                                                    FSkeletalMeshLODRenderData& OutLODModel);
 	
+
+	/**
+	 *Performs a copy of the data found on the index buffers on the mutable mesh to the buffers of the skeletal mesh
+	 * @param OutLODModel - The LOD model to be updated.
+	 * @param InProfileName - Name of the profile to generate.
+	 * @param InMutableMeshVertexBuffers - The mutable buffers to be reading data from
+	 * @param InBoneIndexBuffer - The buffer containing the indices for the bones.
+	 * @return True if the operation could be performed successfully,	 false if not.
+	 */
+	CUSTOMIZABLEOBJECT_API void CopyMutableSkinWeightProfilesBuffers(
+		FSkeletalMeshLODRenderData& OutLODModel,
+		const FName InProfileName,
+		const mu::FMeshBufferSet& InMutableMeshVertexBuffers,
+		const int32 InBoneIndexBuffer);
 }
diff --git a/Engine/Plugins/Experimental/Mutable/Source/CustomizableObjectEditor/Private/MuCOE/CustomizableObjectCompiler.cpp b/Engine/Plugins/Experimental/Mutable/Source/CustomizableObjectEditor/Private/MuCOE/CustomizableObjectCompiler.cpp
--- a/Engine/Plugins/Experimental/Mutable/Source/CustomizableObjectEditor/Private/MuCOE/CustomizableObjectCompiler.cpp
+++ b/Engine/Plugins/Experimental/Mutable/Source/CustomizableObjectEditor/Private/MuCOE/CustomizableObjectCompiler.cpp
@@ -950,6 +950,7 @@ void FCustomizableObjectCompiler::CompileInternal(UCustomizableObject* Object, c
 	Options.bRealTimeMorphTargetsEnabled = Object->bEnableRealTimeMorphTargets;
 	Options.bClothingEnabled = Object->bEnableClothing;
 	Options.b16BitBoneWeightsEnabled = Object->bEnable16BitBoneWeights;
+	Options.bSkinWeightProfilesEnabled = Object->bEnableAltSkinWeightProfiles;
 
 	if (Object->IsLocked() || !UCustomizableObjectSystem::GetInstance()->LockObject(Object))
 	{
@@ -1036,6 +1037,8 @@ void FCustomizableObjectCompiler::CompileInternal(UCustomizableObject* Object, c
 		Object->ClothMeshToMeshVertData = MoveTemp( GenerationContext.ClothMeshToMeshVertData );
 		
 		Object->ContributingClothingAssetsData = MoveTemp( GenerationContext.ContributingClothingAssetsData );
+
+		Object->SkinWeightProfilesInfo = MoveTemp(GenerationContext.SkinWeightProfilesInfo);
 		
 		// A clothing backend, e.g. Chaos cloth, can use 2 config files, one owned by the asset, and another that is shared 
 		// among all assets in a SkeletalMesh. When merging different assets in a skeletalmesh we need to make sure only one of 
diff --git a/Engine/Plugins/Experimental/Mutable/Source/CustomizableObjectEditor/Private/MuCOE/GenerateMutableSource/GenerateMutableSource.h b/Engine/Plugins/Experimental/Mutable/Source/CustomizableObjectEditor/Private/MuCOE/GenerateMutableSource/GenerateMutableSource.h
--- a/Engine/Plugins/Experimental/Mutable/Source/CustomizableObjectEditor/Private/MuCOE/GenerateMutableSource/GenerateMutableSource.h
+++ b/Engine/Plugins/Experimental/Mutable/Source/CustomizableObjectEditor/Private/MuCOE/GenerateMutableSource/GenerateMutableSource.h
@@ -170,6 +170,8 @@ struct FMutableGraphMeshGenerationData
 	int32 MaxNumTriangles = 0;
 	int32 MinNumTriangles = TNumericLimits<int32>::Max();
 
+	TArray<int32> SkinWeightProfilesSemanticIndices;
+
 	// Combine another generated data looking for the most general case.
 	void Combine(const FMutableGraphMeshGenerationData& other)
 	{
@@ -179,6 +181,11 @@ struct FMutableGraphMeshGenerationData
 		MaxBoneIndexTypeSizeBytes = FMath::Max(other.MaxBoneIndexTypeSizeBytes, MaxBoneIndexTypeSizeBytes);
 		MaxNumTriangles = FMath::Max(other.MaxNumTriangles, MaxNumTriangles);
 		MinNumTriangles = FMath::Min(other.MinNumTriangles, MinNumTriangles);
+
+		for (int32 SemanticIndex : other.SkinWeightProfilesSemanticIndices)
+		{
+			SkinWeightProfilesSemanticIndices.AddUnique(SemanticIndex);
+		}
 	}
 };
 
@@ -624,6 +631,9 @@ struct FMutableGraphGenerationContext
 	TArray<FCustomizableObjectMeshToMeshVertData> ClothMeshToMeshVertData;
 	TArray<FCustomizableObjectClothingAssetData> ContributingClothingAssetsData;
 
+	// Data used for SkinWeightProfiles reconstruction
+	TArray<FMutableSkinWeightProfileInfo> SkinWeightProfilesInfo;
+
 	// Hierarchy of current ComponentNew nodes, each stored for every LOD
 	struct ObjectParent
 	{
diff --git a/Engine/Plugins/Experimental/Mutable/Source/CustomizableObjectEditor/Private/MuCOE/GenerateMutableSource/GenerateMutableSourceMesh.cpp b/Engine/Plugins/Experimental/Mutable/Source/CustomizableObjectEditor/Private/MuCOE/GenerateMutableSource/GenerateMutableSourceMesh.cpp
--- a/Engine/Plugins/Experimental/Mutable/Source/CustomizableObjectEditor/Private/MuCOE/GenerateMutableSource/GenerateMutableSourceMesh.cpp
+++ b/Engine/Plugins/Experimental/Mutable/Source/CustomizableObjectEditor/Private/MuCOE/GenerateMutableSource/GenerateMutableSourceMesh.cpp
@@ -289,6 +289,121 @@ TObjectPtr<const USkeletalMesh> GetMeshWithBoneRemovalApplied(TObjectPtr<USkelet
 	return CacheEntry.Mesh;
 }
 
+
+void TransferRemovedBonesInfluences(FBoneIndexType* InfluenceBones, uint8* InfluenceWeights, const int32 InfluenceCount, const TMap<int32, int32>& ReplacementMap)
+{
+	for (int32 i = 0; i < InfluenceCount; ++i)
+	{
+		if (const int32* IndexPtr = ReplacementMap.Find(InfluenceBones[i]))
+		{
+			bool bParentFound = false;
+			int32 ParentIndex = *IndexPtr;
+			for (int32 j = 0; j < InfluenceCount; ++j)
+			{
+				if (InfluenceBones[j] == ParentIndex)
+				{
+					InfluenceWeights[j] += InfluenceWeights[i];
+					InfluenceWeights[i] = 0.f;
+					bParentFound = true;
+					break;
+				}
+			}
+
+			if (!bParentFound)
+			{
+				InfluenceBones[i] = ParentIndex;
+			}
+		}
+	}
+}
+
+
+void NormalizeWeights(FBoneIndexType* InfluenceBones, uint8* InfluenceWeights, const int32 InfluenceCount, const int32 MutableInfluenceCount,
+	int32* MutableMaxOrderedWeighsIndices, const int32 MaxSectionBoneMapIndex, const int32 MaxBoneWeight)
+{
+	// First get the indices of the 4 heaviest influences
+	for (int32 i = 0; i < MutableInfluenceCount; ++i)
+	{
+		int32 CurrentMaxWeight = -1;
+
+		for (int32 j = 0; j < InfluenceCount; ++j)
+		{
+			bool bIndexAlreadyUsed = false;
+
+			for (int32 k = 0; k < i; ++k)
+			{
+				if (MutableMaxOrderedWeighsIndices[k] == j)
+				{
+					bIndexAlreadyUsed = true;
+					break;
+				}
+				else if (MutableMaxOrderedWeighsIndices[k] < 0)
+				{
+					break;
+				}
+			}
+
+			if (!bIndexAlreadyUsed && InfluenceWeights[j] > CurrentMaxWeight
+				&& InfluenceBones[j] < MaxSectionBoneMapIndex)
+			{
+				MutableMaxOrderedWeighsIndices[i] = j;
+				CurrentMaxWeight = InfluenceWeights[j];
+			}
+		}
+	}
+
+	// Copy 4 heaviest influences to 4 first indices
+	for (int32 i = 0; i < MutableInfluenceCount; ++i)
+	{
+		if (i < InfluenceCount)
+		{
+			InfluenceWeights[i] = InfluenceWeights[MutableMaxOrderedWeighsIndices[i]];
+			InfluenceBones[i] = InfluenceBones[MutableMaxOrderedWeighsIndices[i]];
+		}
+		else
+		{
+			InfluenceWeights[i] = 0;
+			InfluenceBones[i] = 0;
+		}
+	}
+
+	// Actually renormalize the first 4 influences
+	int32 TotalWeight = 0;
+
+	for (int32 j = 0; j < MutableInfluenceCount; ++j)
+	{
+		TotalWeight += InfluenceWeights[j];
+	}
+
+	if (TotalWeight > 0)
+	{
+		int32 AssignedWeight = 0;
+
+		for (int32 j = 1; j < MAX_TOTAL_INFLUENCES; ++j)
+		{
+			if (j < MutableInfluenceCount)
+			{
+				float Aux = InfluenceWeights[j];
+				int32 Res = FMath::RoundToInt(Aux / TotalWeight * MaxBoneWeight);
+				AssignedWeight += Res;
+				InfluenceWeights[j] = Res;
+			}
+			else
+			{
+				InfluenceWeights[j] = 0;
+			}
+		}
+
+		InfluenceWeights[0] = MaxBoneWeight - AssignedWeight;
+	}
+	else
+	{
+		FMemory::Memzero(InfluenceWeights, MutableInfluenceCount);
+		InfluenceWeights[0] = MaxBoneWeight;
+	}
+}
+
+
 mu::MeshPtr ConvertSkeletalMeshToMutable(USkeletalMesh* InSkeletalMesh, int LOD, int MaterialIndex, FMutableGraphGenerationContext& GenerationContext, const UCustomizableObjectNode* CurrentNode)
 {
 	// Get the mesh generation flags to use
@@ -585,175 +700,31 @@ mu::MeshPtr ConvertSkeletalMeshToMutable(USkeletalMesh* InSkeletalMesh, int LOD,
 			}
 		}
 
-		if (GenerationContext.Options.bExtraBoneInfluencesEnabled && MaxSectionInfluences < EXTRA_BONE_INFLUENCES)
-		{
-			for (int32 VertexIndex = VertexStart; VertexIndex < VertexStart + VertexCount && VertexIndex < Vertices.Num(); ++VertexIndex)
-			{
-				FSoftSkinVertex& Vertex = Vertices[VertexIndex];
+		const int32 MaxSectionBoneMapIndex = ImportedModel->LODModels[LOD].Sections[MaterialIndex].BoneMap.Num();
 
-				for (int32 i = MaxSectionInfluences; i < EXTRA_BONE_INFLUENCES; ++i)
-				{
-					check(Vertex.InfluenceWeights[i] == 0);
-				}
-			}
-		}
-		else if (!GenerationContext.Options.bExtraBoneInfluencesEnabled)
+		for (int32 VertexIndex = VertexStart; VertexIndex < VertexStart + VertexCount && VertexIndex < Vertices.Num(); ++VertexIndex)
 		{
-			const int32 MaxBoneWeightValue = BoneWeightFormat == MBF_NUINT16 ? 65535 : 255;
-			const int32 MaxSectionBoneMapIndex = ImportedModel->LODModels[LOD].Sections[MaterialIndex].BoneMap.Num();
+			FSoftSkinVertex& Vertex = Vertices[VertexIndex];
 
-			// Renormalize to 4 weights per vertex
-			for (int32 VertexIndex = VertexStart; VertexIndex < VertexStart + VertexCount && VertexIndex < Vertices.Num(); ++VertexIndex)
+			// Transfer removed bones influences to parent bones
+			if (!InfluencesToReplaceMap.IsEmpty())
 			{
-				FSoftSkinVertex& Vertex = Vertices[VertexIndex];
+				TransferRemovedBonesInfluences(&Vertex.InfluenceBones[0], &Vertex.InfluenceWeights[0], MaxSectionInfluences, InfluencesToReplaceMap);
+			}
 
+			if (!GenerationContext.Options.bExtraBoneInfluencesEnabled)
+			{
+				// Normalize weights
 				const int32 MaxMutableWeights = 4;
 				int32 MaxOrderedWeighsIndices[MaxMutableWeights] = { -1, -1, -1, -1 };
 
-				// Transfer removed bones influences to parent bones
-				if (InfluencesToReplaceMap.Num() > 0)
-				{
-					for (int32 i = 0; i < MaxSectionInfluences; ++i)
-					{
-						if (const int32* IndexPtr = InfluencesToReplaceMap.Find(Vertex.InfluenceBones[i]))
-						{
-							bool bParentFound = false;
-							int32 ParentIndex = *IndexPtr;
-							for (int32 j = 0; j < MaxSectionInfluences; ++j)
-							{
-								if (Vertex.InfluenceBones[j] == ParentIndex)
-								{
-									Vertex.InfluenceWeights[j] += Vertex.InfluenceWeights[i];
-									Vertex.InfluenceWeights[i] = 0.f;
-									bParentFound = true;
-									break;
-								}
-							}
-
-							if (!bParentFound)
-							{
-								Vertex.InfluenceBones[i] = ParentIndex;
-							}
-						}
-					}
-				}
-
-				// First get the indices of the 4 heaviest influences
-				for (int32 i = 0; i < MaxMutableWeights; ++i)
-				{
-					int32 CurrentMaxWeight = -1;
-
-					for (int32 j = 0; j < MaxSectionInfluences; ++j)
-					{
-						bool bIndexAlreadyUsed = false;
-
-						for (int32 k = 0; k < i; ++k)
-						{
-							if (MaxOrderedWeighsIndices[k] == j)
-							{
-								bIndexAlreadyUsed = true;
-								break;
-							}
-							else if (MaxOrderedWeighsIndices[k] < 0)
-							{
-								break;
-							}
-						}
-
-						if (!bIndexAlreadyUsed && Vertex.InfluenceWeights[j] > CurrentMaxWeight
-							&& Vertex.InfluenceBones[j] < MaxSectionBoneMapIndex)
-						{
-							MaxOrderedWeighsIndices[i] = j;
-							CurrentMaxWeight = Vertex.InfluenceWeights[j];
-						}
-					}
-				}
-
-				// Copy 4 heaviest influences to 4 first indices
-				for (int32 i = 0; i < MaxMutableWeights; ++i)
-				{
-					if (i < MaxSectionInfluences)
-					{
-						Vertex.InfluenceWeights[i] = Vertex.InfluenceWeights[MaxOrderedWeighsIndices[i]];
-						Vertex.InfluenceBones[i] = Vertex.InfluenceBones[MaxOrderedWeighsIndices[i]];
-					}
-					else
-					{
-						Vertex.InfluenceWeights[i] = 0;
-						Vertex.InfluenceBones[i] = 0;
-					}
-				}
-
-				// Actually renormalize the first 4 influences
-				int32 TotalWeight = 0;
-
-				for (int32 j = 0; j < MaxMutableWeights; ++j)
-				{
-					TotalWeight += Vertex.InfluenceWeights[j];
-				}
-
-				if (TotalWeight > 0)
-				{
-					int32 AssignedWeight = 0;
-
-					for (int32 j = 1; j < MAX_TOTAL_INFLUENCES; ++j)
-					{
-						if (j < MaxMutableWeights)
-						{
-							float Aux = Vertex.InfluenceWeights[j];
-							int32 Res = FMath::RoundToInt(Aux / TotalWeight * MaxBoneWeightValue);
-							AssignedWeight += Res;
-							Vertex.InfluenceWeights[j] = Res;
-						}
-						else
-						{
-							Vertex.InfluenceWeights[j] = 0;
-						}
-					}
-
-					Vertex.InfluenceWeights[0] = MaxBoneWeightValue - AssignedWeight;
-				}
-				else
-				{
-					Vertex.InfluenceWeights[0] = MaxBoneWeightValue;
-
-					for (int32 j = 1; j < MaxMutableWeights; ++j)
-					{
-						Vertex.InfluenceWeights[j] = 0;
-					}
-				}
+				const int32 MaxBoneWeightValue = BoneWeightFormat == MBF_NUINT16 ? 65535 : 255;
+				NormalizeWeights(&Vertex.InfluenceBones[0], &Vertex.InfluenceWeights[0], MaxSectionInfluences, MaxMutableWeights,
+					&MaxOrderedWeighsIndices[0], MaxSectionBoneMapIndex, MaxBoneWeightValue);
 			}
-		}
-		else if (InfluencesToReplaceMap.Num() > 0)
-		{
-			// Transfer removed bones influences to parent bones
-			for (int32 VertexIndex = VertexStart; VertexIndex < VertexStart + VertexCount && VertexIndex < Vertices.Num(); ++VertexIndex)
+			else if (MaxSectionInfluences < EXTRA_BONE_INFLUENCES)
 			{
-				FSoftSkinVertex& Vertex = Vertices[VertexIndex];
-
-				for (int32 i = 0; i < MAX_TOTAL_INFLUENCES; ++i)
-				{
-					if (const int32* ParentBoneIndexPtr = InfluencesToReplaceMap.Find(Vertex.InfluenceBones[i]))
-					{
-						bool bParentFound = false;
-						int32 ParentIndex = *ParentBoneIndexPtr;
-						for (int32 j = 0; j < MAX_TOTAL_INFLUENCES; ++j)
-						{
-							if (Vertex.InfluenceBones[j] == ParentIndex)
-							{
-								Vertex.InfluenceWeights[j] += Vertex.InfluenceWeights[i];
-								Vertex.InfluenceWeights[i] = 0.f;
-								bParentFound = true;
-								break;
-							}
-						}
-
-						if (!bParentFound)
-						{
-							Vertex.InfluenceBones[i] = ParentIndex;
-						}
-					}
-				}
+				FMemory::Memzero(&Vertex.InfluenceWeights[MaxSectionInfluences], EXTRA_BONE_INFLUENCES - MaxSectionInfluences);
 			}
 		}
 
@@ -992,6 +963,7 @@ mu::MeshPtr ConvertSkeletalMeshToMutable(USkeletalMesh* InSkeletalMesh, int LOD,
 			const int Offsets[ChannelCount] = { 0 };
 
 			MutableMesh->GetVertexBuffers().SetBuffer(nextBufferIndex, ElementSize, ChannelCount, Semantics, SemanticIndices, Formats, Components, Offsets);
+			++nextBufferIndex;
 		}
 
 		TArrayView<int32> ClothSectionBufferView(reinterpret_cast<int32*>(MutableMesh->GetVertexBuffers().GetBufferData(nextBufferIndex)), VertexCount);
@@ -1118,6 +1090,112 @@ mu::MeshPtr ConvertSkeletalMeshToMutable(USkeletalMesh* InSkeletalMesh, int LOD,
 		}
 	}
 
+
+	// SkinWeightProfiles vertex info.
+	if (GenerationContext.Options.bSkinWeightProfilesEnabled)
+	{
+		using namespace mu;
+
+		// TODO: Remove BoneWeightFormat after merge
+		const int32 BoneWeightTypeSizeBytes = sizeof(TDecay<decltype(DeclVal<FRawSkinWeight>().InfluenceWeights[0])>::Type);
+		MESH_BUFFER_FORMAT BoneWeightFormat = BoneWeightTypeSizeBytes == 1 ? MBF_NUINT8 : MBF_NUINT16;
+
+		// Limit skinning weights if necessary
+		// \todo: make it more flexible to support 3 or 5 or 1 weight, since there is support for this in 4.25
+		const int32 MutableBonesPerVertex = GenerationContext.Options.bExtraBoneInfluencesEnabled ? EXTRA_BONE_INFLUENCES : 4;
+		const int32 BoneIndicesSize = MutableBonesPerVertex * sizeof(FBoneIndexType);
+		const int32 BoneWeightsSize = MutableBonesPerVertex * BoneWeightTypeSizeBytes;
+		const int32 SkinWeightProfileVertexSize = sizeof(int32) + BoneIndicesSize + BoneWeightsSize;
+
+		const int32 MaxSectionInfluences = ImportedModel->LODModels[LOD].Sections[MaterialIndex].GetMaxBoneInfluences();
+		const int32 MaxSectionBoneMapIndex = ImportedModel->LODModels[LOD].Sections[MaterialIndex].BoneMap.Num();
+
+		const TArray<FSkinWeightProfileInfo>& SkinWeightProfilesInfo = SkeletalMesh->GetSkinWeightProfiles();
+		for (const FSkinWeightProfileInfo& Profile : SkinWeightProfilesInfo)
+		{
+			const FImportedSkinWeightProfileData* ImportedProfileData = ImportedModel->LODModels[LOD].SkinWeightProfiles.Find(Profile.Name);
+			if (!ImportedProfileData)
+			{
+				continue;
+			}
+
+			check(Vertices.Num() == ImportedProfileData->SkinWeights.Num());
+
+			TArray<uint8> MutSkinWeights;
+			MutSkinWeights.SetNumZeroed(VertexCount * SkinWeightProfileVertexSize);
+			uint8* MutSkinWeightData = MutSkinWeights.GetData();
+
+			for (int32 VertexIndex = VertexStart; VertexIndex < VertexStart + VertexCount; ++VertexIndex)
+			{
+				FRawSkinWeight SkinWeight = ImportedProfileData->SkinWeights[VertexIndex];
+
+				if (!InfluencesToReplaceMap.IsEmpty())
+				{
+					TransferRemovedBonesInfluences(&SkinWeight.InfluenceBones[0], &SkinWeight.InfluenceWeights[0], MaxSectionInfluences, InfluencesToReplaceMap);
+				}
+
+				if (!GenerationContext.Options.bExtraBoneInfluencesEnabled)
+				{
+					// Normalize weights
+					const int32 MaxMutableWeights = 4;
+					int32 MaxOrderedWeighsIndices[MaxMutableWeights] = { -1, -1, -1, -1 };
+
+					const int32 MaxBoneWeightValue = BoneWeightFormat == MBF_NUINT16 ? 65535 : 255;
+					NormalizeWeights(&SkinWeight.InfluenceBones[0], &SkinWeight.InfluenceWeights[0], MaxSectionInfluences, MaxMutableWeights,
+						&MaxOrderedWeighsIndices[0], MaxSectionBoneMapIndex, MaxBoneWeightValue);
+				}
+				else if (MaxSectionInfluences < MutableBonesPerVertex)
+				{
+					FMemory::Memzero(&SkinWeight.InfluenceWeights[MaxSectionInfluences], MutableBonesPerVertex - MaxSectionInfluences);
+				}
+
+				if (FMemory::Memcmp(&Vertices[VertexIndex].InfluenceBones[0], &SkinWeight.InfluenceBones[0], BoneIndicesSize) == 0
+					&&
+					FMemory::Memcmp(&Vertices[VertexIndex].InfluenceWeights[0], &SkinWeight.InfluenceWeights[0], BoneWeightsSize) == 0)
+				{
+					MutSkinWeightData += SkinWeightProfileVertexSize;
+					continue;
+				}
+
+				int32 SkinWeightVertexHash = 0;
+				for (int32 InfluenceIndex = 0; InfluenceIndex < MutableBonesPerVertex; ++InfluenceIndex)
+				{
+					SkinWeightVertexHash = HashCombine(SkinWeightVertexHash, SkinWeight.InfluenceBones[InfluenceIndex]);
+					SkinWeightVertexHash = HashCombine(SkinWeightVertexHash, SkinWeight.InfluenceWeights[InfluenceIndex]);
+				}
+
+				FMemory::Memcpy(MutSkinWeightData, &SkinWeightVertexHash, sizeof(int32));
+				MutSkinWeightData += sizeof(int32);
+				FMemory::Memcpy(MutSkinWeightData, &SkinWeight.InfluenceBones[0], BoneIndicesSize);
+				MutSkinWeightData += BoneIndicesSize;
+				FMemory::Memcpy(MutSkinWeightData, &SkinWeight.InfluenceWeights[0], BoneWeightsSize);
+				MutSkinWeightData += BoneWeightsSize;
+			}
+
+			const int32 ProfileIndex = GenerationContext.SkinWeightProfilesInfo.AddUnique({Profile.Name, false, 0});
+			const int32 ProfileSemanticIndex = ProfileIndex + 10;
+
+			const FName PlatformName = *GenerationContext.Options.TargetPlatform->PlatformName();
+			FMutableSkinWeightProfileInfo& MutSkinWeightProfileInfo = GenerationContext.SkinWeightProfilesInfo[ProfileIndex];
+			MutSkinWeightProfileInfo.DefaultProfile = MutSkinWeightProfileInfo.DefaultProfile || Profile.DefaultProfile.GetValueForPlatform(PlatformName);
+			MutSkinWeightProfileInfo.DefaultProfileFromLODIndex = FMath::Min(MutSkinWeightProfileInfo.DefaultProfileFromLODIndex, Profile.DefaultProfileFromLODIndex.GetValueForPlatform(PlatformName));
+
+			// Set up SkinWeightPRofile BufferData
+			const int32 ElementSize = sizeof(int32) + sizeof(FBoneIndexType) + BoneWeightTypeSizeBytes;
+			const int32 ChannelCount = 3;
+			const MESH_BUFFER_SEMANTIC Semantics[ChannelCount] = { MBS_OTHER, MBS_BONEINDICES, MBS_BONEWEIGHTS };
+			const int32 SemanticIndices[ChannelCount] = { ProfileSemanticIndex, ProfileSemanticIndex, ProfileSemanticIndex };
+			const MESH_BUFFER_FORMAT Formats[ChannelCount] = { MBF_INT32, MBF_UINT16, BoneWeightFormat };
+			const int32 Components[ChannelCount] = { 1, MutableBonesPerVertex, MutableBonesPerVertex };
+			const int32 Offsets[ChannelCount] = { 0, sizeof(int32), sizeof(int32) + BoneIndicesSize };
+
+			MutableMesh->GetVertexBuffers().SetBufferCount(nextBufferIndex + 1);
+			MutableMesh->GetVertexBuffers().SetBuffer(nextBufferIndex, ElementSize, ChannelCount, Semantics, SemanticIndices, Formats, Components, Offsets);
+			FMemory::Memcpy(MutableMesh->GetVertexBuffers().GetBufferData(nextBufferIndex), MutSkinWeights.GetData(), VertexCount * SkinWeightProfileVertexSize);
+			++nextBufferIndex;
+		}
+	}
+
 	// Indices
 	{
 		int IndexStart = ImportedModel->LODModels[LOD].Sections[MaterialIndex].BaseIndex;
@@ -2383,6 +2461,21 @@ mu::NodeMeshPtr GenerateMutableSourceMesh(const UEdGraphPin * Pin,
 					}
 				}
 
+				if (GenerationContext.Options.bSkinWeightProfilesEnabled)
+				{
+					FSkeletalMeshModel* ImportModel = Helper_GetImportedModel(TypedNodeSkel->SkeletalMesh);
+
+					const int32 SkinWeightProfilesCount = GenerationContext.SkinWeightProfilesInfo.Num();
+					for (int32 ProfileIndex = 0; ProfileIndex < SkinWeightProfilesCount; ++ProfileIndex)
+					{
+						if (ImportModel->LODModels[LOD].SkinWeightProfiles.Find(GenerationContext.SkinWeightProfilesInfo[ProfileIndex].Name))
+						{
+							const int32 ProfileScemanticIndex = ProfileIndex + 10;
+							MeshData.SkinWeightProfilesSemanticIndices.AddUnique(ProfileScemanticIndex);
+						}
+					}
+				}
+
 				if (!TypedNodeSkel->AnimInstance.IsNull())
 				{
 					int32 SlotIndex= TypedNodeSkel->AnimBlueprintSlot;
diff --git a/Engine/Plugins/Experimental/Mutable/Source/CustomizableObjectEditor/Private/MuCOE/GenerateMutableSource/GenerateMutableSourceSurface.cpp b/Engine/Plugins/Experimental/Mutable/Source/CustomizableObjectEditor/Private/MuCOE/GenerateMutableSource/GenerateMutableSourceSurface.cpp
--- a/Engine/Plugins/Experimental/Mutable/Source/CustomizableObjectEditor/Private/MuCOE/GenerateMutableSource/GenerateMutableSourceSurface.cpp
+++ b/Engine/Plugins/Experimental/Mutable/Source/CustomizableObjectEditor/Private/MuCOE/GenerateMutableSource/GenerateMutableSourceSurface.cpp
@@ -161,6 +161,8 @@ void SetSurfaceFormat( mu::FMeshBufferSet& OutVertexBufferFormat, mu::FMeshBuffe
 		++MutableBufferCount;
 	}
 
+	MutableBufferCount += MeshData.SkinWeightProfilesSemanticIndices.Num();
+
 	OutVertexBufferFormat.SetBufferCount(MutableBufferCount);
 
 	int32 CurrentVertexBuffer = 0;
@@ -240,6 +242,11 @@ void SetSurfaceFormat( mu::FMeshBufferSet& OutVertexBufferFormat, mu::FMeshBuffe
 		++CurrentVertexBuffer;
 	}
 
+	for (int32 ProfileSemanticIndex : MeshData.SkinWeightProfilesSemanticIndices)
+	{
+		MutableMeshBufferUtils::SetupSkinWeightProfileBuffer(CurrentVertexBuffer, MeshData.MaxBoneIndexTypeSizeBytes, 1, MutableBonesPerVertex, ProfileSemanticIndex, OutVertexBufferFormat);
+		++CurrentVertexBuffer;
+	}
 
 	// Index buffer
 	MutableMeshBufferUtils::SetupIndexBuffer(OutIndexBufferFormat);
diff --git a/Engine/Plugins/Experimental/Mutable/Source/MutableRuntime/Private/MuR/MeshBufferSet.cpp b/Engine/Plugins/Experimental/Mutable/Source/MutableRuntime/Private/MuR/MeshBufferSet.cpp
--- a/Engine/Plugins/Experimental/Mutable/Source/MutableRuntime/Private/MuR/MeshBufferSet.cpp
+++ b/Engine/Plugins/Experimental/Mutable/Source/MutableRuntime/Private/MuR/MeshBufferSet.cpp
@@ -257,24 +257,16 @@ namespace mu
 		*pBuffer = -1;
 		*pChannel = -1;
 
-		int index = 0;
-
 		for (size_t b = 0; b < m_buffers.Num(); ++b)
 		{
 			for (size_t c = 0; c < m_buffers[b].m_channels.Num(); ++c)
 			{
-				if (m_buffers[b].m_channels[c].m_semantic == semantic)
+				if (m_buffers[b].m_channels[c].m_semantic == semantic &&
+					m_buffers[b].m_channels[c].m_semanticIndex == semanticIndex)
 				{
-					if (index == semanticIndex)
-					{
-						*pBuffer = int(b);
-						*pChannel = int(c);
-						return;
-					}
-					else
-					{
-						++index;
-					}
+					*pBuffer = int(b);
+					*pChannel = int(c);
+					return;
 				}
 			}
 		}
diff --git a/Engine/Plugins/Experimental/Mutable/Source/MutableRuntime/Private/MuR/OpMeshFormat.cpp b/Engine/Plugins/Experimental/Mutable/Source/MutableRuntime/Private/MuR/OpMeshFormat.cpp
--- a/Engine/Plugins/Experimental/Mutable/Source/MutableRuntime/Private/MuR/OpMeshFormat.cpp
+++ b/Engine/Plugins/Experimental/Mutable/Source/MutableRuntime/Private/MuR/OpMeshFormat.cpp
@@ -603,59 +603,60 @@ namespace mu
 		// Make sure that the bone indices will fit in this format, or extend it.
 		if (formatVertices)
 		{
-			auto& buffs = pSource->GetVertexBuffers();
+			const FMeshBufferSet& VertexBuffers = pSource->GetVertexBuffers();
 
-			int semanticIndex = 0;
-			while (true)
+			const int32 BufferCount = VertexBuffers.GetBufferCount();
+			for (int32 BufferIndex = 0; BufferIndex < BufferCount; ++BufferIndex)
 			{
-				int buf = 0;
-				int chan = 0;
-				buffs.FindChannel(MBS_BONEINDICES, semanticIndex, &buf, &chan);
-				if (buf < 0)
-					break;
-
-				int resultBuf = 0;
-				int resultChan = 0;
-				auto& formBuffs = pResult->GetVertexBuffers();
-				formBuffs.FindChannel(MBS_BONEINDICES, semanticIndex, &resultBuf, &resultChan);
-				if (resultBuf >= 0)
+				const int32 ChannelCount = VertexBuffers.GetBufferChannelCount(BufferIndex);
+				for (int32 ChannelIndex = 0; ChannelIndex < ChannelCount; ++ChannelIndex)
 				{
-					UntypedMeshBufferIteratorConst it(buffs, MBS_BONEINDICES, semanticIndex);
-					int32_t maxBoneIndex = 0;
-					for (int v = 0; v < buffs.GetElementCount(); ++v)
+					const MESH_BUFFER_CHANNEL& Channel = VertexBuffers.m_buffers[BufferIndex].m_channels[ChannelIndex];
+
+					if (Channel.m_semantic == MBS_BONEINDICES)
 					{
-						auto va = it.GetAsVec8i();
-						for (int c = 0; c < it.GetComponents(); ++c)
+						int resultBuf = 0;
+						int resultChan = 0;
+						auto& formBuffs = pResult->GetVertexBuffers();
+						formBuffs.FindChannel(MBS_BONEINDICES, Channel.m_semanticIndex, &resultBuf, &resultChan);
+						if (resultBuf >= 0)
 						{
-							maxBoneIndex = FMath::Max(maxBoneIndex, va[c]);
-						}
-						++it;
-					}
+							UntypedMeshBufferIteratorConst it(VertexBuffers, MBS_BONEINDICES, Channel.m_semanticIndex);
+							int32_t maxBoneIndex = 0;
+							for (int v = 0; v < VertexBuffers.GetElementCount(); ++v)
+							{
+								auto va = it.GetAsVec8i();
+								for (int c = 0; c < it.GetComponents(); ++c)
+								{
+									maxBoneIndex = FMath::Max(maxBoneIndex, va[c]);
+								}
+								++it;
+							}
 
-					auto& format = formBuffs.m_buffers[resultBuf].m_channels[resultChan].m_format;
-					if (maxBoneIndex > 0xffff && (format == MBF_UINT8 || format == MBF_UINT16))
-					{
-						format = MBF_UINT32;
-						formBuffs.UpdateOffsets(resultBuf);
-					}
-					else if (maxBoneIndex > 0x7fff && (format == MBF_INT8 || format == MBF_INT16))
-					{
-						format = MBF_UINT32;
-						formBuffs.UpdateOffsets(resultBuf);
-					}
-					else if (maxBoneIndex > 0xff && format == MBF_UINT8)
-					{
-						format = MBF_UINT16;
-						formBuffs.UpdateOffsets(resultBuf);
-					}
-					else if (maxBoneIndex > 0x7f && format == MBF_INT8)
-					{
-						format = MBF_INT16;
-						formBuffs.UpdateOffsets(resultBuf);
+							auto& format = formBuffs.m_buffers[resultBuf].m_channels[resultChan].m_format;
+							if (maxBoneIndex > 0xffff && (format == MBF_UINT8 || format == MBF_UINT16))
+							{
+								format = MBF_UINT32;
+								formBuffs.UpdateOffsets(resultBuf);
+							}
+							else if (maxBoneIndex > 0x7fff && (format == MBF_INT8 || format == MBF_INT16))
+							{
+								format = MBF_UINT32;
+								formBuffs.UpdateOffsets(resultBuf);
+							}
+							else if (maxBoneIndex > 0xff && format == MBF_UINT8)
+							{
+								format = MBF_UINT16;
+								formBuffs.UpdateOffsets(resultBuf);
+							}
+							else if (maxBoneIndex > 0x7f && format == MBF_INT8)
+							{
+								format = MBF_INT16;
+								formBuffs.UpdateOffsets(resultBuf);
+							}
+						}
 					}
 				}
-
-				semanticIndex++;
 			}
 		}
 
diff --git a/Engine/Plugins/Experimental/Mutable/Source/MutableRuntime/Private/MuR/OpMeshMerge.h b/Engine/Plugins/Experimental/Mutable/Source/MutableRuntime/Private/MuR/OpMeshMerge.h
--- a/Engine/Plugins/Experimental/Mutable/Source/MutableRuntime/Private/MuR/OpMeshMerge.h
+++ b/Engine/Plugins/Experimental/Mutable/Source/MutableRuntime/Private/MuR/OpMeshMerge.h
@@ -754,27 +754,37 @@ namespace mu
                     pResult->GetSkeleton()  ? pResult->GetSkeleton()->GetBoneCount() : 0;
                 if (pResultBoneCount > 256)
 				{
-					int foundBuffer = -1;
-					int foundChannel = -1;
-					pResult->GetVertexBuffers().FindChannel(MBS_BONEINDICES, 0, &foundBuffer, &foundChannel);
+					// Desired format
+					MESH_BUFFER_FORMAT format = pResultBoneCount > 65535 ? MBF_UINT32 : MBF_UINT16;
 
-					if (foundChannel >= 0)
+					// Iterate all vertex buffers and update the format
+					FMeshBufferSet& VertexBuffers = pResult->GetVertexBuffers();
+					for (int32 VertexBufferIndex = 0; VertexBufferIndex < VertexBuffers.m_buffers.Num(); ++VertexBufferIndex)
 					{
-						MESH_BUFFER& boneBuffer = pResult->GetVertexBuffers().m_buffers[foundBuffer];
+						MESH_BUFFER& result = VertexBuffers.m_buffers[VertexBufferIndex];
 
-						MESH_BUFFER_FORMAT format = pResultBoneCount > 65535 ? MBF_UINT32 : MBF_UINT16;
-						boneBuffer.m_channels[foundChannel].m_format = format;
+						const int32 elemSize = VertexBuffers.GetElementSize(VertexBufferIndex);
+						const int32 firstSize = firstCount * elemSize;
 
-						// Reset offsets
-						int offset = 0;
-						for (int32 c = 0; c < boneBuffer.m_channels.Num(); ++c)
+						const int32 ChannelsCount = VertexBuffers.GetBufferChannelCount(VertexBufferIndex);
+						for (int32 ChannelIndex = 0; ChannelIndex < ChannelsCount; ++ChannelIndex)
 						{
-							boneBuffer.m_channels[c].m_offset = (uint8_t)offset;
-							offset += boneBuffer.m_channels[c].m_componentCount
-								*
-								GetMeshFormatData(boneBuffer.m_channels[c].m_format).m_size;
+							if (result.m_channels[ChannelIndex].m_semantic == MBS_BONEINDICES)
+							{
+								result.m_channels[ChannelIndex].m_format = format;
+
+								// Reset offsets
+								int32 offset = 0;
+								for (int32 AuxChannelIndex = 0; AuxChannelIndex < ChannelsCount; ++AuxChannelIndex)
+								{
+									result.m_channels[AuxChannelIndex].m_offset = (uint8_t)offset;
+									offset += result.m_channels[AuxChannelIndex].m_componentCount
+										*
+										GetMeshFormatData(result.m_channels[AuxChannelIndex].m_format).m_size;
+								}
+								result.m_elementSize = offset;
+							}
 						}
-						boneBuffer.m_elementSize = offset;
 					}
 				}
 
