[mutable] Fix SkinWeightProfileInfos being added multiple times to the SkeletalMesh

#rnx

diff --git a/Engine/Plugins/Experimental/Mutable/Source/CustomizableObject/Private/MuCO/CustomizableObjectInstance.cpp b/Engine/Plugins/Experimental/Mutable/Source/CustomizableObject/Private/MuCO/CustomizableObjectInstance.cpp
--- a/Engine/Plugins/Experimental/Mutable/Source/CustomizableObject/Private/MuCO/CustomizableObjectInstance.cpp
+++ b/Engine/Plugins/Experimental/Mutable/Source/CustomizableObject/Private/MuCO/CustomizableObjectInstance.cpp
@@ -4041,7 +4041,14 @@ bool UCustomizableInstancePrivateData::BuildSkeletalMeshRenderData(const TShared
 				}
 
 				const FMutableSkinWeightProfileInfo& Profile = SkinWeightProfilesInfo[ProfileIndex];
-				SkeletalMesh->AddSkinWeightProfile({ Profile.Name, Profile.DefaultProfile, Profile.DefaultProfileFromLODIndex });
+
+				const FSkinWeightProfileInfo* ExistingProfile = SkeletalMesh->GetSkinWeightProfiles().FindByPredicate(
+					[&Profile](const FSkinWeightProfileInfo& P) { return P.Name == Profile.Name; });
+				
+				if (!ExistingProfile)
+				{
+					SkeletalMesh->AddSkinWeightProfile({ Profile.Name, Profile.DefaultProfile, Profile.DefaultProfileFromLODIndex });
+				}
 
 				UnrealConversionUtils::CopyMutableSkinWeightProfilesBuffers(
 					LODModel,
