[mutable] Fix cloth tether reconstruction minor bug, missing shared configs and crash setting cloth buffers when compiling.

#rnx

diff --git a/Engine/Plugins/Experimental/Mutable/Source/CustomizableObject/Private/MuCO/CustomizableObjectInstance.cpp b/Engine/Plugins/Experimental/Mutable/Source/CustomizableObject/Private/MuCO/CustomizableObjectInstance.cpp
--- a/Engine/Plugins/Experimental/Mutable/Source/CustomizableObject/Private/MuCO/CustomizableObjectInstance.cpp
+++ b/Engine/Plugins/Experimental/Mutable/Source/CustomizableObject/Private/MuCO/CustomizableObjectInstance.cpp
@@ -3080,7 +3080,6 @@ namespace
 
 }
 
-
 void UCustomizableInstancePrivateData::BuildClothingData(const TSharedPtr<FMutableOperationData>&OperationData, USkeletalMesh * SkeletalMesh, UCustomizableObjectInstance * CustomizableObjectInstance, int32 ComponentIndex)
 {
 	MUTABLE_CPUPROFILER_SCOPE(UCustomizableInstancePrivateData::BuildClothingData);
@@ -3567,7 +3566,7 @@ void UCustomizableInstancePrivateData::BuildClothingData(const TSharedPtr<FMutab
 					{
 						const int32 Index0 = IndexMap[Tether.Get<0>()];
 						const int32 Index1 = IndexMap[Tether.Get<1>()];
-						if ((Index0 > 0) & (Index1 > 0))
+						if ((Index0 >= 0) & (Index1 >= 0))
 						{
 							DstTetherCluster.Emplace(Index0, Index1, Tether.Get<2>());
 						}
@@ -3883,28 +3882,40 @@ void UCustomizableInstancePrivateData::BuildClothingData(const TSharedPtr<FMutab
 			NewClothingAssets[I]->CalculateReferenceBoneIndex();	
 			NewClothingAssets[I]->PhysicsAsset = ClothingPhysicsAssets[I];
 
-			for ( const FCustomizableObjectClothConfigData& ConfigData : ContributingClothingAssetsData[I].ConfigsData )
+			auto CreateNewClothConfigFromData = [](UObject* Outer, const FCustomizableObjectClothConfigData& ConfigData) 
+			-> UClothConfigCommon* 
 			{
 				UClass* ClothConfigClass = FindObject<UClass>(nullptr, *ConfigData.ClassPath);
 				if (ClothConfigClass)
 				{
-					UClothConfigCommon* ClothConfig = NewObject<UClothConfigCommon>(NewClothingAssets[I], ClothConfigClass);
+					UClothConfigCommon* ClothConfig = NewObject<UClothConfigCommon>(Outer, ClothConfigClass);
 					if (ClothConfig)
 					{
-						FMemoryReaderView MemoryReader( ConfigData.ConfigBytes );
+						FMemoryReaderView MemoryReader(ConfigData.ConfigBytes);
 						ClothConfig->Serialize(MemoryReader);
-						NewClothingAssets[I]->ClothConfigs.Add(ConfigData.ConfigName, ClothConfig);
+
+						return ClothConfig;
 					}
 				}
+
+				return nullptr;
+			};
+
+			for (const FCustomizableObjectClothConfigData& ConfigData : ContributingClothingAssetsData[I].ConfigsData)
+			{
+				UClothConfigCommon* ClothConfig = CreateNewClothConfigFromData(NewClothingAssets[I], ConfigData);
+				if (ClothConfig)
+				{
+					NewClothingAssets[I]->ClothConfigs.Add(ConfigData.ConfigName, ClothConfig);
+				}
 			}
 
 			for (const FCustomizableObjectClothConfigData& ConfigData : ClothSharedConfigsData)
 			{	
-				UClass* ClothConfigClass = FindObject<UClass>(nullptr, *ConfigData.ClassPath);
-				
-				if (ClothConfigClass)
+				UClothConfigCommon* ClothConfig = CreateNewClothConfigFromData(NewClothingAssets[I], ConfigData);
+				if (ClothConfig)
 				{
-					UClothSharedConfigCommon* ClothConfig = NewObject<UClothSharedConfigCommon>(NewClothingAssets[I], ClothConfigClass);
+					NewClothingAssets[I]->ClothConfigs.Add(ConfigData.ConfigName, ClothConfig);
 				}
 			}
 
diff --git a/Engine/Plugins/Experimental/Mutable/Source/CustomizableObjectEditor/Private/MuCOE/CustomizableObjectCompiler.cpp b/Engine/Plugins/Experimental/Mutable/Source/CustomizableObjectEditor/Private/MuCOE/CustomizableObjectCompiler.cpp
--- a/Engine/Plugins/Experimental/Mutable/Source/CustomizableObjectEditor/Private/MuCOE/CustomizableObjectCompiler.cpp
+++ b/Engine/Plugins/Experimental/Mutable/Source/CustomizableObjectEditor/Private/MuCOE/CustomizableObjectCompiler.cpp
@@ -1032,31 +1032,31 @@ void FCustomizableObjectCompiler::CompileInternal(UCustomizableObject* Object, c
 		// to avoid a possibly expensive copy.
 		Object->ContributingMorphTargetsInfo = MoveTemp(GenerationContext.ContributingMorphTargetsInfo);
 		Object->MorphTargetReconstructionData = MoveTemp(GenerationContext.MorphTargetReconstructionData);
-
-		// Clothing	
-		Object->ClothMeshToMeshVertData = MoveTemp( GenerationContext.ClothMeshToMeshVertData );
 		
-		Object->ContributingClothingAssetsData = MoveTemp( GenerationContext.ContributingClothingAssetsData );
-
 		Object->SkinWeightProfilesInfo = MoveTemp(GenerationContext.SkinWeightProfilesInfo);
 		
+		// Clothing	
+		Object->ClothMeshToMeshVertData = MoveTemp(GenerationContext.ClothMeshToMeshVertData);
+		Object->ContributingClothingAssetsData = MoveTemp(GenerationContext.ContributingClothingAssetsData);
+		Object->ClothSharedConfigsData.Empty();
+
 		// A clothing backend, e.g. Chaos cloth, can use 2 config files, one owned by the asset, and another that is shared 
 		// among all assets in a SkeletalMesh. When merging different assets in a skeletalmesh we need to make sure only one of 
 		// the shared is used. In that case we will keep the first visited of a type and will be stored separated from the asset.
 		// TODO: Shared configs, which typically controls the quality of the simulation (iterations, etc), probably should be specified 
 		// somewhere else to give more control with which config ends up used. 
-		auto IsSharedConfigData = []( const FCustomizableObjectClothConfigData& ConfigData ) -> bool
+		auto IsSharedConfigData = [](const FCustomizableObjectClothConfigData& ConfigData) -> bool
 		{
 			 const UClass* ConfigClass = FindObject<UClass>(nullptr, *ConfigData.ClassPath);
 			 return ConfigClass ? static_cast<bool>( Cast<UClothSharedConfigCommon>(ConfigClass->GetDefaultObject() ) ) : false;
 		};
 		
 		// Find shared configs to be used (One of each type) 
-		for ( FCustomizableObjectClothingAssetData& ClothingAssetData : Object->ContributingClothingAssetsData )
+		for (FCustomizableObjectClothingAssetData& ClothingAssetData : Object->ContributingClothingAssetsData)
 		{
-			 for ( FCustomizableObjectClothConfigData& ClothConfigData : ClothingAssetData.ConfigsData )
+			 for (FCustomizableObjectClothConfigData& ClothConfigData : ClothingAssetData.ConfigsData)
 			 {
-				  if ( IsSharedConfigData( ClothConfigData ) )
+				  if (IsSharedConfigData(ClothConfigData))
 				  {
 					  FCustomizableObjectClothConfigData* FoundConfig = Object->ClothSharedConfigsData.FindByPredicate(
 						   [Name = ClothConfigData.ConfigName](const FCustomizableObjectClothConfigData& Other)
@@ -1066,15 +1066,14 @@ void FCustomizableObjectCompiler::CompileInternal(UCustomizableObject* Object, c
 
 					  if (!FoundConfig)
 					  {
-						   FCustomizableObjectClothConfigData& ObjectConfig = Object->ClothSharedConfigsData.AddDefaulted_GetRef();
-						   ObjectConfig = ClothConfigData;
+						   Object->ClothSharedConfigsData.AddDefaulted_GetRef() = ClothConfigData;
 					  }
 				  }
 			 }
 		}
 		
 		// Remove shared configs
-		for ( FCustomizableObjectClothingAssetData& ClothingAssetData : Object->ContributingClothingAssetsData )
+		for (FCustomizableObjectClothingAssetData& ClothingAssetData : Object->ContributingClothingAssetsData)
 		{
 			 ClothingAssetData.ConfigsData.RemoveAllSwap(IsSharedConfigData);
 		}
