Non-unity/pch compile fixes when building UnrealEditor linux

diff --git a/Engine/Plugins/AI/MLAdapter/Source/MLAdapter/Private/Sensors/MLAdapterSensor_Camera.cpp b/Engine/Plugins/AI/MLAdapter/Source/MLAdapter/Private/Sensors/MLAdapterSensor_Camera.cpp
--- a/Engine/Plugins/AI/MLAdapter/Source/MLAdapter/Private/Sensors/MLAdapterSensor_Camera.cpp
+++ b/Engine/Plugins/AI/MLAdapter/Source/MLAdapter/Private/Sensors/MLAdapterSensor_Camera.cpp
@@ -13,6 +13,8 @@
 #include "Engine/LocalPlayer.h"
 #include "Engine/GameViewportClient.h"
 #include "GameFramework/GameUserSettings.h"
+#include "TextureResource.h"
+#include "UnrealClient.h"
 #include "UnrealEngine.h"
 
 #include "Misc/ScopeLock.h"
diff --git a/Engine/Plugins/AI/MLAdapter/Source/MLAdapter/Private/Sensors/MLAdapterSensor_Input.cpp b/Engine/Plugins/AI/MLAdapter/Source/MLAdapter/Private/Sensors/MLAdapterSensor_Input.cpp
--- a/Engine/Plugins/AI/MLAdapter/Source/MLAdapter/Private/Sensors/MLAdapterSensor_Input.cpp
+++ b/Engine/Plugins/AI/MLAdapter/Source/MLAdapter/Private/Sensors/MLAdapterSensor_Input.cpp
@@ -5,6 +5,7 @@
 #include "GameFramework/PlayerInput.h"
 #include "Engine/GameViewportClient.h"
 #include "Engine/World.h"
+#include "InputKeyEventArgs.h"
 #include "MLAdapterInputHelper.h"
 #include "MLAdapterSpace.h"
 #include "Debug/DebugHelpers.h"
diff --git a/Engine/Plugins/AI/MLAdapter/Source/MLAdapter/Public/Sensors/MLAdapterSensor_Camera.h b/Engine/Plugins/AI/MLAdapter/Source/MLAdapter/Public/Sensors/MLAdapterSensor_Camera.h
--- a/Engine/Plugins/AI/MLAdapter/Source/MLAdapter/Public/Sensors/MLAdapterSensor_Camera.h
+++ b/Engine/Plugins/AI/MLAdapter/Source/MLAdapter/Public/Sensors/MLAdapterSensor_Camera.h
@@ -2,6 +2,7 @@
 
 #pragma once
 
+#include "Engine/EngineTypes.h"
 #include "Sensors/MLAdapterSensor.h"
 #include "Engine/TextureRenderTarget2D.h"
 #include "MLAdapterSensor_Camera.generated.h"
diff --git a/Engine/Plugins/AI/MassAI/Source/MassNavigationEditor/Private/MassNavigationEditorModule.cpp b/Engine/Plugins/AI/MassAI/Source/MassNavigationEditor/Private/MassNavigationEditorModule.cpp
--- a/Engine/Plugins/AI/MassAI/Source/MassNavigationEditor/Private/MassNavigationEditorModule.cpp
+++ b/Engine/Plugins/AI/MassAI/Source/MassNavigationEditor/Private/MassNavigationEditorModule.cpp
@@ -3,6 +3,7 @@
 #include "MassNavigationEditorModule.h"
 #include "Modules/ModuleManager.h"
 #include "AssetToolsModule.h"
+#include "PropertyEditorModule.h"
 #include "ToolMenus.h"
 
 #define LOCTEXT_NAMESPACE "MassNavigationEditor"
diff --git a/Engine/Plugins/Animation/AnimationLocomotionLibrary/Source/Runtime/Private/AnimDistanceMatchingLibrary.cpp b/Engine/Plugins/Animation/AnimationLocomotionLibrary/Source/Runtime/Private/AnimDistanceMatchingLibrary.cpp
--- a/Engine/Plugins/Animation/AnimationLocomotionLibrary/Source/Runtime/Private/AnimDistanceMatchingLibrary.cpp
+++ b/Engine/Plugins/Animation/AnimationLocomotionLibrary/Source/Runtime/Private/AnimDistanceMatchingLibrary.cpp
@@ -3,6 +3,7 @@
 #include "AnimDistanceMatchingLibrary.h"
 #include "Animation/AnimNode_SequencePlayer.h"
 #include "Animation/AnimSequence.h"
+#include "AnimationRuntime.h"
 #include "AnimNodes/AnimNode_SequenceEvaluator.h"
 #include "Curves/CurveFloat.h"
 #include "Animation/AnimCurveCompressionCodec_UniformIndexable.h"
diff --git a/Engine/Plugins/Animation/AnimationWarping/Source/Runtime/Private/BoneControllers/AnimNode_FootPlacement.cpp b/Engine/Plugins/Animation/AnimationWarping/Source/Runtime/Private/BoneControllers/AnimNode_FootPlacement.cpp
--- a/Engine/Plugins/Animation/AnimationWarping/Source/Runtime/Private/BoneControllers/AnimNode_FootPlacement.cpp
+++ b/Engine/Plugins/Animation/AnimationWarping/Source/Runtime/Private/BoneControllers/AnimNode_FootPlacement.cpp
@@ -3,6 +3,7 @@
 #include "BoneControllers/AnimNode_FootPlacement.h"
 #include "AnimationRuntime.h"
 #include "Animation/AnimInstanceProxy.h"
+#include "Components/SkeletalMeshComponent.h"
 #include "GameFramework/Character.h"
 #include "GameFramework/CharacterMovementComponent.h"
 #include "Components/CapsuleComponent.h"
diff --git a/Engine/Plugins/Animation/AnimationWarping/Source/Runtime/Private/BoneControllers/AnimNode_SlopeWarping.cpp b/Engine/Plugins/Animation/AnimationWarping/Source/Runtime/Private/BoneControllers/AnimNode_SlopeWarping.cpp
--- a/Engine/Plugins/Animation/AnimationWarping/Source/Runtime/Private/BoneControllers/AnimNode_SlopeWarping.cpp
+++ b/Engine/Plugins/Animation/AnimationWarping/Source/Runtime/Private/BoneControllers/AnimNode_SlopeWarping.cpp
@@ -3,6 +3,7 @@
 #include "BoneControllers/AnimNode_SlopeWarping.h"
 #include "AnimationRuntime.h"
 #include "Animation/AnimInstanceProxy.h"
+#include "Components/SkeletalMeshComponent.h"
 #include "GameFramework/Character.h"
 #include "GameFramework/CharacterMovementComponent.h"
 #include "Components/CapsuleComponent.h"
diff --git a/Engine/Plugins/Animation/ControlRig/Source/ControlRig/Private/Units/Animation/RigUnit_AnimAttribute.h b/Engine/Plugins/Animation/ControlRig/Source/ControlRig/Private/Units/Animation/RigUnit_AnimAttribute.h
--- a/Engine/Plugins/Animation/ControlRig/Source/ControlRig/Private/Units/Animation/RigUnit_AnimAttribute.h
+++ b/Engine/Plugins/Animation/ControlRig/Source/ControlRig/Private/Units/Animation/RigUnit_AnimAttribute.h
@@ -1,6 +1,7 @@
 // Copyright Epic Games, Inc. All Rights Reserved.
 
 #pragma once
+#include "Engine/SkeletalMesh.h"
 #include "Units/RigUnit.h"
 #include "Components/SkeletalMeshComponent.h"
 #include "Units/RigDispatchFactory.h"
diff --git a/Engine/Plugins/Compositing/Composure/Source/Composure/Private/CompositingElements/CompositingElementTransforms.cpp b/Engine/Plugins/Compositing/Composure/Source/Composure/Private/CompositingElements/CompositingElementTransforms.cpp
--- a/Engine/Plugins/Compositing/Composure/Source/Composure/Private/CompositingElements/CompositingElementTransforms.cpp
+++ b/Engine/Plugins/Compositing/Composure/Source/Composure/Private/CompositingElements/CompositingElementTransforms.cpp
@@ -2,6 +2,7 @@
 
 #include "CompositingElements/CompositingElementTransforms.h"
 #include "CompositingElements/CompElementRenderTargetPool.h"
+#include "Materials/Material.h"
 #include "Materials/MaterialInterface.h"
 
 /* TCompositingTargetSwapChain
diff --git a/Engine/Plugins/Compositing/LensDistortion/Source/LensDistortion/Private/LensDistortionRendering.cpp b/Engine/Plugins/Compositing/LensDistortion/Source/LensDistortion/Private/LensDistortionRendering.cpp
--- a/Engine/Plugins/Compositing/LensDistortion/Source/LensDistortion/Private/LensDistortionRendering.cpp
+++ b/Engine/Plugins/Compositing/LensDistortion/Source/LensDistortion/Private/LensDistortionRendering.cpp
@@ -13,6 +13,7 @@
 #include "ShaderParameterUtils.h"
 #include "Logging/MessageLog.h"
 #include "Internationalization/Internationalization.h"
+#include "TextureResource.h"
 
 
 static const uint32 kGridSubdivisionX = 32;
diff --git a/Engine/Plugins/Developer/Concert/ConcertApp/DisasterRecoveryClient/Source/DisasterRecoveryClient/Private/DisasterRecoveryClientModule.cpp b/Engine/Plugins/Developer/Concert/ConcertApp/DisasterRecoveryClient/Source/DisasterRecoveryClient/Private/DisasterRecoveryClientModule.cpp
--- a/Engine/Plugins/Developer/Concert/ConcertApp/DisasterRecoveryClient/Source/DisasterRecoveryClient/Private/DisasterRecoveryClientModule.cpp
+++ b/Engine/Plugins/Developer/Concert/ConcertApp/DisasterRecoveryClient/Source/DisasterRecoveryClient/Private/DisasterRecoveryClientModule.cpp
@@ -2,6 +2,7 @@
 
 #include "IDisasterRecoveryClientModule.h"
 
+#include "Framework/Application/SlateApplication.h"
 #include "IConcertSyncClientModule.h"
 #include "IConcertModule.h"
 #include "IConcertClient.h"
diff --git a/Engine/Plugins/Developer/RenderDocPlugin/Source/RenderDocPlugin/Private/RenderDocPluginModule.cpp b/Engine/Plugins/Developer/RenderDocPlugin/Source/RenderDocPlugin/Private/RenderDocPluginModule.cpp
--- a/Engine/Plugins/Developer/RenderDocPlugin/Source/RenderDocPlugin/Private/RenderDocPluginModule.cpp
+++ b/Engine/Plugins/Developer/RenderDocPlugin/Source/RenderDocPlugin/Private/RenderDocPluginModule.cpp
@@ -1,6 +1,7 @@
 // Copyright Epic Games, Inc. All Rights Reserved.
 
 #include "RenderDocPluginModule.h"
+#include "Framework/Application/SlateApplication.h"
 
 #if PLATFORM_WINDOWS
 #include "Windows/AllowWindowsPlatformTypes.h"
diff --git a/Engine/Plugins/Developer/RenderDocPlugin/Source/RenderDocPlugin/Private/SRenderDocPluginEditorExtension.cpp b/Engine/Plugins/Developer/RenderDocPlugin/Source/RenderDocPlugin/Private/SRenderDocPluginEditorExtension.cpp
--- a/Engine/Plugins/Developer/RenderDocPlugin/Source/RenderDocPlugin/Private/SRenderDocPluginEditorExtension.cpp
+++ b/Engine/Plugins/Developer/RenderDocPlugin/Source/RenderDocPlugin/Private/SRenderDocPluginEditorExtension.cpp
@@ -20,6 +20,7 @@
 #include "Framework/MultiBox/MultiBoxDefs.h"
 #include "Interfaces/IMainFrameModule.h"
 #include "LevelEditor.h"
+#include "RHI.h"
 #include "Widgets/Input/SButton.h"
 #include "Widgets/Images/SImage.h"
 
diff --git a/Engine/Plugins/Editor/WaveformEditor/Source/WaveformEditor/Private/WaveformEditor.cpp b/Engine/Plugins/Editor/WaveformEditor/Source/WaveformEditor/Private/WaveformEditor.cpp
--- a/Engine/Plugins/Editor/WaveformEditor/Source/WaveformEditor/Private/WaveformEditor.cpp
+++ b/Engine/Plugins/Editor/WaveformEditor/Source/WaveformEditor/Private/WaveformEditor.cpp
@@ -5,6 +5,7 @@
 #include "AudioDevice.h"
 #include "Components/AudioComponent.h"
 #include "Misc/TransactionObjectEvent.h"
+#include "PropertyEditorModule.h"
 #include "Sound/SoundWave.h"
 #include "Styling/AppStyle.h"
 #include "SWaveformPanel.h"
diff --git a/Engine/Plugins/Editor/WaveformEditor/Source/WaveformEditor/Public/WaveformEditor.h b/Engine/Plugins/Editor/WaveformEditor/Source/WaveformEditor/Public/WaveformEditor.h
--- a/Engine/Plugins/Editor/WaveformEditor/Source/WaveformEditor/Public/WaveformEditor.h
+++ b/Engine/Plugins/Editor/WaveformEditor/Source/WaveformEditor/Public/WaveformEditor.h
@@ -3,6 +3,7 @@
 #pragma once
 
 #include "EditorUndoClient.h"
+#include "IWaveformTransformation.h"
 #include "Misc/NotifyHook.h"
 #include "Toolkits/AssetEditorToolkit.h"
 #include "WaveformEditorTransportController.h"
diff --git a/Engine/Plugins/Editor/WaveformEditor/Source/WaveformTransformationsWidgets/Private/WaveformTransformationRendererBase.cpp b/Engine/Plugins/Editor/WaveformEditor/Source/WaveformTransformationsWidgets/Private/WaveformTransformationRendererBase.cpp
--- a/Engine/Plugins/Editor/WaveformEditor/Source/WaveformTransformationsWidgets/Private/WaveformTransformationRendererBase.cpp
+++ b/Engine/Plugins/Editor/WaveformEditor/Source/WaveformTransformationsWidgets/Private/WaveformTransformationRendererBase.cpp
@@ -1,6 +1,7 @@
 // Copyright Epic Games, Inc. All Rights Reserved.
 
 #include "WaveformTransformationRendererBase.h"
+#include "Input/Reply.h"
 
 int32 FWaveformTransformationRendererBase::OnPaint(const FPaintArgs& Args, const FGeometry& AllottedGeometry, const FSlateRect& MyCullingRect, FSlateWindowElementList& OutDrawElements, int32 LayerId, const FWidgetStyle& InWidgetStyle, bool bParentEnabled) const
 {
diff --git a/Engine/Plugins/Editor/WaveformEditor/Source/WaveformTransformationsWidgets/Private/WaveformTransformationTrimFadeRenderer.cpp b/Engine/Plugins/Editor/WaveformEditor/Source/WaveformTransformationsWidgets/Private/WaveformTransformationTrimFadeRenderer.cpp
--- a/Engine/Plugins/Editor/WaveformEditor/Source/WaveformTransformationsWidgets/Private/WaveformTransformationTrimFadeRenderer.cpp
+++ b/Engine/Plugins/Editor/WaveformEditor/Source/WaveformTransformationsWidgets/Private/WaveformTransformationTrimFadeRenderer.cpp
@@ -1,6 +1,9 @@
 // Copyright Epic Games, Inc. All Rights Reserved.
 
 #include "WaveformTransformationTrimFadeRenderer.h"
+#include "Layout/Geometry.h"
+#include "Rendering/DrawElements.h"
+#include "Widgets/SWidget.h"
 
 FWaveformTransformationTrimFadeRenderer::FWaveformTransformationTrimFadeRenderer(const TObjectPtr<UWaveformTransformationTrimFade> TransformationToRender)
 {
diff --git a/Engine/Plugins/Editor/WaveformEditor/Source/WaveformTransformationsWidgets/Private/WaveformTransformationsRenderManager.cpp b/Engine/Plugins/Editor/WaveformEditor/Source/WaveformTransformationsWidgets/Private/WaveformTransformationsRenderManager.cpp
--- a/Engine/Plugins/Editor/WaveformEditor/Source/WaveformTransformationsWidgets/Private/WaveformTransformationsRenderManager.cpp
+++ b/Engine/Plugins/Editor/WaveformEditor/Source/WaveformTransformationsWidgets/Private/WaveformTransformationsRenderManager.cpp
@@ -2,6 +2,7 @@
 
 #include "WaveformTransformationsRenderManager.h"
 
+#include "Audio.h"
 #include "DSP/BufferVectorOperations.h"
 #include "DSP/FloatArrayMath.h"
 #include "IWaveformTransformationRenderer.h"
diff --git a/Engine/Plugins/Enterprise/DataprepEditor/Source/DataprepCore/Private/DataprepAssetInterface.cpp b/Engine/Plugins/Enterprise/DataprepEditor/Source/DataprepCore/Private/DataprepAssetInterface.cpp
--- a/Engine/Plugins/Enterprise/DataprepEditor/Source/DataprepCore/Private/DataprepAssetInterface.cpp
+++ b/Engine/Plugins/Enterprise/DataprepEditor/Source/DataprepCore/Private/DataprepAssetInterface.cpp
@@ -15,6 +15,7 @@
 
 #include "Kismet2/KismetEditorUtilities.h"
 #include "UObject/EnterpriseObjectVersion.h"
+#include "UObject/UObjectIterator.h"
 
 #define LOCTEXT_NAMESPACE "DataprepAssetInterface"
 
diff --git a/Engine/Plugins/Enterprise/DataprepEditor/Source/DataprepCore/Private/DataprepCorePrivateUtils.cpp b/Engine/Plugins/Enterprise/DataprepEditor/Source/DataprepCore/Private/DataprepCorePrivateUtils.cpp
--- a/Engine/Plugins/Enterprise/DataprepEditor/Source/DataprepCore/Private/DataprepCorePrivateUtils.cpp
+++ b/Engine/Plugins/Enterprise/DataprepEditor/Source/DataprepCore/Private/DataprepCorePrivateUtils.cpp
@@ -21,6 +21,7 @@
 #include "MessageLogModule.h"
 #include "Modules/ModuleManager.h"
 #include "StaticMeshAttributes.h"
+#include "StaticMeshResources.h"
 
 #define LOCTEXT_NAMESPACE "DataprepAsset"
 
diff --git a/Engine/Plugins/Enterprise/DataprepEditor/Source/DataprepCore/Private/DataprepCoreUtils.cpp b/Engine/Plugins/Enterprise/DataprepEditor/Source/DataprepCore/Private/DataprepCoreUtils.cpp
--- a/Engine/Plugins/Enterprise/DataprepEditor/Source/DataprepCore/Private/DataprepCoreUtils.cpp
+++ b/Engine/Plugins/Enterprise/DataprepEditor/Source/DataprepCore/Private/DataprepCoreUtils.cpp
@@ -1,6 +1,7 @@
 // Copyright Epic Games, Inc. All Rights Reserved.
 
 #include "DataprepCoreUtils.h"
+#include "Engine/Texture.h"
 
 #ifdef NEW_DATASMITHSCENE_WORKFLOW
 #include "DataprepAssetUserData.h"
diff --git a/Engine/Plugins/Enterprise/DataprepEditor/Source/DataprepCore/Private/DataprepOperation.cpp b/Engine/Plugins/Enterprise/DataprepEditor/Source/DataprepCore/Private/DataprepOperation.cpp
--- a/Engine/Plugins/Enterprise/DataprepEditor/Source/DataprepCore/Private/DataprepOperation.cpp
+++ b/Engine/Plugins/Enterprise/DataprepEditor/Source/DataprepCore/Private/DataprepOperation.cpp
@@ -3,6 +3,7 @@
 #include "DataprepOperation.h"
 
 #include "DataprepCoreLogCategory.h"
+#include "Engine/Texture2D.h"
 #include "IDataprepLogger.h"
 #include "IDataprepProgressReporter.h"
 
diff --git a/Engine/Plugins/Enterprise/DataprepEditor/Source/DataprepCore/Private/SelectionSystem/DataprepObjectSelectionFilter.cpp b/Engine/Plugins/Enterprise/DataprepEditor/Source/DataprepCore/Private/SelectionSystem/DataprepObjectSelectionFilter.cpp
--- a/Engine/Plugins/Enterprise/DataprepEditor/Source/DataprepCore/Private/SelectionSystem/DataprepObjectSelectionFilter.cpp
+++ b/Engine/Plugins/Enterprise/DataprepEditor/Source/DataprepCore/Private/SelectionSystem/DataprepObjectSelectionFilter.cpp
@@ -1,6 +1,7 @@
 // Copyright Epic Games, Inc. All Rights Reserved.
 
 #include "SelectionSystem/DataprepObjectSelectionFilter.h"
+#include "Engine/Level.h"
 #include "SelectionSystem/DataprepSelectionSystemStructs.h"
 #include "Shared/DataprepCorePrivateUtils.h"
 
diff --git a/Engine/Plugins/Enterprise/DataprepEditor/Source/DataprepEditor/Private/DataprepEditorUtils.cpp b/Engine/Plugins/Enterprise/DataprepEditor/Source/DataprepEditor/Private/DataprepEditorUtils.cpp
--- a/Engine/Plugins/Enterprise/DataprepEditor/Source/DataprepEditor/Private/DataprepEditorUtils.cpp
+++ b/Engine/Plugins/Enterprise/DataprepEditor/Source/DataprepEditor/Private/DataprepEditorUtils.cpp
@@ -2,9 +2,13 @@
 
 #include "DataprepEditorUtils.h"
 
+#include "Components/StaticMeshComponent.h"
 #include "DataprepAsset.h"
 #include "DataprepParameterizableObject.h"
 #include "DataprepOperation.h"
+#include "Engine/Texture.h"
+#include "Materials/Material.h"
+#include "Materials/MaterialInterface.h"
 #include "Parameterization/DataprepParameterizationUtils.h"
 #include "SchemaActions/DataprepFilterMenuActionCollector.h"
 #include "Widgets/Parameterization/SDataprepLinkToParameter.h"
diff --git a/Engine/Plugins/Enterprise/DataprepEditor/Source/DataprepEditor/Private/DataprepSnapshot.cpp b/Engine/Plugins/Enterprise/DataprepEditor/Source/DataprepEditor/Private/DataprepSnapshot.cpp
--- a/Engine/Plugins/Enterprise/DataprepEditor/Source/DataprepEditor/Private/DataprepSnapshot.cpp
+++ b/Engine/Plugins/Enterprise/DataprepEditor/Source/DataprepEditor/Private/DataprepSnapshot.cpp
@@ -35,6 +35,7 @@
 #include "Settings/LevelEditorMiscSettings.h"
 #include "UnrealExporter.h"
 #include "UObject/PropertyPortFlags.h"
+#include "UObject/UObjectIterator.h"
 
 #define LOCTEXT_NAMESPACE "DataprepEditor"
 
diff --git a/Engine/Plugins/Enterprise/DataprepEditor/Source/DataprepEditor/Private/DataprepStats.cpp b/Engine/Plugins/Enterprise/DataprepEditor/Source/DataprepEditor/Private/DataprepStats.cpp
--- a/Engine/Plugins/Enterprise/DataprepEditor/Source/DataprepEditor/Private/DataprepStats.cpp
+++ b/Engine/Plugins/Enterprise/DataprepEditor/Source/DataprepEditor/Private/DataprepStats.cpp
@@ -8,7 +8,9 @@
 #include "Components/StaticMeshComponent.h"
 #include "Components/PrimitiveComponent.h"
 #include "Components/InstancedStaticMeshComponent.h"
+#include "Engine/Level.h"
 #include "Engine/MapBuildDataRegistry.h"
+#include "Engine/SkeletalMesh.h"
 #include "Engine/World.h"
 #include "Engine/StaticMesh.h"
 #include "Model.h"
diff --git a/Engine/Plugins/Enterprise/DataprepEditor/Source/DataprepEditor/Private/Widgets/SDataprepEditorViewport.cpp b/Engine/Plugins/Enterprise/DataprepEditor/Source/DataprepEditor/Private/Widgets/SDataprepEditorViewport.cpp
--- a/Engine/Plugins/Enterprise/DataprepEditor/Source/DataprepEditor/Private/Widgets/SDataprepEditorViewport.cpp
+++ b/Engine/Plugins/Enterprise/DataprepEditor/Source/DataprepEditor/Private/Widgets/SDataprepEditorViewport.cpp
@@ -16,6 +16,7 @@
 #include "Editor.h"
 #include "Engine/Engine.h"
 #include "Engine/StaticMesh.h"
+#include "StaticMeshResources.h"
 #include "StaticMeshSceneProxy.h"
 #include "Engine/TextureCube.h"
 #include "Engine/World.h"
diff --git a/Engine/Plugins/Enterprise/DataprepEditor/Source/DataprepLibraries/Private/DataprepEditingOperations.cpp b/Engine/Plugins/Enterprise/DataprepEditor/Source/DataprepLibraries/Private/DataprepEditingOperations.cpp
--- a/Engine/Plugins/Enterprise/DataprepEditor/Source/DataprepLibraries/Private/DataprepEditingOperations.cpp
+++ b/Engine/Plugins/Enterprise/DataprepEditor/Source/DataprepLibraries/Private/DataprepEditingOperations.cpp
@@ -37,6 +37,8 @@
 #include "DetailLayoutBuilder.h"
 #include "DetailWidgetRow.h"
 #include "Widgets/Input/SButton.h"
+#include "Widgets/Input/SComboButton.h"
+#include "Widgets/Layout/SBox.h"
 
 #define LOCTEXT_NAMESPACE "DatasmithEditingOperations"
 
diff --git a/Engine/Plugins/Enterprise/DataprepEditor/Source/DataprepLibraries/Private/DataprepLibrariesModule.cpp b/Engine/Plugins/Enterprise/DataprepEditor/Source/DataprepLibraries/Private/DataprepLibrariesModule.cpp
--- a/Engine/Plugins/Enterprise/DataprepEditor/Source/DataprepLibraries/Private/DataprepLibrariesModule.cpp
+++ b/Engine/Plugins/Enterprise/DataprepEditor/Source/DataprepLibraries/Private/DataprepLibrariesModule.cpp
@@ -11,6 +11,7 @@
 #include "ContentBrowserModule.h"
 #include "Dialogs/DlgPickAssetPath.h"
 #include "IContentBrowserSingleton.h"
+#include "Materials/MaterialInterface.h"
 #include "Modules/ModuleManager.h"
 #include "PropertyEditorModule.h"
 #include "ToolMenus.h"
diff --git a/Engine/Plugins/Enterprise/DataprepEditor/Source/DataprepLibraries/Private/DataprepOperations.cpp b/Engine/Plugins/Enterprise/DataprepEditor/Source/DataprepLibraries/Private/DataprepOperations.cpp
--- a/Engine/Plugins/Enterprise/DataprepEditor/Source/DataprepLibraries/Private/DataprepOperations.cpp
+++ b/Engine/Plugins/Enterprise/DataprepEditor/Source/DataprepLibraries/Private/DataprepOperations.cpp
@@ -4,6 +4,7 @@
 
 #include "DataprepOperationsLibrary.h"
 
+#include "Engine/Texture2D.h"
 #include "GenericPlatform/GenericPlatformTime.h"
 #include "StaticMeshResources.h"
 
@@ -14,6 +15,7 @@
 #include "DetailLayoutBuilder.h"
 #include "DetailWidgetRow.h"
 #include "Framework/MultiBox/MultiBoxBuilder.h"
+#include "UObject/UObjectIterator.h"
 #include "Widgets/Input/STextComboBox.h"
 #include "Widgets/Input/SComboButton.h"
 #include "Widgets/Input/SEditableTextBox.h"
diff --git a/Engine/Plugins/Enterprise/DataprepEditor/Source/DataprepLibraries/Private/DataprepOperationsLibrary.cpp b/Engine/Plugins/Enterprise/DataprepEditor/Source/DataprepLibraries/Private/DataprepOperationsLibrary.cpp
--- a/Engine/Plugins/Enterprise/DataprepEditor/Source/DataprepLibraries/Private/DataprepOperationsLibrary.cpp
+++ b/Engine/Plugins/Enterprise/DataprepEditor/Source/DataprepLibraries/Private/DataprepOperationsLibrary.cpp
@@ -1,6 +1,7 @@
 // Copyright Epic Games, Inc. All Rights Reserved.
 
 #include "DataprepOperationsLibrary.h"
+#include "Components/StaticMeshComponent.h"
 #include "DataprepOperationsLibraryUtil.h"
 
 #include "DataprepCoreUtils.h"
@@ -18,6 +19,7 @@
 #include "Engine/StaticMesh.h"
 #include "Engine/StaticMeshActor.h"
 #include "Engine/Texture.h"
+#include "Engine/Texture2D.h"
 #include "Engine/World.h"
 #include "EngineUtils.h"
 #include "GameFramework/Actor.h"
@@ -35,6 +37,7 @@
 #include "StaticMeshOperations.h"
 #include "UObject/SoftObjectPath.h"
 #include "StaticMeshEditorSubsystem.h"
+#include "UObject/UObjectIterator.h"
 
 DEFINE_LOG_CATEGORY(LogDataprep);
 
diff --git a/Engine/Plugins/Enterprise/DataprepEditor/Source/DataprepLibraries/Private/DataprepOperationsLibraryUtil.cpp b/Engine/Plugins/Enterprise/DataprepEditor/Source/DataprepLibraries/Private/DataprepOperationsLibraryUtil.cpp
--- a/Engine/Plugins/Enterprise/DataprepEditor/Source/DataprepLibraries/Private/DataprepOperationsLibraryUtil.cpp
+++ b/Engine/Plugins/Enterprise/DataprepEditor/Source/DataprepLibraries/Private/DataprepOperationsLibraryUtil.cpp
@@ -2,6 +2,7 @@
 
 #include "DataprepOperationsLibraryUtil.h"
 
+#include "Components/StaticMeshComponent.h"
 #include "Engine/StaticMesh.h"
 #include "Engine/StaticMeshActor.h"
 #include "Materials/Material.h"
@@ -9,6 +10,7 @@
 #include "Materials/MaterialInterface.h"
 #include "StaticMeshAttributes.h"
 #include "StaticMeshOperations.h"
+#include "StaticMeshResources.h"
 
 namespace DataprepOperationsLibraryUtil
 {
diff --git a/Engine/Plugins/Enterprise/DataprepEditor/Source/DataprepLibraries/Private/DataprepSelectionTransforms.cpp b/Engine/Plugins/Enterprise/DataprepEditor/Source/DataprepLibraries/Private/DataprepSelectionTransforms.cpp
--- a/Engine/Plugins/Enterprise/DataprepEditor/Source/DataprepLibraries/Private/DataprepSelectionTransforms.cpp
+++ b/Engine/Plugins/Enterprise/DataprepEditor/Source/DataprepLibraries/Private/DataprepSelectionTransforms.cpp
@@ -1,11 +1,14 @@
 // Copyright Epic Games, Inc. All Rights Reserved.
 
 #include "DataprepSelectionTransforms.h"
+#include "Components/StaticMeshComponent.h"
 #include "DataprepCorePrivateUtils.h"
 #include "Engine/StaticMesh.h"
 #include "Engine/StaticMeshActor.h"
 #include "Engine/Texture.h"
+#include "Materials/Material.h"
 #include "Materials/MaterialInstance.h"
+#include "UObject/UObjectIterator.h"
 
 #define LOCTEXT_NAMESPACE "DataprepSelectionTransforms"
 
diff --git a/Engine/Plugins/Enterprise/DatasmithCADImporter/Source/ParametricSurfaceExtension/Private/DataprepTessellationOperation.cpp b/Engine/Plugins/Enterprise/DatasmithCADImporter/Source/ParametricSurfaceExtension/Private/DataprepTessellationOperation.cpp
--- a/Engine/Plugins/Enterprise/DatasmithCADImporter/Source/ParametricSurfaceExtension/Private/DataprepTessellationOperation.cpp
+++ b/Engine/Plugins/Enterprise/DatasmithCADImporter/Source/ParametricSurfaceExtension/Private/DataprepTessellationOperation.cpp
@@ -2,6 +2,7 @@
 
 #include "DataprepTessellationOperation.h"
 
+#include "Components/StaticMeshComponent.h"
 #include "ParametricSurfaceBlueprintLibrary.h"
 #include "DatasmithAdditionalData.h"
 #include "DatasmithAssetImportData.h"
diff --git a/Engine/Plugins/Enterprise/DatasmithContent/Source/DatasmithContent/Private/DatasmithAssetImportData.cpp b/Engine/Plugins/Enterprise/DatasmithContent/Source/DatasmithContent/Private/DatasmithAssetImportData.cpp
--- a/Engine/Plugins/Enterprise/DatasmithContent/Source/DatasmithContent/Private/DatasmithAssetImportData.cpp
+++ b/Engine/Plugins/Enterprise/DatasmithContent/Source/DatasmithContent/Private/DatasmithAssetImportData.cpp
@@ -6,8 +6,10 @@
 #include "DatasmithScene.h"
 
 #include "Engine/StaticMesh.h"
+#include "Engine/Texture.h"
 #include "EngineUtils.h"
 #include "HAL/FileManager.h"
+#include "Materials/MaterialInterface.h"
 #include "Misc/FileHelper.h"
 #include "Misc/Paths.h"
 #include "UObject/EnterpriseObjectVersion.h"
diff --git a/Engine/Plugins/Enterprise/DatasmithImporter/Source/DatasmithImporter/Private/ActorFactoryDatasmithScene.cpp b/Engine/Plugins/Enterprise/DatasmithImporter/Source/DatasmithImporter/Private/ActorFactoryDatasmithScene.cpp
--- a/Engine/Plugins/Enterprise/DatasmithImporter/Source/DatasmithImporter/Private/ActorFactoryDatasmithScene.cpp
+++ b/Engine/Plugins/Enterprise/DatasmithImporter/Source/DatasmithImporter/Private/ActorFactoryDatasmithScene.cpp
@@ -2,6 +2,7 @@
 
 #include "ActorFactoryDatasmithScene.h"
 
+#include "Components/StaticMeshComponent.h"
 #include "DatasmithAssetImportData.h"
 #include "DatasmithBlueprintLibrary.h"
 #include "DatasmithImportContext.h"
@@ -9,6 +10,7 @@
 #include "DatasmithScene.h"
 #include "DatasmithSceneActor.h"
 #include "DatasmithSceneFactory.h"
+#include "Engine/Level.h"
 #include "Utility/DatasmithImporterUtils.h"
 
 #include "AssetRegistry/AssetData.h"
diff --git a/Engine/Plugins/Enterprise/DatasmithImporter/Source/DatasmithImporter/Private/DatasmithBlueprintLibrary.cpp b/Engine/Plugins/Enterprise/DatasmithImporter/Source/DatasmithImporter/Private/DatasmithBlueprintLibrary.cpp
--- a/Engine/Plugins/Enterprise/DatasmithImporter/Source/DatasmithImporter/Private/DatasmithBlueprintLibrary.cpp
+++ b/Engine/Plugins/Enterprise/DatasmithImporter/Source/DatasmithImporter/Private/DatasmithBlueprintLibrary.cpp
@@ -2,6 +2,7 @@
 
 #include "DatasmithBlueprintLibrary.h"
 
+#include "Components/StaticMeshComponent.h"
 #include "DatasmithImporterHelper.h"
 #include "DatasmithImportFactory.h"
 #include "DatasmithImportOptions.h"
@@ -12,6 +13,7 @@
 #include "DatasmithTranslatableSource.h"
 #include "DatasmithTranslatorManager.h"
 #include "ObjectElements/DatasmithUSceneElement.h"
+#include "StaticMeshResources.h"
 #include "Utility/DatasmithImporterUtils.h"
 #include "Utility/DatasmithImportFactoryHelper.h"
 #include "Utility/DatasmithMeshHelper.h"
diff --git a/Engine/Plugins/Enterprise/DatasmithImporter/Source/DatasmithImporter/Private/DatasmithFileProducer.cpp b/Engine/Plugins/Enterprise/DatasmithImporter/Source/DatasmithImporter/Private/DatasmithFileProducer.cpp
--- a/Engine/Plugins/Enterprise/DatasmithImporter/Source/DatasmithImporter/Private/DatasmithFileProducer.cpp
+++ b/Engine/Plugins/Enterprise/DatasmithImporter/Source/DatasmithImporter/Private/DatasmithFileProducer.cpp
@@ -2,6 +2,8 @@
 
 #include "DatasmithFileProducer.h"
 
+#include "Engine/Level.h"
+#include "Engine/TextureLightProfile.h"
 #include "IDataprepLogger.h"
 #include "IDataprepProgressReporter.h"
 
@@ -59,6 +61,7 @@
 #include "Widgets/Input/SCheckBox.h"
 #include "Widgets/Input/SButton.h"
 #include "Widgets/Input/SEditableTextBox.h"
+#include "Widgets/Layout/SBox.h"
 #include "Widgets/SBoxPanel.h"
 #include "Widgets/SWidget.h"
 #include "Widgets/Text/STextBlock.h"
diff --git a/Engine/Plugins/Enterprise/DatasmithImporter/Source/DatasmithImporter/Private/DatasmithImporter.cpp b/Engine/Plugins/Enterprise/DatasmithImporter/Source/DatasmithImporter/Private/DatasmithImporter.cpp
--- a/Engine/Plugins/Enterprise/DatasmithImporter/Source/DatasmithImporter/Private/DatasmithImporter.cpp
+++ b/Engine/Plugins/Enterprise/DatasmithImporter/Source/DatasmithImporter/Private/DatasmithImporter.cpp
@@ -25,6 +25,7 @@
 #include "DatasmithTextureImporter.h"
 #include "DatasmithTranslator.h"
 #include "DatasmithUtils.h"
+#include "Engine/Texture2D.h"
 #include "IDatasmithSceneElements.h"
 #include "LevelVariantSets.h"
 #include "ObjectTemplates/DatasmithObjectTemplate.h"
diff --git a/Engine/Plugins/Enterprise/DatasmithImporter/Source/DatasmithImporter/Private/Utility/DatasmithImporterUtils.cpp b/Engine/Plugins/Enterprise/DatasmithImporter/Source/DatasmithImporter/Private/Utility/DatasmithImporterUtils.cpp
--- a/Engine/Plugins/Enterprise/DatasmithImporter/Source/DatasmithImporter/Private/Utility/DatasmithImporterUtils.cpp
+++ b/Engine/Plugins/Enterprise/DatasmithImporter/Source/DatasmithImporter/Private/Utility/DatasmithImporterUtils.cpp
@@ -14,6 +14,7 @@
 #include "DatasmithSceneXmlReader.h"
 #include "DatasmithSceneXmlWriter.h"
 #include "DatasmithTranslatorManager.h"
+#include "Engine/TextureLightProfile.h"
 #include "LevelVariantSets.h"
 
 #include "ObjectTemplates/DatasmithActorTemplate.h"
diff --git a/Engine/Plugins/Enterprise/DatasmithImporter/Source/DatasmithImporter/Public/Utility/DatasmithImporterUtils.h b/Engine/Plugins/Enterprise/DatasmithImporter/Source/DatasmithImporter/Public/Utility/DatasmithImporterUtils.h
--- a/Engine/Plugins/Enterprise/DatasmithImporter/Source/DatasmithImporter/Public/Utility/DatasmithImporterUtils.h
+++ b/Engine/Plugins/Enterprise/DatasmithImporter/Source/DatasmithImporter/Public/Utility/DatasmithImporterUtils.h
@@ -6,6 +6,7 @@
 #include "DatasmithScene.h"
 #include "IDatasmithSceneElements.h"
 #include "InterchangeManager.h"
+#include "Materials/MaterialInterface.h"
 #include "UObject/UObjectHash.h"
 #include "Engine/Texture.h"
 
diff --git a/Engine/Plugins/Enterprise/GLTFExporter/Source/GLTFExporter/Private/Converters/GLTFTextureUtility.h b/Engine/Plugins/Enterprise/GLTFExporter/Source/GLTFExporter/Private/Converters/GLTFTextureUtility.h
--- a/Engine/Plugins/Enterprise/GLTFExporter/Source/GLTFExporter/Private/Converters/GLTFTextureUtility.h
+++ b/Engine/Plugins/Enterprise/GLTFExporter/Source/GLTFExporter/Private/Converters/GLTFTextureUtility.h
@@ -10,6 +10,7 @@
 class UTextureCube;
 class UTextureRenderTarget2D;
 class UTextureRenderTargetCube;
+class UTexture2D;
 
 struct FGLTFTextureUtility
 {
diff --git a/Engine/Plugins/Enterprise/LidarPointCloud/Source/LidarPointCloudEditor/Private/LidarPointCloudEdModeToolkit.cpp b/Engine/Plugins/Enterprise/LidarPointCloud/Source/LidarPointCloudEditor/Private/LidarPointCloudEdModeToolkit.cpp
--- a/Engine/Plugins/Enterprise/LidarPointCloud/Source/LidarPointCloudEditor/Private/LidarPointCloudEdModeToolkit.cpp
+++ b/Engine/Plugins/Enterprise/LidarPointCloud/Source/LidarPointCloudEditor/Private/LidarPointCloudEdModeToolkit.cpp
@@ -15,6 +15,7 @@
 #include "Widgets/Input/SButton.h"
 #include "Widgets/Layout/SWrapBox.h"
 #include "Widgets/Layout/SSpacer.h"
+#include "Widgets/Views/SHeaderRow.h"
 
 #define LOCTEXT_NAMESPACE "LidarEditMode"
 
diff --git a/Engine/Plugins/Enterprise/LidarPointCloud/Source/LidarPointCloudEditor/Private/LidarPointCloudEditorHelper.cpp b/Engine/Plugins/Enterprise/LidarPointCloud/Source/LidarPointCloudEditor/Private/LidarPointCloudEditorHelper.cpp
--- a/Engine/Plugins/Enterprise/LidarPointCloud/Source/LidarPointCloudEditor/Private/LidarPointCloudEditorHelper.cpp
+++ b/Engine/Plugins/Enterprise/LidarPointCloud/Source/LidarPointCloudEditor/Private/LidarPointCloudEditorHelper.cpp
@@ -14,6 +14,8 @@
 #include "Engine/StaticMeshActor.h"
 #include "Misc/ScopedSlowTask.h"
 #include "PhysicsEngine/BodySetup.h"
+#include "StaticMeshResources.h"
+#include "UObject/UObjectIterator.h"
 
 #define LOCTEXT_NAMESPACE "LidarPointCloudEditorHelper"
 
diff --git a/Engine/Plugins/Enterprise/LidarPointCloud/Source/LidarPointCloudEditor/Private/LidarPointCloudEditorViewport.h b/Engine/Plugins/Enterprise/LidarPointCloud/Source/LidarPointCloudEditor/Private/LidarPointCloudEditorViewport.h
--- a/Engine/Plugins/Enterprise/LidarPointCloud/Source/LidarPointCloudEditor/Private/LidarPointCloudEditorViewport.h
+++ b/Engine/Plugins/Enterprise/LidarPointCloud/Source/LidarPointCloudEditor/Private/LidarPointCloudEditorViewport.h
@@ -8,6 +8,7 @@
 #include "AdvancedPreviewScene.h"
 #include "SEditorViewport.h"
 #include "Editor/UnrealEd/Public/SCommonEditorViewportToolbarBase.h"
+#include "Widgets/Layout/SBorder.h"
 
 class FLidarPointCloudEditor;
 class FLidarPointCloudEditorViewportClient;
diff --git a/Engine/Plugins/Enterprise/MDLImporter/Source/MDLImporter/Private/MDLMaterialImporter.cpp b/Engine/Plugins/Enterprise/MDLImporter/Source/MDLImporter/Private/MDLMaterialImporter.cpp
--- a/Engine/Plugins/Enterprise/MDLImporter/Source/MDLImporter/Private/MDLMaterialImporter.cpp
+++ b/Engine/Plugins/Enterprise/MDLImporter/Source/MDLImporter/Private/MDLMaterialImporter.cpp
@@ -9,6 +9,7 @@
 #include "mdl/MaterialCollection.h"
 
 #include "Materials/MaterialInterface.h"
+#include "Misc/Paths.h"
 
 void FMdlMaterialImporter::AddSearchPath(const FString& SearchPath)
 {
diff --git a/Engine/Plugins/Experimental/AR/PointCloud/Source/PointCloud/Private/PointCloudSceneProxy.cpp b/Engine/Plugins/Experimental/AR/PointCloud/Source/PointCloud/Private/PointCloudSceneProxy.cpp
--- a/Engine/Plugins/Experimental/AR/PointCloud/Source/PointCloud/Private/PointCloudSceneProxy.cpp
+++ b/Engine/Plugins/Experimental/AR/PointCloud/Source/PointCloud/Private/PointCloudSceneProxy.cpp
@@ -1,6 +1,7 @@
 // Copyright Epic Games, Inc. All Rights Reserved.
 
 #include "PointCloudSceneProxy.h"
+#include "Engine/Texture.h"
 #include "PointCloud.h"
 #include "PointCloudComponent.h"
 
@@ -10,6 +11,7 @@
 #include "RenderingThread.h"
 #include "Containers/ResourceArray.h"
 #include "EngineGlobals.h"
+#include "TextureResource.h"
 #include "VertexFactory.h"
 #include "MaterialShared.h"
 #include "Materials/Material.h"
diff --git a/Engine/Plugins/Experimental/ActorPalette/Source/ActorPalette/Private/ActorPaletteViewport.cpp b/Engine/Plugins/Experimental/ActorPalette/Source/ActorPalette/Private/ActorPaletteViewport.cpp
--- a/Engine/Plugins/Experimental/ActorPalette/Source/ActorPalette/Private/ActorPaletteViewport.cpp
+++ b/Engine/Plugins/Experimental/ActorPalette/Source/ActorPalette/Private/ActorPaletteViewport.cpp
@@ -1,6 +1,8 @@
 // Copyright Epic Games, Inc. All Rights Reserved.
 
 #include "ActorPaletteViewport.h"
+#include "Framework/Application/SlateApplication.h"
+#include "Widgets/Input/SComboButton.h"
 #include "Widgets/SBoxPanel.h"
 #include "Widgets/SViewport.h"
 #include "Widgets/Text/STextBlock.h"
diff --git a/Engine/Plugins/Experimental/Animation/PoseSearch/Source/Editor/Private/PoseSearchDebugger.cpp b/Engine/Plugins/Experimental/Animation/PoseSearch/Source/Editor/Private/PoseSearchDebugger.cpp
--- a/Engine/Plugins/Experimental/Animation/PoseSearch/Source/Editor/Private/PoseSearchDebugger.cpp
+++ b/Engine/Plugins/Experimental/Animation/PoseSearch/Source/Editor/Private/PoseSearchDebugger.cpp
@@ -1,6 +1,8 @@
 ﻿// Copyright Epic Games, Inc. All Rights Reserved.
 
 #include "PoseSearchDebugger.h"
+#include "Engine/SkeletalMesh.h"
+#include "Engine/SkinnedAsset.h"
 #include "IAnimationProvider.h"
 #include "IGameplayProvider.h"
 #include "PoseSearch/PoseSearch.h"
@@ -10,6 +12,7 @@
 #include "Animation/MirrorDataTable.h"
 #include "Animation/AnimRootMotionProvider.h"
 #include "Preferences/PersonaOptions.h"
+#include "PropertyEditorModule.h"
 #include "RewindDebuggerInterface/Public/IRewindDebugger.h"
 #include "ObjectTrace.h"
 #include "TraceServices/Model/AnalysisSession.h"
@@ -28,6 +31,7 @@
 #include "Editor.h"
 #include "Algo/AllOf.h"
 #include "Styling/AppStyle.h"
+#include "Widgets/Views/SListView.h"
 
 #define LOCTEXT_NAMESPACE "PoseSearchDebugger"
 
diff --git a/Engine/Plugins/Experimental/Animation/PoseSearch/Source/Editor/Private/SPoseSearchDatabaseViewport.cpp b/Engine/Plugins/Experimental/Animation/PoseSearch/Source/Editor/Private/SPoseSearchDatabaseViewport.cpp
--- a/Engine/Plugins/Experimental/Animation/PoseSearch/Source/Editor/Private/SPoseSearchDatabaseViewport.cpp
+++ b/Engine/Plugins/Experimental/Animation/PoseSearch/Source/Editor/Private/SPoseSearchDatabaseViewport.cpp
@@ -1,6 +1,7 @@
 // Copyright Epic Games, Inc. All Rights Reserved.
 
 #include "SPoseSearchDatabaseViewport.h"
+#include "Framework/Application/SlateApplication.h"
 #include "PoseSearchDatabaseViewportClient.h"
 #include "PoseSearchDatabaseViewModel.h"
 #include "PoseSearchDatabaseEditor.h"
@@ -8,6 +9,7 @@
 #include "PoseSearchDatabaseEditorCommands.h"
 
 #include "SSimpleTimeSlider.h"
+#include "Viewports.h"
 #include "Widgets/Input/SButton.h"
 
 namespace UE::PoseSearch
diff --git a/Engine/Plugins/Experimental/Animation/PoseSearch/Source/Runtime/Private/AnimNode_MotionMatching.cpp b/Engine/Plugins/Experimental/Animation/PoseSearch/Source/Runtime/Private/AnimNode_MotionMatching.cpp
--- a/Engine/Plugins/Experimental/Animation/PoseSearch/Source/Runtime/Private/AnimNode_MotionMatching.cpp
+++ b/Engine/Plugins/Experimental/Animation/PoseSearch/Source/Runtime/Private/AnimNode_MotionMatching.cpp
@@ -8,6 +8,8 @@
 #include "Animation/BlendSpace.h"
 #include "Animation/MotionTrajectoryTypes.h"
 #include "Animation/AnimRootMotionProvider.h"
+#include "Components/SkeletalMeshComponent.h"
+#include "DrawDebugHelpers.h"
 #include "PoseSearch/AnimNode_PoseSearchHistoryCollector.h"
 #include "Trace/PoseSearchTraceLogger.h"
 
diff --git a/Engine/Plugins/Experimental/Animation/PoseSearch/Source/Runtime/Private/PoseSearchLibrary.cpp b/Engine/Plugins/Experimental/Animation/PoseSearch/Source/Runtime/Private/PoseSearchLibrary.cpp
--- a/Engine/Plugins/Experimental/Animation/PoseSearch/Source/Runtime/Private/PoseSearchLibrary.cpp
+++ b/Engine/Plugins/Experimental/Animation/PoseSearch/Source/Runtime/Private/PoseSearchLibrary.cpp
@@ -2,6 +2,7 @@
 
 #include "PoseSearch/PoseSearchLibrary.h"
 
+#include "Animation/AnimInstanceProxy.h"
 #include "Animation/AnimSequence.h"
 #include "Animation/AnimNode_Inertialization.h"
 #include "Animation/AnimNode_SequencePlayer.h"
diff --git a/Engine/Plugins/Experimental/Animation/PoseSearch/Source/Runtime/Private/Trace/PoseSearchTraceLogger.cpp b/Engine/Plugins/Experimental/Animation/PoseSearch/Source/Runtime/Private/Trace/PoseSearchTraceLogger.cpp
--- a/Engine/Plugins/Experimental/Animation/PoseSearch/Source/Runtime/Private/Trace/PoseSearchTraceLogger.cpp
+++ b/Engine/Plugins/Experimental/Animation/PoseSearch/Source/Runtime/Private/Trace/PoseSearchTraceLogger.cpp
@@ -2,8 +2,10 @@
 
 #include "PoseSearchTraceLogger.h"
 #include "Animation/AnimInstanceProxy.h"
+#include "Components/SkeletalMeshComponent.h"
 #include "Trace/Trace.inl"
 #include "Animation/AnimNodeBase.h"
+#include "TraceFilter.h"
 
 
 UE_TRACE_CHANNEL_DEFINE(PoseSearchChannel);
diff --git a/Engine/Plugins/Experimental/ChaosVehiclesPlugin/Source/ChaosVehicles/Private/ChaosVehicleManager.cpp b/Engine/Plugins/Experimental/ChaosVehiclesPlugin/Source/ChaosVehicles/Private/ChaosVehicleManager.cpp
--- a/Engine/Plugins/Experimental/ChaosVehiclesPlugin/Source/ChaosVehicles/Private/ChaosVehicleManager.cpp
+++ b/Engine/Plugins/Experimental/ChaosVehiclesPlugin/Source/ChaosVehicles/Private/ChaosVehicleManager.cpp
@@ -1,6 +1,7 @@
 // Copyright Epic Games, Inc. All Rights Reserved.
 
 #include "ChaosVehicleManager.h"
+#include "Engine/World.h"
 #include "Physics/Experimental/PhysScene_Chaos.h"
 #include "UObject/UObjectIterator.h"
 
diff --git a/Engine/Plugins/Experimental/Chooser/Source/ChooserEditor/Private/ChooserTableEditor.h b/Engine/Plugins/Experimental/Chooser/Source/ChooserEditor/Private/ChooserTableEditor.h
--- a/Engine/Plugins/Experimental/Chooser/Source/ChooserEditor/Private/ChooserTableEditor.h
+++ b/Engine/Plugins/Experimental/Chooser/Source/ChooserEditor/Private/ChooserTableEditor.h
@@ -6,15 +6,16 @@
 #include "Toolkits/IToolkitHost.h"
 #include "Toolkits/AssetEditorToolkit.h"
 #include "Editor/PropertyEditor/Public/PropertyEditorDelegates.h"
+#include "Widgets/Views/SListView.h"
 #include "Widgets/Views/STableViewBase.h"
 #include "Chooser.h"
 #include "EditorUndoClient.h"
 #include "ChooserTableEditor.generated.h"
 
+class SComboButton;
 class SEditableText;
 class IDetailsView;
 
-
 // Class used for chooser editor details customization
 UCLASS()
 class UChooserRowDetails : public UObject
diff --git a/Engine/Plugins/Experimental/ColorCorrectRegions/Source/ColorCorrectRegions/Private/ColorCorrectWindow.cpp b/Engine/Plugins/Experimental/ColorCorrectRegions/Source/ColorCorrectRegions/Private/ColorCorrectWindow.cpp
--- a/Engine/Plugins/Experimental/ColorCorrectRegions/Source/ColorCorrectRegions/Private/ColorCorrectWindow.cpp
+++ b/Engine/Plugins/Experimental/ColorCorrectRegions/Source/ColorCorrectRegions/Private/ColorCorrectWindow.cpp
@@ -8,6 +8,7 @@
 #include "Components/StaticMeshComponent.h"
 #include "Engine/StaticMesh.h"
 #include "Engine/Texture2D.h"
+#include "Misc/EnumRange.h"
 #include "UObject/ConstructorHelpers.h"
 
 ENUM_RANGE_BY_COUNT(EColorCorrectWindowType, EColorCorrectWindowType::MAX)
diff --git a/Engine/Plugins/Experimental/Enterprise/DatasmithInterchange/Source/DatasmithInterchange/Private/DatasmithInterchangeModule.cpp b/Engine/Plugins/Experimental/Enterprise/DatasmithInterchange/Source/DatasmithInterchange/Private/DatasmithInterchangeModule.cpp
--- a/Engine/Plugins/Experimental/Enterprise/DatasmithInterchange/Source/DatasmithInterchange/Private/DatasmithInterchangeModule.cpp
+++ b/Engine/Plugins/Experimental/Enterprise/DatasmithInterchange/Source/DatasmithInterchange/Private/DatasmithInterchangeModule.cpp
@@ -2,6 +2,7 @@
 
 #include "DatasmithInterchangeModule.h"
 
+#include "Engine/Blueprint.h"
 #include "InterchangeDatasmithAreaLightFactory.h"
 #include "InterchangeDatasmithLog.h"
 #include "InterchangeDatasmithPipeline.h"
diff --git a/Engine/Plugins/Experimental/Enterprise/DatasmithInterchange/Source/DatasmithInterchange/Private/InterchangeDatasmithAreaLightFactory.cpp b/Engine/Plugins/Experimental/Enterprise/DatasmithInterchange/Source/DatasmithInterchange/Private/InterchangeDatasmithAreaLightFactory.cpp
--- a/Engine/Plugins/Experimental/Enterprise/DatasmithInterchange/Source/DatasmithInterchange/Private/InterchangeDatasmithAreaLightFactory.cpp
+++ b/Engine/Plugins/Experimental/Enterprise/DatasmithInterchange/Source/DatasmithInterchange/Private/InterchangeDatasmithAreaLightFactory.cpp
@@ -2,6 +2,7 @@
 
 #include "InterchangeDatasmithAreaLightFactory.h"
 
+#include "Engine/Blueprint.h"
 #include "InterchangeDatasmithLog.h"
 #include "InterchangeDatasmithAreaLightFactoryNode.h"
 
diff --git a/Engine/Plugins/Experimental/Enterprise/DatasmithRuntime/Source/Private/ActorImporter.cpp b/Engine/Plugins/Experimental/Enterprise/DatasmithRuntime/Source/Private/ActorImporter.cpp
--- a/Engine/Plugins/Experimental/Enterprise/DatasmithRuntime/Source/Private/ActorImporter.cpp
+++ b/Engine/Plugins/Experimental/Enterprise/DatasmithRuntime/Source/Private/ActorImporter.cpp
@@ -2,6 +2,7 @@
 
 #include "SceneImporter.h"
 
+#include "Components/StaticMeshComponent.h"
 #include "DatasmithRuntimeUtils.h"
 #include "LogCategory.h"
 
@@ -15,6 +16,7 @@
 #include "Engine/World.h"
 #include "GameFramework/Actor.h"
 #include "ProfilingDebugging/MiscTrace.h"
+#include "RenderingThread.h"
 
 namespace DatasmithRuntime
 {
diff --git a/Engine/Plugins/Experimental/GeometryCacheAbcFile/Source/GeometryCacheAbcFile/Private/GeometryCacheAbcStream.cpp b/Engine/Plugins/Experimental/GeometryCacheAbcFile/Source/GeometryCacheAbcFile/Private/GeometryCacheAbcStream.cpp
--- a/Engine/Plugins/Experimental/GeometryCacheAbcFile/Source/GeometryCacheAbcFile/Private/GeometryCacheAbcStream.cpp
+++ b/Engine/Plugins/Experimental/GeometryCacheAbcFile/Source/GeometryCacheAbcFile/Private/GeometryCacheAbcStream.cpp
@@ -8,6 +8,8 @@
 #include "GeometryCacheTrackAbcFile.h"
 #include "HAL/IConsoleManager.h"
 #include "Misc/CoreMisc.h"
+#include "Serialization/MemoryReader.h"
+#include "Serialization/MemoryWriter.h"
 
 static bool GAbcStreamCacheInDDC = true;
 static FAutoConsoleVariableRef CVarAbcStreamCacheInDDC(
diff --git a/Engine/Plugins/Experimental/Mutable/Source/CustomizableObject/Private/MuCO/CustomizableInstancePrivateData.h b/Engine/Plugins/Experimental/Mutable/Source/CustomizableObject/Private/MuCO/CustomizableInstancePrivateData.h
--- a/Engine/Plugins/Experimental/Mutable/Source/CustomizableObject/Private/MuCO/CustomizableInstancePrivateData.h
+++ b/Engine/Plugins/Experimental/Mutable/Source/CustomizableObject/Private/MuCO/CustomizableInstancePrivateData.h
@@ -21,7 +21,7 @@
 /** CustomizableObject Instance flags for internal use  */
 enum ECOInstanceFlags
 {
-	None							= 0,
+	ECONone							= 0,  // Should not use the name None here.. it collides with other enum in global namespace
 
 	// Update process
 	Updating						= 1 << 0,	//
@@ -310,7 +310,7 @@ public:
 
 private:
 	
-	ECOInstanceFlags InstanceFlagsPrivate = ECOInstanceFlags::None;
+	ECOInstanceFlags InstanceFlagsPrivate = ECOInstanceFlags::ECONone;
 
 	// Struct used by BuildMaterials() to identify common materials between LODs
 	struct FMutableMaterialPlaceholder
diff --git a/Engine/Plugins/Experimental/Mutable/Source/CustomizableObject/Private/MuCO/CustomizableObjectInstance.cpp b/Engine/Plugins/Experimental/Mutable/Source/CustomizableObject/Private/MuCO/CustomizableObjectInstance.cpp
--- a/Engine/Plugins/Experimental/Mutable/Source/CustomizableObject/Private/MuCO/CustomizableObjectInstance.cpp
+++ b/Engine/Plugins/Experimental/Mutable/Source/CustomizableObject/Private/MuCO/CustomizableObjectInstance.cpp
@@ -230,7 +230,7 @@ void UCustomizableInstancePrivateData::SetMinMaxLODToLoad(UCustomizableObjectIns
 	NewMinLOD = FMath::Min(FMath::Max(NewMinLOD, static_cast<int32>(FirstLODAvailable)), FirstLODAvailable + NumMaxLODsToStream);
 
 	const bool bIsDowngradeLODUpdate = LastUpdateMinLOD >= 0 && NewMinLOD > LastUpdateMinLOD;
-	SetCOInstanceFlags(bIsDowngradeLODUpdate ? PendingLODsDowngrade : None);
+	SetCOInstanceFlags(bIsDowngradeLODUpdate ? PendingLODsDowngrade : ECONone);
 
 	if (!bIsDowngradeLODUpdate && bLimitLODUpgrades)
 	{
@@ -254,7 +254,7 @@ void UCustomizableInstancePrivateData::SetMinMaxLODToLoad(UCustomizableObjectIns
 	NewMaxLOD = FMath::Max(NewMaxLOD, NewMinLOD);
 
 	// Save the new LODs
-	SetCOInstanceFlags(Public->Descriptor.MinLODToLoad != NewMinLOD || Public->Descriptor.MaxLODToLoad != NewMaxLOD ? PendingLODsUpdate : None);
+	SetCOInstanceFlags(Public->Descriptor.MinLODToLoad != NewMinLOD || Public->Descriptor.MaxLODToLoad != NewMaxLOD ? PendingLODsUpdate : ECONone);
 
 	Public->Descriptor.MinLODToLoad = NewMinLOD;
 	Public->Descriptor.MaxLODToLoad = NewMaxLOD;
@@ -265,7 +265,7 @@ void UCustomizableInstancePrivateData::PrepareForUpdate(const TSharedPtr<FMutabl
 {
 	const bool bNumLODsUpdated = NumLODsAvailable != OperationData->NumLODsAvailable;
 	
-	SetCOInstanceFlags(bNumLODsUpdated ? PendingMeshUpdate : None);
+	SetCOInstanceFlags(bNumLODsUpdated ? PendingMeshUpdate : ECONone);
 
 	NumLODsAvailable = OperationData->NumLODsAvailable;
 
diff --git a/Engine/Plugins/Experimental/PhysicsControl/Source/PhysicsControl/Private/PhysicsControlComponent.cpp b/Engine/Plugins/Experimental/PhysicsControl/Source/PhysicsControl/Private/PhysicsControlComponent.cpp
--- a/Engine/Plugins/Experimental/PhysicsControl/Source/PhysicsControl/Private/PhysicsControlComponent.cpp
+++ b/Engine/Plugins/Experimental/PhysicsControl/Source/PhysicsControl/Private/PhysicsControlComponent.cpp
@@ -1,6 +1,7 @@
 // Copyright Epic Games, Inc. All Rights Reserved.
 
 #include "PhysicsControlComponent.h"
+#include "Engine/World.h"
 #include "Physics/Experimental/PhysScene_Chaos.h"
 #include "PhysicsControlComponentLog.h"
 #include "PhysicsControlRecord.h"
@@ -14,6 +15,7 @@
 #include "Components/StaticMeshComponent.h"
 #include "Components/SkeletalMeshComponent.h"
 #include "PhysicsEngine/PhysicsAsset.h"
+#include "SceneManagement.h"
 
 //======================================================================================================================
 UPhysicsControlComponent::UPhysicsControlComponent(const FObjectInitializer& ObjectInitializer)
diff --git a/Engine/Plugins/Experimental/PhysicsControl/Source/PhysicsControl/Private/PhysicsControlComponentImpl.cpp b/Engine/Plugins/Experimental/PhysicsControl/Source/PhysicsControl/Private/PhysicsControlComponentImpl.cpp
--- a/Engine/Plugins/Experimental/PhysicsControl/Source/PhysicsControl/Private/PhysicsControlComponentImpl.cpp
+++ b/Engine/Plugins/Experimental/PhysicsControl/Source/PhysicsControl/Private/PhysicsControlComponentImpl.cpp
@@ -1,11 +1,13 @@
 // Copyright Epic Games, Inc. All Rights Reserved.
 
 #include "PhysicsControlComponentImpl.h"
+#include "Engine/SkeletalMesh.h"
 #include "PhysicsControlComponentLog.h"
 #include "PhysicsControlComponent.h"
 #include "PhysicsControlComponentHelpers.h"
 
 #include "Components/SkeletalMeshComponent.h"
+#include "PhysicsEngine/ConstraintInstance.h"
 
 //======================================================================================================================
 bool FPhysicsControlComponentImpl::GetBoneData(
diff --git a/Engine/Plugins/Experimental/RenderGrid/Source/RenderGridEditor/Private/UI/Components/SRenderGridEditableTextBlock.cpp b/Engine/Plugins/Experimental/RenderGrid/Source/RenderGridEditor/Private/UI/Components/SRenderGridEditableTextBlock.cpp
--- a/Engine/Plugins/Experimental/RenderGrid/Source/RenderGridEditor/Private/UI/Components/SRenderGridEditableTextBlock.cpp
+++ b/Engine/Plugins/Experimental/RenderGrid/Source/RenderGridEditor/Private/UI/Components/SRenderGridEditableTextBlock.cpp
@@ -7,6 +7,7 @@
 #include "SlateOptMacros.h"
 #include "Styling/AppStyle.h"
 #include "Widgets/Input/SButton.h"
+#include "Widgets/Text/STextBlock.h"
 
 #define LOCTEXT_NAMESPACE "SRenderGridEditableTextBlock"
 
diff --git a/Engine/Plugins/Experimental/RenderGrid/Source/RenderGridEditor/Private/UI/Components/SRenderGridFileSelectorTextBlock.cpp b/Engine/Plugins/Experimental/RenderGrid/Source/RenderGridEditor/Private/UI/Components/SRenderGridFileSelectorTextBlock.cpp
--- a/Engine/Plugins/Experimental/RenderGrid/Source/RenderGridEditor/Private/UI/Components/SRenderGridFileSelectorTextBlock.cpp
+++ b/Engine/Plugins/Experimental/RenderGrid/Source/RenderGridEditor/Private/UI/Components/SRenderGridFileSelectorTextBlock.cpp
@@ -4,10 +4,12 @@
 
 #include "DesktopPlatformModule.h"
 #include "EditorFontGlyphs.h"
+#include "Framework/Application/SlateApplication.h"
 #include "PropertyCustomizationHelpers.h"
 #include "SlateOptMacros.h"
 #include "Styling/AppStyle.h"
 #include "Widgets/Input/SButton.h"
+#include "Widgets/Text/STextBlock.h"
 
 #define LOCTEXT_NAMESPACE "SRenderGridFileSelectorTextBlock"
 
diff --git a/Engine/Plugins/Experimental/RenderGrid/Source/RenderGridEditor/Private/UI/SRenderGridJobList.cpp b/Engine/Plugins/Experimental/RenderGrid/Source/RenderGridEditor/Private/UI/SRenderGridJobList.cpp
--- a/Engine/Plugins/Experimental/RenderGrid/Source/RenderGridEditor/Private/UI/SRenderGridJobList.cpp
+++ b/Engine/Plugins/Experimental/RenderGrid/Source/RenderGridEditor/Private/UI/SRenderGridJobList.cpp
@@ -3,6 +3,8 @@
 #include "UI/SRenderGridJobList.h"
 
 #include "Commands/RenderGridEditorCommands.h"
+#include "Framework/Application/SlateApplication.h"
+#include "Layout/WidgetPath.h"
 #include "UI/Components/SRenderGridDragHandle.h"
 #include "UI/Components/SRenderGridEditableTextBlock.h"
 #include "UI/Components/SRenderGridFileSelectorTextBlock.h"
@@ -21,6 +23,7 @@
 #include "Styling/AppStyle.h"
 #include "Tracks/MovieSceneSubTrack.h"
 #include "Widgets/Input/SCheckBox.h"
+#include "Widgets/Input/SComboButton.h"
 #include "Widgets/Input/SSearchBox.h"
 #include "Widgets/Layout/SScaleBox.h"
 
diff --git a/Engine/Plugins/Experimental/Text3D/Source/Text3D/Public/Text3DComponent.h b/Engine/Plugins/Experimental/Text3D/Source/Text3D/Public/Text3DComponent.h
--- a/Engine/Plugins/Experimental/Text3D/Source/Text3D/Public/Text3DComponent.h
+++ b/Engine/Plugins/Experimental/Text3D/Source/Text3D/Public/Text3DComponent.h
@@ -2,6 +2,7 @@
 
 #pragma once
 
+#include "Components/SceneComponent.h"
 #include "CoreMinimal.h"
 #include "BevelType.h"
 #include "Mesh.h"
diff --git a/Engine/Plugins/Experimental/VirtualProductionUtilities/Source/VPMaterialsEditor/Private/VPMaterialsEditorModule.cpp b/Engine/Plugins/Experimental/VirtualProductionUtilities/Source/VPMaterialsEditor/Private/VPMaterialsEditorModule.cpp
--- a/Engine/Plugins/Experimental/VirtualProductionUtilities/Source/VPMaterialsEditor/Private/VPMaterialsEditorModule.cpp
+++ b/Engine/Plugins/Experimental/VirtualProductionUtilities/Source/VPMaterialsEditor/Private/VPMaterialsEditorModule.cpp
@@ -1,6 +1,7 @@
 // Copyright Epic Games, Inc. All Rights Reserved.
 
 #include "VPMaterialsEditorModule.h"
+#include "Components/ActorComponent.h"
 #include "SMaterialDynamicWidgets.h"
 
 #include "MaterialList.h"
diff --git a/Engine/Plugins/Experimental/Web/WebAPI/Source/WebAPIEditor/Private/Assets/WebAPIDefinitionAssetEditorToolkit.cpp b/Engine/Plugins/Experimental/Web/WebAPI/Source/WebAPIEditor/Private/Assets/WebAPIDefinitionAssetEditorToolkit.cpp
--- a/Engine/Plugins/Experimental/Web/WebAPI/Source/WebAPIEditor/Private/Assets/WebAPIDefinitionAssetEditorToolkit.cpp
+++ b/Engine/Plugins/Experimental/Web/WebAPI/Source/WebAPIEditor/Private/Assets/WebAPIDefinitionAssetEditorToolkit.cpp
@@ -1,4 +1,4 @@
-﻿// Copyright Epic Games, Inc. All Rights Reserved.
+// Copyright Epic Games, Inc. All Rights Reserved.
 
 #include "WebAPIDefinitionAssetEditorToolkit.h"
 
@@ -6,6 +6,7 @@
 #include "EditorViewportTabContent.h"
 #include "LevelEditor.h"
 #include "LevelEditorActions.h"
+#include "PropertyEditorModule.h"
 #include "WebAPIEditorCommands.h"
 #include "WebAPIEditorModule.h"
 #include "Async/Async.h"
diff --git a/Engine/Plugins/Importers/USDImporter/Source/GeometryCacheUSD/Private/GeometryCacheUSDStream.cpp b/Engine/Plugins/Importers/USDImporter/Source/GeometryCacheUSD/Private/GeometryCacheUSDStream.cpp
--- a/Engine/Plugins/Importers/USDImporter/Source/GeometryCacheUSD/Private/GeometryCacheUSDStream.cpp
+++ b/Engine/Plugins/Importers/USDImporter/Source/GeometryCacheUSD/Private/GeometryCacheUSDStream.cpp
@@ -6,6 +6,8 @@
 #include "GeometryCacheTrackUSD.h"
 #include "HAL/IConsoleManager.h"
 #include "Misc/CoreMisc.h"
+#include "Serialization/MemoryReader.h"
+#include "Serialization/MemoryWriter.h"
 #include "USDGeomMeshConversion.h"
 
 static bool GUsdStreamCacheInDDC = true;
diff --git a/Engine/Plugins/Importers/USDImporter/Source/USDExporter/Private/AnimSequenceExporterUSD.cpp b/Engine/Plugins/Importers/USDImporter/Source/USDExporter/Private/AnimSequenceExporterUSD.cpp
--- a/Engine/Plugins/Importers/USDImporter/Source/USDExporter/Private/AnimSequenceExporterUSD.cpp
+++ b/Engine/Plugins/Importers/USDImporter/Source/USDExporter/Private/AnimSequenceExporterUSD.cpp
@@ -6,6 +6,7 @@
 #include "Engine/SkeletalMesh.h"
 #include "EngineAnalytics.h"
 #include "MaterialExporterUSD.h"
+#include "Misc/EngineVersion.h"
 #include "SkeletalMeshExporterUSDOptions.h"
 #include "UnrealUSDWrapper.h"
 #include "USDClassesModule.h"
diff --git a/Engine/Plugins/Importers/USDImporter/Source/USDExporter/Private/StaticMeshExporterUSD.cpp b/Engine/Plugins/Importers/USDImporter/Source/USDExporter/Private/StaticMeshExporterUSD.cpp
--- a/Engine/Plugins/Importers/USDImporter/Source/USDExporter/Private/StaticMeshExporterUSD.cpp
+++ b/Engine/Plugins/Importers/USDImporter/Source/USDExporter/Private/StaticMeshExporterUSD.cpp
@@ -4,7 +4,9 @@
 
 #include "EngineAnalytics.h"
 #include "MaterialExporterUSD.h"
+#include "Misc/EngineVersion.h"
 #include "StaticMeshExporterUSDOptions.h"
+#include "StaticMeshResources.h"
 #include "USDClassesModule.h"
 #include "USDConversionUtils.h"
 #include "USDGeomMeshConversion.h"
diff --git a/Engine/Plugins/Importers/USDImporter/Source/USDExporter/Private/USDConversionBlueprintContext.cpp b/Engine/Plugins/Importers/USDImporter/Source/USDExporter/Private/USDConversionBlueprintContext.cpp
--- a/Engine/Plugins/Importers/USDImporter/Source/USDExporter/Private/USDConversionBlueprintContext.cpp
+++ b/Engine/Plugins/Importers/USDImporter/Source/USDExporter/Private/USDConversionBlueprintContext.cpp
@@ -2,6 +2,7 @@
 
 #include "USDConversionBlueprintContext.h"
 
+#include "PrimitiveSceneProxy.h"
 #include "USDConversionUtils.h"
 #include "USDGeomMeshConversion.h"
 #include "USDLightConversion.h"
diff --git a/Engine/Plugins/Importers/USDImporter/Source/USDStageEditor/Private/SUSDSaveDialog.cpp b/Engine/Plugins/Importers/USDImporter/Source/USDStageEditor/Private/SUSDSaveDialog.cpp
--- a/Engine/Plugins/Importers/USDImporter/Source/USDStageEditor/Private/SUSDSaveDialog.cpp
+++ b/Engine/Plugins/Importers/USDImporter/Source/USDStageEditor/Private/SUSDSaveDialog.cpp
@@ -18,6 +18,9 @@
 #include "Widgets/Layout/SUniformGridPanel.h"
 #include "Widgets/SBoxPanel.h"
 #include "Widgets/SWindow.h"
+#include "Widgets/Text/STextBlock.h"
+#include "Widgets/Views/SHeaderRow.h"
+#include "Widgets/Views/SListView.h"
 
 #define LOCTEXT_NAMESPACE "SUsdSaveDialog"
 
diff --git a/Engine/Plugins/Importers/USDImporter/Source/USDStageEditor/Private/SUSDStage.cpp b/Engine/Plugins/Importers/USDImporter/Source/USDStageEditor/Private/SUSDStage.cpp
--- a/Engine/Plugins/Importers/USDImporter/Source/USDStageEditor/Private/SUSDStage.cpp
+++ b/Engine/Plugins/Importers/USDImporter/Source/USDStageEditor/Private/SUSDStage.cpp
@@ -39,6 +39,8 @@
 #include "ScopedTransaction.h"
 #include "Styling/AppStyle.h"
 #include "UObject/StrongObjectPtr.h"
+#include "Widgets/Input/SComboBox.h"
+#include "Widgets/Input/SComboButton.h"
 #include "Widgets/Input/SEditableTextBox.h"
 #include "Widgets/Input/SButton.h"
 #include "Widgets/Input/SSpinBox.h"
diff --git a/Engine/Plugins/Importers/USDImporter/Source/USDStageEditor/Private/Widgets/SUSDPrimPropertiesList.cpp b/Engine/Plugins/Importers/USDImporter/Source/USDStageEditor/Private/Widgets/SUSDPrimPropertiesList.cpp
--- a/Engine/Plugins/Importers/USDImporter/Source/USDStageEditor/Private/Widgets/SUSDPrimPropertiesList.cpp
+++ b/Engine/Plugins/Importers/USDImporter/Source/USDStageEditor/Private/Widgets/SUSDPrimPropertiesList.cpp
@@ -15,6 +15,7 @@
 #include "Modules/ModuleManager.h"
 #include "ScopedTransaction.h"
 #include "Widgets/Input/SCheckBox.h"
+#include "Widgets/Input/SComboBox.h"
 #include "Widgets/Input/SEditableTextBox.h"
 #include "Widgets/Input/SNumericEntryBox.h"
 
diff --git a/Engine/Plugins/Media/BinkMedia/Source/BinkMediaPlayer/Private/Assets/BinkMediaPlayer.cpp b/Engine/Plugins/Media/BinkMedia/Source/BinkMediaPlayer/Private/Assets/BinkMediaPlayer.cpp
--- a/Engine/Plugins/Media/BinkMedia/Source/BinkMediaPlayer/Private/Assets/BinkMediaPlayer.cpp
+++ b/Engine/Plugins/Media/BinkMedia/Source/BinkMediaPlayer/Private/Assets/BinkMediaPlayer.cpp
@@ -8,6 +8,7 @@
 #include "BinkMediaTexture.h"
 #include "BinkFunctionLibrary.h"
 
+#include "Engine/GameViewportClient.h"
 #include "Misc/Paths.h"
 #include "RenderingThread.h"
 
@@ -17,6 +18,7 @@
 #include "Internationalization/Internationalization.h"
 #include "Internationalization/Culture.h"
 #include "SubtitleManager.h"
+#include "UObject/UObjectIterator.h"
 
 //#include "Media/Public/IMediaEventSink.h" //For EMediaEvent
 
diff --git a/Engine/Plugins/Media/BinkMedia/Source/BinkMediaPlayer/Private/BinkFunctionLibrary.cpp b/Engine/Plugins/Media/BinkMedia/Source/BinkMediaPlayer/Private/BinkFunctionLibrary.cpp
--- a/Engine/Plugins/Media/BinkMedia/Source/BinkMediaPlayer/Private/BinkFunctionLibrary.cpp
+++ b/Engine/Plugins/Media/BinkMedia/Source/BinkMediaPlayer/Private/BinkFunctionLibrary.cpp
@@ -4,6 +4,7 @@
 #include "BinkFunctionLibrary.h"
 
 #include "BinkMediaPlayerPrivate.h"
+#include "Engine/GameViewportClient.h"
 #include "Misc/Paths.h"
 #include "Rendering/RenderingCommon.h"
 #include "RenderingThread.h"
diff --git a/Engine/Plugins/Media/BinkMedia/Source/BinkMediaPlayer/Private/BinkMediaPlayerModule.cpp b/Engine/Plugins/Media/BinkMedia/Source/BinkMediaPlayer/Private/BinkMediaPlayerModule.cpp
--- a/Engine/Plugins/Media/BinkMedia/Source/BinkMediaPlayer/Private/BinkMediaPlayerModule.cpp
+++ b/Engine/Plugins/Media/BinkMedia/Source/BinkMediaPlayer/Private/BinkMediaPlayerModule.cpp
@@ -3,14 +3,15 @@
 
 #include "BinkFunctionLibrary.h"
 #include "BinkMediaPlayerPrivate.h"
+#include "Engine/GameViewportClient.h"
+#include "Internationalization/Internationalization.h"
 #include "Runtime/Core/Public/Misc/Paths.h"
 #include "Interfaces/IPluginManager.h"
 #include "Misc/CoreDelegates.h"
+#include "UObject/UObjectIterator.h"
 
 DEFINE_LOG_CATEGORY(LogBink);
 
-#define LOCTEXT_NAMESPACE "BinkMediaPlayerModule"
-
 TSharedPtr<FBinkMovieStreamer, ESPMode::ThreadSafe> MovieStreamer;
 
 TArray< FTexture2DRHIRef > BinkActiveTextureRefs;
@@ -36,6 +37,8 @@ class UFactory;
 
 #include "BinkMediaPlayer.h"
 
+#define LOCTEXT_NAMESPACE "BinkMediaPlayerModule"
+
 struct FBinkMoviePlayerSettingsDetails : public IDetailCustomization {
 	static TSharedRef<IDetailCustomization> MakeInstance() { return MakeShareable(new FBinkMoviePlayerSettingsDetails); }
 	virtual void CustomizeDetails( IDetailLayoutBuilder& DetailLayout ) override {
diff --git a/Engine/Plugins/Media/BinkMedia/Source/BinkMediaPlayerEditor/Private/BinkMediaPlayerCustomization.cpp b/Engine/Plugins/Media/BinkMedia/Source/BinkMediaPlayerEditor/Private/BinkMediaPlayerCustomization.cpp
--- a/Engine/Plugins/Media/BinkMedia/Source/BinkMediaPlayerEditor/Private/BinkMediaPlayerCustomization.cpp
+++ b/Engine/Plugins/Media/BinkMedia/Source/BinkMediaPlayerEditor/Private/BinkMediaPlayerCustomization.cpp
@@ -6,10 +6,12 @@
 #include "BinkMediaPlayer.h"
 #include "BinkMediaPlayerEditorPrivate.h"
 #include "DesktopPlatformModule.h"
+#include "Framework/Application/SlateApplication.h"
 #include "Widgets/Images/SImage.h"
 #include "EditorDirectories.h"
 #include "DetailLayoutBuilder.h"
 #include "DetailCategoryBuilder.h"
+#include "Widgets/Text/STextBlock.h"
 
 #define LOCTEXT_NAMESPACE "FBinkMediaPlayerCustomization"
 
diff --git a/Engine/Plugins/MovieScene/ReplayTracks/Source/ReplayTracks/Private/Systems/MovieSceneReplaySystem.cpp b/Engine/Plugins/MovieScene/ReplayTracks/Source/ReplayTracks/Private/Systems/MovieSceneReplaySystem.cpp
--- a/Engine/Plugins/MovieScene/ReplayTracks/Source/ReplayTracks/Private/Systems/MovieSceneReplaySystem.cpp
+++ b/Engine/Plugins/MovieScene/ReplayTracks/Source/ReplayTracks/Private/Systems/MovieSceneReplaySystem.cpp
@@ -13,6 +13,7 @@
 #include "GameFramework/PlayerController.h"
 #include "GameFramework/SpectatorPawn.h"
 #include "GameFramework/SpectatorPawnMovement.h"
+#include "GameFramework/WorldSettings.h"
 #include "IMovieScenePlayer.h"
 #include "Misc/CoreMiscDefines.h"
 #include "MovieSceneFwd.h"
diff --git a/Engine/Plugins/Runtime/AudioCapture/Source/AudioCapture/Private/AudioCaptureComponent.cpp b/Engine/Plugins/Runtime/AudioCapture/Source/AudioCapture/Private/AudioCaptureComponent.cpp
--- a/Engine/Plugins/Runtime/AudioCapture/Source/AudioCapture/Private/AudioCaptureComponent.cpp
+++ b/Engine/Plugins/Runtime/AudioCapture/Source/AudioCapture/Private/AudioCaptureComponent.cpp
@@ -1,6 +1,7 @@
 // Copyright Epic Games, Inc. All Rights Reserved.
 
 #include "AudioCaptureComponent.h"
+#include "Audio.h"
 #include "AudioAnalytics.h"
 
 static const unsigned int MaxBufferSize = 2 * 5 * 48000;
diff --git a/Engine/Plugins/Runtime/HairStrands/Source/HairStrandsCore/Private/GroomAsset.cpp b/Engine/Plugins/Runtime/HairStrands/Source/HairStrandsCore/Private/GroomAsset.cpp
--- a/Engine/Plugins/Runtime/HairStrands/Source/HairStrandsCore/Private/GroomAsset.cpp
+++ b/Engine/Plugins/Runtime/HairStrands/Source/HairStrandsCore/Private/GroomAsset.cpp
@@ -3,18 +3,21 @@
 #include "GroomAsset.h"
 
 #include "Async/Async.h"
+#include "Engine/Texture2D.h"
 #include "EngineUtils.h"
 #include "GroomAssetImportData.h"
 #include "GroomBuilder.h"
 #include "HairCardsBuilder.h"
 #include "GroomImportOptions.h"
 #include "GroomSettings.h"
+#include "Materials/MaterialInterface.h"
 #include "RenderingThread.h"
 #include "Engine/AssetUserData.h"
 #include "HairStrandsVertexFactory.h"
 #include "HAL/LowLevelMemTracker.h"
 #include "Misc/App.h"
 #include "Misc/Paths.h"
+#include "RHIStaticStates.h"
 #include "Serialization/LargeMemoryReader.h"
 #include "Serialization/LargeMemoryWriter.h"
 #include "Serialization/MemoryReader.h"
@@ -34,6 +37,7 @@
 #include "GroomDeformerBuilder.h"
 #include "HairCardsVertexFactory.h"
 #include "PSOPrecache.h"
+#include "UObject/UObjectIterator.h"
 
 #if WITH_EDITORONLY_DATA
 #include "DerivedDataCache.h"
diff --git a/Engine/Plugins/Runtime/HairStrands/Source/HairStrandsCore/Private/GroomBindingAsset.cpp b/Engine/Plugins/Runtime/HairStrands/Source/HairStrandsCore/Private/GroomBindingAsset.cpp
--- a/Engine/Plugins/Runtime/HairStrands/Source/HairStrandsCore/Private/GroomBindingAsset.cpp
+++ b/Engine/Plugins/Runtime/HairStrands/Source/HairStrandsCore/Private/GroomBindingAsset.cpp
@@ -8,6 +8,7 @@
 #include "GroomAsset.h"
 #include "GroomBindingBuilder.h"
 #include "HAL/LowLevelMemTracker.h"
+#include "UObject/AnimObjectVersion.h"
 #include "UObject/ObjectSaveContext.h"
 
 #include UE_INLINE_GENERATED_CPP_BY_NAME(GroomBindingAsset)
diff --git a/Engine/Plugins/Runtime/HairStrands/Source/HairStrandsCore/Private/GroomBuilder.cpp b/Engine/Plugins/Runtime/HairStrands/Source/HairStrandsCore/Private/GroomBuilder.cpp
--- a/Engine/Plugins/Runtime/HairStrands/Source/HairStrandsCore/Private/GroomBuilder.cpp
+++ b/Engine/Plugins/Runtime/HairStrands/Source/HairStrandsCore/Private/GroomBuilder.cpp
@@ -10,6 +10,7 @@
 #include "RenderGraphBuilder.h"
 #include "RenderGraphUtils.h"
 #include "GlobalShader.h"
+#include "RHIGPUReadback.h"
 #include "ShaderParameterStruct.h"
 
 #include "Async/ParallelFor.h"
diff --git a/Engine/Plugins/Runtime/HairStrands/Source/HairStrandsCore/Private/GroomComponent.cpp b/Engine/Plugins/Runtime/HairStrands/Source/HairStrandsCore/Private/GroomComponent.cpp
--- a/Engine/Plugins/Runtime/HairStrands/Source/HairStrandsCore/Private/GroomComponent.cpp
+++ b/Engine/Plugins/Runtime/HairStrands/Source/HairStrandsCore/Private/GroomComponent.cpp
@@ -6,6 +6,7 @@
 #include "MaterialShared.h"
 #include "Engine/CollisionProfile.h"
 #include "Engine/Engine.h"
+#include "NiagaraSystem.h"
 #include "PrimitiveSceneProxy.h"
 #include "PhysicsEngine/PhysicsAsset.h"
 #include "HairStrandsRendering.h"
diff --git a/Engine/Plugins/Runtime/HairStrands/Source/HairStrandsCore/Private/GroomDebug.cpp b/Engine/Plugins/Runtime/HairStrands/Source/HairStrandsCore/Private/GroomDebug.cpp
--- a/Engine/Plugins/Runtime/HairStrands/Source/HairStrandsCore/Private/GroomDebug.cpp
+++ b/Engine/Plugins/Runtime/HairStrands/Source/HairStrandsCore/Private/GroomDebug.cpp
@@ -9,6 +9,7 @@
 #include "CachedGeometry.h"
 #include "GroomGeometryCache.h"
 #include "Components/SkeletalMeshComponent.h"
+#include "RHIStaticStates.h"
 #include "SkeletalRenderPublic.h"
 #include "UnrealEngine.h"
 #include "SystemTextures.h"
diff --git a/Engine/Plugins/Runtime/HairStrands/Source/HairStrandsCore/Private/GroomResources.cpp b/Engine/Plugins/Runtime/HairStrands/Source/HairStrandsCore/Private/GroomResources.cpp
--- a/Engine/Plugins/Runtime/HairStrands/Source/HairStrandsCore/Private/GroomResources.cpp
+++ b/Engine/Plugins/Runtime/HairStrands/Source/HairStrandsCore/Private/GroomResources.cpp
@@ -1,6 +1,7 @@
 // Copyright Epic Games, Inc. All Rights Reserved.
 
 #include "GroomResources.h"
+#include "Engine/Texture2D.h"
 #include "EngineUtils.h"
 #include "GroomAssetImportData.h"
 #include "GroomBuilder.h"
@@ -10,11 +11,13 @@
 #include "Engine/AssetUserData.h"
 #include "HairStrandsVertexFactory.h"
 #include "Misc/Paths.h"
+#include "RHIStaticStates.h"
 #include "Serialization/LargeMemoryReader.h"
 #include "Serialization/LargeMemoryWriter.h"
 #include "Serialization/MemoryReader.h"
 #include "Serialization/MemoryWriter.h"
 #include "Serialization/BulkData.h"
+#include "TextureResource.h"
 #include "UObject/PhysicsObjectVersion.h"
 #include "UObject/ReleaseObjectVersion.h"
 #include "NiagaraSystem.h"
diff --git a/Engine/Plugins/Runtime/HairStrands/Source/HairStrandsCore/Private/GroomTextureBuilder.cpp b/Engine/Plugins/Runtime/HairStrands/Source/HairStrandsCore/Private/GroomTextureBuilder.cpp
--- a/Engine/Plugins/Runtime/HairStrands/Source/HairStrandsCore/Private/GroomTextureBuilder.cpp
+++ b/Engine/Plugins/Runtime/HairStrands/Source/HairStrandsCore/Private/GroomTextureBuilder.cpp
@@ -1,18 +1,22 @@
 // Copyright Epic Games, Inc. All Rights Reserved.
 
 #include "GroomTextureBuilder.h"
+#include "Engine/Texture2D.h"
 #include "HairStrandsCore.h"
 #include "Rendering/SkeletalMeshRenderData.h"
 #include "Rendering/SkeletalMeshLODRenderData.h"
 #include "GroomAsset.h"
 #include "RenderGraphBuilder.h"
 #include "RenderGraphUtils.h"
+#include "RHI.h"
+#include "RHIStaticStates.h"
 #include "ShaderParameterStruct.h"
 #include "GlobalShader.h"
 #include "ShaderPrint.h"
 #include "CommonRenderResources.h"
 #include "HairStrandsMeshProjection.h"
 #include "Engine/StaticMesh.h"
+#include "TextureResource.h"
 
 static int32 GHairStrandsTextureDilationPassCount = 8;
 static FAutoConsoleVariableRef CVarHairStrandsTextureDilationPassCount(TEXT("r.HairStrands.Textures.DilationCount"), GHairStrandsTextureDilationPassCount, TEXT("Number of dilation pass run onto the generated hair strands textures (Default:8)."));
diff --git a/Engine/Plugins/Runtime/HairStrands/Source/HairStrandsCore/Private/HairCardsBuilder.cpp b/Engine/Plugins/Runtime/HairStrands/Source/HairStrandsCore/Private/HairCardsBuilder.cpp
--- a/Engine/Plugins/Runtime/HairStrands/Source/HairStrandsCore/Private/HairCardsBuilder.cpp
+++ b/Engine/Plugins/Runtime/HairStrands/Source/HairStrandsCore/Private/HairCardsBuilder.cpp
@@ -1,6 +1,7 @@
 // Copyright Epic Games, Inc. All Rights Reserved. 
 
 #include "HairCardsBuilder.h"
+#include "Engine/Texture2D.h"
 #include "HairStrandsCore.h"
 #include "HairStrandsDatas.h"
 #include "HairCardsDatas.h"
@@ -9,6 +10,7 @@
 #include "Math/Box.h"
 #include "RenderGraphBuilder.h"
 #include "RenderGraphUtils.h"
+#include "RHIStaticStates.h"
 #include "ShaderParameterStruct.h"
 #include "GlobalShader.h"
 #include "GroomAsset.h"
@@ -26,6 +28,7 @@
 #include "StaticMeshAttributes.h"
 #include "PhysicsEngine/BodySetup.h"
 #include "StaticMeshOperations.h"
+#include "TextureResource.h"
 
 #if WITH_EDITOR
 
diff --git a/Engine/Plugins/Runtime/HairStrands/Source/HairStrandsEditor/Private/GroomComponentDetailsCustomization.h b/Engine/Plugins/Runtime/HairStrands/Source/HairStrandsEditor/Private/GroomComponentDetailsCustomization.h
--- a/Engine/Plugins/Runtime/HairStrands/Source/HairStrandsEditor/Private/GroomComponentDetailsCustomization.h
+++ b/Engine/Plugins/Runtime/HairStrands/Source/HairStrandsEditor/Private/GroomComponentDetailsCustomization.h
@@ -10,6 +10,8 @@
 
 class IDetailLayoutBuilder;
 class IDetailCategoryBuilder;
+class IDetailChildrenBuilder;
+class IPropertyHandle;
 
 //////////////////////////////////////////////////////////////////////////
 // FGroomComponentDetailsCustomization
diff --git a/Engine/Plugins/Runtime/HairStrands/Source/HairStrandsEditor/Public/GroomCustomAssetEditorToolkit.h b/Engine/Plugins/Runtime/HairStrands/Source/HairStrandsEditor/Public/GroomCustomAssetEditorToolkit.h
--- a/Engine/Plugins/Runtime/HairStrands/Source/HairStrandsEditor/Public/GroomCustomAssetEditorToolkit.h
+++ b/Engine/Plugins/Runtime/HairStrands/Source/HairStrandsEditor/Public/GroomCustomAssetEditorToolkit.h
@@ -10,6 +10,7 @@
 #include "Widgets/SGroomEditorViewport.h"
 #include "HairStrandsInterface.h"
 
+class FAssetThumbnailPool;
 class UGroomAsset; 
 class UGroomComponent;
 class IDetailsView;
diff --git a/Engine/Plugins/Runtime/MassGameplay/Source/MassMovementEditor/Public/MassMovementStyleRefDetails.h b/Engine/Plugins/Runtime/MassGameplay/Source/MassMovementEditor/Public/MassMovementStyleRefDetails.h
--- a/Engine/Plugins/Runtime/MassGameplay/Source/MassMovementEditor/Public/MassMovementStyleRefDetails.h
+++ b/Engine/Plugins/Runtime/MassGameplay/Source/MassMovementEditor/Public/MassMovementStyleRefDetails.h
@@ -5,6 +5,7 @@
 #include "IPropertyTypeCustomization.h"
 
 class IPropertyHandle;
+class SWidget;
 
 /**
  * Type customization for FMassMovementStyleRef.
diff --git a/Engine/Plugins/Runtime/MassGameplay/Source/MassSpawner/Private/MassSpawner.cpp b/Engine/Plugins/Runtime/MassGameplay/Source/MassSpawner/Private/MassSpawner.cpp
--- a/Engine/Plugins/Runtime/MassGameplay/Source/MassSpawner/Private/MassSpawner.cpp
+++ b/Engine/Plugins/Runtime/MassGameplay/Source/MassSpawner/Private/MassSpawner.cpp
@@ -1,6 +1,7 @@
 // Copyright Epic Games, Inc. All Rights Reserved.
 
 #include "MassSpawner.h"
+#include "Engine/Engine.h"
 #include "Engine/World.h"
 #include "UObject/ConstructorHelpers.h"
 #include "Components/SceneComponent.h"
diff --git a/Engine/Plugins/Runtime/MixedRealityCaptureFramework/Source/MixedRealityCaptureFramework/Private/MrcGarbageMatteCaptureComponent.cpp b/Engine/Plugins/Runtime/MixedRealityCaptureFramework/Source/MixedRealityCaptureFramework/Private/MrcGarbageMatteCaptureComponent.cpp
--- a/Engine/Plugins/Runtime/MixedRealityCaptureFramework/Source/MixedRealityCaptureFramework/Private/MrcGarbageMatteCaptureComponent.cpp
+++ b/Engine/Plugins/Runtime/MixedRealityCaptureFramework/Source/MixedRealityCaptureFramework/Private/MrcGarbageMatteCaptureComponent.cpp
@@ -6,6 +6,7 @@
 #include "Engine/StaticMesh.h"
 #include "Engine/StaticMeshActor.h"
 #include "Engine/TextureRenderTarget2D.h"
+#include "Engine/World.h"
 #include "IMrcFocalDriver.h"
 #include "Materials/Material.h"
 #include "MrcCalibrationData.h"
diff --git a/Engine/Plugins/Runtime/NetworkPredictionExtras/Source/NetworkPredictionExtras/Public/CharacterMotionSimulation.h b/Engine/Plugins/Runtime/NetworkPredictionExtras/Source/NetworkPredictionExtras/Public/CharacterMotionSimulation.h
--- a/Engine/Plugins/Runtime/NetworkPredictionExtras/Source/NetworkPredictionExtras/Public/CharacterMotionSimulation.h
+++ b/Engine/Plugins/Runtime/NetworkPredictionExtras/Source/NetworkPredictionExtras/Public/CharacterMotionSimulation.h
@@ -2,6 +2,7 @@
 
 #pragma once
 
+#include "Engine/HitResult.h"
 #include "Misc/StringBuilder.h"
 #include "BaseMovementSimulation.h"
 #include "NetworkPredictionReplicationProxy.h"
diff --git a/Engine/Plugins/Runtime/Soundscape/Source/SoundScape/Private/SoundscapeSubsystem.cpp b/Engine/Plugins/Runtime/Soundscape/Source/SoundScape/Private/SoundscapeSubsystem.cpp
--- a/Engine/Plugins/Runtime/Soundscape/Source/SoundScape/Private/SoundscapeSubsystem.cpp
+++ b/Engine/Plugins/Runtime/Soundscape/Source/SoundScape/Private/SoundscapeSubsystem.cpp
@@ -2,11 +2,13 @@
 
 
 #include "SoundscapeSubsystem.h"
+#include "DrawDebugHelpers.h"
 #include "SoundscapeSettings.h"
 #include "SoundScapePalette.h"
 #include "AudioDevice.h"
 #include "HAL/IConsoleManager.h"
 #include "TimerManager.h"
+#include "UObject/UObjectIterator.h"
 
 static FAutoConsoleCommandWithWorld GResetSoundscape(
 	TEXT("soundscape.ResetSoundscape"),
diff --git a/Engine/Plugins/Runtime/Steam/SteamVR/Source/SteamVR/Private/SteamVRAssetManager.cpp b/Engine/Plugins/Runtime/Steam/SteamVR/Source/SteamVR/Private/SteamVRAssetManager.cpp
--- a/Engine/Plugins/Runtime/Steam/SteamVR/Source/SteamVR/Private/SteamVRAssetManager.cpp
+++ b/Engine/Plugins/Runtime/Steam/SteamVR/Source/SteamVR/Private/SteamVRAssetManager.cpp
@@ -2,7 +2,9 @@
 
 #include "SteamVRAssetManager.h"
 #include "ISteamVRPlugin.h" // STEAMVR_SUPPORTED_PLATFORMS
+#include "Materials/Material.h"
 #include "ProceduralMeshComponent.h"
+#include "TextureResource.h"
 #include "Tickable.h" // for FTickableGameObject
 #include "Engine/Engine.h" // for GEngine
 #include "Engine/World.h"
diff --git a/Engine/Plugins/Runtime/XR/XRVisualization/Source/XRVisualization/Private/XRVisualizationFunctionLibrary.cpp b/Engine/Plugins/Runtime/XR/XRVisualization/Source/XRVisualization/Private/XRVisualizationFunctionLibrary.cpp
--- a/Engine/Plugins/Runtime/XR/XRVisualization/Source/XRVisualization/Private/XRVisualizationFunctionLibrary.cpp
+++ b/Engine/Plugins/Runtime/XR/XRVisualization/Source/XRVisualization/Private/XRVisualizationFunctionLibrary.cpp
@@ -1,9 +1,11 @@
 // Copyright Epic Games, Inc. All Rights Reserved.
 
 #include "XRVisualizationFunctionLibrary.h"
+#include "Engine/Level.h"
 #include "EngineGlobals.h"
 #include "Engine/Engine.h"
 #include "Features/IModularFeatures.h"
+#include "Materials/MaterialInterface.h"
 #include "XRVisualizationModule.h"
 #include "Components/StaticMeshComponent.h"
 #include "DrawDebugHelpers.h"
diff --git a/Engine/Plugins/Runtime/ZoneGraph/Source/ZoneGraph/Private/ZoneGraphData.cpp b/Engine/Plugins/Runtime/ZoneGraph/Source/ZoneGraph/Private/ZoneGraphData.cpp
--- a/Engine/Plugins/Runtime/ZoneGraph/Source/ZoneGraph/Private/ZoneGraphData.cpp
+++ b/Engine/Plugins/Runtime/ZoneGraph/Source/ZoneGraph/Private/ZoneGraphData.cpp
@@ -1,6 +1,8 @@
 // Copyright Epic Games, Inc. All Rights Reserved.
 
 #include "ZoneGraphData.h"
+#include "Engine/Level.h"
+#include "Engine/World.h"
 #include "ZoneGraphTypes.h"
 #include "ZoneGraphRenderingComponent.h"
 #include "ZoneGraphSubsystem.h"
diff --git a/Engine/Plugins/Runtime/ZoneGraph/Source/ZoneGraphEditor/Private/ZoneGraphPropertyUtils.h b/Engine/Plugins/Runtime/ZoneGraph/Source/ZoneGraphEditor/Private/ZoneGraphPropertyUtils.h
--- a/Engine/Plugins/Runtime/ZoneGraph/Source/ZoneGraphEditor/Private/ZoneGraphPropertyUtils.h
+++ b/Engine/Plugins/Runtime/ZoneGraph/Source/ZoneGraphEditor/Private/ZoneGraphPropertyUtils.h
@@ -3,6 +3,7 @@
 #include "CoreMinimal.h"
 #include "PropertyEditorModule.h"
 #include "Editor.h"
+#include "PropertyHandle.h"
 #include "ZoneGraphTypes.h"
 #include "Misc/Guid.h"
 
diff --git a/Engine/Plugins/Runtime/ZoneGraph/Source/ZoneGraphEditor/Private/ZoneGraphTagDetails.h b/Engine/Plugins/Runtime/ZoneGraph/Source/ZoneGraphEditor/Private/ZoneGraphTagDetails.h
--- a/Engine/Plugins/Runtime/ZoneGraph/Source/ZoneGraphEditor/Private/ZoneGraphTagDetails.h
+++ b/Engine/Plugins/Runtime/ZoneGraph/Source/ZoneGraphEditor/Private/ZoneGraphTagDetails.h
@@ -6,6 +6,7 @@
 #include "ZoneGraphTypes.h"
 
 class IPropertyHandle;
+class SWidget;
 
 /**
  * Type customization for FZoneGraphTag.
diff --git a/Engine/Plugins/Runtime/ZoneGraph/Source/ZoneGraphEditor/Private/ZoneGraphTagFilterDetails.h b/Engine/Plugins/Runtime/ZoneGraph/Source/ZoneGraphEditor/Private/ZoneGraphTagFilterDetails.h
--- a/Engine/Plugins/Runtime/ZoneGraph/Source/ZoneGraphEditor/Private/ZoneGraphTagFilterDetails.h
+++ b/Engine/Plugins/Runtime/ZoneGraph/Source/ZoneGraphEditor/Private/ZoneGraphTagFilterDetails.h
@@ -1,6 +1,7 @@
 // Copyright Epic Games, Inc. All Rights Reserved.
 #pragma once
 #include "CoreMinimal.h"
+#include "Layout/Visibility.h"
 #include "UnrealClient.h"
 #include "IPropertyTypeCustomization.h"
 #include "ZoneGraphTypes.h"
diff --git a/Engine/Plugins/Runtime/ZoneGraph/Source/ZoneGraphEditor/Private/ZoneGraphTagInfoDetails.h b/Engine/Plugins/Runtime/ZoneGraph/Source/ZoneGraphEditor/Private/ZoneGraphTagInfoDetails.h
--- a/Engine/Plugins/Runtime/ZoneGraph/Source/ZoneGraphEditor/Private/ZoneGraphTagInfoDetails.h
+++ b/Engine/Plugins/Runtime/ZoneGraph/Source/ZoneGraphEditor/Private/ZoneGraphTagInfoDetails.h
@@ -5,8 +5,11 @@
 #include "IPropertyTypeCustomization.h"
 #include "ZoneGraphTypes.h"
 
+class FReply;
 class IPropertyHandle;
 class SColorBlock;
+struct FGeometry;
+struct FPointerEvent;
 
 /**
  * Type customization for FZoneGraphTagInfo.
diff --git a/Engine/Plugins/Runtime/ZoneGraph/Source/ZoneGraphEditor/Private/ZoneGraphTessellationSettingsDetails.h b/Engine/Plugins/Runtime/ZoneGraph/Source/ZoneGraphEditor/Private/ZoneGraphTessellationSettingsDetails.h
--- a/Engine/Plugins/Runtime/ZoneGraph/Source/ZoneGraphEditor/Private/ZoneGraphTessellationSettingsDetails.h
+++ b/Engine/Plugins/Runtime/ZoneGraph/Source/ZoneGraphEditor/Private/ZoneGraphTessellationSettingsDetails.h
@@ -1,6 +1,7 @@
 // Copyright Epic Games, Inc. All Rights Reserved.
 #pragma once
 #include "CoreMinimal.h"
+#include "Layout/Visibility.h"
 #include "UnrealClient.h"
 #include "IPropertyTypeCustomization.h"
 #include "ZoneGraphTypes.h"
diff --git a/Engine/Plugins/Runtime/ZoneGraph/Source/ZoneGraphEditor/Private/ZoneShapeComponentDetails.cpp b/Engine/Plugins/Runtime/ZoneGraph/Source/ZoneGraphEditor/Private/ZoneShapeComponentDetails.cpp
--- a/Engine/Plugins/Runtime/ZoneGraph/Source/ZoneGraphEditor/Private/ZoneShapeComponentDetails.cpp
+++ b/Engine/Plugins/Runtime/ZoneGraph/Source/ZoneGraphEditor/Private/ZoneShapeComponentDetails.cpp
@@ -6,6 +6,7 @@
 #include "Layout/Visibility.h"
 #include "Misc/Attribute.h"
 #include "Widgets/DeclarativeSyntaxSupport.h"
+#include "Widgets/Input/SComboButton.h"
 #include "Widgets/SWidget.h"
 #include "Widgets/Text/STextBlock.h"
 #include "Widgets/Layout/SBox.h"
diff --git a/Engine/Plugins/Runtime/ZoneGraph/Source/ZoneGraphEditor/Private/ZoneShapeComponentVisualizer.cpp b/Engine/Plugins/Runtime/ZoneGraph/Source/ZoneGraphEditor/Private/ZoneShapeComponentVisualizer.cpp
--- a/Engine/Plugins/Runtime/ZoneGraph/Source/ZoneGraphEditor/Private/ZoneShapeComponentVisualizer.cpp
+++ b/Engine/Plugins/Runtime/ZoneGraph/Source/ZoneGraphEditor/Private/ZoneShapeComponentVisualizer.cpp
@@ -8,6 +8,7 @@
 #include "Framework/Commands/Commands.h"
 #include "Framework/Commands/UICommandList.h"
 #include "Framework/MultiBox/MultiBoxBuilder.h"
+#include "Settings/LevelEditorViewportSettings.h"
 #include "Styling/AppStyle.h"
 #include "Editor.h"
 #include "EditorViewportClient.h"
diff --git a/Engine/Plugins/Runtime/ZoneGraph/Source/ZoneGraphEditor/Private/ZoneShapeComponentVisualizer.h b/Engine/Plugins/Runtime/ZoneGraph/Source/ZoneGraphEditor/Private/ZoneShapeComponentVisualizer.h
--- a/Engine/Plugins/Runtime/ZoneGraph/Source/ZoneGraphEditor/Private/ZoneShapeComponentVisualizer.h
+++ b/Engine/Plugins/Runtime/ZoneGraph/Source/ZoneGraphEditor/Private/ZoneShapeComponentVisualizer.h
@@ -2,10 +2,12 @@
 
 #pragma once
 
+#include "ComponentVisualizer.h"
 #include "CoreMinimal.h"
 #include "InputCoreTypes.h"
 #include "HitProxies.h"
 #include "ComponentVisualizer.h"
+#include "UObject/GCObject.h"
 #include "ZoneShapeComponent.h"
 #include "ZoneShapeComponentVisualizer.generated.h"
 
diff --git a/Engine/Plugins/Runtime/nDisplay/Source/DisplayCluster/Private/Game/EngineClasses/Scene/DisplayClusterCameraComponent.cpp b/Engine/Plugins/Runtime/nDisplay/Source/DisplayCluster/Private/Game/EngineClasses/Scene/DisplayClusterCameraComponent.cpp
--- a/Engine/Plugins/Runtime/nDisplay/Source/DisplayCluster/Private/Game/EngineClasses/Scene/DisplayClusterCameraComponent.cpp
+++ b/Engine/Plugins/Runtime/nDisplay/Source/DisplayCluster/Private/Game/EngineClasses/Scene/DisplayClusterCameraComponent.cpp
@@ -3,6 +3,7 @@
 #include "Components/DisplayClusterCameraComponent.h"
 #include "Components/BillboardComponent.h"
 
+#include "Engine/CollisionProfile.h"
 #include "Engine/Texture2D.h"
 #include "UObject/ConstructorHelpers.h"
 
diff --git a/Engine/Plugins/Runtime/nDisplay/Source/DisplayCluster/Private/Game/EngineClasses/Scene/DisplayClusterOriginComponent.cpp b/Engine/Plugins/Runtime/nDisplay/Source/DisplayCluster/Private/Game/EngineClasses/Scene/DisplayClusterOriginComponent.cpp
--- a/Engine/Plugins/Runtime/nDisplay/Source/DisplayCluster/Private/Game/EngineClasses/Scene/DisplayClusterOriginComponent.cpp
+++ b/Engine/Plugins/Runtime/nDisplay/Source/DisplayCluster/Private/Game/EngineClasses/Scene/DisplayClusterOriginComponent.cpp
@@ -3,6 +3,7 @@
 #include "Components/DisplayClusterOriginComponent.h"
 #include "Components/StaticMeshComponent.h"
 
+#include "Engine/CollisionProfile.h"
 #include "Engine/StaticMesh.h"
 #include "UObject/ConstructorHelpers.h"
 
diff --git a/Engine/Plugins/Runtime/nDisplay/Source/DisplayCluster/Private/Game/EngineClasses/Scene/DisplayClusterWidgetComponent.cpp b/Engine/Plugins/Runtime/nDisplay/Source/DisplayCluster/Private/Game/EngineClasses/Scene/DisplayClusterWidgetComponent.cpp
--- a/Engine/Plugins/Runtime/nDisplay/Source/DisplayCluster/Private/Game/EngineClasses/Scene/DisplayClusterWidgetComponent.cpp
+++ b/Engine/Plugins/Runtime/nDisplay/Source/DisplayCluster/Private/Game/EngineClasses/Scene/DisplayClusterWidgetComponent.cpp
@@ -2,6 +2,7 @@
 
 #include "DisplayClusterWidgetComponent.h"
 
+#include "Engine/Engine.h"
 #include "Slate/WidgetRenderer.h"
 
 UDisplayClusterWidgetComponent::UDisplayClusterWidgetComponent()
diff --git a/Engine/Plugins/Runtime/nDisplay/Source/DisplayCluster/Private/Render/Containers/DisplayClusterRender_MeshComponent.cpp b/Engine/Plugins/Runtime/nDisplay/Source/DisplayCluster/Private/Render/Containers/DisplayClusterRender_MeshComponent.cpp
--- a/Engine/Plugins/Runtime/nDisplay/Source/DisplayCluster/Private/Render/Containers/DisplayClusterRender_MeshComponent.cpp
+++ b/Engine/Plugins/Runtime/nDisplay/Source/DisplayCluster/Private/Render/Containers/DisplayClusterRender_MeshComponent.cpp
@@ -7,6 +7,7 @@
 
 #include "Engine/StaticMesh.h"
 #include "ProceduralMeshComponent.h"
+#include "RenderingThread.h"
 
 //-------------------------------------------------------------------------------
 //
diff --git a/Engine/Plugins/Runtime/nDisplay/Source/DisplayCluster/Private/Render/Synchronization/DisplayClusterRenderSyncPolicyBase.cpp b/Engine/Plugins/Runtime/nDisplay/Source/DisplayCluster/Private/Render/Synchronization/DisplayClusterRenderSyncPolicyBase.cpp
--- a/Engine/Plugins/Runtime/nDisplay/Source/DisplayCluster/Private/Render/Synchronization/DisplayClusterRenderSyncPolicyBase.cpp
+++ b/Engine/Plugins/Runtime/nDisplay/Source/DisplayCluster/Private/Render/Synchronization/DisplayClusterRenderSyncPolicyBase.cpp
@@ -1,6 +1,7 @@
 // Copyright Epic Games, Inc. All Rights Reserved.
 
 #include "Render/Synchronization/DisplayClusterRenderSyncPolicyBase.h"
+#include "Components/Viewport.h"
 #include "Render/Synchronization/DisplayClusterRenderSyncHelper.h"
 
 #include "Cluster/IPDisplayClusterClusterManager.h"
@@ -12,6 +13,7 @@
 #include "Engine/Engine.h"
 #include "Engine/GameViewportClient.h"
 
+#include "RHI.h"
 #include "RHIResources.h"
 
 
diff --git a/Engine/Plugins/Runtime/nDisplay/Source/DisplayCluster/Private/Render/Synchronization/Linux/DisplayClusterRenderSyncHelperLinux.cpp b/Engine/Plugins/Runtime/nDisplay/Source/DisplayCluster/Private/Render/Synchronization/Linux/DisplayClusterRenderSyncHelperLinux.cpp
--- a/Engine/Plugins/Runtime/nDisplay/Source/DisplayCluster/Private/Render/Synchronization/Linux/DisplayClusterRenderSyncHelperLinux.cpp
+++ b/Engine/Plugins/Runtime/nDisplay/Source/DisplayCluster/Private/Render/Synchronization/Linux/DisplayClusterRenderSyncHelperLinux.cpp
@@ -6,6 +6,7 @@
 
 #include "Engine/Engine.h"
 #include "Engine/GameViewportClient.h"
+#include "RHI.h"
 #include "DynamicRHI.h"
 
 
diff --git a/Engine/Plugins/Runtime/nDisplay/Source/DisplayCluster/Public/DisplayClusterLightCardActor.h b/Engine/Plugins/Runtime/nDisplay/Source/DisplayCluster/Public/DisplayClusterLightCardActor.h
--- a/Engine/Plugins/Runtime/nDisplay/Source/DisplayCluster/Public/DisplayClusterLightCardActor.h
+++ b/Engine/Plugins/Runtime/nDisplay/Source/DisplayCluster/Public/DisplayClusterLightCardActor.h
@@ -16,6 +16,7 @@ class USceneComponent;
 class USpringArmComponent;
 class UStaticMeshComponent;
 class UStaticMesh;
+class UTexture;
 
 UENUM(BlueprintType)
 enum class EDisplayClusterLightCardMask : uint8
diff --git a/Engine/Plugins/Runtime/nDisplay/Source/DisplayClusterColorGrading/Private/DataModelGenerators/DisplayClusterColorGradingGenerator_RootActor.cpp b/Engine/Plugins/Runtime/nDisplay/Source/DisplayClusterColorGrading/Private/DataModelGenerators/DisplayClusterColorGradingGenerator_RootActor.cpp
--- a/Engine/Plugins/Runtime/nDisplay/Source/DisplayClusterColorGrading/Private/DataModelGenerators/DisplayClusterColorGradingGenerator_RootActor.cpp
+++ b/Engine/Plugins/Runtime/nDisplay/Source/DisplayClusterColorGrading/Private/DataModelGenerators/DisplayClusterColorGradingGenerator_RootActor.cpp
@@ -23,6 +23,7 @@
 #include "PropertyPathHelpers.h"
 #include "ScopedTransaction.h"
 #include "Widgets/Input/SButton.h"
+#include "Widgets/Input/SComboButton.h"
 #include "Widgets/Text/SInlineEditableTextBlock.h"
 #include "Widgets/Views/SListView.h"
 
diff --git a/Engine/Plugins/Runtime/nDisplay/Source/DisplayClusterConfiguration/Private/DisplayClusterConfigurationModule.cpp b/Engine/Plugins/Runtime/nDisplay/Source/DisplayClusterConfiguration/Private/DisplayClusterConfigurationModule.cpp
--- a/Engine/Plugins/Runtime/nDisplay/Source/DisplayClusterConfiguration/Private/DisplayClusterConfigurationModule.cpp
+++ b/Engine/Plugins/Runtime/nDisplay/Source/DisplayClusterConfiguration/Private/DisplayClusterConfigurationModule.cpp
@@ -3,6 +3,7 @@
 #include "DisplayClusterConfigurationModule.h"
 #include "DisplayClusterConfigurationMgr.h"
 
+#include "Misc/CommandLine.h"
 #include "Modules/ModuleManager.h"
 
 #include "Engine/Engine.h"
diff --git a/Engine/Plugins/Runtime/nDisplay/Source/DisplayClusterConfigurator/Private/DisplayClusterConfiguratorModule.cpp b/Engine/Plugins/Runtime/nDisplay/Source/DisplayClusterConfigurator/Private/DisplayClusterConfiguratorModule.cpp
--- a/Engine/Plugins/Runtime/nDisplay/Source/DisplayClusterConfigurator/Private/DisplayClusterConfiguratorModule.cpp
+++ b/Engine/Plugins/Runtime/nDisplay/Source/DisplayClusterConfigurator/Private/DisplayClusterConfiguratorModule.cpp
@@ -5,6 +5,8 @@
 #include "DisplayClusterConfiguratorCommands.h"
 #include "DisplayClusterConfiguratorAssetTypeActions.h"
 #include "DisplayClusterConfiguratorVersionUtils.h"
+#include "Framework/Application/SlateApplication.h"
+#include "PropertyEditorModule.h"
 #include "Settings/DisplayClusterConfiguratorSettings.h"
 #include "Views/Details/DisplayClusterRootActorDetailsCustomization.h"
 #include "Views/OutputMapping/DisplayClusterConfiguratorOutputMappingCommands.h"
diff --git a/Engine/Plugins/Runtime/nDisplay/Source/DisplayClusterConfigurator/Private/Views/Details/Cluster/DisplayClusterConfiguratorMediaCustomization.cpp b/Engine/Plugins/Runtime/nDisplay/Source/DisplayClusterConfigurator/Private/Views/Details/Cluster/DisplayClusterConfiguratorMediaCustomization.cpp
--- a/Engine/Plugins/Runtime/nDisplay/Source/DisplayClusterConfigurator/Private/Views/Details/Cluster/DisplayClusterConfiguratorMediaCustomization.cpp
+++ b/Engine/Plugins/Runtime/nDisplay/Source/DisplayClusterConfigurator/Private/Views/Details/Cluster/DisplayClusterConfiguratorMediaCustomization.cpp
@@ -10,6 +10,7 @@
 #include "IDetailChildrenBuilder.h"
 #include "IPropertyTypeCustomization.h"
 #include "PropertyHandle.h"
+#include "Widgets/Input/SComboBox.h"
 
 
 void FDisplayClusterConfiguratorMediaCustomization::CustomizeChildren(TSharedRef<IPropertyHandle> InPropertyHandle, IDetailChildrenBuilder& InChildBuilder, IPropertyTypeCustomizationUtils& CustomizationUtils)
diff --git a/Engine/Plugins/Runtime/nDisplay/Source/DisplayClusterConfigurator/Private/Views/Details/Cluster/DisplayClusterConfiguratorMediaCustomization.h b/Engine/Plugins/Runtime/nDisplay/Source/DisplayClusterConfigurator/Private/Views/Details/Cluster/DisplayClusterConfiguratorMediaCustomization.h
--- a/Engine/Plugins/Runtime/nDisplay/Source/DisplayClusterConfigurator/Private/Views/Details/Cluster/DisplayClusterConfiguratorMediaCustomization.h
+++ b/Engine/Plugins/Runtime/nDisplay/Source/DisplayClusterConfigurator/Private/Views/Details/Cluster/DisplayClusterConfiguratorMediaCustomization.h
@@ -6,7 +6,7 @@
 
 template<class T>
 class SComboBox;
-
+class SWidget;
 
 class FDisplayClusterConfiguratorMediaCustomization final
 	: public FDisplayClusterConfiguratorBaseTypeCustomization
diff --git a/Engine/Plugins/Runtime/nDisplay/Source/DisplayClusterConfigurator/Private/Views/Details/Cluster/DisplayClusterConfiguratorNodeSelectionCustomization.h b/Engine/Plugins/Runtime/nDisplay/Source/DisplayClusterConfigurator/Private/Views/Details/Cluster/DisplayClusterConfiguratorNodeSelectionCustomization.h
--- a/Engine/Plugins/Runtime/nDisplay/Source/DisplayClusterConfigurator/Private/Views/Details/Cluster/DisplayClusterConfiguratorNodeSelectionCustomization.h
+++ b/Engine/Plugins/Runtime/nDisplay/Source/DisplayClusterConfigurator/Private/Views/Details/Cluster/DisplayClusterConfiguratorNodeSelectionCustomization.h
@@ -3,12 +3,14 @@
 #pragma once
 
 #include "CoreMinimal.h"
+#include "Misc/Attribute.h"
 #include "Views/Details/DisplayClusterConfiguratorBaseTypeCustomization.h"
 
 class FDisplayClusterConfiguratorBlueprintEditor;
 class ADisplayClusterRootActor;
 class UDisplayClusterConfigurationData;
 class SDisplayClusterConfigurationSearchableComboBox;
+class SWidget;
 
 /**
  * Helper class used to construct and manage a custom array property widget that displays for each element a dropdown containing
diff --git a/Engine/Plugins/Runtime/nDisplay/Source/DisplayClusterConfigurator/Private/Views/Details/Components/DisplayClusterConfiguratorScreenComponentDetailsCustomization.cpp b/Engine/Plugins/Runtime/nDisplay/Source/DisplayClusterConfigurator/Private/Views/Details/Components/DisplayClusterConfiguratorScreenComponentDetailsCustomization.cpp
--- a/Engine/Plugins/Runtime/nDisplay/Source/DisplayClusterConfigurator/Private/Views/Details/Components/DisplayClusterConfiguratorScreenComponentDetailsCustomization.cpp
+++ b/Engine/Plugins/Runtime/nDisplay/Source/DisplayClusterConfigurator/Private/Views/Details/Components/DisplayClusterConfiguratorScreenComponentDetailsCustomization.cpp
@@ -7,6 +7,8 @@
 #include "DetailLayoutBuilder.h"
 #include "DetailCategoryBuilder.h"
 #include "DetailWidgetRow.h"
+#include "Widgets/Input/SComboBox.h"
+#include "Widgets/Text/STextBlock.h"
 
 #define LOCTEXT_NAMESPACE "DisplayClusterConfiguratorDetailCustomization"
 
diff --git a/Engine/Plugins/Runtime/nDisplay/Source/DisplayClusterConfigurator/Private/Views/OutputMapping/EdNodes/DisplayClusterConfiguratorViewportNode.h b/Engine/Plugins/Runtime/nDisplay/Source/DisplayClusterConfigurator/Private/Views/OutputMapping/EdNodes/DisplayClusterConfiguratorViewportNode.h
--- a/Engine/Plugins/Runtime/nDisplay/Source/DisplayClusterConfigurator/Private/Views/OutputMapping/EdNodes/DisplayClusterConfiguratorViewportNode.h
+++ b/Engine/Plugins/Runtime/nDisplay/Source/DisplayClusterConfigurator/Private/Views/OutputMapping/EdNodes/DisplayClusterConfiguratorViewportNode.h
@@ -2,6 +2,7 @@
 
 #pragma once
 
+#include "Math/TransformCalculus2D.h"
 #include "Views/OutputMapping/EdNodes/DisplayClusterConfiguratorBaseNode.h"
 
 #include "Engine/Texture.h"
diff --git a/Engine/Plugins/Runtime/nDisplay/Source/DisplayClusterConfigurator/Private/Views/Viewport/DisplayClusterConfiguratorSCSEditorViewportClient.cpp b/Engine/Plugins/Runtime/nDisplay/Source/DisplayClusterConfigurator/Private/Views/Viewport/DisplayClusterConfiguratorSCSEditorViewportClient.cpp
--- a/Engine/Plugins/Runtime/nDisplay/Source/DisplayClusterConfigurator/Private/Views/Viewport/DisplayClusterConfiguratorSCSEditorViewportClient.cpp
+++ b/Engine/Plugins/Runtime/nDisplay/Source/DisplayClusterConfigurator/Private/Views/Viewport/DisplayClusterConfiguratorSCSEditorViewportClient.cpp
@@ -1,6 +1,8 @@
 // Copyright Epic Games, Inc. All Rights Reserved.
 
 #include "Views/Viewport/DisplayClusterConfiguratorSCSEditorViewportClient.h"
+#include "Materials/Material.h"
+#include "Settings/LevelEditorViewportSettings.h"
 #include "Views/Viewport/DisplayClusterConfiguratorSCSEditorViewport.h"
 #include "DisplayClusterConfiguratorBlueprintEditor.h"
 #include "DisplayClusterWorldOriginComponent.h"
diff --git a/Engine/Plugins/Runtime/nDisplay/Source/DisplayClusterLightCardEditor/Private/DisplayClusterLightCardEditor.h b/Engine/Plugins/Runtime/nDisplay/Source/DisplayClusterLightCardEditor/Private/DisplayClusterLightCardEditor.h
--- a/Engine/Plugins/Runtime/nDisplay/Source/DisplayClusterLightCardEditor/Private/DisplayClusterLightCardEditor.h
+++ b/Engine/Plugins/Runtime/nDisplay/Source/DisplayClusterLightCardEditor/Private/DisplayClusterLightCardEditor.h
@@ -3,6 +3,7 @@
 #pragma once
 
 #include "CoreMinimal.h"
+#include "GameFramework/Actor.h"
 #include "Widgets/DeclarativeSyntaxSupport.h"
 
 #include "DisplayClusterLightCardEditorProxyType.h"
diff --git a/Engine/Plugins/Runtime/nDisplay/Source/DisplayClusterLightCardEditor/Private/LightCardTemplates/SDisplayClusterLightCardTemplateList.cpp b/Engine/Plugins/Runtime/nDisplay/Source/DisplayClusterLightCardEditor/Private/LightCardTemplates/SDisplayClusterLightCardTemplateList.cpp
--- a/Engine/Plugins/Runtime/nDisplay/Source/DisplayClusterLightCardEditor/Private/LightCardTemplates/SDisplayClusterLightCardTemplateList.cpp
+++ b/Engine/Plugins/Runtime/nDisplay/Source/DisplayClusterLightCardEditor/Private/LightCardTemplates/SDisplayClusterLightCardTemplateList.cpp
@@ -5,6 +5,8 @@
 #include "DisplayClusterLightCardTemplateDragDropOp.h"
 #include "DisplayClusterLightCardTemplateHelpers.h"
 #include "DisplayClusterLightCardEditor.h"
+#include "Engine/Texture2D.h"
+#include "Framework/Application/SlateApplication.h"
 #include "LightCardTemplates/DisplayClusterLightCardTemplate.h"
 
 #include "AssetRegistry/AssetRegistryModule.h"
diff --git a/Engine/Plugins/Runtime/nDisplay/Source/DisplayClusterLightCardEditor/Private/Viewport/DisplayClusterLightCardEditorViewportClient.cpp b/Engine/Plugins/Runtime/nDisplay/Source/DisplayClusterLightCardEditor/Private/Viewport/DisplayClusterLightCardEditorViewportClient.cpp
--- a/Engine/Plugins/Runtime/nDisplay/Source/DisplayClusterLightCardEditor/Private/Viewport/DisplayClusterLightCardEditorViewportClient.cpp
+++ b/Engine/Plugins/Runtime/nDisplay/Source/DisplayClusterLightCardEditor/Private/Viewport/DisplayClusterLightCardEditorViewportClient.cpp
@@ -2,6 +2,7 @@
 
 #include "DisplayClusterLightCardEditorViewportClient.h"
 
+#include "Application/ThrottleManager.h"
 #include "DisplayClusterLightcardEditorViewport.h"
 #include "DisplayClusterLightCardEditorWidget.h"
 #include "LightCardTemplates/DisplayClusterLightCardTemplate.h"
@@ -18,6 +19,7 @@
 #include "DisplayClusterLightCardEditorUtils.h"
 #include "IDisplayClusterScenePreview.h"
 #include "DisplayClusterLightCardEditor.h"
+#include "SceneManagement.h"
 #include "Settings/DisplayClusterLightCardEditorSettings.h"
 
 #include "AudioDevice.h"
@@ -47,7 +49,9 @@
 #include "RayTracingDebugVisualizationMenuCommands.h"
 #include "Renderer/Private/SceneRendering.h"
 #include "ScopedTransaction.h"
+#include "Settings/LevelEditorViewportSettings.h"
 #include "Slate/SceneViewport.h"
+#include "TextureResource.h"
 #include "UnrealEdGlobals.h"
 #include "UnrealWidget.h"
 #include "Components/BillboardComponent.h"
