[mutable] Fix full skeleton deformations caused by the reshape of the root and change physics to deform selection.

-Small refactors to bring the code base closer to the unreal standards.
-Now Bones in poses have a usage flag that allows for a more informed merge.
-Root Bone and bones that are not skinned can not be deformed.
-Bones reshaped that have not moved are not marked as reshaped.
-Remove reshape unused flags and skip conversion to bool arguments.
-Changed physics to deform selection methods to be on parity with bone selection.

diff --git a/Engine/Plugins/Experimental/Mutable/Source/CustomizableObjectEditor/Classes/MuCOE/Nodes/CustomizableObjectNodeMeshMorph.h b/Engine/Plugins/Experimental/Mutable/Source/CustomizableObjectEditor/Classes/MuCOE/Nodes/CustomizableObjectNodeMeshMorph.h
--- a/Engine/Plugins/Experimental/Mutable/Source/CustomizableObjectEditor/Classes/MuCOE/Nodes/CustomizableObjectNodeMeshMorph.h
+++ b/Engine/Plugins/Experimental/Mutable/Source/CustomizableObjectEditor/Classes/MuCOE/Nodes/CustomizableObjectNodeMeshMorph.h
@@ -69,19 +69,22 @@ public:
 	UPROPERTY(EditAnywhere, Category = ReshapeBones, meta = (EditCondition = "bReshapeSkeleton"))
 	EBoneDeformSelectionMethod SelectionMethod;
 	
-	/** Enables the deformation of all bones of the skeleton */
-	UPROPERTY()
-	bool bDeformAllBones_DEPRECATED = false;
-
 	/** Array with selected bones that will be deformed */
 	UPROPERTY(EditAnywhere, Category = ReshapeBones, meta = (EditCondition = "bReshapeSkeleton && (SelectionMethod == EBoneDeformSelectionMethod::ONLY_SELECTED || SelectionMethod == EBoneDeformSelectionMethod::ALL_BUT_SELECTED)"))
 	TArray<FMeshReshapeBoneReference> BonesToDeform;
 
-	UPROPERTY(EditAnywhere, Category = ReshapePhysics, meta = (EditCondition="bReshapePhysicsVolumes"))
-	bool bDeformAllPhysicsBodies = false;
+	UPROPERTY(EditAnywhere, Category = ReshapePhysics, meta = (DisplayName = "Selection Method", EditCondition = "bReshapePhysicsVolumes"))
+	EBoneDeformSelectionMethod PhysicsSelectionMethod;
 
-	UPROPERTY(EditAnywhere, Category = ReshapePhysics, meta=(EditCondition="bReshapePhysicsVolumes && !bDeformAllPhysicsBodies"))
+	/** Array with bones with physics bodies that will be deformed */
+	UPROPERTY(EditAnywhere, Category = ReshapePhysics, meta = (EditCondition = "bReshapePhysicsVolumes && (PhysicsSelectionMethod == EBoneDeformSelectionMethod::ONLY_SELECTED || PhysicsSelectionMethod == EBoneDeformSelectionMethod::ALL_BUT_SELECTED)"))
 	TArray<FMeshReshapeBoneReference> PhysicsBodiesToDeform;
 
+	UPROPERTY()
+	bool bDeformAllBones_DEPRECATED = false;
+
+	UPROPERTY()
+	bool bDeformAllPhysicsBodies_DEPRECATED = false;
+
 };
 
diff --git a/Engine/Plugins/Experimental/Mutable/Source/CustomizableObjectEditor/Classes/MuCOE/Nodes/CustomizableObjectNodeMeshReshape.h b/Engine/Plugins/Experimental/Mutable/Source/CustomizableObjectEditor/Classes/MuCOE/Nodes/CustomizableObjectNodeMeshReshape.h
--- a/Engine/Plugins/Experimental/Mutable/Source/CustomizableObjectEditor/Classes/MuCOE/Nodes/CustomizableObjectNodeMeshReshape.h
+++ b/Engine/Plugins/Experimental/Mutable/Source/CustomizableObjectEditor/Classes/MuCOE/Nodes/CustomizableObjectNodeMeshReshape.h
@@ -69,20 +69,21 @@ public:
 	UPROPERTY(EditAnywhere, Category = ReshapeBones, meta = (EditCondition = "bReshapeSkeleton"))
 	EBoneDeformSelectionMethod SelectionMethod;
 
-	/** Enables the deformation of all bones of the skeleton */
-	UPROPERTY()
-	bool bDeformAllBones_DEPRECATED = false;
-
 	/** Array with selected bones that will be deformed */
 	UPROPERTY(EditAnywhere, Category = ReshapeBones, meta = (EditCondition = "bReshapeSkeleton && (SelectionMethod == EBoneDeformSelectionMethod::ONLY_SELECTED || SelectionMethod == EBoneDeformSelectionMethod::ALL_BUT_SELECTED)"))
 	TArray<FMeshReshapeBoneReference> BonesToDeform;
 
-	/** Enables the deformation of all physic bodies*/
-	UPROPERTY(EditAnywhere, Category = ReshapePhysics, meta = (EditCondition = "bReshapePhysicsVolumes"))
-	bool bDeformAllPhysicsBodies = false;
+	UPROPERTY(EditAnywhere, Category = ReshapePhysics, meta = (DisplayName = "Selection Method", EditCondition = "bReshapePhysicsVolumes"))
+	EBoneDeformSelectionMethod PhysicsSelectionMethod;
 
 	/** Array with bones with physics bodies that will be deformed */
-	UPROPERTY(EditAnywhere, Category = ReshapePhysics, meta = (EditCondition = "bReshapePhysicsVolumes && !bDeformAllPhysicsBodies"))
+	UPROPERTY(EditAnywhere, Category = ReshapePhysics, meta = (EditCondition = "bReshapePhysicsVolumes && (PhysicsSelectionMethod == EBoneDeformSelectionMethod::ONLY_SELECTED || PhysicsSelectionMethod == EBoneDeformSelectionMethod::ALL_BUT_SELECTED)"))
 	TArray<FMeshReshapeBoneReference> PhysicsBodiesToDeform;
+
+	UPROPERTY()
+	bool bDeformAllBones_DEPRECATED = false;
+
+	UPROPERTY()
+	bool bDeformAllPhysicsBodies_DEPRECATED = false;
 };
 
diff --git a/Engine/Plugins/Experimental/Mutable/Source/CustomizableObjectEditor/Classes/MuCOE/Nodes/CustomizableObjectNodeSwitchBase.cpp b/Engine/Plugins/Experimental/Mutable/Source/CustomizableObjectEditor/Classes/MuCOE/Nodes/CustomizableObjectNodeSwitchBase.cpp
--- a/Engine/Plugins/Experimental/Mutable/Source/CustomizableObjectEditor/Classes/MuCOE/Nodes/CustomizableObjectNodeSwitchBase.cpp
+++ b/Engine/Plugins/Experimental/Mutable/Source/CustomizableObjectEditor/Classes/MuCOE/Nodes/CustomizableObjectNodeSwitchBase.cpp
@@ -129,7 +129,7 @@ void UCustomizableObjectNodeSwitchBase::EnumParameterPostEditChangeProperty(FPos
 {
 	if (const UEdGraphPin* SwitchPin = SwitchParameter())
 	{
-		if (const UEdGraphPin* ConnectedPin = FollowInputPin(*SwitchPin); ConnectedPin->GetOwningNode() == Parameters.Node)
+		if (const UEdGraphPin* ConnectedPin = FollowInputPin(*SwitchPin); ConnectedPin && ConnectedPin->GetOwningNode() == Parameters.Node)
 		{
 			// Using MarkForReconstruct instead of Super::ReconstructNode because when we Copy Paste this node, it crashes sometimes due to reconctructing the node while constructing it.
 			MarkForReconstruct();
diff --git a/Engine/Plugins/Experimental/Mutable/Source/CustomizableObjectEditor/Private/MuCOE/GenerateMutableSource/GenerateMutableSource.cpp b/Engine/Plugins/Experimental/Mutable/Source/CustomizableObjectEditor/Private/MuCOE/GenerateMutableSource/GenerateMutableSource.cpp
--- a/Engine/Plugins/Experimental/Mutable/Source/CustomizableObjectEditor/Private/MuCOE/GenerateMutableSource/GenerateMutableSource.cpp
+++ b/Engine/Plugins/Experimental/Mutable/Source/CustomizableObjectEditor/Private/MuCOE/GenerateMutableSource/GenerateMutableSource.cpp
@@ -598,7 +598,7 @@ mu::NodeMeshApplyPosePtr CreateNodeMeshApplyPose(mu::NodeMeshPtr InputMeshNode,
 	for (int32 i = 0; i < ArrayBoneName.Num(); ++i)
 	{
 		MutableSkeleton->SetBoneName(i, TCHAR_TO_ANSI(*ArrayBoneName[i]));
-		MutableMesh->SetBonePose(i, TCHAR_TO_ANSI(*ArrayBoneName[i]), (FTransform3f)ArrayTransform[i], true);
+		MutableMesh->SetBonePose(i, TCHAR_TO_ANSI(*ArrayBoneName[i]), (FTransform3f)ArrayTransform[i], mu::EBoneUsageFlags::Skinning);
 	}
 
 	mu::NodeMeshApplyPosePtr NodeMeshApplyPose = new mu::NodeMeshApplyPose;
diff --git a/Engine/Plugins/Experimental/Mutable/Source/CustomizableObjectEditor/Private/MuCOE/GenerateMutableSource/GenerateMutableSourceMesh.cpp b/Engine/Plugins/Experimental/Mutable/Source/CustomizableObjectEditor/Private/MuCOE/GenerateMutableSource/GenerateMutableSourceMesh.cpp
--- a/Engine/Plugins/Experimental/Mutable/Source/CustomizableObjectEditor/Private/MuCOE/GenerateMutableSource/GenerateMutableSourceMesh.cpp
+++ b/Engine/Plugins/Experimental/Mutable/Source/CustomizableObjectEditor/Private/MuCOE/GenerateMutableSource/GenerateMutableSourceMesh.cpp
@@ -403,6 +403,66 @@ void NormalizeWeights(FBoneIndexType* InfluenceBones, uint8* InfluenceWeights, c
 	}
 }
 
+void SetAndPropagatePoseBoneUsage(
+		mu::Mesh& MutableMesh, int32 PoseIndex, mu::EBoneUsageFlags Usage, 
+		const TMap<FName, int32>* BoneNameToSkeletonIndexMap = nullptr,
+		const TMap<FName, int32>* BoneNameToPoseIndexMap = nullptr)
+{
+	if (!MutableMesh.GetSkeleton())
+	{
+		return;
+	}
+
+	const mu::Skeleton& MutableSkeleton = *MutableMesh.GetSkeleton();
+
+	if (PoseIndex < 0 || PoseIndex >= MutableMesh.BonePoses.Num())
+	{
+		check(false);
+		return;
+	}
+
+	int32 BoneIndex = [&]() 
+	{
+		if (BoneNameToSkeletonIndexMap)
+		{
+			const int32* Found = BoneNameToSkeletonIndexMap->Find(FName(MutableMesh.BonePoses[PoseIndex].BoneName.c_str()));
+			return Found ? *Found : INDEX_NONE;
+		}
+		else
+		{
+			return MutableSkeleton.FindBone(MutableMesh.BonePoses[PoseIndex].BoneName.c_str());
+		}
+	}();
+
+	while (BoneIndex != INDEX_NONE)
+	{
+		PoseIndex = [&]()
+		{
+			if (BoneNameToPoseIndexMap)
+			{
+				const int32* Found = BoneNameToPoseIndexMap->Find(FName(MutableSkeleton.GetBoneName(BoneIndex)));
+				return Found ? *Found : INDEX_NONE;
+			}
+			else
+			{
+				return MutableMesh.FindBonePose(MutableSkeleton.GetBoneName(BoneIndex));
+			}
+		}();
+
+		if (PoseIndex == INDEX_NONE)
+		{
+			check(false);
+			return;
+		}
+
+		EnumAddFlags(MutableMesh.BonePoses[PoseIndex].BoneUsageFlags, Usage);
+
+		BoneIndex = MutableSkeleton.GetBoneParent(BoneIndex);
+	}
+
+}
+
+
 mu::MeshPtr ConvertSkeletalMeshToMutable(USkeletalMesh* InSkeletalMesh, int LOD, int MaterialIndex, FMutableGraphGenerationContext& GenerationContext, const UCustomizableObjectNode* CurrentNode)
 {
 	// Get the mesh generation flags to use
@@ -1342,10 +1402,18 @@ mu::MeshPtr ConvertSkeletalMeshToMutable(USkeletalMesh* InSkeletalMesh, int LOD,
 			FMatrix44f BaseInvMatrix = SkeletalMesh->GetRefBasesInvMatrix()[RefSkelIndex];
 			FTransform3f BaseInvTransform;
 			BaseInvTransform.SetFromMatrix(BaseInvMatrix);
-			MutableMesh->SetBonePose(BoneMapIndex, TCHAR_TO_ANSI(*BoneName), BaseInvTransform.Inverse(), bSkinned);
+
+			mu::EBoneUsageFlags BoneUsageFlags = bSkinned ? mu::EBoneUsageFlags::Skinning : mu::EBoneUsageFlags::None;
+			
+			if (SkeletalMesh->GetRefSkeleton().GetParentIndex(RefSkelIndex) == INDEX_NONE)
+			{
+				EnumAddFlags(BoneUsageFlags, mu::EBoneUsageFlags::Root);
+			}
+
+			MutableMesh->SetBonePose(BoneMapIndex, TCHAR_TO_ANSI(*BoneName), BaseInvTransform.Inverse(), BoneUsageFlags);
 		}
 	}
-	
+
 	if (!bIgnorePhysics && SkeletalMesh->GetPhysicsAsset() && MutableMesh->GetSkeleton())
 	{
 		// Find BodySetups with relevant bones.
@@ -1370,19 +1438,21 @@ mu::MeshPtr ConvertSkeletalMeshToMutable(USkeletalMesh* InSkeletalMesh, int LOD,
 			const int32 SkeletonBoneCount = MutableMesh->GetSkeleton()->GetBoneCount();
 			for ( int32 I = 0; I < SkeletonBoneCount; ++I )
 			{
-				FString SkeletonBoneName = MutableMesh->GetSkeleton()->GetBoneName( I );
+				FString SkeletonBoneName = MutableMesh->GetSkeleton()->GetBoneName(I);
 				
 				if (SkeletonBoneName.Equals(BodyBoneName))
 				{
 					RelevantBodySetups.Add(BodySetup);
 					DiscardedBodySetups[BodySetupIndex] = 0;
+					int32 BonePoseIndex = MutableMesh->FindBonePose(TCHAR_TO_ANSI(*BodyBoneName));
+					
+					EnumAddFlags(MutableMesh->BonePoses[BonePoseIndex].BoneUsageFlags, mu::EBoneUsageFlags::Physics);
 				}
 			}
 		}
 
 		const int32 NumDiscardedSetups = Algo::CountIf(DiscardedBodySetups, [](const uint8& V) { return V; });
 
-
 		constexpr bool bOptOutOfIncompleteBodyWarnings = true;
 		if (NumDiscardedSetups > 0 && !bOptOutOfIncompleteBodyWarnings)
 		{
@@ -1524,6 +1594,41 @@ mu::MeshPtr ConvertSkeletalMeshToMutable(USkeletalMesh* InSkeletalMesh, int LOD,
 		MutableMesh->SetPhysicsBody(PhysicsBody);
 	}
 	
+	// Set Bone Parenting usages. This has to be done after all primary usages are set.
+	for (int32 I = MutableMesh->GetBonePoseCount() - 1; I >= 0; --I)
+	{
+		mu::Mesh::FBonePose& BonePose = MutableMesh->BonePoses[I];
+
+		constexpr mu::EBoneUsageFlags FlagsToPropagate =
+				mu::EBoneUsageFlags::Skinning | mu::EBoneUsageFlags::Physics | mu::EBoneUsageFlags::Deform;
+		if (EnumHasAnyFlags(BonePose.BoneUsageFlags, FlagsToPropagate))
+		{
+			const int32 BoneIndex = MutableMesh->GetSkeleton()->FindBone(MutableMesh->GetBonePoseName(I));
+
+			if (BoneIndex == INDEX_NONE)
+			{
+				continue;
+			}
+
+			const int32 ParentIndex = MutableMesh->GetSkeleton()->GetBoneParent(BoneIndex);
+
+			if (ParentIndex == INDEX_NONE)
+			{	
+				continue;
+			}
+
+			const mu::EBoneUsageFlags ParentPropagationFlags =
+				(EnumHasAnyFlags(BonePose.BoneUsageFlags, mu::EBoneUsageFlags::Skinning) 
+					? mu::EBoneUsageFlags::SkinningParent : mu::EBoneUsageFlags::None) |
+				(EnumHasAnyFlags(BonePose.BoneUsageFlags, mu::EBoneUsageFlags::Physics) 
+					? mu::EBoneUsageFlags::PhysicsParent : mu::EBoneUsageFlags::None) |
+				(EnumHasAnyFlags(BonePose.BoneUsageFlags, mu::EBoneUsageFlags::Deform) 
+					? mu::EBoneUsageFlags::DeformParent : mu::EBoneUsageFlags::None);
+
+			SetAndPropagatePoseBoneUsage(*MutableMesh, ParentIndex, ParentPropagationFlags);
+		}
+	}
+
 	return MutableMesh;
 }
 
@@ -2072,46 +2177,115 @@ bool GetAndValidateReshapeBonesToDeform(
 
 
 bool GetAndValidateReshapePhysicsToDeform(
-	TArray<FString>& OutPhysiscsToDeform,
-	const TArray<FMeshReshapeBoneReference>& InPhysicsToDeform,
-	const TArray<USkeletalMesh*>& SkeletalMeshes,
-	const UCustomizableObjectNode* Node,
-	FMutableGraphGenerationContext& GenerationContext)
+		TArray<FString>& OutPhysiscsToDeform,
+		const TArray<FMeshReshapeBoneReference>& InPhysicsToDeform,
+		const TArray<USkeletalMesh*>& SkeletalMeshes,
+		EBoneDeformSelectionMethod SelectionMethod,
+		const UCustomizableObjectNode* Node,
+		FMutableGraphGenerationContext& GenerationContext)
 {
-	bool bSetRefreshWarning = false;
 
-	int32 NumPhysicsToDeform = InPhysicsToDeform.Num();
-	for (int32 BodyIndex = 0; BodyIndex < NumPhysicsToDeform; ++BodyIndex)
+	const bool bIsReferenceSkeletalMeshMethod =
+			SelectionMethod == EBoneDeformSelectionMethod::DEFORM_REF_SKELETON ||
+			SelectionMethod == EBoneDeformSelectionMethod::DEFORM_NONE_REF_SKELETON;
+
+	// Get the participant bone names.
+	TArray<FName> BoneNamesInUserSelection;
+	if (bIsReferenceSkeletalMeshMethod)
 	{
-		bool bMissingBone = true;
-		bool bMissingBody = true;
+		const FReferenceSkeleton& RefSkeleton =
+			GenerationContext.ComponentInfos[GenerationContext.CurrentMeshComponent].RefSkeletalMesh->GetRefSkeleton();
 
-		const FName BodyBoneName = InPhysicsToDeform[BodyIndex].BoneName;
-		for (const USkeletalMesh* SkeletalMesh : SkeletalMeshes)
+		const int32 RefSkeletonNumBones = RefSkeleton.GetNum();
+		BoneNamesInUserSelection.SetNum(RefSkeletonNumBones);
+		for (int32 I = 0; I < RefSkeletonNumBones; ++I)
 		{
-			if (SkeletalMesh->GetRefSkeleton().FindBoneIndex(BodyBoneName) != INDEX_NONE)
-			{
-				bMissingBone = false;
-			}
+			BoneNamesInUserSelection[I] = RefSkeleton.GetBoneName(I);
+		}
+	}
+	else
+	{
+		BoneNamesInUserSelection.Reserve(InPhysicsToDeform.Num());
+		Algo::Transform(InPhysicsToDeform, BoneNamesInUserSelection, [](const FMeshReshapeBoneReference& B) { return B.BoneName; });
+	}
+
+	int32 NumUserSelectedBones = BoneNamesInUserSelection.Num();
+
+	struct FMissingBoneState
+	{
+		uint8 bMissingBone : 1;
+		uint8 bMissingBody : 1;
+	};
+
+	TArray<FMissingBoneState> MissingBones;
+	MissingBones.Init(FMissingBoneState{ false, true }, NumUserSelectedBones);
 
-			UPhysicsAsset* PhysicsAsset = SkeletalMesh->GetPhysicsAsset();
-			if (PhysicsAsset && !bMissingBone)
+	for (const USkeletalMesh* SkeletalMesh : SkeletalMeshes)
+	{
+		const FReferenceSkeleton& RefSkeleton = bIsReferenceSkeletalMeshMethod 
+				? GenerationContext.ComponentInfos[GenerationContext.CurrentMeshComponent].RefSkeletalMesh->GetRefSkeleton()
+				: SkeletalMesh->GetRefSkeleton();
+
+		UPhysicsAsset* PhysicsAsset = SkeletalMesh->GetPhysicsAsset();
+
+		TArray<uint8> BoneInclusionSet;
+		BoneInclusionSet.Init(0, PhysicsAsset->SkeletalBodySetups.Num());
+
+		// Find to which SkeletalBodySetups the user selection bones belong to. 
+		for (int32 IndexToDeform = 0; IndexToDeform < NumUserSelectedBones; ++IndexToDeform)
+		{
+			const FName& BodyBoneName = BoneNamesInUserSelection[IndexToDeform];
+			MissingBones[IndexToDeform].bMissingBone = RefSkeleton.FindBoneIndex(BodyBoneName) == INDEX_NONE;
+
+			if (PhysicsAsset && !MissingBones[IndexToDeform].bMissingBone)
 			{
 				const int32 FoundIndex = PhysicsAsset->SkeletalBodySetups.IndexOfByPredicate(
-					[&BodyBoneName]( const TObjectPtr<USkeletalBodySetup>& Setup ) {  return Setup->BoneName == BodyBoneName; });
+					[&BodyBoneName](const TObjectPtr<USkeletalBodySetup>& Setup) {  return Setup->BoneName == BodyBoneName; });
 
 				if (FoundIndex != INDEX_NONE)
 				{
-					OutPhysiscsToDeform.AddUnique(BodyBoneName.ToString());
-					bMissingBody = false;
-					break;
+					BoneInclusionSet[FoundIndex] = 1;
+					MissingBones[IndexToDeform].bMissingBody = false;
 				}
 			}
 		}
 
-		if (bMissingBone)
+		const bool bFlipSelection =
+				SelectionMethod == EBoneDeformSelectionMethod::ALL_BUT_SELECTED ||
+				SelectionMethod == EBoneDeformSelectionMethod::DEFORM_NONE_REF_SKELETON;
+		if (bFlipSelection)
 		{
-			FString BoneName = InPhysicsToDeform[BodyIndex].BoneName.ToString();
+			for (uint8& Elem : BoneInclusionSet) 
+			{ 
+				Elem = 1 - Elem; 
+			}
+		}
+
+		// Append the bones in the inclusion set to the output bone names list.
+		const int32 BoneInclusionSetNum = BoneInclusionSet.Num();
+		for (int32 I = 0; I < BoneInclusionSetNum; ++I)
+		{
+			if (BoneInclusionSet[I])
+			{
+				FName SetupBoneName = PhysicsAsset->SkeletalBodySetups[I]->BoneName;
+				OutPhysiscsToDeform.AddUnique(SetupBoneName.ToString());
+			}
+		}
+	}
+
+	// Don't warn if the selection is not explicit.
+	if (SelectionMethod != EBoneDeformSelectionMethod::ONLY_SELECTED)
+	{
+		return false;
+	}
+
+	// Emit warnings if some explicitly selected bone is not present or has no phyiscs attached.
+	bool bSetRefreshWarning = false;
+	for (int32 IndexToDeform = 0; IndexToDeform < NumUserSelectedBones; ++IndexToDeform)
+	{
+		if (MissingBones[IndexToDeform].bMissingBone)
+		{
+			FString BoneName = BoneNamesInUserSelection[IndexToDeform].ToString();
 			FString msg = FString::Printf(TEXT("Could not find the selected Physics Body bone to Deform [%s] in Skeleton"), *BoneName);
 
 			GenerationContext.Compiler->CompilerLog(FText::FromString(msg), Node);
@@ -2119,17 +2293,17 @@ bool GetAndValidateReshapePhysicsToDeform(
 			bSetRefreshWarning = true;
 		}
 
-		if (bMissingBody && !bMissingBone)
+		if (MissingBones[IndexToDeform].bMissingBody)
 		{
-			FString BoneName = InPhysicsToDeform[BodyIndex].BoneName.ToString();
+			FString BoneName = BoneNamesInUserSelection[IndexToDeform].ToString();
 			FString msg = FString::Printf(TEXT("Selected Bone to Deform [%s] does not have any physics body attached."), *BoneName);
 
 			GenerationContext.Compiler->CompilerLog(FText::FromString(msg), Node);
 
 			bSetRefreshWarning = true;
 		}
-			
 	}
+
 	return bSetRefreshWarning;
 }
 
@@ -2255,7 +2429,6 @@ mu::NodeMeshPtr GenerateMorphMesh(const UEdGraphPin* Pin,
 
 			UCustomizableObjectNodeMeshMorph* TypedMorphNode = TypedNodeMorphs[MorphIndex];
 
-			Result->SetDeformAllPhysics(TypedMorphNode->bDeformAllPhysicsBodies);
 			Result->SetReshapeSkeleton(TypedMorphNode->bReshapeSkeleton);
 			Result->SetReshapePhysicsVolumes(TypedMorphNode->bReshapePhysicsVolumes);	
 			{
@@ -2278,11 +2451,15 @@ mu::NodeMeshPtr GenerateMorphMesh(const UEdGraphPin* Pin,
 					}
 				}
 
-				if (TypedMorphNode->bReshapePhysicsVolumes && !TypedMorphNode->bDeformAllPhysicsBodies)
+				if (TypedMorphNode->bReshapePhysicsVolumes)
 				{
 					TArray<FString> PhysicsToDeform;
+
+					const EBoneDeformSelectionMethod SelectionMethod = TypedMorphNode->PhysicsSelectionMethod;
 					bWarningFound = bWarningFound || GetAndValidateReshapePhysicsToDeform(
-						PhysicsToDeform, TypedMorphNode->PhysicsBodiesToDeform, SkeletalMeshesToDeform, TypedMorphNode, GenerationContext);
+						PhysicsToDeform, 
+						TypedMorphNode->PhysicsBodiesToDeform, SkeletalMeshesToDeform, SelectionMethod, 
+						TypedMorphNode, GenerationContext);
 	
 					for (const FString& PhysicsBoneName : PhysicsToDeform)
 					{
@@ -2963,7 +3140,6 @@ mu::NodeMeshPtr GenerateMutableSourceMesh(const UEdGraphPin * Pin,
 			MeshNode->SetReshapeSkeleton(TypedNodeReshape->bReshapeSkeleton);
 			MeshNode->SetReshapePhysicsVolumes(TypedNodeReshape->bReshapePhysicsVolumes);
 			MeshNode->SetEnableRigidParts(TypedNodeReshape->bEnableRigidParts);
-			MeshNode->SetDeformAllPhysics(TypedNodeReshape->bDeformAllPhysicsBodies);
 				
 			const UEdGraphPin* ConnectedPin = FollowInputPin(*TypedNodeReshape->BaseMeshPin());
 			const UEdGraphPin* SourceMeshPin = ConnectedPin ? FindMeshBaseSource(*ConnectedPin, false) : nullptr;
@@ -2984,11 +3160,14 @@ mu::NodeMeshPtr GenerateMutableSourceMesh(const UEdGraphPin * Pin,
 				}
 			}
 
-			if (TypedNodeReshape->bReshapePhysicsVolumes && !TypedNodeReshape->bDeformAllPhysicsBodies)
+			if (TypedNodeReshape->bReshapePhysicsVolumes)
 			{
+				EBoneDeformSelectionMethod SelectionMethod = TypedNodeReshape->PhysicsSelectionMethod;
 				TArray<FString> PhysicsToDeform;
 				bWarningFound = bWarningFound || GetAndValidateReshapePhysicsToDeform(
-					PhysicsToDeform, TypedNodeReshape->PhysicsBodiesToDeform, SkeletalMeshesToDeform, TypedNodeReshape, GenerationContext);
+					PhysicsToDeform, 
+					TypedNodeReshape->PhysicsBodiesToDeform, SkeletalMeshesToDeform, SelectionMethod, 
+					TypedNodeReshape, GenerationContext);
 
 				for (const FString& PhysicsBoneName : PhysicsToDeform)
 				{
diff --git a/Engine/Plugins/Experimental/Mutable/Source/MutableRuntime/Private/MuR/CodeRunner.cpp b/Engine/Plugins/Experimental/Mutable/Source/MutableRuntime/Private/MuR/CodeRunner.cpp
--- a/Engine/Plugins/Experimental/Mutable/Source/MutableRuntime/Private/MuR/CodeRunner.cpp
+++ b/Engine/Plugins/Experimental/Mutable/Source/MutableRuntime/Private/MuR/CodeRunner.cpp
@@ -2027,14 +2027,8 @@ namespace mu
 					}
 
 					// -----------
-					Ptr<Mesh> pResult = MeshBindShapeReshape( BaseMesh.get(), Shape.get(), 	
-						BonesToDeform, PhysicsToDeform,
-						args.flags & uint32(OP::EMeshBindShapeFlags::DeformAllBones),
-						args.flags & uint32(OP::EMeshBindShapeFlags::DeformAllPhysics),	
-						args.flags & uint32(OP::EMeshBindShapeFlags::ReshapeVertices),
-						args.flags & uint32(OP::EMeshBindShapeFlags::ReshapeSkeleton),
-						args.flags & uint32(OP::EMeshBindShapeFlags::ReshapePhysicsVolumes),
-						args.flags & uint32(OP::EMeshBindShapeFlags::EnableRigidParts));
+					Ptr<Mesh> pResult = MeshBindShapeReshape( BaseMesh.get(), Shape.get(), BonesToDeform, PhysicsToDeform, 
+						    static_cast<EMeshBindShapeFlags>(args.flags));
 					
 					GetMemory().SetMesh(item, pResult);
 				}	
@@ -2078,10 +2072,8 @@ namespace mu
 				Ptr<const Mesh> BaseMesh = GetMemory().GetMesh(CACHE_ADDRESS(args.mesh, item));
 				Ptr<const Mesh> Shape = GetMemory().GetMesh(CACHE_ADDRESS(args.shape, item));
 
-				Ptr<Mesh> pResult = MeshApplyShape(BaseMesh.get(), Shape.get(),
-					args.flags & uint32(OP::EMeshBindShapeFlags::ReshapeVertices),
-					args.flags & uint32(OP::EMeshBindShapeFlags::ReshapeSkeleton),
-					args.flags & uint32(OP::EMeshBindShapeFlags::ReshapePhysicsVolumes));
+
+				Ptr<Mesh> pResult = MeshApplyShape(BaseMesh.get(), Shape.get(), static_cast<EMeshBindShapeFlags>(args.flags));
 
 				GetMemory().SetMesh(item, pResult);
 				break;
@@ -2129,7 +2121,7 @@ namespace mu
 					Ptr<Mesh> ResultPtr = MorphedMesh->Clone(CloneFlags);
 					ResultPtr->SetSkeleton(ReshapeMesh->GetSkeleton().get());
 					ResultPtr->SetPhysicsBody(ReshapeMesh->GetPhysicsBody().get());
-					ResultPtr->m_bonePoses = ReshapeMesh->m_bonePoses;
+					ResultPtr->BonePoses = ReshapeMesh->BonePoses;
 
 					GetMemory().SetMesh(item, ResultPtr);
 				}
diff --git a/Engine/Plugins/Experimental/Mutable/Source/MutableRuntime/Private/MuR/Mesh.cpp b/Engine/Plugins/Experimental/Mutable/Source/MutableRuntime/Private/MuR/Mesh.cpp
--- a/Engine/Plugins/Experimental/Mutable/Source/MutableRuntime/Private/MuR/Mesh.cpp
+++ b/Engine/Plugins/Experimental/Mutable/Source/MutableRuntime/Private/MuR/Mesh.cpp
@@ -71,12 +71,12 @@ MeshPtr Mesh::Clone() const
 	// physics body doen't need to be deep cloned either as they are also assumed to be shared.
 	
 	// Clone bone poses
-	pResult->m_bonePoses = m_bonePoses;
+	pResult->BonePoses = BonePoses;
 
     return pResult;
 }
 
-MeshPtr Mesh::Clone( EMeshCloneFlags Flags ) const
+MeshPtr Mesh::Clone(EMeshCloneFlags Flags) const
 {
     //MUTABLE_CPUPROFILER_SCOPE(MeshClone);
 	LLM_SCOPE_BYNAME(TEXT("MutableRuntime"));
@@ -86,55 +86,55 @@ MeshPtr Mesh::Clone( EMeshCloneFlags Flags ) const
     pResult->m_internalId = m_internalId;
 	pResult->m_staticFormatFlags = m_staticFormatFlags;
 
-	if ( EnumHasAnyFlags( Flags, EMeshCloneFlags::WithSurfaces ))
+	if (EnumHasAnyFlags(Flags, EMeshCloneFlags::WithSurfaces ))
 	{
 		pResult->m_surfaces = m_surfaces;
 	}
 
-	if ( EnumHasAnyFlags( Flags, EMeshCloneFlags::WithSkeleton ))
+	if (EnumHasAnyFlags(Flags, EMeshCloneFlags::WithSkeleton ))
 	{
 		pResult->m_pSkeleton = m_pSkeleton;
 	}
 
-	if ( EnumHasAnyFlags( Flags, EMeshCloneFlags::WithPhysicsBody ))
+	if (EnumHasAnyFlags(Flags, EMeshCloneFlags::WithPhysicsBody ))
 	{
 		pResult->m_pPhysicsBody = m_pPhysicsBody;
 	}
 
-	if ( EnumHasAnyFlags( Flags, EMeshCloneFlags::WithFaceGroups ))
+	if (EnumHasAnyFlags(Flags, EMeshCloneFlags::WithFaceGroups ))
 	{
 		pResult->m_faceGroups = m_faceGroups;
 	}
 
-	if ( EnumHasAnyFlags( Flags, EMeshCloneFlags::WithTags ))
+	if (EnumHasAnyFlags(Flags, EMeshCloneFlags::WithTags ))
 	{
 		pResult->m_tags = m_tags;
 	}
 
     // Clone the main buffers
-	if ( EnumHasAnyFlags( Flags, EMeshCloneFlags::WithVertexBuffers ))
+	if (EnumHasAnyFlags(Flags, EMeshCloneFlags::WithVertexBuffers ))
     {
 		pResult->m_VertexBuffers = m_VertexBuffers;
 	}
 
-	if ( EnumHasAnyFlags( Flags, EMeshCloneFlags::WithIndexBuffers ))
+	if (EnumHasAnyFlags(Flags, EMeshCloneFlags::WithIndexBuffers ))
     {
 		pResult->m_IndexBuffers = m_IndexBuffers;
 	}
 
-	if ( EnumHasAnyFlags( Flags, EMeshCloneFlags::WithFaceBuffers ))
+	if (EnumHasAnyFlags(Flags, EMeshCloneFlags::WithFaceBuffers ))
     {
 		pResult->m_FaceBuffers = m_FaceBuffers;
 	}
 
 	// Clone additional buffers
-	if ( EnumHasAnyFlags( Flags, EMeshCloneFlags::WithAdditionalBuffers ))
+	if (EnumHasAnyFlags(Flags, EMeshCloneFlags::WithAdditionalBuffers))
 	{
 		pResult->m_AdditionalBuffers = m_AdditionalBuffers;
 	}
 
     // Clone the layout	
-	if ( EnumHasAnyFlags( Flags, EMeshCloneFlags::WithLayouts ))
+	if (EnumHasAnyFlags(Flags, EMeshCloneFlags::WithLayouts))
 	{
 		pResult->m_layouts = m_layouts;
 	}
@@ -146,7 +146,7 @@ MeshPtr Mesh::Clone( EMeshCloneFlags Flags ) const
 	// Clone bone poses
 	if (EnumHasAnyFlags(Flags, EMeshCloneFlags::WithPoses))
 	{
-		pResult->m_bonePoses = m_bonePoses;
+		pResult->BonePoses = BonePoses;
 	}
 
     return pResult;
@@ -459,29 +459,21 @@ void Mesh::SetTag( int tagIndex, const char* strName )
 
 
 //---------------------------------------------------------------------------------------------
-const char* Mesh::GetBonePoseName(int32 boneIndex) const
+const char* Mesh::GetBonePoseName(int32 BoneIndex) const
 {
-	check(boneIndex >= 0 && boneIndex < GetBonePoseCount());
-	if (boneIndex >= 0 && boneIndex < GetBonePoseCount())
+	check(BoneIndex >= 0 && BoneIndex < GetBonePoseCount());
+	if (BoneIndex >= 0 && BoneIndex < GetBonePoseCount())
 	{
-		return m_bonePoses[boneIndex].m_boneName.c_str();
+		return BonePoses[BoneIndex].BoneName.c_str();
 	}
 	return "";
 }
 
 
 //---------------------------------------------------------------------------------------------
-int32 Mesh::FindBonePose(const char* strName) const
+int32 Mesh::FindBonePose(const char* StrName) const
 {
-	for (int32 i = 0; i < m_bonePoses.Num(); ++i)
-	{
-		if (m_bonePoses[i].m_boneName == strName)
-		{
-			return i;
-		}
-	}
-
-	return INDEX_NONE;
+	return BonePoses.IndexOfByPredicate([StrName](const FBonePose& Pose) { return Pose.BoneName == StrName; });
 }
 
 
@@ -489,26 +481,23 @@ int32 Mesh::FindBonePose(const char* strName) const
 void mu::Mesh::SetBonePoseCount(int32 count)
 {
 	LLM_SCOPE_BYNAME(TEXT("MutableRuntime"));
-	m_bonePoses.SetNum(count);
+	BonePoses.SetNum(count);
 }
 
 
 //---------------------------------------------------------------------------------------------
 int32 mu::Mesh::GetBonePoseCount() const
 {
-	return m_bonePoses.Num();
+	return BonePoses.Num();
 }
 
 
 //---------------------------------------------------------------------------------------------
-void mu::Mesh::SetBonePose(int32 boneIndex, const char* strName, FTransform3f transform, bool bSkinned)
+void mu::Mesh::SetBonePose(int32 BoneIndex, const char* StrName, FTransform3f Transform, EBoneUsageFlags BoneUsageFlags)
 {
-	if (boneIndex >= 0 && boneIndex < m_bonePoses.Num())
+	if (BoneIndex >= 0 && BoneIndex < BonePoses.Num())
 	{
-		FBonePose& pose = m_bonePoses[boneIndex];
-		pose.m_boneName = strName;
-		pose.m_boneTransform = transform;
-		pose.m_boneSkinned = bSkinned;
+		BonePoses[BoneIndex] = FBonePose{ StrName, BoneUsageFlags, Transform };
 	}
 	else 
 	{
@@ -518,10 +507,10 @@ void mu::Mesh::SetBonePose(int32 boneIndex, const char* strName, FTransform3f tr
 
 
 //---------------------------------------------------------------------------------------------
-void mu::Mesh::GetBoneTransform(int32 boneIndex, FTransform3f& transform) const
+void mu::Mesh::GetBoneTransform(int32 BoneIndex, FTransform3f& Transform) const
 {
-	check(boneIndex >= 0 && boneIndex < m_bonePoses.Num());
-	transform = boneIndex >= INDEX_NONE ? m_bonePoses[boneIndex].m_boneTransform : FTransform3f::Identity;
+	check(BoneIndex >= 0 && BoneIndex < BonePoses.Num());
+	Transform = BoneIndex > INDEX_NONE ? BonePoses[BoneIndex].BoneTransform : FTransform3f::Identity;
 }
 
 
diff --git a/Engine/Plugins/Experimental/Mutable/Source/MutableRuntime/Private/MuR/OpMeshApplyPose.h b/Engine/Plugins/Experimental/Mutable/Source/MutableRuntime/Private/MuR/OpMeshApplyPose.h
--- a/Engine/Plugins/Experimental/Mutable/Source/MutableRuntime/Private/MuR/OpMeshApplyPose.h
+++ b/Engine/Plugins/Experimental/Mutable/Source/MutableRuntime/Private/MuR/OpMeshApplyPose.h
@@ -40,7 +40,7 @@ namespace mu
 
             if( poseBoneIndex != INDEX_NONE )
             {
-				skinTransforms[b] = pPose->m_bonePoses[poseBoneIndex].m_boneTransform;
+				skinTransforms[b] = pPose->BonePoses[poseBoneIndex].BoneTransform;
             }
             else
             {
diff --git a/Engine/Plugins/Experimental/Mutable/Source/MutableRuntime/Private/MuR/OpMeshBind.h b/Engine/Plugins/Experimental/Mutable/Source/MutableRuntime/Private/MuR/OpMeshBind.h
--- a/Engine/Plugins/Experimental/Mutable/Source/MutableRuntime/Private/MuR/OpMeshBind.h
+++ b/Engine/Plugins/Experimental/Mutable/Source/MutableRuntime/Private/MuR/OpMeshBind.h
@@ -458,31 +458,19 @@ namespace mu
 
 	inline TTuple<TArray<FReshapeVertexBindingData>, TArray<int32>> BindPhysicsBodies( 
 			const PhysicsBody* PBody, FShapeMeshTree& ShapeMeshTree, const Mesh* pMesh, 
-			bool bDeformAllPhysics, const TArray<string>& PhysicsToDeform )
+			const TArray<string>& PhysicsToDeform )
 	{
 		TTuple<TArray<FReshapeVertexBindingData>, TArray<int32>> ReturnValue;
 
 		TArray<int32>& BodiesToDeformIndices = ReturnValue.Get<1>();
 	
-		if (!bDeformAllPhysics)
-		{ 
-			BodiesToDeformIndices.Reserve(PhysicsToDeform.Num());
-			const int32 NumBodies = PBody->GetBodyCount();
-			for (int32 I = 0; I < NumBodies; ++I)
-			{
-				if ( PhysicsToDeform.Contains(string(PBody->GetBodyBoneName(I))) )
-				{
-					BodiesToDeformIndices.Add(I);
-				}
-			}
-		}
-		else
+		BodiesToDeformIndices.Reserve(PhysicsToDeform.Num());
+		const int32 NumBodies = PBody->GetBodyCount();
+		for (int32 I = 0; I < NumBodies; ++I)
 		{
-			const int32 NumBodies = PBody->GetBodyCount();
-			BodiesToDeformIndices.SetNumUninitialized(NumBodies);
-			for (int32 I = 0; I < NumBodies; ++I)
+			if ( PhysicsToDeform.Contains(PBody->GetBodyBoneName(I)))
 			{
-				BodiesToDeformIndices[I] = I;
+				BodiesToDeformIndices.Add(I);
 			}
 		}
 		
@@ -733,8 +721,7 @@ namespace mu
 	}
 
 	inline TTuple<TArray<FReshapeVertexBindingData>, TArray<int32>> BindPose(
-			const Mesh* Mesh, FShapeMeshTree& ShapeMeshTree, 
-			bool bDeformAllBones, const TArray<string>& BonesToDeform )
+			const Mesh* Mesh, FShapeMeshTree& ShapeMeshTree, const TArray<string>& BonesToDeform )
 	{
 		UE::Geometry::FAxisAlignedBox3d ShapeAABBox = ShapeMeshTree.GetBoundingBox();
 
@@ -746,36 +733,34 @@ namespace mu
 
 		TArray<FReshapeVertexBindingData>& SkeletonBindDataArray = ReturnValue.Get<0>();
 		TArray<int32>& BoneIndices = ReturnValue.Get<1>();
-		
-		for ( int32 b = 0; b < Mesh->GetBonePoseCount(); ++b)
-		{
-			if (!bDeformAllBones)
-			{
-				bool bFound = false;
-				for (int32 bd = 0; bd < BonesToDeform.Num(); ++bd)
-				{
-					if (Mesh->GetBonePoseName(b) == BonesToDeform[bd])
-					{
-						bFound = true;
-						break;
-					}
-				}
 
-				if (!bFound)
-				{
-					continue;
-				}
-			}
+		const int32 BoneCount = Mesh->GetBonePoseCount();
+		SkeletonBindDataArray.Reserve(BoneCount);
+		BoneIndices.Reserve(BoneCount);
 
-			const FTransform3f& t = Mesh->m_bonePoses[b].m_boneTransform;
+		for ( int32 BoneIndex = 0; BoneIndex < BoneCount; ++BoneIndex)
+		{
+			const EBoneUsageFlags BoneUsageFlags = Mesh->BonePoses[BoneIndex].BoneUsageFlags;
 
-			FVector3f Origin = t.GetLocation();
+			// Only deform skinned bones.
+			if (EnumHasAnyFlags(BoneUsageFlags, EBoneUsageFlags::SkinningParent | EBoneUsageFlags::Root) &&
+				!EnumHasAnyFlags(BoneUsageFlags, EBoneUsageFlags::Skinning))
+			{
+				continue;
+			}
 
-			FReshapeVertexBindingData SkeletonBindData;
-			BindReshapePoint(ShapeMeshTree, Origin, SkeletonBindData, BindValidityTolerance);
+			if (!BonesToDeform.Contains(Mesh->GetBonePoseName(BoneIndex)))
+			{
+				continue;
+			}
 
-			SkeletonBindDataArray.Add(SkeletonBindData);
-			BoneIndices.Add(b);
+			BoneIndices.Add(BoneIndex);
+			
+			BindReshapePoint(
+					ShapeMeshTree, 
+					Mesh->BonePoses[BoneIndex].BoneTransform.GetLocation(),
+					SkeletonBindDataArray.AddDefaulted_GetRef(),
+					BindValidityTolerance);
 		}
 
 		return ReturnValue;
@@ -787,9 +772,8 @@ namespace mu
     //---------------------------------------------------------------------------------------------
     inline MeshPtr MeshBindShapeReshape(
 			const Mesh* BaseMesh, const Mesh* ShapeMesh, 
-			const TArray<string>& BonesToDeform, const TArray<string>& PhysicsToDeform,
-			bool bDeformAllBones, bool bDeformAllPhyiscs,
-			bool reshapeVertices, bool reshapeSkeleton, bool reshapePhysicsVolumes, bool bEnableRigidParts )
+			const TArray<string>& BonesToDeform, const TArray<string>& PhysicsToDeform, 
+			EMeshBindShapeFlags BindFlags)
     {
 		MUTABLE_CPUPROFILER_SCOPE(MeshBindShape);
 		
@@ -798,12 +782,17 @@ namespace mu
 			return nullptr;
 		}
 
+		const bool bReshapeVertices = EnumHasAnyFlags(BindFlags, EMeshBindShapeFlags::ReshapeVertices);
+		const bool bReshapeSkeleton = EnumHasAnyFlags(BindFlags, EMeshBindShapeFlags::ReshapeSkeleton);
+		const bool bReshapePhysics = EnumHasAnyFlags(BindFlags, EMeshBindShapeFlags::ReshapePhysicsVolumes);
+		const bool bEnableRigidParts = EnumHasAnyFlags(BindFlags, EMeshBindShapeFlags::EnableRigidParts);
+
 		// Early out if nothing will be modified and the vertices discarted. return null in this
 		// case indicating nothing has modified so the Base Mesh can be reused.
-		const bool bSkeletonModification = BaseMesh->GetSkeleton() && reshapeSkeleton;
-		const bool bPhysicsModification = BaseMesh->GetPhysicsBody() && reshapePhysicsVolumes;
+		const bool bSkeletonModification = BaseMesh->GetSkeleton() && bReshapeSkeleton;
+		const bool bPhysicsModification = BaseMesh->GetPhysicsBody() && bReshapePhysics;
 
-		if (!reshapeVertices && !bSkeletonModification && !bPhysicsModification)
+		if (!bReshapeVertices && !bSkeletonModification && !bPhysicsModification)
 		{
 			return nullptr;
 		}
@@ -876,7 +865,7 @@ namespace mu
 		
 		// If no vertices are needed, it is assumed we only want to reshape physics or skeleton
 		// In that case, remove everything except physics bodies, the skeleton and pose.
-		if (!reshapeVertices)
+		if (!bReshapeVertices)
 		{
 			constexpr EMeshCloneFlags CloneFlags = 
 					EMeshCloneFlags::WithSkeleton    | 
@@ -891,7 +880,7 @@ namespace mu
 		}
 
 		int32 BindingDataIndex = 0;
-		if (reshapeVertices)
+		if (bReshapeVertices)
 		{
 			TArray<FReshapeVertexBindingData> VerticesBindData = BindVerticesReshape( BaseMesh, ShapeMeshTree, bEnableRigidParts );
 			
@@ -909,12 +898,12 @@ namespace mu
 
 		// Bind the skeleton bones
 		// \TODO: Build bind data only for actually modified bones?
-		if (reshapeSkeleton && (bDeformAllBones || BonesToDeform.Num()))
+		if (bReshapeSkeleton && BonesToDeform.Num())
 		{
 			MUTABLE_CPUPROFILER_SCOPE(BindSkeleton);
 
 			TTuple<TArray<FReshapeVertexBindingData>, TArray<int32>> SkeletonBindingData = 
-					BindPose(Result.get(), ShapeMeshTree, bDeformAllBones, BonesToDeform);
+					BindPose(Result.get(), ShapeMeshTree, BonesToDeform);
 			const TArray<FReshapeVertexBindingData>& SkeletonBindDataArray = SkeletonBindingData.Get<0>();
 			const TArray<int32>& BoneIndices = SkeletonBindingData.Get<1>();
 
@@ -944,13 +933,13 @@ namespace mu
 			Result->m_AdditionalBuffers.Emplace(EMeshBufferType::SkeletonDeformBinding, MoveTemp(SkeletonBuffer));
 		}
 
-		const PhysicsBody* PhysicsBody = Result->m_pPhysicsBody.get();
-		if (reshapePhysicsVolumes && PhysicsBody && (bDeformAllPhyiscs || PhysicsToDeform.Num()))
+		const PhysicsBody* ResultPhysicsBody = Result->m_pPhysicsBody.get();
+		if (bReshapePhysics && ResultPhysicsBody && PhysicsToDeform.Num())
 		{
 			MUTABLE_CPUPROFILER_SCOPE(BindPhysicsBody);
 
 			TTuple<TArray<FReshapeVertexBindingData>, TArray<int32>> PhysicsBindingData = 
-					BindPhysicsBodies(PhysicsBody, ShapeMeshTree, Result.get(), bDeformAllPhyiscs, PhysicsToDeform);
+					BindPhysicsBodies(ResultPhysicsBody, ShapeMeshTree, Result.get(), PhysicsToDeform);
 			const TArray<FReshapeVertexBindingData>& PhysicsBindDataArray = PhysicsBindingData.Get<0>();
 			const TArray<int32>& DeformedBodyIndices = PhysicsBindingData.Get<1>();
 
diff --git a/Engine/Plugins/Experimental/Mutable/Source/MutableRuntime/Private/MuR/OpMeshFormat.cpp b/Engine/Plugins/Experimental/Mutable/Source/MutableRuntime/Private/MuR/OpMeshFormat.cpp
--- a/Engine/Plugins/Experimental/Mutable/Source/MutableRuntime/Private/MuR/OpMeshFormat.cpp
+++ b/Engine/Plugins/Experimental/Mutable/Source/MutableRuntime/Private/MuR/OpMeshFormat.cpp
@@ -715,7 +715,7 @@ namespace mu
 
 		pResult->m_AdditionalBuffers = pSource->m_AdditionalBuffers;
 
-		pResult->m_bonePoses = pSource->m_bonePoses;
+		pResult->BonePoses = pSource->BonePoses;
 
 		return pResult;
 
diff --git a/Engine/Plugins/Experimental/Mutable/Source/MutableRuntime/Private/MuR/OpMeshMerge.h b/Engine/Plugins/Experimental/Mutable/Source/MutableRuntime/Private/MuR/OpMeshMerge.h
--- a/Engine/Plugins/Experimental/Mutable/Source/MutableRuntime/Private/MuR/OpMeshMerge.h
+++ b/Engine/Plugins/Experimental/Mutable/Source/MutableRuntime/Private/MuR/OpMeshMerge.h
@@ -515,35 +515,43 @@ namespace mu
 		{
 			MUTABLE_CPUPROFILER_SCOPE(Pose);
 
-			pResult->m_bonePoses.Reserve(pResult->GetSkeleton()->GetBoneCount());
+			pResult->BonePoses.Reserve(pResult->GetSkeleton()->GetBoneCount());
 			
 			// Copy poses from the first mesh
-			pResult->m_bonePoses = pFirst->m_bonePoses;
+			pResult->BonePoses = pFirst->BonePoses;
 
 			// Add or override bone poses
-			for (const Mesh::FBonePose& pSecondBonePose : pSecond->m_bonePoses)
+			for (const Mesh::FBonePose& SecondBonePose : pSecond->BonePoses)
 			{
-				const int32 resultBoneIndex = pResult->FindBonePose(pSecondBonePose.m_boneName.c_str());
+				const int32 ResultBoneIndex = pResult->FindBonePose(SecondBonePose.BoneName.c_str());
 
-				if (resultBoneIndex != INDEX_NONE)
+				if (ResultBoneIndex != INDEX_NONE)
 				{
-					Mesh::FBonePose& pResultBonePose = pResult->m_bonePoses[resultBoneIndex];
+					Mesh::FBonePose& ResultBonePose = pResult->BonePoses[ResultBoneIndex];
 
-					// override pose if the bone is not used by the skinning of the first mesh
-					if (!pResultBonePose.m_boneSkinned)
+					// TODO: Not sure how to tune this priority, review it.
+					// For now use a similar strategy as before. 
+					auto ComputeBoneMergePriority = [](const Mesh::FBonePose& BonePose)
 					{
-						pResultBonePose.m_boneName = pSecondBonePose.m_boneName;
-						pResultBonePose.m_boneSkinned = pSecondBonePose.m_boneSkinned;
-						pResultBonePose.m_boneTransform = pSecondBonePose.m_boneTransform;
-					}
+						return (EnumHasAnyFlags(BonePose.BoneUsageFlags, EBoneUsageFlags::Skinning) ? 1 : 0) +
+							   (EnumHasAnyFlags(BonePose.BoneUsageFlags, EBoneUsageFlags::Reshaped) ? 1 : 0);
+					};
+
+					if (ComputeBoneMergePriority(ResultBonePose) < ComputeBoneMergePriority(SecondBonePose))
+					{
+						//ResultBonePose.BoneName = SecondBonePose.BoneName;
+						ResultBonePose.BoneTransform = SecondBonePose.BoneTransform;
+						// Merge usage flags
+						EnumAddFlags(ResultBonePose.BoneUsageFlags, SecondBonePose.BoneUsageFlags);					
+					}	
 				}
 				else
 				{
-					pResult->m_bonePoses.Add(pSecondBonePose);
+					pResult->BonePoses.Add(SecondBonePose);
 				}
 			}
 
-			pResult->m_bonePoses.Shrink();
+			pResult->BonePoses.Shrink();
 		}
 
 		
diff --git a/Engine/Plugins/Experimental/Mutable/Source/MutableRuntime/Private/MuR/OpMeshReshape.h b/Engine/Plugins/Experimental/Mutable/Source/MutableRuntime/Private/MuR/OpMeshReshape.h
--- a/Engine/Plugins/Experimental/Mutable/Source/MutableRuntime/Private/MuR/OpMeshReshape.h
+++ b/Engine/Plugins/Experimental/Mutable/Source/MutableRuntime/Private/MuR/OpMeshReshape.h
@@ -366,19 +366,24 @@ namespace mu
 
 			const FReshapeVertexBindingData& Binding = BindingData[b];
 
-			FTransform3f& t = Result->m_bonePoses[BoneIndex].m_boneTransform;
+			check(!EnumHasAnyFlags(Result->BonePoses[BoneIndex].BoneUsageFlags, EBoneUsageFlags::Root));
 
-			FVector3f NewPosition;
+			FTransform3f& T = Result->BonePoses[BoneIndex].BoneTransform;
+
+			FVector3f NewPosition(ForceInitToZero);
 			FQuat4f TangentSpaceCorrection;
-			bool bModified = GetDeform(Shape, Binding, NewPosition, TangentSpaceCorrection);
 
-			if (bModified)
-			{
-				// The pose has been modified and should not be replaced
-				Result->m_bonePoses[BoneIndex].m_boneSkinned = true;
+			const bool bModified = GetDeform(Shape, Binding, NewPosition, TangentSpaceCorrection);
+			const bool bHasChanged = bModified && FVector3f::DistSquared(NewPosition, T.GetLocation()) > UE_SMALL_NUMBER;
+
+			// Only set it if has actually moved.
+			if (bHasChanged)
+			{	
+				// Mark as reshaped. 
+				EnumAddFlags(Result->BonePoses[BoneIndex].BoneUsageFlags, EBoneUsageFlags::Reshaped);
 
 				// TODO: Review if the rotation also needs to be applied.
-				t.SetLocation(FMath::Lerp(t.GetLocation(), NewPosition, Binding.Weight));
+				T.SetLocation(FMath::Lerp(T.GetLocation(), NewPosition, Binding.Weight));
 			}
 		}
 	}
@@ -535,8 +540,7 @@ namespace mu
 	//! Rebuild the (previously bound) mesh data for a new shape.
 	//! Proof-of-concept implementation.
 	//---------------------------------------------------------------------------------------------
-	inline MeshPtr MeshApplyShape(const Mesh* BaseMesh, const Mesh* ShapeMesh, 
-			bool bReshapeVertices, bool bReshapeSkeleton, bool bReshapePhysicsVolumes)
+	inline MeshPtr MeshApplyShape(const Mesh* BaseMesh, const Mesh* ShapeMesh, EMeshBindShapeFlags BindFlags)
 	{
 		MUTABLE_CPUPROFILER_SCOPE(MeshApplyReshape);
 
@@ -545,6 +549,10 @@ namespace mu
 			return nullptr;
 		}
 		
+		const bool bReshapeVertices = EnumHasAnyFlags(BindFlags, EMeshBindShapeFlags::ReshapeVertices);
+		const bool bReshapeSkeleton = EnumHasAnyFlags(BindFlags, EMeshBindShapeFlags::ReshapeSkeleton);
+		const bool bReshapePhysicsVolumes = EnumHasAnyFlags(BindFlags, EMeshBindShapeFlags::ReshapePhysicsVolumes);
+
 		// Early out if nothing will be modified and the vertices discarted.
 		const bool bSkeletonModification = BaseMesh->GetSkeleton() && bReshapeSkeleton;
 		const bool bPhysicsModification = BaseMesh->GetPhysicsBody() && bReshapePhysicsVolumes;
diff --git a/Engine/Plugins/Experimental/Mutable/Source/MutableRuntime/Private/MuR/Operations.h b/Engine/Plugins/Experimental/Mutable/Source/MutableRuntime/Private/MuR/Operations.h
--- a/Engine/Plugins/Experimental/Mutable/Source/MutableRuntime/Private/MuR/Operations.h
+++ b/Engine/Plugins/Experimental/Mutable/Source/MutableRuntime/Private/MuR/Operations.h
@@ -8,9 +8,9 @@
 #include "HAL/PlatformMath.h"
 #include "HAL/UnrealMemory.h"
 #include "Misc/AssertionMacros.h"
+#include "Misc/EnumClassFlags.h"
 #include "Templates/Function.h"
 
-
 #define MUTABLE_OP_MAX_INTERPOLATE_COUNT	6
 #define MUTABLE_OP_MAX_ADD_COUNT			7
 #define MUTABLE_OP_MAX_SWIZZLE_CHANNELS		4
@@ -355,7 +355,15 @@ namespace mu
         COUNT
 
     };
-
+	enum class EMeshBindShapeFlags : uint32
+	{
+		None				   = 0,
+		ReshapeSkeleton		   = 1 << 0,
+		EnableRigidParts       = 1 << 2,
+		ReshapePhysicsVolumes  = 1 << 4,
+		ReshapeVertices		   = 1 << 5,
+	};
+	ENUM_CLASS_FLAGS(EMeshBindShapeFlags);
 
     //---------------------------------------------------------------------------------------------
     //!
@@ -904,18 +912,6 @@ namespace mu
 			ADDRESS scalarB;
 		};
 
-		enum class EMeshBindShapeFlags : uint32
-		{
-			None				   = 0,
-			ReshapeSkeleton		   = 1 << 0,
-			DiscardInvalidBindings = 1 << 1,
-			EnableRigidParts       = 1 << 2,
-			DeformAllBones		   = 1 << 3,
-			ReshapePhysicsVolumes  = 1 << 4,
-			ReshapeVertices		   = 1 << 5,
-			DeformAllPhysics	   = 1 << 6,
-		};
-
 		struct MeshBindShapeArgs
 		{
 			ADDRESS mesh;
diff --git a/Engine/Plugins/Experimental/Mutable/Source/MutableRuntime/Public/MuR/Mesh.h b/Engine/Plugins/Experimental/Mutable/Source/MutableRuntime/Public/MuR/Mesh.h
--- a/Engine/Plugins/Experimental/Mutable/Source/MutableRuntime/Public/MuR/Mesh.h
+++ b/Engine/Plugins/Experimental/Mutable/Source/MutableRuntime/Public/MuR/Mesh.h
@@ -75,6 +75,21 @@ namespace mu
 
 	MUTABLE_DEFINE_POD_VECTOR_SERIALISABLE(MESH_SURFACE);
 
+	enum class EBoneUsageFlags : uint32
+	{
+		None		   = 0,
+		Root		   = 1 << 1,
+		Skinning	   = 1 << 2,
+		SkinningParent = 1 << 3,
+		Physics	       = 1 << 4,
+		PhysicsParent  = 1 << 5,
+		Deform         = 1 << 6,
+		DeformParent   = 1 << 7,
+		Reshaped       = 1 << 8	
+	};
+
+	ENUM_CLASS_FLAGS(EBoneUsageFlags);
+	MUTABLE_DEFINE_ENUM_SERIALISABLE(EBoneUsageFlags)
 
 	//!
 	enum class EMeshBufferType
@@ -112,8 +127,8 @@ namespace mu
 		WithLayouts           = 1 << 11,
 		WithPoses			  = 1 << 12
 	};
-	ENUM_CLASS_FLAGS(EMeshCloneFlags);
 	
+	ENUM_CLASS_FLAGS(EMeshCloneFlags);
 
     //! \brief Mesh object containing any number of buffers with any number of channels.
     //! The buffers can be per-index or per-vertex.
@@ -269,10 +284,10 @@ namespace mu
 		int32 GetBonePoseCount() const;
 
 		//!
-		void SetBonePose(int32 boneIndex, const char* strName, FTransform3f transform, bool bSkinned );
+		void SetBonePose(int32 BoneIndex, const char* StrName, FTransform3f Transform, EBoneUsageFlags BoneUsageFlags);
 
 		//! Return a matrix stored per bone. It is a set of 16-float values.
-		void GetBoneTransform(int32 boneIndex, FTransform3f& transform) const;
+		void GetBoneTransform(int32 BoneIndex, FTransform3f& Transform) const;
 
         //! \}
 
@@ -361,48 +376,53 @@ namespace mu
 
 		struct FBonePose
 		{
-			string m_boneName;
-			uint8 m_boneSkinned : 1;
-			FTransform3f m_boneTransform;
-
-			FBonePose() : m_boneSkinned(0) {};
+			string BoneName;
+			EBoneUsageFlags BoneUsageFlags = EBoneUsageFlags::None;
+			FTransform3f BoneTransform;
 
 			inline void Serialise(OutputArchive& arch) const
 			{
-				const int32 ver = 0;
+				const int32 ver = 1;
 				arch << ver;
 
-				arch << m_boneName;
-				arch << m_boneSkinned;
-				arch << m_boneTransform;
+				arch << BoneName;
+				arch << BoneUsageFlags;
+				arch << BoneTransform;
 			}
 
 			inline void Unserialise(InputArchive& arch)
 			{
 				int32 ver = 0;
 				arch >> ver;
-				check(ver == 0);
+				check(ver <= 1);
 
-				arch >> m_boneName;
+				arch >> BoneName;
 
-				uint8 skinned = 0;
-				arch >> skinned;
-				m_boneSkinned = skinned;
+				if (ver == 0)
+				{
+					uint8 Skinned = 0;
+					arch >> Skinned;
+					BoneUsageFlags = Skinned ? EBoneUsageFlags::Skinning : EBoneUsageFlags::None;
+				}
+				else
+				{
+					arch >> BoneUsageFlags;
+				}
 
-				arch >> m_boneTransform;
+				arch >> BoneTransform;
 			}
 
+
 			//!
-			inline bool operator==(const FBonePose& o) const
+			inline bool operator==(const FBonePose& Other) const
 			{
-				return m_boneSkinned == o.m_boneSkinned
-					&& m_boneName == o.m_boneName;
+				return BoneUsageFlags == Other.BoneUsageFlags
+					&& BoneName == Other.BoneName;
 			}
-
 		};
 		// This is the pose used by this mesh fragment, used to update the transforms of the final skeleton
 		// taking into consideration the meshes being used.
-		TArray<FBonePose> m_bonePoses;
+		TArray<FBonePose> BonePoses;
 
 		//!
 		inline void Serialise(OutputArchive& arch) const
@@ -425,7 +445,7 @@ namespace mu
 
 			arch << m_tags;
 
-			arch << m_bonePoses;
+			arch << BonePoses;
 		}
 
 		//!
@@ -459,19 +479,19 @@ namespace mu
 
 			if (ver >= 13)
 			{
-				arch >> m_bonePoses;
+				arch >> BonePoses;
 			}
 			else if (m_pSkeleton)
 			{
 				const int32 NumBones = m_pSkeleton->GetBoneCount();
-				m_bonePoses.SetNum(NumBones);
+				BonePoses.SetNum(NumBones);
 				check(m_pSkeleton->m_boneTransforms_DEPRECATED.Num() == NumBones);
 
 				for (int32 BoneIndex = 0; BoneIndex < NumBones; ++BoneIndex)
 				{
-					m_bonePoses[BoneIndex].m_boneName = m_pSkeleton->GetBoneName(BoneIndex);
-					m_bonePoses[BoneIndex].m_boneSkinned = true;
-					m_bonePoses[BoneIndex].m_boneTransform = m_pSkeleton->m_boneTransforms_DEPRECATED[BoneIndex];
+					BonePoses[BoneIndex].BoneName = m_pSkeleton->GetBoneName(BoneIndex);
+					BonePoses[BoneIndex].BoneUsageFlags = EBoneUsageFlags::Skinning;
+					BonePoses[BoneIndex].BoneTransform = m_pSkeleton->m_boneTransforms_DEPRECATED[BoneIndex];
 				}
 			}
 		}
@@ -484,7 +504,7 @@ namespace mu
 			if (equal) equal = (m_VertexBuffers == o.m_VertexBuffers);
 			if (equal) equal = (m_FaceBuffers == o.m_FaceBuffers);
 			if (equal) equal = (m_layouts.Num() == o.m_layouts.Num());
-			if (equal) equal = (m_bonePoses.Num() == o.m_bonePoses.Num());
+			if (equal) equal = (BonePoses.Num() == o.BonePoses.Num());
 			if (equal && m_pSkeleton != o.m_pSkeleton)
 			{
 				if (m_pSkeleton && o.m_pSkeleton)
@@ -510,9 +530,9 @@ namespace mu
 				equal &= m_AdditionalBuffers[i] == o.m_AdditionalBuffers[i];
 			}
 
-			for (int32 i = 0; equal && i < m_bonePoses.Num(); ++i)
+			for (int32 i = 0; equal && i < BonePoses.Num(); ++i)
 			{
-				equal &= m_bonePoses[i] == o.m_bonePoses[i];
+				equal &= BonePoses[i] == o.BonePoses[i];
 			}
 
 			return equal;
diff --git a/Engine/Plugins/Experimental/Mutable/Source/MutableTools/Private/MuT/ASTOpMeshApplyShape.cpp b/Engine/Plugins/Experimental/Mutable/Source/MutableTools/Private/MuT/ASTOpMeshApplyShape.cpp
--- a/Engine/Plugins/Experimental/Mutable/Source/MutableTools/Private/MuT/ASTOpMeshApplyShape.cpp
+++ b/Engine/Plugins/Experimental/Mutable/Source/MutableTools/Private/MuT/ASTOpMeshApplyShape.cpp
@@ -20,6 +20,9 @@ namespace mu
 	ASTOpMeshApplyShape::ASTOpMeshApplyShape()
 		: Mesh(this)
 		, Shape(this)
+		, bReshapeSkeleton(false)
+		, bReshapePhysicsVolumes(false)
+		, bReshapeVertices(true)
 	{
 	}
 
@@ -31,16 +34,16 @@ namespace mu
 	}
 
 
-	bool ASTOpMeshApplyShape::IsEqual(const ASTOp& otherUntyped) const
+	bool ASTOpMeshApplyShape::IsEqual(const ASTOp& OtherUntyped) const
 	{
-		if (auto other = dynamic_cast<const ASTOpMeshApplyShape*>(&otherUntyped))
+		if (const ASTOpMeshApplyShape* Other = dynamic_cast<const ASTOpMeshApplyShape*>(&OtherUntyped))
 		{
 			const bool bSameFlags =
-				m_reshapePhysicsVolumes == other->m_reshapePhysicsVolumes &&
-				m_reshapeSkeleton == other->m_reshapeSkeleton &&
-				m_reshapeVertices == other->m_reshapeVertices;
+				bReshapePhysicsVolumes == Other->bReshapePhysicsVolumes &&
+				bReshapeSkeleton == Other->bReshapeSkeleton &&
+				bReshapeVertices == Other->bReshapeVertices;
 
-			return Mesh == other->Mesh && Shape == other->Shape && bSameFlags;
+			return Mesh == Other->Mesh && Shape == Other->Shape && bSameFlags;
 		}
 		return false;
 	}
@@ -50,22 +53,22 @@ namespace mu
 	{
 		uint64 res = std::hash<void*>()(Mesh.child().get());
 		hash_combine(res, Shape.child().get());
-		hash_combine(res, m_reshapeSkeleton);
-		hash_combine(res, m_reshapePhysicsVolumes);
-		hash_combine(res, m_reshapeVertices);
+		hash_combine(res, bool(bReshapeSkeleton));
+		hash_combine(res, bool(bReshapePhysicsVolumes));
+		hash_combine(res, bool(bReshapeVertices));
 		return res;
 	}
 
 
 	mu::Ptr<ASTOp> ASTOpMeshApplyShape::Clone(MapChildFuncRef mapChild) const
 	{
-		Ptr<ASTOpMeshApplyShape> n = new ASTOpMeshApplyShape();
-		n->Mesh = mapChild(Mesh.child());
-		n->Shape = mapChild(Shape.child());
-		n->m_reshapeSkeleton = m_reshapeSkeleton;
-		n->m_reshapePhysicsVolumes = m_reshapePhysicsVolumes;
-		n->m_reshapeVertices = m_reshapeVertices;
-		return n;
+		Ptr<ASTOpMeshApplyShape> NewOp = new ASTOpMeshApplyShape();
+		NewOp->Mesh = mapChild(Mesh.child());
+		NewOp->Shape = mapChild(Shape.child());
+		NewOp->bReshapeSkeleton = bReshapeSkeleton;
+		NewOp->bReshapePhysicsVolumes = bReshapePhysicsVolumes;
+		NewOp->bReshapeVertices = bReshapeVertices;
+		return NewOp;
 	}
 
 
@@ -81,32 +84,24 @@ namespace mu
 		// Already linked?
 		if (!linkedAddress)
 		{
-			OP::MeshApplyShapeArgs args;
-			memset(&args, 0, sizeof(args));
+			OP::MeshApplyShapeArgs Args;
+			FMemory::Memzero(&Args, sizeof(Args));
 
-			if (m_reshapeSkeleton)
-			{
-				args.flags |= uint32(OP::EMeshBindShapeFlags::ReshapeSkeleton);
-			}
+			constexpr EMeshBindShapeFlags NoFlags = EMeshBindShapeFlags::None;
+			EMeshBindShapeFlags BindFlags = NoFlags;
+			EnumAddFlags(BindFlags, bReshapeSkeleton ? EMeshBindShapeFlags::ReshapeSkeleton : NoFlags);
+			EnumAddFlags(BindFlags, bReshapePhysicsVolumes ? EMeshBindShapeFlags::ReshapePhysicsVolumes : NoFlags);
+			EnumAddFlags(BindFlags, bReshapeVertices ? EMeshBindShapeFlags::ReshapeVertices : NoFlags);
 
-			if (m_reshapePhysicsVolumes)
-			{
-				args.flags |= uint32(OP::EMeshBindShapeFlags::ReshapePhysicsVolumes);
-			}
+			Args.flags = static_cast<uint32>(BindFlags);
 
-			if (m_reshapeVertices)
-			{
-				args.flags |= uint32(OP::EMeshBindShapeFlags::ReshapeVertices);
-			}
-
-
-			if (Mesh) args.mesh = Mesh->linkedAddress;
-			if (Shape) args.shape = Shape->linkedAddress;
+			Args.mesh = Mesh ? Mesh->linkedAddress : 0;
+			Args.shape = Shape ? Shape->linkedAddress : 0;
 
 			linkedAddress = (OP::ADDRESS)program.m_opAddress.Num();
 			program.m_opAddress.Add((uint32_t)program.m_byteCode.Num());
 			AppendCode(program.m_byteCode, OP_TYPE::ME_APPLYSHAPE);
-			AppendCode(program.m_byteCode, args);
+			AppendCode(program.m_byteCode, Args);
 		}
 
 	}
diff --git a/Engine/Plugins/Experimental/Mutable/Source/MutableTools/Private/MuT/ASTOpMeshApplyShape.h b/Engine/Plugins/Experimental/Mutable/Source/MutableTools/Private/MuT/ASTOpMeshApplyShape.h
--- a/Engine/Plugins/Experimental/Mutable/Source/MutableTools/Private/MuT/ASTOpMeshApplyShape.h
+++ b/Engine/Plugins/Experimental/Mutable/Source/MutableTools/Private/MuT/ASTOpMeshApplyShape.h
@@ -21,9 +21,10 @@ struct FProgram;
 
 		ASTChild Mesh;
 		ASTChild Shape;
-		bool m_reshapeSkeleton = false;
-		bool m_reshapePhysicsVolumes = false;
-		bool m_reshapeVertices = true;
+
+		uint32 bReshapeSkeleton : 1;
+		uint32 bReshapePhysicsVolumes : 1;
+		uint32 bReshapeVertices : 1;
 
 	public:
 
diff --git a/Engine/Plugins/Experimental/Mutable/Source/MutableTools/Private/MuT/ASTOpMeshBindShape.cpp b/Engine/Plugins/Experimental/Mutable/Source/MutableTools/Private/MuT/ASTOpMeshBindShape.cpp
--- a/Engine/Plugins/Experimental/Mutable/Source/MutableTools/Private/MuT/ASTOpMeshBindShape.cpp
+++ b/Engine/Plugins/Experimental/Mutable/Source/MutableTools/Private/MuT/ASTOpMeshBindShape.cpp
@@ -23,6 +23,9 @@ namespace mu
 	ASTOpMeshBindShape::ASTOpMeshBindShape()
 		: Mesh(this)
 		, Shape(this)
+		, bReshapeSkeleton(false)
+		, bEnableRigidParts(false)
+		, bReshapePhysicsVolumes(false)
 	{
 	}
 
@@ -34,25 +37,22 @@ namespace mu
 	}
 
 
-	bool ASTOpMeshBindShape::IsEqual(const ASTOp& otherUntyped) const
+	bool ASTOpMeshBindShape::IsEqual(const ASTOp& OtherUntyped) const
 	{
-		if (auto other = dynamic_cast<const ASTOpMeshBindShape*>(&otherUntyped))
+		if (const ASTOpMeshBindShape* Other = dynamic_cast<const ASTOpMeshBindShape*>(&OtherUntyped))
 		{
 			const bool bSameFlags =
-				m_reshapeSkeleton == other->m_reshapeSkeleton &&
-				m_discardInvalidBindings == other->m_discardInvalidBindings &&
-				m_enableRigidParts == other->m_enableRigidParts &&
-				m_deformAllBones == other->m_deformAllBones &&
-				m_deformAllPhysics == other->m_deformAllPhysics &&
-				m_reshapePhysicsVolumes == other->m_reshapePhysicsVolumes &&
-				m_reshapeVertices == other->m_reshapeVertices;
+				bReshapeSkeleton == Other->bReshapeSkeleton	&&
+				bEnableRigidParts == Other->bEnableRigidParts &&
+				bReshapePhysicsVolumes == Other->bReshapePhysicsVolumes &&
+				bReshapeVertices == Other->bReshapeVertices;
 
 			return bSameFlags &&
-				Mesh == other->Mesh &&
-				Shape == other->Shape &&
-				m_bonesToDeform == other->m_bonesToDeform &&
-				m_physicsToDeform == other->m_physicsToDeform &&
-				m_bindingMethod == other->m_bindingMethod;
+				Mesh == Other->Mesh &&
+				Shape == Other->Shape &&
+				BonesToDeform == Other->BonesToDeform &&
+				PhysicsToDeform == Other->PhysicsToDeform &&
+				BindingMethod == Other->BindingMethod;
 		}
 
 		return false;
@@ -63,21 +63,18 @@ namespace mu
 	{
 		uint64 res = std::hash<void*>()(Mesh.child().get());
 		hash_combine(res, Shape.child().get());
-		hash_combine(res, m_reshapeSkeleton);
-		hash_combine(res, m_discardInvalidBindings);
-		hash_combine(res, m_enableRigidParts);
-		hash_combine(res, m_deformAllBones);
-		hash_combine(res, m_deformAllPhysics);
-		hash_combine(res, m_reshapePhysicsVolumes);
-		hash_combine(res, m_reshapeVertices);
-		hash_combine(res, m_bindingMethod);
-
-		for (const string& S : m_bonesToDeform)
+		hash_combine(res, bool(bReshapeSkeleton));
+		hash_combine(res, bool(bEnableRigidParts));
+		hash_combine(res, bool(bReshapePhysicsVolumes));
+		hash_combine(res, bool(bReshapeVertices));
+		hash_combine(res, bool(BindingMethod));
+
+		for (const string& S : BonesToDeform)
 		{
 			hash_combine(res, S);
 		}
 
-		for (const string& S : m_physicsToDeform)
+		for (const string& S : PhysicsToDeform)
 		{
 			hash_combine(res, S);
 		}
@@ -88,20 +85,18 @@ namespace mu
 
 	mu::Ptr<ASTOp> ASTOpMeshBindShape::Clone(MapChildFuncRef mapChild) const
 	{
-		Ptr<ASTOpMeshBindShape> n = new ASTOpMeshBindShape();
-		n->Mesh = mapChild(Mesh.child());
-		n->Shape = mapChild(Shape.child());
-		n->m_reshapeSkeleton = m_reshapeSkeleton;
-		n->m_discardInvalidBindings = m_discardInvalidBindings;
-		n->m_enableRigidParts = m_enableRigidParts;
-		n->m_deformAllBones = m_deformAllBones;
-		n->m_deformAllPhysics = m_deformAllPhysics;
-		n->m_reshapePhysicsVolumes = m_reshapePhysicsVolumes;
-		n->m_reshapeVertices = m_reshapeVertices;
-		n->m_bonesToDeform = m_bonesToDeform;
-		n->m_physicsToDeform = m_physicsToDeform;
-		n->m_bindingMethod = m_bindingMethod;
-		return n;
+		Ptr<ASTOpMeshBindShape> NewOp = new ASTOpMeshBindShape();
+		NewOp->Mesh = mapChild(Mesh.child());
+		NewOp->Shape = mapChild(Shape.child());
+		NewOp->bReshapeSkeleton	= bReshapeSkeleton;
+		NewOp->bEnableRigidParts = bEnableRigidParts;
+		NewOp->bReshapePhysicsVolumes = bReshapePhysicsVolumes;
+		NewOp->bReshapeVertices = bReshapeVertices;
+		NewOp->BonesToDeform = BonesToDeform;
+		NewOp->PhysicsToDeform = PhysicsToDeform;
+		NewOp->BindingMethod = BindingMethod;
+
+		return NewOp;
 	}
 
 
@@ -117,70 +112,44 @@ namespace mu
 		// Already linked?
 		if (!linkedAddress)
 		{
-			OP::MeshBindShapeArgs args;
-			FMemory::Memzero(&args, sizeof(args));
-
-			if (m_reshapeSkeleton)
-			{
-				args.flags |= uint32(OP::EMeshBindShapeFlags::ReshapeSkeleton);
-			}
-
-			if (m_discardInvalidBindings)
-			{
-				args.flags |= uint32(OP::EMeshBindShapeFlags::DiscardInvalidBindings);
-			}
-
-			if (m_enableRigidParts)
-			{
-				args.flags |= uint32(OP::EMeshBindShapeFlags::EnableRigidParts);
-			}
-
-			if (m_deformAllBones)
-			{
-				args.flags |= uint32(OP::EMeshBindShapeFlags::DeformAllBones);
-			}
+			OP::MeshBindShapeArgs Args;
+			FMemory::Memzero(&Args, sizeof(Args));
 
-			if (m_deformAllPhysics)
-			{
-				args.flags |= uint32(OP::EMeshBindShapeFlags::DeformAllPhysics);
-			}
+			constexpr EMeshBindShapeFlags NoFlags = EMeshBindShapeFlags::None;
+			EMeshBindShapeFlags BindFlags = NoFlags;
+			EnumAddFlags(BindFlags, bReshapeSkeleton ? EMeshBindShapeFlags::ReshapeSkeleton : NoFlags);
+			EnumAddFlags(BindFlags, bEnableRigidParts ? EMeshBindShapeFlags::EnableRigidParts : NoFlags);
+			EnumAddFlags(BindFlags, bReshapePhysicsVolumes ? EMeshBindShapeFlags::ReshapePhysicsVolumes : NoFlags);
+			EnumAddFlags(BindFlags, bReshapeVertices ? EMeshBindShapeFlags::ReshapeVertices : NoFlags);
 
-			if (m_reshapePhysicsVolumes)
-			{
-				args.flags |= uint32(OP::EMeshBindShapeFlags::ReshapePhysicsVolumes);
-			}
-
-			if (m_reshapeVertices)
-			{
-				args.flags |= uint32(OP::EMeshBindShapeFlags::ReshapeVertices);
-			}
+			Args.flags = static_cast<uint32>(BindFlags);
 
-			args.bindingMethod = m_bindingMethod;
+			Args.bindingMethod = BindingMethod;
 
 			if (Mesh)
 			{
-				args.mesh = Mesh->linkedAddress;
+				Args.mesh = Mesh->linkedAddress;
 			}
 
 			if (Shape)
 			{
-				args.shape = Shape->linkedAddress;
+				Args.shape = Shape->linkedAddress;
 			}
 
 			linkedAddress = (OP::ADDRESS)program.m_opAddress.Num();
 			program.m_opAddress.Add((uint32_t)program.m_byteCode.Num());
 			AppendCode(program.m_byteCode, OP_TYPE::ME_BINDSHAPE);
-			AppendCode(program.m_byteCode, args);
+			AppendCode(program.m_byteCode, Args);
 
-			AppendCode(program.m_byteCode, (int32)m_bonesToDeform.Num());
-			for (const string& S : m_bonesToDeform)
+			AppendCode(program.m_byteCode, (int32)BonesToDeform.Num());
+			for (const string& S : BonesToDeform)
 			{
 				const uint32 index = program.AddConstant(S);
 				AppendCode(program.m_byteCode, index);
 			}
 
-			AppendCode(program.m_byteCode, (int32)m_physicsToDeform.Num());
-			for (const string& S : m_physicsToDeform)
+			AppendCode(program.m_byteCode, (int32)PhysicsToDeform.Num());
+			for (const string& S : PhysicsToDeform)
 			{
 				const uint32 index = program.AddConstant(S);
 				AppendCode(program.m_byteCode, index);
diff --git a/Engine/Plugins/Experimental/Mutable/Source/MutableTools/Private/MuT/ASTOpMeshBindShape.h b/Engine/Plugins/Experimental/Mutable/Source/MutableTools/Private/MuT/ASTOpMeshBindShape.h
--- a/Engine/Plugins/Experimental/Mutable/Source/MutableTools/Private/MuT/ASTOpMeshBindShape.h
+++ b/Engine/Plugins/Experimental/Mutable/Source/MutableTools/Private/MuT/ASTOpMeshBindShape.h
@@ -25,19 +25,15 @@ struct FProgram;
 		ASTChild Mesh;
 		ASTChild Shape;
 
-		TArray<string> m_bonesToDeform;
-		TArray<string> m_physicsToDeform;
-
-		bool m_reshapeSkeleton = false;
-		bool m_discardInvalidBindings = true;
-		bool m_enableRigidParts = false;
-		bool m_deformAllBones = false;
-		bool m_deformAllPhysics = false;
-		bool m_reshapePhysicsVolumes = false;
-		bool m_reshapeVertices = true;
-
-		uint32 m_bindingMethod = 0;
-
+		TArray<string> BonesToDeform;
+		TArray<string> PhysicsToDeform;
+
+		uint32 BindingMethod = 0;
+		
+		uint32 bReshapeSkeleton	      : 1;
+		uint32 bEnableRigidParts      : 1;
+		uint32 bReshapePhysicsVolumes : 1;
+		uint32 bReshapeVertices       : 1;
 
 	public:
 
diff --git a/Engine/Plugins/Experimental/Mutable/Source/MutableTools/Private/MuT/CodeGenerator.cpp b/Engine/Plugins/Experimental/Mutable/Source/MutableTools/Private/MuT/CodeGenerator.cpp
--- a/Engine/Plugins/Experimental/Mutable/Source/MutableTools/Private/MuT/CodeGenerator.cpp
+++ b/Engine/Plugins/Experimental/Mutable/Source/MutableTools/Private/MuT/CodeGenerator.cpp
@@ -2417,8 +2417,7 @@ namespace mu
 				
 				BindOp->Mesh = lastMeshOp;
 				BindOp->Shape = ClipShapeResult.meshOp; 
-				BindOp->m_discardInvalidBindings = false;
-				BindOp->m_bindingMethod = static_cast<uint32>(TypedClipNode->BindingMethod);
+				BindOp->BindingMethod = static_cast<uint32>(TypedClipNode->BindingMethod);
 	
 				ClipOp->Mesh = BindOp;
 
diff --git a/Engine/Plugins/Experimental/Mutable/Source/MutableTools/Private/MuT/CodeGenerator_Mesh.cpp b/Engine/Plugins/Experimental/Mutable/Source/MutableTools/Private/MuT/CodeGenerator_Mesh.cpp
--- a/Engine/Plugins/Experimental/Mutable/Source/MutableTools/Private/MuT/CodeGenerator_Mesh.cpp
+++ b/Engine/Plugins/Experimental/Mutable/Source/MutableTools/Private/MuT/CodeGenerator_Mesh.cpp
@@ -409,23 +409,21 @@ class Node;
 		    Ptr<ASTOpMeshBindShape> OpBind = new ASTOpMeshBindShape();
 		    Ptr<ASTOpMeshApplyShape> OpApply = new ASTOpMeshApplyShape();
 
-			// Setting reshapeVertices to false the bind op will remove all mesh members except 
+			// Setting bReshapeVertices to false the bind op will remove all mesh members except 
 			// PhysicsBodies and the Skeleton.
-            OpBind->m_reshapeVertices = false;
-		    OpBind->m_reshapeSkeleton = node.m_reshapeSkeleton;
-		    OpBind->m_deformAllBones = false; //TODO: Remove deprecated var from node
-		    OpBind->m_bonesToDeform = node.m_bonesToDeform;
-    	    OpBind->m_reshapePhysicsVolumes = node.m_reshapePhysicsVolumes; 
-			OpBind->m_physicsToDeform = node.m_physicsToDeform;
-			OpBind->m_deformAllPhysics = node.m_deformAllPhysics;
-			OpBind->m_bindingMethod = static_cast<uint32>(EShapeBindingMethod::ReshapeClosestProject);
+            OpBind->bReshapeVertices = false;
+		    OpBind->bReshapeSkeleton = node.m_reshapeSkeleton;
+		    OpBind->BonesToDeform = node.m_bonesToDeform;
+    	    OpBind->bReshapePhysicsVolumes = node.m_reshapePhysicsVolumes; 
+			OpBind->PhysicsToDeform = node.m_physicsToDeform;
+			OpBind->BindingMethod = static_cast<uint32>(EShapeBindingMethod::ReshapeClosestProject);
             
 			OpBind->Mesh = BaseResult.meshOp;
             OpBind->Shape = BaseResult.meshOp;
            
-			OpApply->m_reshapeVertices = OpBind->m_reshapeVertices;
-		    OpApply->m_reshapeSkeleton = OpBind->m_reshapeSkeleton;
-		    OpApply->m_reshapePhysicsVolumes = OpBind->m_reshapePhysicsVolumes;
+			OpApply->bReshapeVertices = OpBind->bReshapeVertices;
+		    OpApply->bReshapeSkeleton = OpBind->bReshapeSkeleton;
+		    OpApply->bReshapePhysicsVolumes = OpBind->bReshapePhysicsVolumes;
 
 			OpApply->Mesh = OpBind;
             OpApply->Shape = OpMorph;
@@ -1301,7 +1299,6 @@ class Node;
 			OpClipDeform->ClipShape = baseResult.meshOp;
 		}
 
-    	OpBind->m_discardInvalidBindings = false;
 		OpClipDeform->Mesh = OpBind;
 
 		Result.meshOp = OpClipDeform;
@@ -1398,28 +1395,26 @@ class Node;
 	{
 		const NodeMeshReshape::Private& node = *reshape->GetPrivate();
 
-		Ptr<ASTOpMeshBindShape> opBind = new ASTOpMeshBindShape();
-		Ptr<ASTOpMeshApplyShape> opApply = new ASTOpMeshApplyShape();
+		Ptr<ASTOpMeshBindShape> OpBind = new ASTOpMeshBindShape();
+		Ptr<ASTOpMeshApplyShape> OpApply = new ASTOpMeshApplyShape();
 
-		opBind->m_reshapeSkeleton = node.m_reshapeSkeleton;
-    	opBind->m_enableRigidParts = node.m_enableRigidParts;
-		opBind->m_deformAllBones = false; //TODO: Remove deprecated var from node
-		opBind->m_bonesToDeform = node.m_bonesToDeform;
-    	opBind->m_reshapePhysicsVolumes = node.m_reshapePhysicsVolumes;
-		opBind->m_deformAllPhysics = node.m_deformAllPhysics;
-		opBind->m_physicsToDeform = node.m_physicsToDeform;
-		opBind->m_reshapeVertices = true;
-		opBind->m_bindingMethod = static_cast<uint32>(EShapeBindingMethod::ReshapeClosestProject);
+		OpBind->bReshapeSkeleton = node.m_reshapeSkeleton;
+    	OpBind->bEnableRigidParts = node.m_enableRigidParts;
+		OpBind->BonesToDeform = node.m_bonesToDeform;
+    	OpBind->bReshapePhysicsVolumes = node.m_reshapePhysicsVolumes;
+		OpBind->PhysicsToDeform = node.m_physicsToDeform;
+		OpBind->bReshapeVertices = true;
+		OpBind->BindingMethod = static_cast<uint32>(EShapeBindingMethod::ReshapeClosestProject);
 
-		opApply->m_reshapeVertices = true;	
-		opApply->m_reshapeSkeleton = node.m_reshapeSkeleton;
-		opApply->m_reshapePhysicsVolumes = node.m_reshapePhysicsVolumes;
+		OpApply->bReshapeVertices = OpBind->bReshapeVertices;
+		OpApply->bReshapeSkeleton = OpBind->bReshapeSkeleton;
+		OpApply->bReshapePhysicsVolumes = OpBind->bReshapePhysicsVolumes;
 
 		// Base Mesh
 		if (node.m_pBaseMesh)
 		{
 			GenerateMesh(InOptions, OutResult, node.m_pBaseMesh);
-			opBind->Mesh = OutResult.meshOp;
+			OpBind->Mesh = OutResult.meshOp;
 		}
 		else
 		{
@@ -1439,20 +1434,20 @@ class Node;
 		{
 			FMeshGenerationResult baseResult;
 			GenerateMesh(ShapeOptions, baseResult, node.m_pBaseShape);
-			opBind->Shape = baseResult.meshOp;
+			OpBind->Shape = baseResult.meshOp;
 		}
 
-		opApply->Mesh = opBind;
+		OpApply->Mesh = OpBind;
 
 		// Target Shape
 		if (node.m_pTargetShape)
 		{
 			FMeshGenerationResult targetResult;
 			GenerateMesh(ShapeOptions, targetResult, node.m_pTargetShape);
-			opApply->Shape = targetResult.meshOp;
+			OpApply->Shape = targetResult.meshOp;
 		}
 
-		OutResult.meshOp = opApply;
+		OutResult.meshOp = OpApply;
 	}
 
 }
diff --git a/Engine/Plugins/Experimental/Mutable/Source/MutableTools/Private/MuT/CodeOptimiser.cpp b/Engine/Plugins/Experimental/Mutable/Source/MutableTools/Private/MuT/CodeOptimiser.cpp
--- a/Engine/Plugins/Experimental/Mutable/Source/MutableTools/Private/MuT/CodeOptimiser.cpp
+++ b/Engine/Plugins/Experimental/Mutable/Source/MutableTools/Private/MuT/CodeOptimiser.cpp
@@ -998,7 +998,7 @@ namespace mu
 				auto typedOp = dynamic_cast<ASTOpMeshBindShape*>(node.get());
 				check(typedOp);
 
-				if (typedOp && typedOp->m_reshapeSkeleton)
+				if (typedOp && typedOp->bReshapeSkeleton)
 				{
 					RecurseWithState(typedOp->Mesh.child(), true);
 					return false;
@@ -1012,7 +1012,7 @@ namespace mu
 				auto typedOp = dynamic_cast<ASTOpMeshApplyShape*>(node.get());
 				check(typedOp);
 
-				if (typedOp && typedOp->m_reshapeSkeleton)
+				if (typedOp && typedOp->bReshapeSkeleton)
 				{
 					RecurseWithState(typedOp->Mesh.child(), true);
 					return false;
diff --git a/Engine/Plugins/Experimental/Mutable/Source/MutableTools/Private/MuT/NodeMeshMorph.cpp b/Engine/Plugins/Experimental/Mutable/Source/MutableTools/Private/MuT/NodeMeshMorph.cpp
--- a/Engine/Plugins/Experimental/Mutable/Source/MutableTools/Private/MuT/NodeMeshMorph.cpp
+++ b/Engine/Plugins/Experimental/Mutable/Source/MutableTools/Private/MuT/NodeMeshMorph.cpp
@@ -161,12 +161,6 @@ namespace mu
 		m_pD->m_reshapeSkeleton = bEnable;
 	}	
 
-	//---------------------------------------------------------------------------------------------
-	void NodeMeshMorph::SetDeformAllPhysics(bool bEnable)
-	{
-		m_pD->m_deformAllPhysics = bEnable;
-	}
-
 	//---------------------------------------------------------------------------------------------
 	void NodeMeshMorph::AddBoneToDeform(const char* BoneName)
 	{
diff --git a/Engine/Plugins/Experimental/Mutable/Source/MutableTools/Private/MuT/NodeMeshMorphPrivate.h b/Engine/Plugins/Experimental/Mutable/Source/MutableTools/Private/MuT/NodeMeshMorphPrivate.h
--- a/Engine/Plugins/Experimental/Mutable/Source/MutableTools/Private/MuT/NodeMeshMorphPrivate.h
+++ b/Engine/Plugins/Experimental/Mutable/Source/MutableTools/Private/MuT/NodeMeshMorphPrivate.h
@@ -29,16 +29,13 @@ namespace mu
 		bool m_reshapeSkeleton = false;
 		bool m_reshapePhysicsVolumes = false;
 		
-		bool m_deformAllBones_DEPRECATED = false;
-		bool m_deformAllPhysics = false;
-		
 		TArray<string> m_bonesToDeform;
 		TArray<string> m_physicsToDeform;
 
         //!
 		void Serialise( OutputArchive& arch ) const
 		{
-            uint32_t ver = 4;
+            uint32_t ver = 5;
 			arch << ver;
 
 			arch << m_pFactor;
@@ -50,7 +47,6 @@ namespace mu
 			arch << m_reshapePhysicsVolumes;
 			arch << m_bonesToDeform;
 
-			arch << m_deformAllPhysics;
 			
 			arch << m_physicsToDeform;
 
@@ -61,7 +57,7 @@ namespace mu
 		{
             uint32_t ver;
 			arch >> ver;
-			check(ver <= 4);
+			check(ver <= 5);
 
 			arch >> m_pFactor;
 			arch >> m_pBase;
@@ -89,18 +85,22 @@ namespace mu
 
 			if (ver == 3)
 			{
-				arch >> m_deformAllBones_DEPRECATED;
+				bool bDeformAllBones_DEPRECATED;
+				arch >> bDeformAllBones_DEPRECATED;
+			}
+
+			if (ver >= 3 && ver < 5)
+			{
+				bool bDeformAllPhysics_DEPRECATED;
+				arch >> bDeformAllPhysics_DEPRECATED;
 			}
 
 			if (ver >= 3)
 			{
-				arch >> m_deformAllPhysics;
 				arch >> m_physicsToDeform;
 			}
 			else
 			{
-				m_deformAllBones_DEPRECATED = false;
-				m_deformAllPhysics = false;
 				m_physicsToDeform.Empty();
 			}
 
diff --git a/Engine/Plugins/Experimental/Mutable/Source/MutableTools/Private/MuT/NodeMeshReshape.cpp b/Engine/Plugins/Experimental/Mutable/Source/MutableTools/Private/MuT/NodeMeshReshape.cpp
--- a/Engine/Plugins/Experimental/Mutable/Source/MutableTools/Private/MuT/NodeMeshReshape.cpp
+++ b/Engine/Plugins/Experimental/Mutable/Source/MutableTools/Private/MuT/NodeMeshReshape.cpp
@@ -133,13 +133,7 @@ namespace mu
 	void NodeMeshReshape::SetEnableRigidParts(bool bEnable)
 	{
 		m_pD->m_enableRigidParts = bEnable;
-	}
-
-	//---------------------------------------------------------------------------------------------
-	void NodeMeshReshape::SetDeformAllPhysics(bool bEnable)
-	{
-		m_pD->m_deformAllPhysics = bEnable;
-	}
+	}	
 
 	//---------------------------------------------------------------------------------------------
 	void NodeMeshReshape::AddBoneToDeform(const char* BoneName)
diff --git a/Engine/Plugins/Experimental/Mutable/Source/MutableTools/Private/MuT/NodeMeshReshapePrivate.h b/Engine/Plugins/Experimental/Mutable/Source/MutableTools/Private/MuT/NodeMeshReshapePrivate.h
--- a/Engine/Plugins/Experimental/Mutable/Source/MutableTools/Private/MuT/NodeMeshReshapePrivate.h
+++ b/Engine/Plugins/Experimental/Mutable/Source/MutableTools/Private/MuT/NodeMeshReshapePrivate.h
@@ -26,16 +26,15 @@ namespace mu
 		Ptr<NodeMesh> m_pTargetShape;
 		bool m_reshapeSkeleton = false;
 		bool m_enableRigidParts = false;
-		bool m_deformAllBones_DEPRECATED = false;
-		bool m_deformAllPhysics = false;
 		bool m_reshapePhysicsVolumes = false;
-		
+	
+
 		TArray<string> m_bonesToDeform;
 		TArray<string> m_physicsToDeform;
         //!
 		void Serialise( OutputArchive& arch ) const
 		{
-            uint32_t ver = 6;
+            uint32_t ver = 7;
 			arch << ver;
 
 			arch << m_pBaseMesh;
@@ -45,7 +44,6 @@ namespace mu
 			arch << m_enableRigidParts;
 			arch << m_bonesToDeform;
 			arch << m_reshapePhysicsVolumes;
-			arch << m_deformAllPhysics;
 			arch << m_physicsToDeform;
         }
 
@@ -54,7 +52,7 @@ namespace mu
 		{
             uint32_t ver;
 			arch >> ver;
-			check(ver<=6);
+			check(ver <= 7);
 
 			arch >> m_pBaseMesh;
 			arch >> m_pBaseShape;
@@ -64,7 +62,8 @@ namespace mu
 
 			if (ver <= 5)
 			{
-				arch >> m_deformAllBones_DEPRECATED;
+				bool bDeformAllBones_DEPRECATED;
+				arch >> bDeformAllBones_DEPRECATED;
 			}
 
 			arch >> m_bonesToDeform;
@@ -74,9 +73,14 @@ namespace mu
 				arch >> m_reshapePhysicsVolumes;
 		 	}
 
+			if (ver >= 5 && ver < 7)
+			{
+				bool bDeformAllPhysics_DEPRECATED;
+				arch >> bDeformAllPhysics_DEPRECATED;
+			}
+
 			if (ver >= 5)
 			{
-				arch >> m_deformAllPhysics;
 				arch >> m_physicsToDeform;
 			}
 		}
diff --git a/Engine/Plugins/Experimental/Mutable/Source/MutableTools/Public/MuT/NodeMeshMorph.h b/Engine/Plugins/Experimental/Mutable/Source/MutableTools/Public/MuT/NodeMeshMorph.h
--- a/Engine/Plugins/Experimental/Mutable/Source/MutableTools/Public/MuT/NodeMeshMorph.h
+++ b/Engine/Plugins/Experimental/Mutable/Source/MutableTools/Public/MuT/NodeMeshMorph.h
@@ -80,9 +80,6 @@ namespace mu
 		/** Also deform the mesh skeleton. Disabled by default. */
 		void SetReshapeSkeleton(bool);	
 
-		/** Deform all the physiscs bodies */
-		void SetDeformAllPhysics(bool);
-
 		/** Deform Mesh Physics Volumes */
 		void SetReshapePhysicsVolumes(bool);
 
diff --git a/Engine/Plugins/Experimental/Mutable/Source/MutableTools/Public/MuT/NodeMeshReshape.h b/Engine/Plugins/Experimental/Mutable/Source/MutableTools/Public/MuT/NodeMeshReshape.h
--- a/Engine/Plugins/Experimental/Mutable/Source/MutableTools/Public/MuT/NodeMeshReshape.h
+++ b/Engine/Plugins/Experimental/Mutable/Source/MutableTools/Public/MuT/NodeMeshReshape.h
@@ -59,9 +59,6 @@ namespace mu
 		/** Search for clusters of rigid parts and not deform them. */
 		void SetEnableRigidParts(bool);
 
-		/** Deform all the physiscs bodies */
-		void SetDeformAllPhysics(bool);
-
 		/** Deform Mesh Physics Volumes */
 		void SetReshapePhysicsVolumes(bool);
 
