[mutable] Fix NodeApplyPose not working if the skeleton of the pose asset is different than the reference skeleton

#rnx

diff --git a/Engine/Plugins/Experimental/Mutable/Source/CustomizableObjectEditor/Private/MuCOE/Nodes/CustomizableObjectNodeAnimationPose.cpp b/Engine/Plugins/Experimental/Mutable/Source/CustomizableObjectEditor/Private/MuCOE/Nodes/CustomizableObjectNodeAnimationPose.cpp
--- a/Engine/Plugins/Experimental/Mutable/Source/CustomizableObjectEditor/Private/MuCOE/Nodes/CustomizableObjectNodeAnimationPose.cpp
+++ b/Engine/Plugins/Experimental/Mutable/Source/CustomizableObjectEditor/Private/MuCOE/Nodes/CustomizableObjectNodeAnimationPose.cpp
@@ -108,12 +108,25 @@ void UCustomizableObjectNodeAnimationPose::StaticRetrievePoseInformation(UPoseAs
 		SkeletalMeshComponent->RefreshBoneTransforms();
 		SkeletalMeshComponent->InitAnim(false);
 
-		UAnimInstance* AnimInstance = SkeletalMeshComponent->GetAnimInstance();
-		FBoneContainer& BoneContainerCopy = AnimInstance->GetRequiredBones();
+		// PoseSkeleton might be different than the RefSkeletalMesh's skeleton, use it as reference to extract the pose.
+		USkeleton* PoseSkeleton = PoseAsset->GetSkeleton();
+		const FReferenceSkeleton& PoseRefSkeleton = PoseSkeleton->GetReferenceSkeleton();
+
+		// Use all bones from the pose's skeleton as RequiredBones
+		TArray<FBoneIndexType> RequiredBones;
+		RequiredBones.SetNumUninitialized(PoseRefSkeleton.GetNum());
+
+		for (uint16 BoneIndex = 0; BoneIndex < PoseRefSkeleton.GetNum(); ++BoneIndex)
+		{
+			RequiredBones[BoneIndex] = BoneIndex;
+		}
+
+		FBoneContainer BoneContainer;
+		BoneContainer.InitializeTo(RequiredBones, FCurveEvaluationOption(SkeletalMeshComponent->GetAllowedAnimCurveEvaluate()), *PoseSkeleton);
 
 		// Needs a FMemMark declared before in the stack context so that the memory allocated by the FCompactPose is freed correctly
 		FCompactPose OutPose;
-		OutPose.SetBoneContainer(&BoneContainerCopy);
+		OutPose.SetBoneContainer(&BoneContainer);
 
 		FBlendedCurve OutCurve;
 		UE::Anim::FStackAttributeContainer OutAttributes;
@@ -123,22 +136,15 @@ void UCustomizableObjectNodeAnimationPose::StaticRetrievePoseInformation(UPoseAs
 		OutCurve = OutAnimData.GetCurve();
 
 		FBlendedHeapCurve AnimCurves = SkeletalMeshComponent->AnimCurves;
-		const int32 NumElement = AnimCurves.CurveWeights.Num();
-
-		for (int32 i = 0; i < NumElement; ++i)
-		{
-			OutCurve.CurveWeights.Add(AnimCurves.CurveWeights[i]);
-		}
-
+		OutCurve.CurveWeights = AnimCurves.CurveWeights;
 		OutCurve.bInitialized = AnimCurves.bInitialized;
+		OutCurve.UIDToArrayIndexLUT = AnimCurves.UIDToArrayIndexLUT;
 
 		// Assuming one single pose, with a weigth set to 1.0
 		FAnimExtractContext ExtractionContext;
 		ExtractionContext.bExtractRootMotion = false;
 		ExtractionContext.CurrentTime = 0.0f;
-
-		OutCurve.UIDToArrayIndexLUT = AnimCurves.UIDToArrayIndexLUT;
-
+		
 		const TArray<FSmartName>& PoseNames = PoseAsset->GetPoseNames();
 		ExtractionContext.PoseCurves.Add(FPoseCurve(0, PoseNames[0].UID, 1.0f));
 		FAnimationPoseData SecondOutAnimData(OutPose, OutCurve, OutAttributes);
@@ -147,43 +153,36 @@ void UCustomizableObjectNodeAnimationPose::StaticRetrievePoseInformation(UPoseAs
 		OutCurve = OutAnimData.GetCurve();
 
 		const TArray<FTransform, FAnimStackAllocator>& ArrayPoseBoneTransform = OutPose.GetBones();
-		const TArray<FBoneIndexType>& ArrayPoseBoneIndices = OutPose.GetBoneContainer().GetBoneIndicesArray();
 
-		for (int32 ArrayIndex = 0; ArrayIndex < ArrayPoseBoneIndices.Num(); ArrayIndex++)
+		// Extract final pose by combining the OutPose with the RefSkeleton pose (for bones missing in the PoseRefSkeleton)
+		const FReferenceSkeleton& RefSkeleton = RefSkeletalMesh->GetRefSkeleton();
+		for (int32 BoneIndex = 0; BoneIndex < RefSkeleton.GetNum(); ++BoneIndex)
 		{
-			FBoneIndexType const& PoseBoneIndex = ArrayPoseBoneIndices[ArrayIndex];
-
-			checkSlow(PoseBoneIndex != INDEX_NONE);
+			FName BoneName = RefSkeleton.GetBoneName(BoneIndex);
 
 			FTransform CumulativePoseTransform;
-			FBoneIndexType ParentIndex = PoseBoneIndex;
 
+			int32 ParentIndex = BoneIndex;
 			while (ParentIndex > 0)
 			{
-				for (int32 IndicesIndex = 0; IndicesIndex < ArrayPoseBoneIndices.Num(); ++IndicesIndex)
+				FName ParentName = RefSkeleton.GetBoneName(ParentIndex);
+				if (int32 PoseBoneIndex = PoseRefSkeleton.FindBoneIndex(ParentName); PoseBoneIndex >= 0)
 				{
-					const FBoneIndexType BoneIndex = ArrayPoseBoneIndices[IndicesIndex];
-
-					if (BoneIndex == ParentIndex)
-					{
-						CumulativePoseTransform = CumulativePoseTransform * ArrayPoseBoneTransform[IndicesIndex];
-						break;
-					}
-					else if (BoneIndex > ParentIndex)
-					{
-						break; // Break for, ArrayPoseBoneIndices is sorted.
-					}
+					CumulativePoseTransform *= ArrayPoseBoneTransform[PoseBoneIndex];
 				}
-
-				ParentIndex = RefSkeletalMesh->GetRefSkeleton().GetParentIndex(ParentIndex);
+				else
+				{
+					// Accumulate the transform using the RefSkeleton's transform if the pose doesn't have the bone
+					CumulativePoseTransform *= RefSkeleton.GetRefBonePose()[ParentIndex];
+				}
+			
+				ParentIndex = RefSkeleton.GetParentIndex(ParentIndex);
 			}
 
-			FString Name = RefSkeletalMesh->GetRefSkeleton().GetBoneName(PoseBoneIndex).ToString();
-			int32 SMBoneIndex = SkeletalMeshComponent->GetBoneIndex(FName(*Name));
-			FTransform BoneToComponentTranform = SkeletalMeshComponent->GetEditableComponentSpaceTransforms()[SMBoneIndex];
+			FTransform BoneToComponentTranform = SkeletalMeshComponent->GetEditableComponentSpaceTransforms()[BoneIndex];
 			FTransform TransformToAdd = BoneToComponentTranform.Inverse() * CumulativePoseTransform;
 
-			OutArrayBoneName.Add(Name);
+			OutArrayBoneName.Add(BoneName.ToString());
 			OutArrayTransform.Add(TransformToAdd);
 		}
 	}
