[mutable] Multi-dimensional parameter support in the instance parameter generator.
Additional renaming of structs for conformance with Unreal.

diff --git a/Engine/Plugins/Experimental/Mutable/Source/MutableRuntime/Private/MuR/CodeRunner.cpp b/Engine/Plugins/Experimental/Mutable/Source/MutableRuntime/Private/MuR/CodeRunner.cpp
--- a/Engine/Plugins/Experimental/Mutable/Source/MutableRuntime/Private/MuR/CodeRunner.cpp
+++ b/Engine/Plugins/Experimental/Mutable/Source/MutableRuntime/Private/MuR/CodeRunner.cpp
@@ -253,7 +253,7 @@ namespace mu
 				case DT_INT:        GetMemory().SetInt(item, 0); break;
 				case DT_SCALAR:		GetMemory().SetScalar(item, 0.0f); break;
 				case DT_STRING:		GetMemory().SetString(item, nullptr); break;
-				case DT_COLOUR:		GetMemory().SetColour(item, vec4f()); break;
+				case DT_COLOUR:		GetMemory().SetColour(item, FVector4f()); break;
 				case DT_PROJECTOR:  GetMemory().SetProjector(item, nullptr); break;
 				case DT_MESH:       GetMemory().SetMesh(item, nullptr); break;
 				case DT_IMAGE:      GetMemory().SetImage(item, nullptr); break;
@@ -369,7 +369,7 @@ namespace mu
 
                 if ( args.value )
                 {
-                    vec4f value = GetMemory().GetColour( CACHE_ADDRESS(args.value,item) );
+					FVector4f value = GetMemory().GetColour( CACHE_ADDRESS(args.value,item) );
 
                     OP::ADDRESS nameAd = args.name;
                     check(  nameAd < (uint32)pModel->GetPrivate()->m_program.m_constantStrings.Num() );
@@ -3345,7 +3345,7 @@ namespace mu
 
         case OP_TYPE::IM_SWIZZLE:
         {
-            auto args = pModel->GetPrivate()->m_program.GetOpArgs<OP::ImageSwizzleArgs>(item.at);
+			OP::ImageSwizzleArgs args = pModel->GetPrivate()->m_program.GetOpArgs<OP::ImageSwizzleArgs>(item.at);
             switch (item.stage)
             {
             case 0:
@@ -3370,7 +3370,7 @@ namespace mu
 
         case OP_TYPE::IM_SELECTCOLOUR:
         {
-            auto args = pModel->GetPrivate()->m_program.GetOpArgs<OP::ImageSelectColourArgs>(item.at);
+			OP::ImageSelectColourArgs args = pModel->GetPrivate()->m_program.GetOpArgs<OP::ImageSelectColourArgs>(item.at);
             switch (item.stage)
             {
             case 0:
@@ -3382,9 +3382,9 @@ namespace mu
             case 1:
             {
                 Ptr<const Image> pBase = GetMemory().GetImage( CACHE_ADDRESS(args.base,item) );
-                vec4<float> colour = GetMemory().GetColour(SCHEDULED_OP::FromOpAndOptions(args.colour, item, 0));
+				FVector4f colour = GetMemory().GetColour(SCHEDULED_OP::FromOpAndOptions(args.colour, item, 0));
 
-               ImagePtr pResult = ImageSelectColour( pBase.get(), colour.xyz() );
+				Ptr<Image> pResult = ImageSelectColour( pBase.get(), vec3f(colour) );
 
                 GetMemory().SetImage( item, pResult );
                 break;
@@ -3399,7 +3399,7 @@ namespace mu
 
         case OP_TYPE::IM_COLOURMAP:
         {
-            auto args = pModel->GetPrivate()->m_program.GetOpArgs<OP::ImageColourMapArgs>(item.at);
+			OP::ImageColourMapArgs args = pModel->GetPrivate()->m_program.GetOpArgs<OP::ImageColourMapArgs>(item.at);
             switch (item.stage)
             {
             case 0:
@@ -3438,7 +3438,7 @@ namespace mu
 
         case OP_TYPE::IM_GRADIENT:
         {
-            auto args = pModel->GetPrivate()->m_program.GetOpArgs<OP::ImageGradientArgs>(item.at);
+			OP::ImageGradientArgs args = pModel->GetPrivate()->m_program.GetOpArgs<OP::ImageGradientArgs>(item.at);
             switch (item.stage)
             {
             case 0:
@@ -3449,10 +3449,10 @@ namespace mu
 
             case 1:
             {
-                vec4<float> colour0 = GetMemory().GetColour(SCHEDULED_OP::FromOpAndOptions(args.colour0, item, 0));
-                vec4<float> colour1 = GetMemory().GetColour(SCHEDULED_OP::FromOpAndOptions(args.colour1, item, 0));
+				FVector4f colour0 = GetMemory().GetColour(SCHEDULED_OP::FromOpAndOptions(args.colour0, item, 0));
+				FVector4f colour1 = GetMemory().GetColour(SCHEDULED_OP::FromOpAndOptions(args.colour1, item, 0));
 
-                ImagePtr pResult = ImageGradient( colour0.xyz(), colour1.xyz(),
+                ImagePtr pResult = ImageGradient( colour0, colour1,
                                                   args.size[0],
                                                   args.size[1] );
 
@@ -3469,7 +3469,7 @@ namespace mu
 
         case OP_TYPE::IM_BINARISE:
         {
-            auto args = pModel->GetPrivate()->m_program.GetOpArgs<OP::ImageBinariseArgs>(item.at);
+			OP::ImageBinariseArgs args = pModel->GetPrivate()->m_program.GetOpArgs<OP::ImageBinariseArgs>(item.at);
             switch (item.stage)
             {
             case 0:
@@ -3536,7 +3536,7 @@ namespace mu
 
             case 1:
             {
-                mu::vec4f c = GetMemory().GetColour(SCHEDULED_OP::FromOpAndOptions(args.colour, item, 0));
+				FVector4f c = GetMemory().GetColour(SCHEDULED_OP::FromOpAndOptions(args.colour, item, 0));
 
 				uint16 SizeX = args.size[0];
 				uint16 SizeY = args.size[1];
@@ -4539,7 +4539,7 @@ namespace mu
         case OP_TYPE::CO_CONSTANT:
         {
 			OP::ColourConstantArgs args = program.GetOpArgs<OP::ColourConstantArgs>(item.at);
-            vec4f result;
+			FVector4f result;
             result[0] = args.value[0];
             result[1] = args.value[1];
             result[2] = args.value[2];
@@ -4556,7 +4556,7 @@ namespace mu
             float g=0.0f;
             float b=0.0f;            
             pParams->GetColourValue( args.variable, &r, &g, &b, index );
-            GetMemory().SetColour( item, vec4f(r,g,b,1.0f) );
+            GetMemory().SetColour( item, FVector4f(r,g,b,1.0f) );
             break;
         }
 
@@ -4580,22 +4580,22 @@ namespace mu
 
                 Ptr<const Image> pImage = GetMemory().GetImage(SCHEDULED_OP::FromOpAndOptions(args.image, item, 0));
 
-                vec4f result;
+				FVector4f result;
                 if (pImage)
                 {
                     if (args.filter)
                     {
                         // TODO
-                        result = pImage->Sample(vec2<float>(x, y));
+                        result = pImage->Sample(FVector2f(x, y));
                     }
                     else
                     {
-                        result = pImage->Sample(vec2<float>(x, y));
+                        result = pImage->Sample(FVector2f(x, y));
                     }
                 }
                 else
                 {
-                    result = vec4f();
+                    result = FVector4f();
                 }
 
                 GetMemory().SetColour( item, result );
@@ -4624,7 +4624,7 @@ namespace mu
 
             case 1:
             {
-                vec4f result;
+				FVector4f result;
 
                 for (int t=0;t<MUTABLE_OP_MAX_SWIZZLE_CHANNELS;++t)
                 {
@@ -4660,7 +4660,7 @@ namespace mu
             {
                 Ptr<const Image> pImage = GetMemory().GetImage( CACHE_ADDRESS(args.image,item) );
 
-                vec4f result = vec4f( (float)pImage->GetSizeX(), (float)pImage->GetSizeY(), 0.0f, 0.0f );
+				FVector4f result = FVector4f( (float)pImage->GetSizeX(), (float)pImage->GetSizeY(), 0.0f, 0.0f );
 
                 GetMemory().SetColour( item, result );
                 break;
@@ -4687,7 +4687,7 @@ namespace mu
             {
                 Ptr<const Layout> pLayout = GetMemory().GetLayout( CACHE_ADDRESS(args.layout,item) );
 
-                vec4f result = vec4f(0,0,0,0);
+				FVector4f result = FVector4f(0,0,0,0);
                 if ( pLayout )
                 {
                     int relBlockIndex = pLayout->FindBlock( args.block );
@@ -4705,7 +4705,7 @@ namespace mu
                         // Convert the rect from blocks to pixels
                         FIntPoint grid = pLayout->GetGridSize();
 
-                        result = vec4<float>( float(rectInblocks.min[0]) / float(grid[0]),
+                        result = FVector4f( float(rectInblocks.min[0]) / float(grid[0]),
                                               float(rectInblocks.min[1]) / float(grid[1]),
                                               float(rectInblocks.size[0]) / float(grid[0]),
                                               float(rectInblocks.size[1]) / float(grid[1]) );
@@ -4738,7 +4738,7 @@ namespace mu
 
             case 1:
             {
-                vec4f result = vec4f(1, 1, 1, 1);
+				FVector4f result = FVector4f(1, 1, 1, 1);
 
                 if (args.x)
                 {
@@ -4790,12 +4790,12 @@ namespace mu
                 otype = program.GetOpType( args.b );
                 dtype = GetOpDataType( otype );
                 check( dtype == DT_COLOUR );
-                vec4f a = args.a ? GetMemory().GetColour( CACHE_ADDRESS( args.a, item ) )
-                                 : vec4f( 0, 0, 0, 0 );
-                vec4f b = args.b ? GetMemory().GetColour( CACHE_ADDRESS( args.b, item ) )
-                                 : vec4f( 0, 0, 0, 0 );
+				FVector4f a = args.a ? GetMemory().GetColour( CACHE_ADDRESS( args.a, item ) )
+                                 : FVector4f( 0, 0, 0, 0 );
+				FVector4f b = args.b ? GetMemory().GetColour( CACHE_ADDRESS( args.b, item ) )
+                                 : FVector4f( 0, 0, 0, 0 );
 
-                vec4f result = vec4f(0,0,0,0);
+				FVector4f result = FVector4f(0,0,0,0);
                 switch (args.operation)
                 {
                 case OP::ArithmeticArgs::ADD:
@@ -4845,14 +4845,14 @@ namespace mu
     {
         MUTABLE_CPUPROFILER_SCOPE(RunCode_Projector);
 
-        const auto& program = pModel->GetPrivate()->m_program;
+        const FProgram& program = pModel->GetPrivate()->m_program;
 		OP_TYPE type = program.GetOpType(item.at);
         switch (type)
         {
 
         case OP_TYPE::PR_CONSTANT:
         {
-            auto args = pModel->GetPrivate()->m_program.GetOpArgs<OP::ResourceConstantArgs>(item.at);
+			OP::ResourceConstantArgs args = pModel->GetPrivate()->m_program.GetOpArgs<OP::ResourceConstantArgs>(item.at);
             ProjectorPtr pResult = new Projector();
             pResult->m_value = program.m_constantProjectors[args.value];
             GetMemory().SetProjector( item, pResult );
@@ -4861,13 +4861,13 @@ namespace mu
 
         case OP_TYPE::PR_PARAMETER:
         {
-            auto args = pModel->GetPrivate()->m_program.GetOpArgs<OP::ParameterArgs>(item.at);
+			OP::ParameterArgs args = pModel->GetPrivate()->m_program.GetOpArgs<OP::ParameterArgs>(item.at);
 			Ptr<RangeIndex> index = BuildCurrentOpRangeIndex( item, pParams, pModel, args.variable );
             ProjectorPtr pResult = new Projector();
             pResult->m_value = pParams->GetPrivate()->GetProjectorValue(args.variable,index);
 
             // The type cannot be changed, take it from the default value
-            const auto& def = program.m_parameters[args.variable].m_defaultValue.m_projector;
+            const FProjector& def = program.m_parameters[args.variable].m_defaultValue.m_projector;
             pResult->m_value.type = def.type;
 
             GetMemory().SetProjector( item, pResult );
diff --git a/Engine/Plugins/Experimental/Mutable/Source/MutableRuntime/Private/MuR/CodeRunnerTasked.cpp b/Engine/Plugins/Experimental/Mutable/Source/MutableRuntime/Private/MuR/CodeRunnerTasked.cpp
--- a/Engine/Plugins/Experimental/Mutable/Source/MutableRuntime/Private/MuR/CodeRunnerTasked.cpp
+++ b/Engine/Plugins/Experimental/Mutable/Source/MutableRuntime/Private/MuR/CodeRunnerTasked.cpp
@@ -571,8 +571,8 @@ namespace mu
 	private:
 		int m_imageCompressionQuality = 0;
 		OP::ImageLayerColourArgs Args;
+		FVector4f m_col;
 		Ptr<const Image> m_base;
-		vec4<float> m_col;
 		Ptr<const Image> m_mask;
 		Ptr<const Image> Result;
 	};
@@ -633,16 +633,16 @@ namespace mu
 			// \todo: precalculated tables for softlight
 			switch (EBlendType(Args.blendType))
 			{
-			case EBlendType::BT_NORMAL_COMBINE: ImageNormalCombine(pNew.get(), m_base.get(), m_mask.get(), m_col.xyz()); break;
-			case EBlendType::BT_SOFTLIGHT: ImageSoftLight(pNew.get(), m_base.get(), m_mask.get(), m_col.xyz(), nullptr); break;
-			case EBlendType::BT_HARDLIGHT: ImageHardLight(pNew.get(), m_base.get(), m_mask.get(), m_col.xyz()); break;
-			case EBlendType::BT_BURN: ImageBurn(pNew.get(), m_base.get(), m_mask.get(), m_col.xyz()); break;
-			case EBlendType::BT_DODGE: ImageDodge(pNew.get(), m_base.get(), m_mask.get(), m_col.xyz()); break;
-			case EBlendType::BT_SCREEN: ImageScreen(pNew.get(), m_base.get(), m_mask.get(), m_col.xyz()); break;
-			case EBlendType::BT_OVERLAY: ImageOverlay(pNew.get(), m_base.get(), m_mask.get(), m_col.xyz()); break;
-			case EBlendType::BT_ALPHA_OVERLAY: ImageAlphaOverlay(pNew.get(), m_base.get(), m_mask.get(), m_col.xyz()); break;
-			case EBlendType::BT_MULTIPLY: ImageMultiply(pNew.get(), m_base.get(), m_mask.get(), m_col.xyz()); break;
-			case EBlendType::BT_BLEND: ImageBlend(pNew.get(), m_base.get(), m_mask.get(), m_col.xyz()); break;
+			case EBlendType::BT_NORMAL_COMBINE: ImageNormalCombine(pNew.get(), m_base.get(), m_mask.get(), vec3f(m_col)); break;
+			case EBlendType::BT_SOFTLIGHT: ImageSoftLight(pNew.get(), m_base.get(), m_mask.get(), vec3f(m_col), nullptr); break;
+			case EBlendType::BT_HARDLIGHT: ImageHardLight(pNew.get(), m_base.get(), m_mask.get(), vec3f(m_col)); break;
+			case EBlendType::BT_BURN: ImageBurn(pNew.get(), m_base.get(), m_mask.get(), vec3f(m_col)); break;
+			case EBlendType::BT_DODGE: ImageDodge(pNew.get(), m_base.get(), m_mask.get(), vec3f(m_col)); break;
+			case EBlendType::BT_SCREEN: ImageScreen(pNew.get(), m_base.get(), m_mask.get(), vec3f(m_col)); break;
+			case EBlendType::BT_OVERLAY: ImageOverlay(pNew.get(), m_base.get(), m_mask.get(), vec3f(m_col)); break;
+			case EBlendType::BT_ALPHA_OVERLAY: ImageAlphaOverlay(pNew.get(), m_base.get(), m_mask.get(), vec3f(m_col)); break;
+			case EBlendType::BT_MULTIPLY: ImageMultiply(pNew.get(), m_base.get(), m_mask.get(), vec3f(m_col)); break;
+			case EBlendType::BT_BLEND: ImageBlend(pNew.get(), m_base.get(), m_mask.get(), vec3f(m_col)); break;
 			default: check(false);
 			}
 
@@ -651,15 +651,15 @@ namespace mu
 		{
 			switch (EBlendType(Args.blendType))
 			{
-			case EBlendType::BT_NORMAL_COMBINE: ImageNormalCombine(pNew.get(), m_base.get(), m_col.xyz()); break;
-			case EBlendType::BT_SOFTLIGHT: ImageSoftLight(pNew.get(), m_base.get(), m_col.xyz()); break;
-			case EBlendType::BT_HARDLIGHT: ImageHardLight(pNew.get(), m_base.get(), m_col.xyz()); break;
-			case EBlendType::BT_BURN: ImageBurn(pNew.get(), m_base.get(), m_col.xyz()); break;
-			case EBlendType::BT_DODGE: ImageDodge(pNew.get(), m_base.get(), m_col.xyz()); break;
-			case EBlendType::BT_SCREEN: ImageScreen(pNew.get(), m_base.get(), m_col.xyz()); break;
-			case EBlendType::BT_OVERLAY: ImageOverlay(pNew.get(), m_base.get(), m_col.xyz()); break;
-			case EBlendType::BT_ALPHA_OVERLAY: ImageAlphaOverlay(pNew.get(), m_base.get(), m_col.xyz()); break;
-			case EBlendType::BT_MULTIPLY: ImageMultiply(pNew.get(), m_base.get(), m_col.xyz()); break;
+			case EBlendType::BT_NORMAL_COMBINE: ImageNormalCombine(pNew.get(), m_base.get(), vec3f(m_col)); break;
+			case EBlendType::BT_SOFTLIGHT: ImageSoftLight(pNew.get(), m_base.get(), vec3f(m_col)); break;
+			case EBlendType::BT_HARDLIGHT: ImageHardLight(pNew.get(), m_base.get(), vec3f(m_col)); break;
+			case EBlendType::BT_BURN: ImageBurn(pNew.get(), m_base.get(), vec3f(m_col)); break;
+			case EBlendType::BT_DODGE: ImageDodge(pNew.get(), m_base.get(), vec3f(m_col)); break;
+			case EBlendType::BT_SCREEN: ImageScreen(pNew.get(), m_base.get(), vec3f(m_col)); break;
+			case EBlendType::BT_OVERLAY: ImageOverlay(pNew.get(), m_base.get(), vec3f(m_col)); break;
+			case EBlendType::BT_ALPHA_OVERLAY: ImageAlphaOverlay(pNew.get(), m_base.get(), vec3f(m_col)); break;
+			case EBlendType::BT_MULTIPLY: ImageMultiply(pNew.get(), m_base.get(), vec3f(m_col)); break;
 			case EBlendType::BT_BLEND: FillPlainColourImage(pNew.get(), m_col); break;
 			default: check(false);
 			}
diff --git a/Engine/Plugins/Experimental/Mutable/Source/MutableRuntime/Private/MuR/Image.cpp b/Engine/Plugins/Experimental/Mutable/Source/MutableRuntime/Private/MuR/Image.cpp
--- a/Engine/Plugins/Experimental/Mutable/Source/MutableRuntime/Private/MuR/Image.cpp
+++ b/Engine/Plugins/Experimental/Mutable/Source/MutableRuntime/Private/MuR/Image.cpp
@@ -498,9 +498,9 @@ namespace mu
 	}
 
     //---------------------------------------------------------------------------------------------
-    vec4<float> Image::Sample( vec2<float> coords ) const
+    FVector4f Image::Sample( FVector2f coords ) const
     {
-        vec4<float> result;
+		FVector4f result;
 
         if (m_size[0]==0 || m_size[1]==0) { return result; }
 
@@ -567,7 +567,7 @@ namespace mu
     }
 
     //---------------------------------------------------------------------------------------------
-    bool Image::IsPlainColour( vec4<float>& colour ) const
+    bool Image::IsPlainColour(FVector4f& colour) const
     {
         bool res = true;
 
diff --git a/Engine/Plugins/Experimental/Mutable/Source/MutableRuntime/Private/MuR/ImagePrivate.h b/Engine/Plugins/Experimental/Mutable/Source/MutableRuntime/Private/MuR/ImagePrivate.h
--- a/Engine/Plugins/Experimental/Mutable/Source/MutableRuntime/Private/MuR/ImagePrivate.h
+++ b/Engine/Plugins/Experimental/Mutable/Source/MutableRuntime/Private/MuR/ImagePrivate.h
@@ -136,7 +136,7 @@ namespace mu
     //---------------------------------------------------------------------------------------------
     //!
     //---------------------------------------------------------------------------------------------
-    inline void FillPlainColourImage( Image* pImage, vec4<float> c )
+    inline void FillPlainColourImage( Image* pImage, FVector4f c )
     {
         int pixelCount = pImage->GetSizeX() * pImage->GetSizeY();
 
diff --git a/Engine/Plugins/Experimental/Mutable/Source/MutableRuntime/Private/MuR/Instance.cpp b/Engine/Plugins/Experimental/Mutable/Source/MutableRuntime/Private/MuR/Instance.cpp
--- a/Engine/Plugins/Experimental/Mutable/Source/MutableRuntime/Private/MuR/Instance.cpp
+++ b/Engine/Plugins/Experimental/Mutable/Source/MutableRuntime/Private/MuR/Instance.cpp
@@ -375,7 +375,7 @@ namespace mu
         check( surf>=0 && surf<m_pD->m_lods[lod].m_components[comp].m_surfaces.Num() );
         check( vec>=0 && vec<m_pD->m_lods[lod].m_components[comp].m_surfaces[surf].m_vectors.Num() );
 
-        vec4<float> r = m_pD->m_lods[lod].m_components[comp].m_surfaces[surf].m_vectors[vec].m_vec;
+        FVector4f r = m_pD->m_lods[lod].m_components[comp].m_surfaces[surf].m_vectors[vec].m_vec;
 
 		if ( pX ) { *pX = r[0]; }
 		if ( pY ) { *pY = r[1]; }
@@ -621,7 +621,7 @@ namespace mu
 
 
 	//---------------------------------------------------------------------------------------------
-    int Instance::Private::AddVector( int lod, int comp, int surf, const vec4<float>& vec, const char* strName )
+    int Instance::Private::AddVector( int lod, int comp, int surf, const FVector4f& vec, const char* strName )
 	{
 		// Automatically create the necessary lods and components
 		while ( lod>=GetLODCount() )
diff --git a/Engine/Plugins/Experimental/Mutable/Source/MutableRuntime/Private/MuR/InstancePrivate.h b/Engine/Plugins/Experimental/Mutable/Source/MutableRuntime/Private/MuR/InstancePrivate.h
--- a/Engine/Plugins/Experimental/Mutable/Source/MutableRuntime/Private/MuR/InstancePrivate.h
+++ b/Engine/Plugins/Experimental/Mutable/Source/MutableRuntime/Private/MuR/InstancePrivate.h
@@ -37,7 +37,7 @@ namespace mu
 
 		struct VECTOR
 		{
-			VECTOR( const vec4<float>& v, const char* strName )
+			VECTOR( const FVector4f& v, const char* strName )
 			{
 				m_vec = v;
 				if (strName)
@@ -46,7 +46,7 @@ namespace mu
 				}
 			}
 
-			vec4<float> m_vec;
+			FVector4f m_vec;
 			string m_name;
 		};
 
@@ -151,7 +151,7 @@ namespace mu
 		int AddSurface( int lod, int comp );
         void SetSurfaceName( int lod, int comp, int surf, const char* strName );
         int AddImage( int lod, int comp, int surf, RESOURCE_ID, const char* strName );
-        int AddVector( int lod, int comp, int surf, const vec4<float>&, const char* strName );
+        int AddVector( int lod, int comp, int surf, const FVector4f&, const char* strName );
         int AddScalar( int lod, int comp, int surf, float, const char* strName );
         int AddString( int lod, int comp, int surf, const char* strValue, const char* strName );
     };
diff --git a/Engine/Plugins/Experimental/Mutable/Source/MutableRuntime/Private/MuR/MeshPrivate.h b/Engine/Plugins/Experimental/Mutable/Source/MutableRuntime/Private/MuR/MeshPrivate.h
--- a/Engine/Plugins/Experimental/Mutable/Source/MutableRuntime/Private/MuR/MeshPrivate.h
+++ b/Engine/Plugins/Experimental/Mutable/Source/MutableRuntime/Private/MuR/MeshPrivate.h
@@ -201,9 +201,9 @@ namespace mu
 			return m_components;
 		}
 
-        vec4<float> GetAsVec4f() const
+        FVector4f GetAsVec4f() const
         {
-            vec4<float> res;
+			FVector4f res;
 
             for (int c=0; c< FMath::Min(m_components,4); ++c)
             {
@@ -485,9 +485,9 @@ namespace mu
             return m_components;
         }
 
-		vec4<float> GetAsVec4f() const
+		FVector4f GetAsVec4f() const
 		{
-			vec4<float> res = { 0.0f,0.0f,0.0f,0.0f };
+			FVector4f res(0.0f,0.0f,0.0f,0.0f);
 			for (int c = 0; c < FMath::Min(m_components, 4); ++c)
 			{
 				ConvertData(c, &res[0], MBF_FLOAT32, ptr(), m_format);
diff --git a/Engine/Plugins/Experimental/Mutable/Source/MutableRuntime/Private/MuR/Model.cpp b/Engine/Plugins/Experimental/Mutable/Source/MutableRuntime/Private/MuR/Model.cpp
--- a/Engine/Plugins/Experimental/Mutable/Source/MutableRuntime/Private/MuR/Model.cpp
+++ b/Engine/Plugins/Experimental/Mutable/Source/MutableRuntime/Private/MuR/Model.cpp
@@ -729,8 +729,7 @@ namespace mu
     ModelParametersGenerator::ModelParametersGenerator
         (
             const Model* pModel,
-            System* pSystem,
-            bool considerRelevancy
+            System* pSystem
         )
     {
 		LLM_SCOPE_BYNAME(TEXT("MutableRuntime"));
@@ -740,40 +739,30 @@ namespace mu
         m_pD->m_pSystem = pSystem;
 
         size_t paramCount = pModel->GetPrivate()->m_program.m_parameters.Num();
-        m_pD->m_considerRelevancy = considerRelevancy;
-        if (m_pD->m_considerRelevancy)
-        {
-			TArray<int> currentValues;
-			currentValues.SetNumZeroed(paramCount, 0);
-            m_pD->m_intervals.SetNum( paramCount );
-            m_pD->m_instanceCount = (int)m_pD->BuildIntervals( 0, 0, currentValues );
-        }
-        else
+
+		m_pD->m_instanceCount = 1;
+        for (size_t i=0;i<paramCount; ++i)
         {
-            m_pD->m_instanceCount = 1;
-            for (size_t i=0;i<paramCount; ++i)
+            const FParameterDesc& param =
+                    m_pD->m_pModel->GetPrivate()->m_program.m_parameters[ i ];
+            switch ( param.m_type )
             {
-                const FParameterDesc& param =
-                        m_pD->m_pModel->GetPrivate()->m_program.m_parameters[ i ];
-                switch ( param.m_type )
-                {
 
-                case PARAMETER_TYPE::T_INT:
-                {
-                    m_pD->m_instanceCount *= param.m_possibleValues.Num();
-                    break;
-                }
+            case PARAMETER_TYPE::T_INT:
+            {
+                m_pD->m_instanceCount *= param.m_possibleValues.Num();
+                break;
+            }
 
-                case PARAMETER_TYPE::T_BOOL:
-                {
-                    m_pD->m_instanceCount *= 2;
-                    break;
-                }
+            case PARAMETER_TYPE::T_BOOL:
+            {
+                m_pD->m_instanceCount *= 2;
+                break;
+            }
 
-                default:
-                    // Parameter not discrete, so ignored for combinations.
-                    break;
-                }
+            default:
+                // Parameter not discrete, so ignored for combinations.
+                break;
             }
         }
     }
@@ -805,117 +794,145 @@ namespace mu
 
 
     //---------------------------------------------------------------------------------------------
-    Ptr<Parameters> ModelParametersGenerator::GetInstance( int64 index, float (*randomGenerator )() )
+    Ptr<Parameters> ModelParametersGenerator::GetInstance( int64 index, TFunction<float()> RandomGenerator )
     {
 		LLM_SCOPE_BYNAME(TEXT("MutableRuntime"));
 
 		Ptr<Parameters> res = m_pD->m_pModel->NewParameters();
 
-        if (m_pD->m_considerRelevancy)
+		const FProgram& Program = m_pD->m_pModel->GetPrivate()->m_program;
+        int paramCount = Program.m_parameters.Num();
+        int64 currentInstance = index;
+        for (int i=0;i<paramCount;++i)
         {
-			TArray<int> values = m_pD->GetParameters( (int)index );
-            for ( int32 p=0; p<values.Num(); ++p )
-            {
-                switch ( res->GetType( p ) )
-                {
-                case PARAMETER_TYPE::T_BOOL:
-                    res->SetBoolValue( p, values[p]!=0 );
-                    break;
+            const FParameterDesc& param = Program.m_parameters[ i ];
+			Ptr<RangeIndex> RangeIndex = res->NewRangeIndex(i);
 
-                case PARAMETER_TYPE::T_INT:
-                    res->SetIntValue( p, res->GetIntPossibleValue( p, values[p] ) );
-                    break;
+            switch ( param.m_type )
+            {
 
-                case PARAMETER_TYPE::T_FLOAT:
-                    if ( randomGenerator )
-                    {
-                        res->SetFloatValue( p, randomGenerator() );
-                    }
-                    break;
+            case PARAMETER_TYPE::T_INT:
+            {
+				bool bIsRangeSize = Program.m_ranges.ContainsByPredicate([i](const FRangeDesc& r) { return r.m_dimensionParameter == i; });
+				if (bIsRangeSize)
+				{
+					res->SetIntValue(i, m_pD->DefaultRangeDimension);
+				}
+				else
+				{
+					int numOptions = res->GetIntPossibleValueCount(i);
+					int value = res->GetIntPossibleValue(i, (int)(currentInstance % numOptions));
+					res->SetIntValue(i, value);
+					currentInstance /= numOptions;
+				}
+                break;
+            }
 
-                case PARAMETER_TYPE::T_COLOUR:
-                    if ( randomGenerator )
-                    {
-                        res->SetColourValue( p, randomGenerator(), randomGenerator(), randomGenerator() );
-                    }
-                    break;
+            case PARAMETER_TYPE::T_BOOL:
+            {
+                res->SetBoolValue( i, currentInstance%2!=0 );
+                currentInstance /= 2;
+                break;
+            }
 
-				case PARAMETER_TYPE::T_PROJECTOR:
-					if (randomGenerator)
+            case PARAMETER_TYPE::T_FLOAT:
+			{
+				bool bIsRangeSize = Program.m_ranges.ContainsByPredicate([i](const FRangeDesc& r) { return r.m_dimensionParameter == i; });
+				if (bIsRangeSize)
+				{
+					if (RangeIndex)
 					{
-						// For projectors we just warp the position a little bit just to get something different
-						FVector3f Position,Direction,Up,Scale;
-						float Angle=1.0f;
-						res->GetProjectorValue(p, nullptr, &Position, &Direction, &Up, &Scale, &Angle );
-						Position.X *= 0.9 + randomGenerator() * 0.2f;
-						Position.Y *= 0.9 + randomGenerator() * 0.2f;
-						Position.Z *= 0.9 + randomGenerator() * 0.2f;
-						res->SetProjectorValue(p,Position,Direction,Up,Scale,Angle);
+						for (int32 Dimensions = 0; Dimensions < RangeIndex->GetRangeCount(); ++Dimensions)
+						{
+							for (int32 RangePosition = 0; RangePosition < m_pD->DefaultRangeDimension; ++RangePosition)
+							{
+								RangeIndex->SetPosition(Dimensions, RangePosition);
+								res->SetFloatValue(i, float(m_pD->DefaultRangeDimension), RangeIndex);
+							}
+						}
 					}
-					break;
-
-                default:
-                    break;
-                }
-            }
-        }
-        else
-        {
-            int paramCount = m_pD->m_pModel->GetPrivate()->m_program.m_parameters.Num();
-            int64 currentInstance = index;
-            for (int i=0;i<paramCount;++i)
-            {
-                const FParameterDesc& param = m_pD->m_pModel->GetPrivate()->m_program.m_parameters[ i ];
-                switch ( param.m_type )
-                {
-
-                case PARAMETER_TYPE::T_INT:
-                {
-                    int numOptions = res->GetIntPossibleValueCount( i );
-                    int value = res->GetIntPossibleValue( i, (int)(currentInstance%numOptions) );
-                    res->SetIntValue( i, value );
-                    currentInstance /= numOptions;
-                    break;
-                }
+					else
+					{
+						res->SetFloatValue(i, float(m_pD->DefaultRangeDimension));
+					}
+				}
+				else if (RandomGenerator)
+				{
+					if (RangeIndex)
+					{
+						for (int32 Dimensions = 0; Dimensions < RangeIndex->GetRangeCount(); ++Dimensions)
+						{
+							for (int32 RangePosition = 0; RangePosition < m_pD->DefaultRangeDimension; ++RangePosition)
+							{
+								RangeIndex->SetPosition(Dimensions, RangePosition);
+								res->SetFloatValue(i, RandomGenerator(), RangeIndex);
+							}
+						}
+					}
+					else
+					{
+						res->SetFloatValue(i, RandomGenerator());
+					}
+				}
+				break;
+			}
 
-                case PARAMETER_TYPE::T_BOOL:
+            case PARAMETER_TYPE::T_COLOUR:
+                if (RandomGenerator)
                 {
-                    res->SetBoolValue( i, currentInstance%2!=0 );
-                    currentInstance /= 2;
-                    break;
+					if (RangeIndex)
+					{
+						for (int32 Dimensions = 0; Dimensions < RangeIndex->GetRangeCount(); ++Dimensions)
+						{
+							for (int32 RangePosition = 0; RangePosition < m_pD->DefaultRangeDimension; ++RangePosition)
+							{
+								RangeIndex->SetPosition(Dimensions, RangePosition);
+								res->SetColourValue(i, RandomGenerator(), RandomGenerator(), RandomGenerator(), RangeIndex);
+							}
+						}
+					}
+					else
+					{
+						res->SetColourValue(i, RandomGenerator(), RandomGenerator(), RandomGenerator());
+					}
                 }
+                break;
 
-                case PARAMETER_TYPE::T_FLOAT:
-                    if ( randomGenerator )
-                    {
-                        res->SetFloatValue( i, randomGenerator() );
-                    }
-                    break;
-
-                case PARAMETER_TYPE::T_COLOUR:
-                    if ( randomGenerator )
-                    {
-                        res->SetColourValue( i, randomGenerator(), randomGenerator(), randomGenerator() );
-                    }
-                    break;
+			case PARAMETER_TYPE::T_PROJECTOR:
+				if (RandomGenerator)
+				{
+					// For projectors we just warp the position a little bit just to get something different
+					FVector3f Position, Direction, Up, Scale;
+					float Angle = 1.0f;
+					res->GetProjectorValue(i, nullptr, &Position, &Direction, &Up, &Scale, &Angle);
 
-				case PARAMETER_TYPE::T_PROJECTOR:
-					if (randomGenerator)
+					if (RangeIndex)
 					{
-						// For projectors we just warp the position a little bit just to get something different
-						FVector3f Position, Direction, Up, Scale;
-						float Angle = 1.0f;
-						res->GetProjectorValue(i, nullptr, &Position, &Direction, &Up, &Scale, &Angle);
-						Position.X *= 0.9 + randomGenerator() * 0.2f;
-						Position.Y *= 0.9 + randomGenerator() * 0.2f;
-						Position.Z *= 0.9 + randomGenerator() * 0.2f;
+						for (int32 Dimensions = 0; Dimensions < RangeIndex->GetRangeCount(); ++Dimensions)
+						{
+							for (int32 RangePosition = 0; RangePosition < m_pD->DefaultRangeDimension; ++RangePosition)
+							{
+								RangeIndex->SetPosition(Dimensions, RangePosition);
+								Position.X *= 0.9 + RandomGenerator() * 0.2f;
+								Position.Y *= 0.9 + RandomGenerator() * 0.2f;
+								Position.Z *= 0.9 + RandomGenerator() * 0.2f;
+								res->SetProjectorValue(i, Position, Direction, Up, Scale, Angle, RangeIndex);
+
+							}
+						}
+					}
+					else
+					{
+						Position.X *= 0.9 + RandomGenerator() * 0.2f;
+						Position.Y *= 0.9 + RandomGenerator() * 0.2f;
+						Position.Z *= 0.9 + RandomGenerator() * 0.2f;
 						res->SetProjectorValue(i, Position, Direction, Up, Scale, Angle);
 					}
-					break;
+				}
+				break;
 
-                default:
-                    break;
-                }
+            default:
+                break;
             }
         }
 
@@ -930,20 +947,61 @@ namespace mu
 
         ParametersPtr res = m_pD->m_pModel->NewParameters();
 
-        int paramCount = (int)m_pD->m_pModel->GetPrivate()->m_program.m_parameters.Num();
+		const FProgram& Program = m_pD->m_pModel->GetPrivate()->m_program;
+        int32 ParamCount = Program.m_parameters.Num();
 
-        for (int i=0;i<paramCount;++i)
+        for (int i=0;i<ParamCount;++i)
         {
-            const FParameterDesc& param = m_pD->m_pModel->GetPrivate()->m_program.m_parameters[ i ];
+            const FParameterDesc& param = Program.m_parameters[ i ];
+			Ptr<RangeIndex> RangeIndex = res->NewRangeIndex(i);
+
             switch ( param.m_type )
             {
 
             case PARAMETER_TYPE::T_INT:
             {
-                int numOptions = res->GetIntPossibleValueCount( i );
-                int valueIndex = (int)(FMath::Min(numOptions-1,int(randomGenerator()*numOptions)));
-                int value = res->GetIntPossibleValue( i, valueIndex );
-                res->SetIntValue( i, value );
+				bool bIsRangeSize = Program.m_ranges.ContainsByPredicate([i](const FRangeDesc& r) { return r.m_dimensionParameter == i; });
+				if (bIsRangeSize)
+				{
+					if (RangeIndex)
+					{
+						for (int32 Dimensions = 0; Dimensions < RangeIndex->GetRangeCount(); ++Dimensions)
+						{
+							for (int32 RangePosition = 0; RangePosition < m_pD->DefaultRangeDimension; ++RangePosition)
+							{
+								RangeIndex->SetPosition(Dimensions, RangePosition);
+								res->SetIntValue(i, m_pD->DefaultRangeDimension, RangeIndex);
+							}
+						}
+					}
+					else
+					{
+						res->SetIntValue(i, m_pD->DefaultRangeDimension);
+					}
+				}
+				else
+				{
+					if (RangeIndex)
+					{
+						for (int32 Dimensions = 0; Dimensions < RangeIndex->GetRangeCount(); ++Dimensions)
+						{
+							for (int32 RangePosition = 0; RangePosition < m_pD->DefaultRangeDimension; ++RangePosition)
+							{
+								int numOptions = res->GetIntPossibleValueCount(i);
+								int valueIndex = (int)(FMath::Min(numOptions - 1, int(randomGenerator() * numOptions)));
+								int value = res->GetIntPossibleValue(i, valueIndex);
+								res->SetIntValue(i, value, RangeIndex);
+							}
+						}
+					}
+					else
+					{
+						int numOptions = res->GetIntPossibleValueCount(i);
+						int valueIndex = (int)(FMath::Min(numOptions - 1, int(randomGenerator() * numOptions)));
+						int value = res->GetIntPossibleValue(i, valueIndex);
+						res->SetIntValue(i, value);
+					}
+				}
                 break;
             }
 
@@ -954,25 +1012,90 @@ namespace mu
             }
 
             case PARAMETER_TYPE::T_FLOAT:
-                res->SetFloatValue( i, randomGenerator() );
-                break;
+			{
+				bool bIsRangeSize = Program.m_ranges.ContainsByPredicate([i](const FRangeDesc& r) { return r.m_dimensionParameter == i; });
+				if (bIsRangeSize)
+				{
+					if (RangeIndex)
+					{
+						for (int32 Dimensions = 0; Dimensions < RangeIndex->GetRangeCount(); ++Dimensions)
+						{
+							for (int32 RangePosition = 0; RangePosition < m_pD->DefaultRangeDimension; ++RangePosition)
+							{
+								res->SetFloatValue(i, float(m_pD->DefaultRangeDimension), RangeIndex);
+							}
+						}
+					}
+					else
+					{
+						res->SetFloatValue(i, float(m_pD->DefaultRangeDimension));
+					}
+				}
+				else
+				{
+					if (RangeIndex)
+					{
+						for (int32 Dimensions = 0; Dimensions < RangeIndex->GetRangeCount(); ++Dimensions)
+						{
+							for (int32 RangePosition = 0; RangePosition < m_pD->DefaultRangeDimension; ++RangePosition)
+							{
+								res->SetFloatValue(i, randomGenerator(), RangeIndex);
+							}
+						}
+					}
+					else
+					{
+						res->SetFloatValue(i, randomGenerator());
+					}
+				}
+				break;
+			}
 
             case PARAMETER_TYPE::T_COLOUR:
-                res->SetColourValue( i, randomGenerator(), randomGenerator(), randomGenerator() );
+				if (RangeIndex)
+				{
+					for (int32 Dimensions = 0; Dimensions < RangeIndex->GetRangeCount(); ++Dimensions)
+					{
+						for (int32 RangePosition = 0; RangePosition < m_pD->DefaultRangeDimension; ++RangePosition)
+						{
+							res->SetColourValue(i, randomGenerator(), randomGenerator(), randomGenerator(),RangeIndex);
+						}
+					}
+				}
+				else
+				{
+					res->SetColourValue(i, randomGenerator(), randomGenerator(), randomGenerator());
+				}
                 break;
 
 			case PARAMETER_TYPE::T_PROJECTOR:
+			{
+				// For projectors we just warp the position a little bit just to get something different
+				FVector3f Position, Direction, Up, Scale;
+				float Angle = 1.0f;
+				res->GetProjectorValue(i, nullptr, &Position, &Direction, &Up, &Scale, &Angle);
+				if (RangeIndex)
+				{
+					for (int32 Dimensions = 0; Dimensions < RangeIndex->GetRangeCount(); ++Dimensions)
+					{
+						for (int32 RangePosition = 0; RangePosition < m_pD->DefaultRangeDimension; ++RangePosition)
+						{
+							Position.X *= 0.9 + randomGenerator() * 0.2f;
+							Position.Y *= 0.9 + randomGenerator() * 0.2f;
+							Position.Z *= 0.9 + randomGenerator() * 0.2f;
+							res->SetProjectorValue(i, Position, Direction, Up, Scale, Angle, RangeIndex);
+						}
+					}
+				}
+				else
 				{
-					// For projectors we just warp the position a little bit just to get something different
-					FVector3f Position, Direction, Up, Scale;
-					float Angle = 1.0f;
-					res->GetProjectorValue(i, nullptr, &Position, &Direction, &Up, &Scale, &Angle);
 					Position.X *= 0.9 + randomGenerator() * 0.2f;
 					Position.Y *= 0.9 + randomGenerator() * 0.2f;
 					Position.Z *= 0.9 + randomGenerator() * 0.2f;
 					res->SetProjectorValue(i, Position, Direction, Up, Scale, Angle);
 				}
 				break;
+			}
 
             default:
                 break;
@@ -983,217 +1106,5 @@ namespace mu
     }
 
 
-    //---------------------------------------------------------------------------------------------
-    //---------------------------------------------------------------------------------------------
-    //---------------------------------------------------------------------------------------------
-//    class PartialRelevantParameterVisitor : public PartialDiscreteCoveredCodeVisitor
-//    {
-//    public:
-
-//        PartialRelevantParameterVisitor
-//            (
-//                System::Private* pSystem,
-//                Ptr<const Model>& pModel,
-//                int state,
-//                const Parameters* pParams,
-//                int relevantParameters,
-//                int thisParameter
-//            )
-//            : PartialDiscreteCoveredCodeVisitor( pSystem, pModel, pParams, relevantParameters )
-//        {
-//            m_visited.resize( pModel->GetPrivate()->m_program.m_opAddress.Num(), false );
-//            m_thisParameter = thisParameter;
-//            m_relevant = false;
-
-//            OP::ADDRESS at = pModel->GetPrivate()->m_program.m_states[ state ].m_root;
-//            Run( at );
-
-//        }
-
-
-//        void Visit( OP::ADDRESS at, FProgram& program ) override
-//        {
-//            if ( !m_relevant && !m_visited[at])
-//            {
-//                m_visited[at] = true;
-
-//                switch ( program.m_code[at].type )
-//                {
-//                case OP_TYPE::BO_PARAMETER:
-//                case OP_TYPE::NU_PARAMETER:
-//                case OP_TYPE::SC_PARAMETER:
-//                case OP_TYPE::CO_PARAMETER:
-//                case OP_TYPE::PR_PARAMETER:
-//                case OP_TYPE::IM_PARAMETER:
-//                    if ( program.m_code[at].args.Parameter.variable == (size_t)m_thisParameter )
-//                    {
-//                        m_relevant = true;
-//                    }
-//                    break;
-
-//                default:
-//                    break;
-
-//                }
-
-//                PartialDiscreteCoveredCodeVisitor::Visit( at, program );
-//            }
-//        }
-
-//        //! Result
-//        bool m_relevant;
-
-//    private:
-
-//        int m_thisParameter;
-
-//        //! Flags for visited instructions.
-//        TArray<bool> m_visited;
-//    };
-
-
-    //---------------------------------------------------------------------------------------------
-	uint32 ModelParametersGenerator::Private::BuildIntervals(
-		uint32 currentInstanceIndex,
-		uint32 currentParameter,
-		TArray<int>& currentValues )
-    {
-        size_t paramCount = m_pModel->GetPrivate()->m_program.m_parameters.Num();
-
-        if ( currentParameter>=paramCount )
-        {
-            currentInstanceIndex++;
-        }
-        else
-        {
-            // Find out about parameter relevancy
-            bool relevant = true;
-            {
-                ParametersPtr pParams = m_pModel->NewParameters();
-                for ( size_t p=0; p<currentParameter; ++p )
-                {
-                    switch ( pParams->GetType( (int)p ) )
-                    {
-                    case PARAMETER_TYPE::T_BOOL:
-                        pParams->SetBoolValue( (int)p, currentValues[p]!=0 );
-                        break;
-
-                    case PARAMETER_TYPE::T_INT:
-                        pParams->SetIntValue( (int)p,
-                                              pParams->GetIntPossibleValue( (int)p, currentValues[p] ) );
-                        break;
-
-                    default:
-                        break;
-                    }
-                }
-
-                // todo
-                relevant = true;
-//                PartialRelevantParameterVisitor relevancyVisitor( m_pSystem->GetPrivate(),
-//                                                                  m_pModel,
-//                                                                  0, // state
-//                                                                  pParams.get(),
-//                                                                  (int)currentParameter,
-//                                                                  (int)currentParameter
-//                                                                  );
-//                relevant = relevancyVisitor.m_relevant;
-            }
-
-            if (relevant)
-            {
-                const FParameterDesc& param =
-                        m_pModel->GetPrivate()->m_program.m_parameters[ currentParameter ];
-                switch ( param.m_type )
-                {
-
-                case PARAMETER_TYPE::T_BOOL:
-                {
-                    // False
-                    PARAMETER_INTERVAL_VALUE falseValue;
-                    falseValue.m_minIndex = (int)currentInstanceIndex;
-                    falseValue.m_value = 0;
-                    m_intervals[ currentParameter ].m_intervalValue.Add( falseValue );
-                    currentValues[ currentParameter ] = 0;
-                    currentInstanceIndex = BuildIntervals( currentInstanceIndex,
-                                                           currentParameter+1,
-                                                           currentValues);
-
-                    // True
-                    PARAMETER_INTERVAL_VALUE trueValue;
-                    trueValue.m_minIndex = (int)currentInstanceIndex;
-                    trueValue.m_value = 1;
-                    m_intervals[ currentParameter ].m_intervalValue.Add( trueValue );
-                    currentValues[ currentParameter ] = 1;
-                    currentInstanceIndex = BuildIntervals( currentInstanceIndex,
-                                                           currentParameter+1,
-                                                           currentValues );
-                    break;
-                }
-
-                case PARAMETER_TYPE::T_INT:
-                {
-                    for ( int32 v=0; v<param.m_possibleValues.Num(); ++v )
-                    {
-                        PARAMETER_INTERVAL_VALUE value;
-                        value.m_minIndex = (int)currentInstanceIndex;
-                        value.m_value = v;
-                        m_intervals[ currentParameter ].m_intervalValue.Add( value );
-                        currentValues[ currentParameter ] = v;
-                        currentInstanceIndex = BuildIntervals( currentInstanceIndex,
-                                                               currentParameter+1,
-                                                               currentValues);
-                    }
-                    break;
-                }
-
-                default:
-                    // Continuous parameter.
-                    currentInstanceIndex = BuildIntervals( currentInstanceIndex,
-                                                           currentParameter+1,
-                                                           currentValues);
-                    break;
-                }
-            }
-            else
-            {
-                // Irrelevant parameter
-                currentInstanceIndex = BuildIntervals( currentInstanceIndex,
-                                                       currentParameter+1,
-                                                       currentValues);
-            }
-        }
-
-        return currentInstanceIndex;
-    }
-
-
-    //---------------------------------------------------------------------------------------------
-	TArray<int> ModelParametersGenerator::Private::GetParameters( int instanceIndex )
-    {
-		TArray<int> res;
-		res.SetNumZeroed(m_pModel->GetPrivate()->m_program.m_parameters.Num());
-
-        for ( size_t p=0; p<res.Num(); ++p )
-        {
-            size_t e = 0;
-            while ( m_intervals[p].m_intervalValue.Num()>e+1
-                    &&
-                    m_intervals[p].m_intervalValue[e+1].m_minIndex<=instanceIndex )
-            {
-                ++e;
-            }
-
-            // Avoid degenerated case (parameter always irrelevant?)
-            if ( m_intervals[p].m_intervalValue.Num()>e )
-            {
-                res[p] = m_intervals[p].m_intervalValue[e].m_value;
-            }
-        }
-
-        return res;
-    }
-
-
 }
 
diff --git a/Engine/Plugins/Experimental/Mutable/Source/MutableRuntime/Private/MuR/ModelPrivate.h b/Engine/Plugins/Experimental/Mutable/Source/MutableRuntime/Private/MuR/ModelPrivate.h
--- a/Engine/Plugins/Experimental/Mutable/Source/MutableRuntime/Private/MuR/ModelPrivate.h
+++ b/Engine/Plugins/Experimental/Mutable/Source/MutableRuntime/Private/MuR/ModelPrivate.h
@@ -721,41 +721,16 @@ namespace mu
     public:
 
         //!
-        ModelPtrConst m_pModel;
-        SystemPtr m_pSystem;
+        Ptr<const Model> m_pModel;
+        Ptr<System> m_pSystem;
 
         //! Number of possible instances
         int64 m_instanceCount;
 
-        //!
-        struct PARAMETER_INTERVAL_VALUE
-        {
-            //! Minimum instance index to use these increments
-            int m_minIndex;
-
-            //! Parameter value in this interval
-            int m_value;
-        };
-
-        //!
-        struct PARAMETER_INTERVALS
-        {
-            //! As many entries as necessary
-            //! Sorted by PARAMETER_INCREMENTS_PER_VALUE::m_minIndex
-			TArray<PARAMETER_INTERVAL_VALUE> m_intervalValue;
-        };
-
-        //! An entry for every parameter in the model
-		TArray< PARAMETER_INTERVALS > m_intervals;
-
-        //! Whether to use brute force, or consider the parameter relevancy
-        bool m_considerRelevancy;
-
-        //! Build the interval information
-		uint32 BuildIntervals( uint32 currentInstanceIndex, uint32 currentParameter, TArray<int>& currentValues );
-
-        //! Get the parameter values for a particular instance index
-		TArray<int> GetParameters( int instanceIndex );
+		/** Value used for scalar paramters that control multidemnsional sizes. 
+		* TODO: Make it an option.
+		*/
+		int32 DefaultRangeDimension = 8;
     };
 
 }
diff --git a/Engine/Plugins/Experimental/Mutable/Source/MutableRuntime/Private/MuR/MutableMath.h b/Engine/Plugins/Experimental/Mutable/Source/MutableRuntime/Private/MuR/MutableMath.h
--- a/Engine/Plugins/Experimental/Mutable/Source/MutableRuntime/Private/MuR/MutableMath.h
+++ b/Engine/Plugins/Experimental/Mutable/Source/MutableRuntime/Private/MuR/MutableMath.h
@@ -454,17 +454,6 @@ namespace mu
 			return res;
 		}
 
-		//-----------------------------------------------------------------------------------------
-		inline bool AlmostNull( SCALAR tol = std::numeric_limits<SCALAR>::min() ) const
-		{
-			bool res = true;
-			for( int i=0; i<DIM; ++i )
-			{
-				res &= fabs(m[i])<tol;
-			}
-			return res;
-		}
-
 		//-----------------------------------------------------------------------------------------
 		inline void Serialise( OutputArchive& arch ) const
 		{
@@ -601,6 +590,14 @@ namespace mu
 			this->m[2] = other[2];
 		}
 
+		//-----------------------------------------------------------------------------------------
+		inline vec3(const FVector4f& other)
+		{
+			this->m[0] = other[0];
+			this->m[1] = other[1];
+			this->m[2] = other[2];
+		}
+
         //-----------------------------------------------------------------------------------------
         inline void Set( SCALAR x, SCALAR y, SCALAR z )
         {
@@ -941,84 +938,6 @@ namespace mu
 	}
 
 
-	//---------------------------------------------------------------------------------------------
-	//! Angle between two vectors, in radians
-	//---------------------------------------------------------------------------------------------
-	inline float angle( const vec3f& a, const vec3f& b )
-	{
-		float s = sqrt( dot(a,a) * dot(b,b) );
-
-		if ( s > std::numeric_limits<float>::min() )
-		{
-			return acos( dot(a,b) / s );
-		}
-		else
-		{
-			return 0.0f;
-		}
-	}
-
-
-	//---------------------------------------------------------------------------------------------
-	//!
-	//---------------------------------------------------------------------------------------------
-	template <class SCALAR>
-	class vec4 : public vec<SCALAR,4>
-	{
-	public:
-
-		//-----------------------------------------------------------------------------------------
-		inline vec4()
-		{
-			// Elements are cleared in the parent constructor
-		}
-
-		//-----------------------------------------------------------------------------------------
-		inline vec4( SCALAR x, SCALAR y, SCALAR z, SCALAR w )
-		{
-			this->m[0] = x;
-			this->m[1] = y;
-			this->m[2] = z;
-			this->m[3] = w;
-		}
-
-        //-----------------------------------------------------------------------------------------
-        inline vec4( const vec<SCALAR,4>& other )
-        {
-            this->m[0] = other[0];
-            this->m[1] = other[1];
-            this->m[2] = other[2];
-            this->m[3] = other[3];
-        }
-
-        //-----------------------------------------------------------------------------------------
-        inline vec4( const vec<SCALAR,3>& other, float w )
-        {
-            this->m[0] = other[0];
-            this->m[1] = other[1];
-            this->m[2] = other[2];
-            this->m[3] = w;
-        }
-
-        //-----------------------------------------------------------------------------------------
-        inline vec3<SCALAR> xyz() const
-        {
-            return vec3<SCALAR>( this->m[0], this->m[1], this->m[2] );
-        }
-
-        //-----------------------------------------------------------------------------------------
-        inline vec2<SCALAR> xy() const
-        {
-            return vec2<SCALAR>( this->m[0], this->m[1] );
-        }
-
-	};
-
-	//!
-	typedef vec4<float> vec4f;
-    MUTABLE_DEFINE_POD_VECTOR_SERIALISABLE( vec4f );
-
-
 	//!
     template <class VECTOR>
     class box
@@ -1527,17 +1446,6 @@ namespace mu
 			return res;
 		}
 
-		//-----------------------------------------------------------------------------------------
-		inline bool AlmostNull( SCALAR tol = std::numeric_limits<SCALAR>::min() ) const
-		{
-			bool res = true;
-			for( int i=0; i<DIM; ++i )
-			{
-				res &= m[i].AlmostNull( tol );
-			}
-			return res;
-		}
-
 		//-----------------------------------------------------------------------------------------
 		inline void Serialise( OutputArchive& arch ) const
 		{
diff --git a/Engine/Plugins/Experimental/Mutable/Source/MutableRuntime/Private/MuR/OpImageColourMap.h b/Engine/Plugins/Experimental/Mutable/Source/MutableRuntime/Private/MuR/OpImageColourMap.h
--- a/Engine/Plugins/Experimental/Mutable/Source/MutableRuntime/Private/MuR/OpImageColourMap.h
+++ b/Engine/Plugins/Experimental/Mutable/Source/MutableRuntime/Private/MuR/OpImageColourMap.h
@@ -19,22 +19,22 @@ namespace mu
                                     pSource->GetLODCount(),
 									pSource->GetFormat() );
 
-        uint8_t* pDestBuf = pDest->GetData();
-        const uint8_t* pSourceBuf = pSource->GetData();
-        const uint8_t* pMaskBuf = pMask->GetData();
+        uint8* pDestBuf = pDest->GetData();
+        const uint8* pSourceBuf = pSource->GetData();
+        const uint8* pMaskBuf = pMask->GetData();
 
 		// Generic implementation
 		int pixelCount = (int)pSource->CalculatePixelCount();
 
 		// Make a palette for faster conversion
-        uint8_t palette[256][4];
+        uint8 palette[256][4];
 		for ( int i=0; i<256; ++i )
 		{
-			vec4<float> c = pMap->Sample( vec2<float>( float(i)/256.0f, 0.0f ) );
-            palette[i][0] = (uint8_t)FMath::Max( 0, FMath::Min( 255, int(c[0]*255.0f) ) );
-            palette[i][1] = (uint8_t)FMath::Max( 0, FMath::Min( 255, int(c[1]*255.0f) ) );
-            palette[i][2] = (uint8_t)FMath::Max( 0, FMath::Min( 255, int(c[2]*255.0f) ) );
-            palette[i][3] = (uint8_t)FMath::Max( 0, FMath::Min( 255, int(c[3]*255.0f) ) );
+			FVector4f c = pMap->Sample(FVector2f( float(i)/256.0f, 0.0f ) );
+            palette[i][0] = (uint8)FMath::Max( 0, FMath::Min( 255, int(c[0]*255.0f) ) );
+            palette[i][1] = (uint8)FMath::Max( 0, FMath::Min( 255, int(c[1]*255.0f) ) );
+            palette[i][2] = (uint8)FMath::Max( 0, FMath::Min( 255, int(c[2]*255.0f) ) );
+            palette[i][3] = (uint8)FMath::Max( 0, FMath::Min( 255, int(c[3]*255.0f) ) );
 		}
 
 		switch ( pSource->GetFormat() )
@@ -43,7 +43,7 @@ namespace mu
 		{
 			for ( int i=0; i<pixelCount; ++i )
 			{
-                uint8_t m_8 = pMaskBuf[i];
+                uint8 m_8 = pMaskBuf[i];
 				if (m_8>127)
 				{
 					pDestBuf[i] = palette[ pSourceBuf[i] ][0];
@@ -60,7 +60,7 @@ namespace mu
 		{
 			for ( int i=0; i<pixelCount; ++i )
 			{
-                uint8_t m_8 = pMaskBuf[i];
+                uint8 m_8 = pMaskBuf[i];
 				if (m_8>127)
 				{
 					pDestBuf[3*i+0] = palette[ pSourceBuf[3*i+0] ][0];
@@ -81,7 +81,7 @@ namespace mu
         {
             for ( int i=0; i<pixelCount; ++i )
             {
-                uint8_t m_8 = pMaskBuf[i];
+                uint8 m_8 = pMaskBuf[i];
                 if (m_8>127)
                 {
                     pDestBuf[4*i+0] = palette[ pSourceBuf[4*i+0] ][0];
@@ -104,7 +104,7 @@ namespace mu
         {
             for ( int i=0; i<pixelCount; ++i )
             {
-                uint8_t m_8 = pMaskBuf[i];
+                uint8 m_8 = pMaskBuf[i];
                 if (m_8>127)
                 {
                     pDestBuf[4*i+0] = palette[ pSourceBuf[4*i+0] ][2];
diff --git a/Engine/Plugins/Experimental/Mutable/Source/MutableRuntime/Private/MuR/OpImageGradient.h b/Engine/Plugins/Experimental/Mutable/Source/MutableRuntime/Private/MuR/OpImageGradient.h
--- a/Engine/Plugins/Experimental/Mutable/Source/MutableRuntime/Private/MuR/OpImageGradient.h
+++ b/Engine/Plugins/Experimental/Mutable/Source/MutableRuntime/Private/MuR/OpImageGradient.h
@@ -10,8 +10,8 @@
 namespace mu
 {
 
-	inline ImagePtr ImageGradient( const vec3<float>& c0,
-								   const vec3<float>& c1,
+	inline ImagePtr ImageGradient( const FVector4f& c0,
+								   const FVector4f& c1,
 								   int sizeX, int sizeY )
 	{
 		ImagePtr pDest = new Image((uint16)sizeX, (uint16)sizeY, 1, EImageFormat::IF_RGB_UBYTE );
diff --git a/Engine/Plugins/Experimental/Mutable/Source/MutableRuntime/Private/MuR/OpImageProject.cpp b/Engine/Plugins/Experimental/Mutable/Source/MutableRuntime/Private/MuR/OpImageProject.cpp
--- a/Engine/Plugins/Experimental/Mutable/Source/MutableRuntime/Private/MuR/OpImageProject.cpp
+++ b/Engine/Plugins/Experimental/Mutable/Source/MutableRuntime/Private/MuR/OpImageProject.cpp
@@ -1416,7 +1416,7 @@ struct AdjacentFaces
 };
 
 
-void PlanarlyProjectVertex(const FVector3f& unfoldedPosition, vec4f& projectedPosition,
+void PlanarlyProjectVertex(const FVector3f& unfoldedPosition, FVector4f& projectedPosition,
 	const FProjector& projector, const FVector3f& projectorPosition, const FVector3f& projectorDirection, const FVector3f& s, const FVector3f& u)
 {
 	float x = FVector3f::DotProduct( unfoldedPosition - projectorPosition, s ) / projector.scale[0] + 0.5f;
@@ -2020,7 +2020,7 @@ void MeshProject_Optimised_Wrapping( const Mesh* pMesh,
                 GetBarycentricCoords(out_intersection, v3_0, v3_1, v3_2, outIntersectionBaricentric[0], outIntersectionBaricentric[1], outIntersectionBaricentric[2]);
 				FVector2f outIntersectionUV = pVertices[i0].uv * outIntersectionBaricentric[0] + pVertices[i1].uv * outIntersectionBaricentric[1] + pVertices[i2].uv * outIntersectionBaricentric[2];
 
-                vec4f proj4D_v0, proj4D_v1, proj4D_v2;
+				FVector4f proj4D_v0, proj4D_v1, proj4D_v2;
                 PlanarlyProjectVertex(pVertices[i0].pos, proj4D_v0, projector2, projectorPosition2, projectorDirection2, s2, u2);
                 PlanarlyProjectVertex(pVertices[i1].pos, proj4D_v1, projector2, projectorPosition2, projectorDirection2, s2, u2);
                 PlanarlyProjectVertex(pVertices[i2].pos, proj4D_v2, projector2, projectorPosition2, projectorDirection2, s2, u2);
diff --git a/Engine/Plugins/Experimental/Mutable/Source/MutableRuntime/Private/MuR/OpMeshDifference.h b/Engine/Plugins/Experimental/Mutable/Source/MutableRuntime/Private/MuR/OpMeshDifference.h
--- a/Engine/Plugins/Experimental/Mutable/Source/MutableRuntime/Private/MuR/OpMeshDifference.h
+++ b/Engine/Plugins/Experimental/Mutable/Source/MutableRuntime/Private/MuR/OpMeshDifference.h
@@ -92,7 +92,7 @@ struct MORPH_VERTEX
 		int differentVertexCount = 0;
 		TArray<bool> isVertexDifferent;
 		isVertexDifferent.SetNumZeroed(vcount);
-		TArray< vec4<float> > deltas;
+		TArray< FVector4f > deltas;
 		deltas.SetNum(vcount * numChannels);
 
 		for ( int c=0; c<numChannels; ++c )
@@ -105,12 +105,12 @@ struct MORPH_VERTEX
 
 			for ( size_t v=0; v<vcount; ++v )
 			{
-				vec4<float> base = baseIt.GetAsVec4f();			
-				vec4<float> target = targetIt.GetAsVec4f();		
-				vec4<float> delta = target-base;
+				FVector4f base = baseIt.GetAsVec4f();
+				FVector4f target = targetIt.GetAsVec4f();
+				FVector4f delta = target-base;
 				deltas[ v*numChannels+c ] = delta;		
 
-				if (!delta.AlmostNull())
+				if (!delta.Equals(FVector4f(0,0,0,0),UE_SMALL_NUMBER))
 				{
 					if ( !isVertexDifferent[v] )
 					{
@@ -185,7 +185,7 @@ struct MORPH_VERTEX
 					// all channels
 					for ( int c=0; c<numChannels; ++c )
 					{
-						*((vec4<float>*)pDestData) = deltas[ v*numChannels+c ];		
+						*((FVector4f*)pDestData) = deltas[ v*numChannels+c ];
 						pDestData += 4*4;											
 					}
 				}
diff --git a/Engine/Plugins/Experimental/Mutable/Source/MutableRuntime/Private/MuR/OpMeshFormat.cpp b/Engine/Plugins/Experimental/Mutable/Source/MutableRuntime/Private/MuR/OpMeshFormat.cpp
--- a/Engine/Plugins/Experimental/Mutable/Source/MutableRuntime/Private/MuR/OpMeshFormat.cpp
+++ b/Engine/Plugins/Experimental/Mutable/Source/MutableRuntime/Private/MuR/OpMeshFormat.cpp
@@ -201,49 +201,10 @@ namespace mu
 
 					// Special case for derived channel data
 					bool generated = false;
-					if (resultSemantic == MBS_TANGENTSIGN)
-					{
-						// Look for the full tangent space
-						int tanXBuf, tanXChan, tanYBuf, tanYChan, tanZBuf, tanZChan;
-						Source.FindChannel(MBS_TANGENT, resultSemanticIndex, &tanXBuf, &tanXChan);
-						Source.FindChannel(MBS_BINORMAL, resultSemanticIndex, &tanYBuf, &tanYChan);
-						Source.FindChannel(MBS_NORMAL, resultSemanticIndex, &tanZBuf, &tanZChan);
-
-						if (tanXBuf >= 0 && tanYBuf >= 0 && tanZBuf >= 0)
-						{
-							generated = true;
-							UntypedMeshBufferIteratorConst xIt(Source, MBS_TANGENT, resultSemanticIndex);
-							UntypedMeshBufferIteratorConst yIt(Source, MBS_BINORMAL, resultSemanticIndex);
-							UntypedMeshBufferIteratorConst zIt(Source, MBS_NORMAL, resultSemanticIndex);
-							for (int v = 0; v < vCount; ++v)
-							{
-								mat3f mat;
-								mat[0] = xIt.GetAsVec4f().xyz();
-								mat[1] = yIt.GetAsVec4f().xyz();
-								mat[2] = zIt.GetAsVec4f().xyz();
-								float sign = mat.GetDeterminant() < 0 ? -1.0f : 1.0f;
-								ConvertData(0, pResultBuf, resultFormat, &sign, MBF_FLOAT32);
-
-								for (int i = 1; i < resultComponents; ++i)
-								{
-									// Add zeros
-									FMemory::Memzero
-									(
-										pResultBuf + GetMeshFormatData(resultFormat).m_size * i,
-										GetMeshFormatData(resultFormat).m_size
-									);
-								}
-								pResultBuf += resultElemSize;
-								xIt++;
-								yIt++;
-								zIt++;
-							}
-						}
-					}
 
 					// If we have to add colour channels, we will add them as white, to be neutral.
 					// \todo: normal channels also should have special values.
-					else if (resultSemantic == MBS_COLOUR)
+					if (resultSemantic == MBS_COLOUR)
 					{
 						generated = true;
 
@@ -380,9 +341,9 @@ namespace mu
 								xIt += v;
 								yIt += v;
 								zIt += v;
-								mat[0] = xIt.GetAsVec4f().xyz();
-								mat[1] = yIt.GetAsVec4f().xyz();
-								mat[2] = zIt.GetAsVec4f().xyz();
+								mat[0] = vec3f(xIt.GetAsVec3f());
+								mat[1] = vec3f(yIt.GetAsVec3f());
+								mat[2] = vec3f(zIt.GetAsVec3f());
 
 								uint8_t sign = 0;
 								if (resultFormat == MBF_PACKEDDIR8_W_TANGENTSIGN)
diff --git a/Engine/Plugins/Experimental/Mutable/Source/MutableRuntime/Private/MuR/OpMeshMorph.h b/Engine/Plugins/Experimental/Mutable/Source/MutableRuntime/Private/MuR/OpMeshMorph.h
--- a/Engine/Plugins/Experimental/Mutable/Source/MutableRuntime/Private/MuR/OpMeshMorph.h
+++ b/Engine/Plugins/Experimental/Mutable/Source/MutableRuntime/Private/MuR/OpMeshMorph.h
@@ -77,16 +77,16 @@ namespace mu
                 		continue;
                 	}
                 
-                    auto channelBaseIter = baseChannelsIters[c] + v;
-                    auto channelMorphIter = morphChannelsIters[c] + m;
+                    UntypedMeshBufferIterator channelBaseIter = baseChannelsIters[c] + v;
+					UntypedMeshBufferIteratorConst channelMorphIter = morphChannelsIters[c] + m;
                    
-                    const auto dstChannelFormat = baseChannelsIters[c].GetFormat();
-                    const auto dstChannelComps = baseChannelsIters[c].GetComponents();
+                    const MESH_BUFFER_FORMAT dstChannelFormat = baseChannelsIters[c].GetFormat();
+                    const int dstChannelComps = baseChannelsIters[c].GetComponents();
 
                     // Apply Morph to range found above.
                     for ( int32 r = 0; r < runSize; ++r, ++channelBaseIter, ++channelMorphIter )
                     {
-                        const vec4<float> value = channelBaseIter.GetAsVec4f() + channelMorphIter.GetAsVec4f()*factor;
+                        const FVector4f value = channelBaseIter.GetAsVec4f() + channelMorphIter.GetAsVec4f()*factor;
 
                         // TODO: Optimize this for the specific components.
                         // Max 4 components
diff --git a/Engine/Plugins/Experimental/Mutable/Source/MutableRuntime/Private/MuR/System.cpp b/Engine/Plugins/Experimental/Mutable/Source/MutableRuntime/Private/MuR/System.cpp
--- a/Engine/Plugins/Experimental/Mutable/Source/MutableRuntime/Private/MuR/System.cpp
+++ b/Engine/Plugins/Experimental/Mutable/Source/MutableRuntime/Private/MuR/System.cpp
@@ -36,7 +36,6 @@
 
 namespace mu
 {
-    static_assert( sizeof(vec4f) == 16, "UNEXPECTED_STRUCT_PACKING" );
     static_assert( sizeof(mat4f) == 64, "UNEXPECTED_STRUCT_PACKING" );
 
     //!
@@ -700,7 +699,7 @@ namespace mu
 		float* pG,
 		float* pB)
 	{
-		vec4f col;
+		FVector4f col;
 
 		mu::OP_TYPE opType = pModel->GetPrivate()->m_program.GetOpType(at);
 		if (GetOpDataType(opType) == DT_COLOUR)
diff --git a/Engine/Plugins/Experimental/Mutable/Source/MutableRuntime/Private/MuR/SystemPrivate.h b/Engine/Plugins/Experimental/Mutable/Source/MutableRuntime/Private/MuR/SystemPrivate.h
--- a/Engine/Plugins/Experimental/Mutable/Source/MutableRuntime/Private/MuR/SystemPrivate.h
+++ b/Engine/Plugins/Experimental/Mutable/Source/MutableRuntime/Private/MuR/SystemPrivate.h
@@ -20,8 +20,8 @@ namespace mu
 	class Colour : public RefCounted
 	{
 	public:
-		Colour(vec4f v = vec4f()) : m_colour(v) {}
-		vec4f m_colour;
+		Colour(FVector4f v = FVector4f()) : m_colour(v) {}
+		FVector4f m_colour;
 	};
 	typedef Ptr<Colour> ColourPtr;
 
@@ -595,18 +595,18 @@ namespace mu
 			return pResult ? pResult->m_value : 0;
 		}
 
-		vec4f GetColour(CACHE_ADDRESS at)
+		FVector4f GetColour(CACHE_ADDRESS at)
 		{
-			if (!at.at) return vec4f();
+			if (!at.at) return FVector4f();
 			auto d = m_resources.get_ptr(at);
-			if (!d) return vec4f();
+			if (!d) return FVector4f();
 
 			Ptr<const Colour> pResult;
 			if (at.at)
 			{
 				pResult = (const Colour*)d->Value.get();
 			}
-			return pResult ? pResult->m_colour : vec4f();
+			return pResult ? pResult->m_colour : FVector4f();
 		}
 
 		Ptr<const Projector> GetProjector(CACHE_ADDRESS at)
@@ -703,7 +703,7 @@ namespace mu
 			m_resources[at] = TPair<int, Ptr<const RefCounted>>(1, pResult);
 		}
 
-		void SetColour(CACHE_ADDRESS at, const vec4f& v)
+		void SetColour(CACHE_ADDRESS at, const FVector4f& v)
 		{
 			check(at.at < m_resources.size_code());
 
diff --git a/Engine/Plugins/Experimental/Mutable/Source/MutableRuntime/Public/MuR/Image.h b/Engine/Plugins/Experimental/Mutable/Source/MutableRuntime/Public/MuR/Image.h
--- a/Engine/Plugins/Experimental/Mutable/Source/MutableRuntime/Public/MuR/Image.h
+++ b/Engine/Plugins/Experimental/Mutable/Source/MutableRuntime/Public/MuR/Image.h
@@ -319,7 +319,7 @@ namespace mu
 		//-----------------------------------------------------------------------------------------
 
 		//! Sample the image and return an RGB colour.
-		vec4<float> Sample(vec2<float> coords) const;
+		FVector4f Sample(FVector2f coords) const;
 
 		//! Calculate the size of the image data in bytes, regardless of what is allocated in
 		//! m_data, only using the image descriptions. For non-block-compressed images, it returns
@@ -354,7 +354,7 @@ namespace mu
 		int32 GetMipsDataSize() const;
 
 		//! See if all the pixels in the image are equal and return the colour.
-		bool IsPlainColour(vec4<float>& colour) const;
+		bool IsPlainColour(FVector4f& colour) const;
 
 		//! See if all the pixels in the alpha channel are the max value of the pixel format
 		//! (white).
diff --git a/Engine/Plugins/Experimental/Mutable/Source/MutableRuntime/Public/MuR/MeshBufferSet.h b/Engine/Plugins/Experimental/Mutable/Source/MutableRuntime/Public/MuR/MeshBufferSet.h
--- a/Engine/Plugins/Experimental/Mutable/Source/MutableRuntime/Public/MuR/MeshBufferSet.h
+++ b/Engine/Plugins/Experimental/Mutable/Source/MutableRuntime/Public/MuR/MeshBufferSet.h
@@ -110,7 +110,7 @@ namespace mu
 		MBS_OTHER,
 
         //! Sign to define the orientation of the tangent space.
-        MBS_TANGENTSIGN,
+        MBS_TANGENTSIGN_DEPRECATED,
 
 		//! Semantics usefule for mesh binding.
 		MBS_TRIANGLEINDEX,
diff --git a/Engine/Plugins/Experimental/Mutable/Source/MutableRuntime/Public/MuR/Model.h b/Engine/Plugins/Experimental/Mutable/Source/MutableRuntime/Public/MuR/Model.h
--- a/Engine/Plugins/Experimental/Mutable/Source/MutableRuntime/Public/MuR/Model.h
+++ b/Engine/Plugins/Experimental/Mutable/Source/MutableRuntime/Public/MuR/Model.h
@@ -131,12 +131,7 @@ namespace mu
     //!
     //! This class can be used to generate all the possible instances of a model, or to generate
     //! random instances. It only takes into account discrete parameters (bools and ints) but not
-    //! continuous parameters like colours and floats.
-    //!
-    //! The use of this class can be very expensive for complex models, since the number of
-    //! combinations grows geometrically. If considerRelevancy is set to false, a brute force
-    //! approach will be used, but this will generate duplicated instances when changing parameters
-    //! that are not relevant.
+    //! continuous parameters like colours and floats, which may be set to random values.
     //!
     //! \ingroup runtime
     class MUTABLERUNTIME_API ModelParametersGenerator : public RefCounted
@@ -144,7 +139,7 @@ namespace mu
     public:
 
         //!
-        ModelParametersGenerator( const Model* pModel, System* pSystem, bool considerRelevancy=true );
+        ModelParametersGenerator( const Model* pModel, System* pSystem );
 
         //! Return the number of different possible instances that can be built from the model.
         int64 GetInstanceCount();
@@ -153,12 +148,12 @@ namespace mu
         //! \param index is an index in the range of 0 to GetInstanceCount-1
         //! \param randomGenerator if not null, this function will be used to generate random values
         //! for the continuous parameters of the instance.
-        ParametersPtr GetInstance( int64 index, float (*randomGenerator )() );
+        ParametersPtr GetInstance( int64 index, TFunction<float()> RandomGenerator);
 
         //! Return the parameters of one of a random instance of the model.
         //! \param randomGenerator This function will be used to generate random values
         //! for the continuous parameters of the instance.
-        ParametersPtr GetRandomInstance(TFunctionRef<float()> randomGenerator );
+        ParametersPtr GetRandomInstance(TFunctionRef<float()> RandomGenerator );
 
         //-----------------------------------------------------------------------------------------
         // Interface pattern
diff --git a/Engine/Plugins/Experimental/Mutable/Source/MutableTools/Private/MuT/AST.cpp b/Engine/Plugins/Experimental/Mutable/Source/MutableTools/Private/MuT/AST.cpp
--- a/Engine/Plugins/Experimental/Mutable/Source/MutableTools/Private/MuT/AST.cpp
+++ b/Engine/Plugins/Experimental/Mutable/Source/MutableTools/Private/MuT/AST.cpp
@@ -871,14 +871,14 @@ FImageDesc ASTOp::GetImageDesc( bool, GetImageDescContext* )
 }
 
 
-bool ASTOp::IsImagePlainConstant( vec4<float>& ) const
+bool ASTOp::IsImagePlainConstant(FVector4f&) const
 {
     check(false);
     return false;
 }
 
 
-bool ASTOp::IsColourConstant( vec4<float>& ) const
+bool ASTOp::IsColourConstant(FVector4f&) const
 {
     check(false);
     return false;
@@ -1562,7 +1562,7 @@ int ASTOpFixed::EvaluateInt( ASTOpList& /*facts*/, bool &unknown ) const
 }
 
 
-bool ASTOpFixed::IsImagePlainConstant( vec4<float>& colour ) const
+bool ASTOpFixed::IsImagePlainConstant(FVector4f& colour) const
 {
     bool res = false;
     switch( op.type )
@@ -1596,13 +1596,13 @@ bool ASTOpFixed::IsImagePlainConstant( vec4<float>& colour ) const
         res = children[op.args.ImageInterpolate3.target0]->IsColourConstant( colour );
         if (res)
         {
-            vec4<float> baseColour;
+			FVector4f baseColour;
             res = children[op.args.ImageInterpolate3.target1]->IsColourConstant( colour );
             res &= (colour==baseColour);
         }
         if (res)
         {
-            vec4<float> baseColour;
+			FVector4f baseColour;
             res = children[op.args.ImageInterpolate3.target2]->IsColourConstant( colour );
             res &= (colour==baseColour);
         }
@@ -1620,7 +1620,7 @@ bool ASTOpFixed::IsImagePlainConstant( vec4<float>& colour ) const
 
 
 //-------------------------------------------------------------------------------------------------
-bool ASTOpFixed::IsColourConstant( vec4<float>& colour ) const
+bool ASTOpFixed::IsColourConstant(FVector4f& colour) const
 {
     bool res = false;
     switch( op.type )
diff --git a/Engine/Plugins/Experimental/Mutable/Source/MutableTools/Private/MuT/AST.h b/Engine/Plugins/Experimental/Mutable/Source/MutableTools/Private/MuT/AST.h
--- a/Engine/Plugins/Experimental/Mutable/Source/MutableTools/Private/MuT/AST.h
+++ b/Engine/Plugins/Experimental/Mutable/Source/MutableTools/Private/MuT/AST.h
@@ -683,8 +683,8 @@ namespace mu
         //! Return the size in pixels of the layout grid block for the image operation
         virtual void GetLayoutBlockSize( int* pBlockX, int* pBlockY );
 
-        virtual bool IsImagePlainConstant( vec4<float>& colour ) const;
-        virtual bool IsColourConstant( vec4<float>& colour ) const;
+        virtual bool IsImagePlainConstant(FVector4f& colour ) const;
+        virtual bool IsColourConstant(FVector4f& colour ) const;
 
         //!
         virtual bool GetNonBlackRect( FImageRect& maskUsage ) const;
@@ -978,8 +978,8 @@ namespace mu
         Ptr<ASTOp> OptimiseSink(const MODEL_OPTIMIZATION_OPTIONS&, OPTIMIZE_SINK_CONTEXT&) const override;
         FBoolEvalResult EvaluateBool( ASTOpList& facts, FEvaluateBoolCache* cache ) const override;
         int EvaluateInt( ASTOpList& facts, bool &unknown ) const override;
-        bool IsImagePlainConstant( vec4<float>& colour ) const override;
-        bool IsColourConstant( vec4<float>& colour ) const override;
+        bool IsImagePlainConstant(FVector4f& colour ) const override;
+        bool IsColourConstant(FVector4f& colour ) const override;
         Ptr<ImageSizeExpression> GetImageSizeExpression() const override;
 
         //---------------------------------------------------------------------------------------------
diff --git a/Engine/Plugins/Experimental/Mutable/Source/MutableTools/Private/MuT/ASTOpConstantResource.cpp b/Engine/Plugins/Experimental/Mutable/Source/MutableTools/Private/MuT/ASTOpConstantResource.cpp
--- a/Engine/Plugins/Experimental/Mutable/Source/MutableTools/Private/MuT/ASTOpConstantResource.cpp
+++ b/Engine/Plugins/Experimental/Mutable/Source/MutableTools/Private/MuT/ASTOpConstantResource.cpp
@@ -268,7 +268,7 @@ namespace mu
 
 
 	//-------------------------------------------------------------------------------------------------
-	bool ASTOpConstantResource::IsImagePlainConstant(vec4<float>& colour) const
+	bool ASTOpConstantResource::IsImagePlainConstant(FVector4f& colour) const
 	{
 		bool res = false;
 		switch (type)
@@ -285,7 +285,7 @@ namespace mu
 				if (pImage->m_flags & Image::IF_IS_PLAIN_COLOUR)
 				{
 					res = true;
-					colour = pImage->Sample(vec2<float>(0, 0));
+					colour = pImage->Sample(FVector2f(0, 0));
 				}
 				else
 				{
diff --git a/Engine/Plugins/Experimental/Mutable/Source/MutableTools/Private/MuT/ASTOpConstantResource.h b/Engine/Plugins/Experimental/Mutable/Source/MutableTools/Private/MuT/ASTOpConstantResource.h
--- a/Engine/Plugins/Experimental/Mutable/Source/MutableTools/Private/MuT/ASTOpConstantResource.h
+++ b/Engine/Plugins/Experimental/Mutable/Source/MutableTools/Private/MuT/ASTOpConstantResource.h
@@ -70,7 +70,7 @@ template <class SCALAR> class vec4;
 			BLOCK_LAYOUT_SIZE_CACHE* cache) override;
 		void GetLayoutBlockSize(int* pBlockX, int* pBlockY) override;
 		bool GetNonBlackRect(FImageRect& maskUsage) const override;
-		bool IsImagePlainConstant(vec4<float>& colour) const override;
+		bool IsImagePlainConstant(FVector4f& colour) const override;
 		Ptr<ImageSizeExpression> GetImageSizeExpression() const override;
 	};
 
diff --git a/Engine/Plugins/Experimental/Mutable/Source/MutableTools/Private/MuT/ASTOpImageCompose.cpp b/Engine/Plugins/Experimental/Mutable/Source/MutableTools/Private/MuT/ASTOpImageCompose.cpp
--- a/Engine/Plugins/Experimental/Mutable/Source/MutableTools/Private/MuT/ASTOpImageCompose.cpp
+++ b/Engine/Plugins/Experimental/Mutable/Source/MutableTools/Private/MuT/ASTOpImageCompose.cpp
@@ -167,7 +167,7 @@ mu::Ptr<ImageSizeExpression> ASTOpImageCompose::GetImageSizeExpression() const
 }
 
 
-bool ASTOpImageCompose::IsImagePlainConstant(vec4<float>& colour) const
+bool ASTOpImageCompose::IsImagePlainConstant(FVector4f& colour) const
 {
 	bool res = false;
 
@@ -178,7 +178,7 @@ bool ASTOpImageCompose::IsImagePlainConstant(vec4<float>& colour) const
 
 	if (res && Base.child())
 	{
-		vec4<float> baseColour;
+		FVector4f baseColour;
 		res = Base->IsImagePlainConstant(baseColour);
 		res &= (colour == baseColour);
 	}
diff --git a/Engine/Plugins/Experimental/Mutable/Source/MutableTools/Private/MuT/ASTOpImageCompose.h b/Engine/Plugins/Experimental/Mutable/Source/MutableTools/Private/MuT/ASTOpImageCompose.h
--- a/Engine/Plugins/Experimental/Mutable/Source/MutableTools/Private/MuT/ASTOpImageCompose.h
+++ b/Engine/Plugins/Experimental/Mutable/Source/MutableTools/Private/MuT/ASTOpImageCompose.h
@@ -40,7 +40,7 @@ template <class SCALAR> class vec4;
 		Ptr<ASTOp> OptimiseSemantic(const MODEL_OPTIMIZATION_OPTIONS& options) const;
 		FImageDesc GetImageDesc(bool returnBestOption, GetImageDescContext* context) override;
 		void GetLayoutBlockSize(int* pBlockX, int* pBlockY) override;
-		bool IsImagePlainConstant(vec4<float>& colour) const override;
+		bool IsImagePlainConstant(FVector4f& colour) const override;
 		Ptr<ImageSizeExpression> GetImageSizeExpression() const override;
 
 	};
diff --git a/Engine/Plugins/Experimental/Mutable/Source/MutableTools/Private/MuT/ASTOpImageMipmap.cpp b/Engine/Plugins/Experimental/Mutable/Source/MutableTools/Private/MuT/ASTOpImageMipmap.cpp
--- a/Engine/Plugins/Experimental/Mutable/Source/MutableTools/Private/MuT/ASTOpImageMipmap.cpp
+++ b/Engine/Plugins/Experimental/Mutable/Source/MutableTools/Private/MuT/ASTOpImageMipmap.cpp
@@ -221,7 +221,7 @@ void ASTOpImageMipmap::GetLayoutBlockSize(int* pBlockX, int* pBlockY)
 }
 
 
-bool ASTOpImageMipmap::IsImagePlainConstant(vec4<float>& colour) const
+bool ASTOpImageMipmap::IsImagePlainConstant(FVector4f& colour) const
 {
 	bool res = false;
 	if (Source.child())
diff --git a/Engine/Plugins/Experimental/Mutable/Source/MutableTools/Private/MuT/ASTOpImageMipmap.h b/Engine/Plugins/Experimental/Mutable/Source/MutableTools/Private/MuT/ASTOpImageMipmap.h
--- a/Engine/Plugins/Experimental/Mutable/Source/MutableTools/Private/MuT/ASTOpImageMipmap.h
+++ b/Engine/Plugins/Experimental/Mutable/Source/MutableTools/Private/MuT/ASTOpImageMipmap.h
@@ -50,7 +50,7 @@ template <class SCALAR> class vec4;
 		Ptr<ASTOp> OptimiseSink(const MODEL_OPTIMIZATION_OPTIONS& options, OPTIMIZE_SINK_CONTEXT& context) const override;
 		FImageDesc GetImageDesc(bool returnBestOption, GetImageDescContext* context) override;
 		void GetLayoutBlockSize(int* pBlockX, int* pBlockY) override;
-		bool IsImagePlainConstant(vec4<float>& colour) const override;
+		bool IsImagePlainConstant(FVector4f& colour) const override;
 		Ptr<ImageSizeExpression> GetImageSizeExpression() const override;
 
 	};
diff --git a/Engine/Plugins/Experimental/Mutable/Source/MutableTools/Private/MuT/ASTOpImagePixelFormat.cpp b/Engine/Plugins/Experimental/Mutable/Source/MutableTools/Private/MuT/ASTOpImagePixelFormat.cpp
--- a/Engine/Plugins/Experimental/Mutable/Source/MutableTools/Private/MuT/ASTOpImagePixelFormat.cpp
+++ b/Engine/Plugins/Experimental/Mutable/Source/MutableTools/Private/MuT/ASTOpImagePixelFormat.cpp
@@ -244,7 +244,7 @@ void ASTOpImagePixelFormat::GetLayoutBlockSize(int* pBlockX, int* pBlockY)
 }
 
 
-bool ASTOpImagePixelFormat::IsImagePlainConstant(vec4<float>& colour) const
+bool ASTOpImagePixelFormat::IsImagePlainConstant(FVector4f& colour) const
 {
 	bool res = false;
 	if (Source.child())
diff --git a/Engine/Plugins/Experimental/Mutable/Source/MutableTools/Private/MuT/ASTOpImagePixelFormat.h b/Engine/Plugins/Experimental/Mutable/Source/MutableTools/Private/MuT/ASTOpImagePixelFormat.h
--- a/Engine/Plugins/Experimental/Mutable/Source/MutableTools/Private/MuT/ASTOpImagePixelFormat.h
+++ b/Engine/Plugins/Experimental/Mutable/Source/MutableTools/Private/MuT/ASTOpImagePixelFormat.h
@@ -37,7 +37,7 @@ template <class SCALAR> class vec4;
 		Ptr<ASTOp> OptimiseSink(const MODEL_OPTIMIZATION_OPTIONS& options, OPTIMIZE_SINK_CONTEXT& context) const override;
 		FImageDesc GetImageDesc(bool returnBestOption, GetImageDescContext* context) override;
 		void GetLayoutBlockSize(int* pBlockX, int* pBlockY) override;
-		bool IsImagePlainConstant(vec4<float>& colour) const override;
+		bool IsImagePlainConstant(FVector4f& colour) const override;
 		Ptr<ImageSizeExpression> GetImageSizeExpression() const override;
 
 	};
diff --git a/Engine/Plugins/Experimental/Mutable/Source/MutableTools/Private/MuT/ASTOpSwitch.cpp b/Engine/Plugins/Experimental/Mutable/Source/MutableTools/Private/MuT/ASTOpSwitch.cpp
--- a/Engine/Plugins/Experimental/Mutable/Source/MutableTools/Private/MuT/ASTOpSwitch.cpp
+++ b/Engine/Plugins/Experimental/Mutable/Source/MutableTools/Private/MuT/ASTOpSwitch.cpp
@@ -401,7 +401,7 @@ namespace mu
 
 
 	//-------------------------------------------------------------------------------------------------
-	bool ASTOpSwitch::IsImagePlainConstant(vec4<float>&) const
+	bool ASTOpSwitch::IsImagePlainConstant(FVector4f&) const
 	{
 		// We could check if every option is plain and exactly the same colour, but probably it is
 		// not worth.
diff --git a/Engine/Plugins/Experimental/Mutable/Source/MutableTools/Private/MuT/ASTOpSwitch.h b/Engine/Plugins/Experimental/Mutable/Source/MutableTools/Private/MuT/ASTOpSwitch.h
--- a/Engine/Plugins/Experimental/Mutable/Source/MutableTools/Private/MuT/ASTOpSwitch.h
+++ b/Engine/Plugins/Experimental/Mutable/Source/MutableTools/Private/MuT/ASTOpSwitch.h
@@ -74,7 +74,7 @@ template <class SCALAR> class vec4;
 		void GetLayoutBlockSize(int* pBlockX, int* pBlockY) override;
 		Ptr<ASTOp> OptimiseSemantic(const MODEL_OPTIMIZATION_OPTIONS&) const override;
 		bool GetNonBlackRect(FImageRect& maskUsage) const override;
-		bool IsImagePlainConstant(vec4<float>& colour) const override;
+		bool IsImagePlainConstant(FVector4f& colour) const override;
 		Ptr<ImageSizeExpression> GetImageSizeExpression() const override;
 
 		// Own interface
diff --git a/Engine/Plugins/Experimental/Mutable/Source/MutableTools/Private/MuT/CodeOptimiser_Semantic.cpp b/Engine/Plugins/Experimental/Mutable/Source/MutableTools/Private/MuT/CodeOptimiser_Semantic.cpp
--- a/Engine/Plugins/Experimental/Mutable/Source/MutableTools/Private/MuT/CodeOptimiser_Semantic.cpp
+++ b/Engine/Plugins/Experimental/Mutable/Source/MutableTools/Private/MuT/CodeOptimiser_Semantic.cpp
@@ -366,15 +366,15 @@ namespace mu
             auto maskAt = children[op.args.ImageLayer.mask].child();
             if ( maskAt )
             {
-                vec4<float> colour;
+                FVector4f colour;
                 if ( maskAt->IsImagePlainConstant( colour ))
                 {
-                    if ( colour.xyz().AlmostNull() )
+                    if ( colour.IsNearlyZero3(UE_SMALL_NUMBER) )
                     {
                         // If the mask is black, we can skip the entire operation
                         at = children[op.args.ImageLayer.base].child();
                     }
-                    else if ( colour.xyz().AlmostEqual( vec3<float>(1,1,1) ) )
+                    else if ( colour.Equals( FVector4f(1,1,1,1), UE_SMALL_NUMBER ) )
                     {
                         // If the mask is white, we can remove it
                         auto nop = mu::Clone<ASTOpFixed>(this);
