GenerateMutableSourceMesh: Fix compile error in UE5/Main due to change of type of InfluenceWeights in SkeletalMeshLodModel.h.
#rnx

diff --git a/Engine/Plugins/Experimental/Mutable/Source/CustomizableObjectEditor/Private/MuCOE/GenerateMutableSource/GenerateMutableSourceMesh.cpp b/Engine/Plugins/Experimental/Mutable/Source/CustomizableObjectEditor/Private/MuCOE/GenerateMutableSource/GenerateMutableSourceMesh.cpp
--- a/Engine/Plugins/Experimental/Mutable/Source/CustomizableObjectEditor/Private/MuCOE/GenerateMutableSource/GenerateMutableSourceMesh.cpp
+++ b/Engine/Plugins/Experimental/Mutable/Source/CustomizableObjectEditor/Private/MuCOE/GenerateMutableSource/GenerateMutableSourceMesh.cpp
@@ -290,7 +290,7 @@ TObjectPtr<const USkeletalMesh> GetMeshWithBoneRemovalApplied(TObjectPtr<USkelet
 }
 
 
-void TransferRemovedBonesInfluences(FBoneIndexType* InfluenceBones, uint8* InfluenceWeights, const int32 InfluenceCount, const TMap<int32, int32>& ReplacementMap)
+void TransferRemovedBonesInfluences(FBoneIndexType* InfluenceBones, uint16* InfluenceWeights, const int32 InfluenceCount, const TMap<int32, int32>& ReplacementMap)
 {
 	for (int32 i = 0; i < InfluenceCount; ++i)
 	{
@@ -318,7 +318,7 @@ void TransferRemovedBonesInfluences(FBoneIndexType* InfluenceBones, uint8* Influ
 }
 
 
-void NormalizeWeights(FBoneIndexType* InfluenceBones, uint8* InfluenceWeights, const int32 InfluenceCount, const int32 MutableInfluenceCount,
+void NormalizeWeights(FBoneIndexType* InfluenceBones, uint16* InfluenceWeights, const int32 InfluenceCount, const int32 MutableInfluenceCount,
 	int32* MutableMaxOrderedWeighsIndices, const int32 MaxSectionBoneMapIndex, const int32 MaxBoneWeight)
 {
 	// First get the indices of the 4 heaviest influences
@@ -398,7 +398,7 @@ void NormalizeWeights(FBoneIndexType* InfluenceBones, uint8* InfluenceWeights, c
 	}
 	else
 	{
-		FMemory::Memzero(InfluenceWeights, MutableInfluenceCount);
+		FMemory::Memzero(InfluenceWeights, MutableInfluenceCount*sizeof(InfluenceWeights[0]));
 		InfluenceWeights[0] = MaxBoneWeight;
 	}
 }
