[Mutable] Removed LOD from Descriptor hash used to cache Instances.

Now there are two kinds of hashes, one includes the LOD (Runtime) and the other does not. The Runtime Hash does not hash the LOD, being able to compare the already existing LODs.

#rnx

diff --git a/Engine/Plugins/Experimental/Mutable/Source/CustomizableObject/Classes/MuCO/CustomizableObjectInstance.h b/Engine/Plugins/Experimental/Mutable/Source/CustomizableObject/Classes/MuCO/CustomizableObjectInstance.h
--- a/Engine/Plugins/Experimental/Mutable/Source/CustomizableObject/Classes/MuCO/CustomizableObjectInstance.h
+++ b/Engine/Plugins/Experimental/Mutable/Source/CustomizableObject/Classes/MuCO/CustomizableObjectInstance.h
@@ -595,11 +595,11 @@ public:
 	/** Instance updated. */
 	void Updated(EUpdateResult Result);
 	
-	/** Hash representing the actual state of the instance (meshes, textures...). This does not include the parameters. */
-	uint32 GetDescriptorHash() const;
+	/** Return the UCustomizableObjectInstance::Descriptor hash on the last update request. */
+	FDescriptorRuntimeHash GetDescriptorRuntimeHash() const;
 
-	/** Hash representing the state of the instance when the update was requested. */
-	uint32 GetUpdateDescriptorHash() const;
+	/** Return the UCustomizableObjectInstance::Descriptor hash on the last successful update. */
+	FDescriptorRuntimeHash GetUpdateDescriptorRuntimeHash() const;
 
 	// --------------------------------------------------------------------
 
@@ -663,11 +663,11 @@ private:
 	UPROPERTY( Transient )
 	TObjectPtr<UCustomizableInstancePrivateData> PrivateData;
 
-	/** Hash of the UCustomizableObjectInstance::Descriptor when UpdateSkeletalMeshAsync has been called. */
-	uint32 UpdateDescriptorHash = 0;
+	/** Hash of the UCustomizableObjectInstance::Descriptor on the last update request. */
+	FDescriptorRuntimeHash UpdateDescriptorRuntimeHash;
 	
 	/** Hash of the UCustomizableObjectInstance::Descriptor on the last successful update. */
-	uint32 DescriptorHash = 0;
+	FDescriptorRuntimeHash DescriptorRuntimeHash;
 
 	// Deprecated properties
 	
diff --git a/Engine/Plugins/Experimental/Mutable/Source/CustomizableObject/Private/MuCO/CustomizableObjectInstance.cpp b/Engine/Plugins/Experimental/Mutable/Source/CustomizableObject/Private/MuCO/CustomizableObjectInstance.cpp
--- a/Engine/Plugins/Experimental/Mutable/Source/CustomizableObject/Private/MuCO/CustomizableObjectInstance.cpp
+++ b/Engine/Plugins/Experimental/Mutable/Source/CustomizableObject/Private/MuCO/CustomizableObjectInstance.cpp
@@ -798,6 +798,8 @@ void UCustomizableInstancePrivateData::DoUpdateSkeletalMeshAsync(UCustomizableOb
 		UE_LOG(LogMutable, Verbose, TEXT("LOD change: %d, %d -> %d, %d"), LastUpdateMinLOD, LastUpdateMaxLOD, Instance->Descriptor.MinLODToLoad, Instance->Descriptor.MaxLODToLoad);
 	}
 
+	Instance->UpdateDescriptorRuntimeHash = FDescriptorRuntimeHash(Instance->Descriptor);
+
 	UCustomizableObjectSystem::GetInstance()->GetPrivate()->InitUpdateSkeletalMesh(Instance, Priority);
 
 #if WITH_EDITOR
@@ -823,8 +825,6 @@ void UCustomizableInstancePrivateData::DoUpdateSkeletalMesh(UCustomizableObjectI
 	check(IsInGameThread());
 	check(bAsync);
 
-	CustomizableInstance->UpdateDescriptorHash = GetTypeHash(CustomizableInstance->Descriptor);
-
 	UCustomizableObjectSystem* System = UCustomizableObjectSystem::GetInstance();
 	
 	UCustomizableObject* CustomizableObject = CustomizableInstance->GetCustomizableObject();
@@ -2452,7 +2452,7 @@ void UCustomizableObjectInstance::Updated(EUpdateResult Result)
 {
 	if (Result == EUpdateResult::Success)
 	{
-		DescriptorHash = UpdateDescriptorHash;
+		DescriptorRuntimeHash = UpdateDescriptorRuntimeHash;
 	}
 
 	// Call Customizable Skeletal Components updated callbacks.
@@ -2475,15 +2475,15 @@ void UCustomizableObjectInstance::Updated(EUpdateResult Result)
 }
 
 
-uint32 UCustomizableObjectInstance::GetDescriptorHash() const
+FDescriptorRuntimeHash UCustomizableObjectInstance::GetDescriptorRuntimeHash() const
 {
-	return DescriptorHash;	
+	return DescriptorRuntimeHash;
 }
 
 
-uint32 UCustomizableObjectInstance::GetUpdateDescriptorHash() const
+FDescriptorRuntimeHash UCustomizableObjectInstance::GetUpdateDescriptorRuntimeHash() const
 {
-	return UpdateDescriptorHash;
+	return UpdateDescriptorRuntimeHash;
 }
 
 
diff --git a/Engine/Plugins/Experimental/Mutable/Source/CustomizableObject/Private/MuCO/CustomizableObjectInstanceDescriptor.cpp b/Engine/Plugins/Experimental/Mutable/Source/CustomizableObject/Private/MuCO/CustomizableObjectInstanceDescriptor.cpp
--- a/Engine/Plugins/Experimental/Mutable/Source/CustomizableObject/Private/MuCO/CustomizableObjectInstanceDescriptor.cpp
+++ b/Engine/Plugins/Experimental/Mutable/Source/CustomizableObject/Private/MuCO/CustomizableObjectInstanceDescriptor.cpp
@@ -2087,57 +2087,86 @@ void FCustomizableObjectInstanceDescriptor::CreateParametersLookupTable()
 	}
 }
 
-
-uint32 GetTypeHash(const FCustomizableObjectInstanceDescriptor& Key)
+FDescriptorHash::FDescriptorHash(const FCustomizableObjectInstanceDescriptor& Descriptor)
 {
-	uint32 Hash = 0;
-
-	if (Key.CustomizableObject) 
+	if (Descriptor.CustomizableObject)
 	{
-		Hash = HashCombine(Hash, GetTypeHash(Key.CustomizableObject->GetPathName()));
-		Hash = HashCombine(Hash, GetTypeHash(Key.CustomizableObject->GetCompilationGuid()));
+		Hash = HashCombine(Hash, GetTypeHash(Descriptor.CustomizableObject->GetPathName()));
+		Hash = HashCombine(Hash, GetTypeHash(Descriptor.CustomizableObject->GetCompilationGuid()));
 	}
 
-	for (const FCustomizableObjectBoolParameterValue& Value : Key.BoolParameters)
+	for (const FCustomizableObjectBoolParameterValue& Value : Descriptor.BoolParameters)
 	{
 		Hash = HashCombine(Hash, GetTypeHash(Value));
 	}	
 
-	for (const FCustomizableObjectIntParameterValue& Value : Key.IntParameters)
+	for (const FCustomizableObjectIntParameterValue& Value : Descriptor.IntParameters)
 	{
 		Hash = HashCombine(Hash, GetTypeHash(Value));
 	}
 	
-	for (const FCustomizableObjectFloatParameterValue& Value : Key.FloatParameters)
+	for (const FCustomizableObjectFloatParameterValue& Value : Descriptor.FloatParameters)
 	{
 		Hash = HashCombine(Hash, GetTypeHash(Value));
 	}
 
-	for (const FCustomizableObjectTextureParameterValue& Value : Key.TextureParameters)
+	for (const FCustomizableObjectTextureParameterValue& Value : Descriptor.TextureParameters)
 	{
 		Hash = HashCombine(Hash, GetTypeHash(Value));
 	}
 
-	for (const FCustomizableObjectVectorParameterValue& Value : Key.VectorParameters)
+	for (const FCustomizableObjectVectorParameterValue& Value : Descriptor.VectorParameters)
 	{
 		Hash = HashCombine(Hash, GetTypeHash(Value));
 	}
 
-	for (const FCustomizableObjectProjectorParameterValue& Value : Key.ProjectorParameters)
+	for (const FCustomizableObjectProjectorParameterValue& Value : Descriptor.ProjectorParameters)
 	{
 		Hash = HashCombine(Hash, GetTypeHash(Value));
 	}
 	
-	Hash = HashCombine(Hash, GetTypeHash(Key.State));
-	Hash = HashCombine(Hash, GetTypeHash(Key.bBuildParameterDecorations));
-	Hash = HashCombine(Hash, GetTypeHash(Key.MinLODToLoad));
-	Hash = HashCombine(Hash, GetTypeHash(Key.MaxLODToLoad));
+	Hash = HashCombine(Hash, GetTypeHash(Descriptor.State));
+	Hash = HashCombine(Hash, GetTypeHash(Descriptor.bBuildParameterDecorations));
+	Hash = HashCombine(Hash, GetTypeHash(Descriptor.MinLODToLoad));
+	Hash = HashCombine(Hash, GetTypeHash(Descriptor.MaxLODToLoad));
 
-	for (const TTuple<FName, FMultilayerProjector>& Pair : Key.MultilayerProjectors)
+	for (const TTuple<FName, FMultilayerProjector>& Pair : Descriptor.MultilayerProjectors)
 	{
 		// Hash = HashCombine(Hash, GetTypeHash(Pair.Key)); // Already hashed by the FMultilayerProjector.
 		Hash = HashCombine(Hash, GetTypeHash(Pair.Value));
 	}
-	
-	return Hash;
+}
+
+
+bool FDescriptorHash::operator==(const FDescriptorHash& Other) const
+{
+	return Hash == Other.Hash;
+}
+
+
+bool FDescriptorHash::operator<(const FDescriptorHash& Other) const
+{
+	return Hash < Other.Hash;
+}
+
+
+FString FDescriptorHash::ToString() const
+{
+	return FString::Printf(TEXT("%u"), Hash);
+}
+
+
+FDescriptorRuntimeHash::FDescriptorRuntimeHash(const FCustomizableObjectInstanceDescriptor& Descriptor) :
+	FDescriptorHash(Descriptor)
+{
+	MinLODToLoad = Descriptor.MinLODToLoad;
+	MaxLODToLoad = Descriptor.MaxLODToLoad;	
+}
+
+
+bool FDescriptorRuntimeHash::IsSubset(const FDescriptorRuntimeHash& Other) const
+{
+	return FDescriptorHash::operator==(Other) &&
+		MinLODToLoad >= Other.MinLODToLoad &&
+		MaxLODToLoad <= Other.MaxLODToLoad;
 }
diff --git a/Engine/Plugins/Experimental/Mutable/Source/CustomizableObject/Private/MuCO/CustomizableObjectSystem.cpp b/Engine/Plugins/Experimental/Mutable/Source/CustomizableObject/Private/MuCO/CustomizableObjectSystem.cpp
--- a/Engine/Plugins/Experimental/Mutable/Source/CustomizableObject/Private/MuCO/CustomizableObjectSystem.cpp
+++ b/Engine/Plugins/Experimental/Mutable/Source/CustomizableObject/Private/MuCO/CustomizableObjectSystem.cpp
@@ -811,14 +811,14 @@ void FCustomizableObjectSystemPrivate::InitUpdateSkeletalMesh(UCustomizableObjec
 		// Skip the update if the Instance has the same descriptor (pending or applied).
 		// If there is a hash collision, at worse, we would add an unnecessary update.
 		{
-			const uint32 UpdateDescriptorHash = Public->GetUpdateDescriptorHash();
+			const FDescriptorRuntimeHash UpdateDescriptorHash = Public->GetUpdateDescriptorRuntimeHash();
 
 			if (const FMutableQueueElem* QueueElem = MutableOperationQueue.Get(Public))
 			{
 				if (const TSharedPtr<FMutableOperation>& QueuedOperation = QueueElem->Operation;
 					QueuedOperation &&
 					QueuedOperation->Type == FMutableOperation::EOperationType::Update &&
-					QueuedOperation->InstanceDescriptorHash == UpdateDescriptorHash)
+					UpdateDescriptorHash.IsSubset(QueuedOperation->InstanceDescriptorRuntimeHash))
 				{
 					return;
 				}
@@ -826,12 +826,12 @@ void FCustomizableObjectSystemPrivate::InitUpdateSkeletalMesh(UCustomizableObjec
 		
 			if (CurrentMutableOperation &&
 				CurrentMutableOperation->Type == FMutableOperation::EOperationType::Update &&
-				CurrentMutableOperation->InstanceDescriptorHash == UpdateDescriptorHash)
+				UpdateDescriptorHash.IsSubset(CurrentMutableOperation->InstanceDescriptorRuntimeHash))
 			{
 				return;
 			}
 			
-			if (Public->GetDescriptorHash() == UpdateDescriptorHash)
+			if (UpdateDescriptorHash.IsSubset(Public->GetDescriptorRuntimeHash()))
 			{
 				UpdateSkeletalMesh(Public);
 				return;
@@ -2421,7 +2421,7 @@ FMutableOperation FMutableOperation::CreateInstanceUpdate(UCustomizableObjectIns
 	Op.bNeverStream = bInNeverStream;
 	Op.MipsToSkip = InMipsToSkip;
 	Op.CustomizableObjectInstance = InCustomizableObjectInstance;
-	Op.InstanceDescriptorHash = InCustomizableObjectInstance->GetUpdateDescriptorHash();
+	Op.InstanceDescriptorRuntimeHash = InCustomizableObjectInstance->GetUpdateDescriptorRuntimeHash();
 	Op.bStarted = false;
 	Op.bBuildParameterDecorations = InCustomizableObjectInstance->GetBuildParameterDecorations();
 	Op.Parameters = InCustomizableObjectInstance->GetPrivate()->ReloadParametersFromObject(InCustomizableObjectInstance, false);
diff --git a/Engine/Plugins/Experimental/Mutable/Source/CustomizableObject/Private/MuCO/CustomizableObjectSystemPrivate.h b/Engine/Plugins/Experimental/Mutable/Source/CustomizableObject/Private/MuCO/CustomizableObjectSystemPrivate.h
--- a/Engine/Plugins/Experimental/Mutable/Source/CustomizableObject/Private/MuCO/CustomizableObjectSystemPrivate.h
+++ b/Engine/Plugins/Experimental/Mutable/Source/CustomizableObject/Private/MuCO/CustomizableObjectSystemPrivate.h
@@ -64,8 +64,8 @@ public:
 	// It is weak because we don't want to lock it in case it becomes irrelevant in the game while operations are pending and it needs to be destroyed.
 	TWeakObjectPtr<UCustomizableObjectInstance> CustomizableObjectInstance;
 
-	/** Hash of the UCustomizableObjectInstance::Descriptor (state of the instance parameters + state). */
-	uint32 InstanceDescriptorHash;
+	/** Hash of the UCustomizableObjectInstance::Descriptor at the time of the update request. */
+	FDescriptorRuntimeHash InstanceDescriptorRuntimeHash;
 
 	//! This is used to calculate stats.
 	double StartUpdateTime = 0.0;
diff --git a/Engine/Plugins/Experimental/Mutable/Source/CustomizableObject/Public/MuCO/CustomizableObjectInstanceDescriptor.h b/Engine/Plugins/Experimental/Mutable/Source/CustomizableObject/Public/MuCO/CustomizableObjectInstanceDescriptor.h
--- a/Engine/Plugins/Experimental/Mutable/Source/CustomizableObject/Public/MuCO/CustomizableObjectInstanceDescriptor.h
+++ b/Engine/Plugins/Experimental/Mutable/Source/CustomizableObject/Public/MuCO/CustomizableObjectInstanceDescriptor.h
@@ -6,7 +6,6 @@
 #include "Containers/Map.h"
 #include "Containers/UnrealString.h"
 #include "Math/Color.h"
-#include "Math/UnrealMathSSE.h"
 #include "MuCO/CustomizableObjectParameterTypeDefinitions.h"
 #include "MuCO/CustomizableObject.h"
 #include "MuCO/MultilayerProjector.h"
@@ -18,6 +17,8 @@ class FArchive;
 class UCustomizableInstancePrivateData;
 class UCustomizableObject;
 class UCustomizableObjectInstance;
+class FDescriptorHash;
+class FDescriptorRuntimeHash;
 
 namespace mu
 {
@@ -30,7 +31,7 @@ namespace mu
 
 /** Set of parameters + state that defines a CustomizableObjectInstance.
  *
- * This object has the same parameters + state interfice as UCustomizableObjectInstance.
+ * This object has the same parameters + state interface as UCustomizableObjectInstance.
  * UCustomizableObjectInstance must share the same interface. Any public methods added here should also end up in the Instance. */
 USTRUCT()
 struct CUSTOMIZABLEOBJECT_API FCustomizableObjectInstanceDescriptor
@@ -336,11 +337,45 @@ private:
 	void CreateParametersLookupTable();
 	
 	// Friends
-	friend CUSTOMIZABLEOBJECT_API uint32 GetTypeHash(const FCustomizableObjectInstanceDescriptor& Key);
+	friend FDescriptorHash;
+	friend FDescriptorRuntimeHash;
 	friend UCustomizableObjectInstance;
 	friend UCustomizableInstancePrivateData;
 	friend FMultilayerProjector;
 };
 
 
-CUSTOMIZABLEOBJECT_API uint32 GetTypeHash(const FCustomizableObjectInstanceDescriptor& Key);
\ No newline at end of file
+/** Hash of the Descriptor. Hashes everything except runtime information. */
+class CUSTOMIZABLEOBJECT_API FDescriptorHash
+{
+public:
+	FDescriptorHash() = default;
+
+	explicit FDescriptorHash(const FCustomizableObjectInstanceDescriptor& Descriptor);
+
+	bool operator==(const FDescriptorHash& Other) const;
+
+	bool operator<(const FDescriptorHash& Other) const;
+
+	FString ToString() const;
+	
+private:
+	uint32 Hash = 0;
+};
+
+
+/** Hash of the Descriptor. Hashes everything including runtime information. */
+class CUSTOMIZABLEOBJECT_API FDescriptorRuntimeHash : public FDescriptorHash
+{
+public:
+	FDescriptorRuntimeHash() = default;
+
+	explicit FDescriptorRuntimeHash(const FCustomizableObjectInstanceDescriptor& Descriptor);
+
+	/** Return true if this Hash is a subset of the other Hash (i.e., this Descriptor is a subset of the other Descriptor). */
+	bool IsSubset(const FDescriptorRuntimeHash& Other) const;
+	
+private:
+	int32 MinLODToLoad = 0;
+	int32 MaxLODToLoad = INT32_MAX;
+};
\ No newline at end of file
