[Mutable] DescriptorHash added different operator.

#rnx

diff --git a/Engine/Plugins/Experimental/Mutable/Source/CustomizableObject/Private/MuCO/CustomizableObjectInstanceDescriptor.cpp b/Engine/Plugins/Experimental/Mutable/Source/CustomizableObject/Private/MuCO/CustomizableObjectInstanceDescriptor.cpp
--- a/Engine/Plugins/Experimental/Mutable/Source/CustomizableObject/Private/MuCO/CustomizableObjectInstanceDescriptor.cpp
+++ b/Engine/Plugins/Experimental/Mutable/Source/CustomizableObject/Private/MuCO/CustomizableObjectInstanceDescriptor.cpp
@@ -2144,6 +2144,12 @@ bool FDescriptorHash::operator==(const FDescriptorHash& Other) const
 }
 
 
+bool FDescriptorHash::operator!=(const FDescriptorHash& Other) const
+{
+	return Hash != Other.Hash;	
+}
+
+
 bool FDescriptorHash::operator<(const FDescriptorHash& Other) const
 {
 	return Hash < Other.Hash;
diff --git a/Engine/Plugins/Experimental/Mutable/Source/CustomizableObject/Public/MuCO/CustomizableObjectInstanceDescriptor.h b/Engine/Plugins/Experimental/Mutable/Source/CustomizableObject/Public/MuCO/CustomizableObjectInstanceDescriptor.h
--- a/Engine/Plugins/Experimental/Mutable/Source/CustomizableObject/Public/MuCO/CustomizableObjectInstanceDescriptor.h
+++ b/Engine/Plugins/Experimental/Mutable/Source/CustomizableObject/Public/MuCO/CustomizableObjectInstanceDescriptor.h
@@ -355,6 +355,8 @@ public:
 
 	bool operator==(const FDescriptorHash& Other) const;
 
+	bool operator!=(const FDescriptorHash& Other) const;
+
 	bool operator<(const FDescriptorHash& Other) const;
 
 	FString ToString() const;
