[mutable] Fix crash opening CO due to memory missalgnment in MeshDifference operation

diff --git a/Engine/Plugins/Experimental/Mutable/Source/MutableRuntime/Private/MuR/OpMeshDifference.h b/Engine/Plugins/Experimental/Mutable/Source/MutableRuntime/Private/MuR/OpMeshDifference.h
--- a/Engine/Plugins/Experimental/Mutable/Source/MutableRuntime/Private/MuR/OpMeshDifference.h
+++ b/Engine/Plugins/Experimental/Mutable/Source/MutableRuntime/Private/MuR/OpMeshDifference.h
@@ -17,7 +17,6 @@ struct MORPH_VERTEX
 };
 #pragma pack(pop)
 
-
 	//---------------------------------------------------------------------------------------------
 	//! Create a diff from the mesh vertices. The meshes must have the same amount of vertices.
 	//! If the channel list is empty all the channels will be compared.
@@ -185,7 +184,7 @@ struct MORPH_VERTEX
 					// all channels
 					for ( int c=0; c<numChannels; ++c )
 					{
-						*((FVector4f*)pDestData) = deltas[ v*numChannels+c ];
+						FMemory::Memcpy(pDestData, &deltas[v * numChannels + c], 16);
 						pDestData += 4*4;											
 					}
 				}
@@ -203,6 +202,7 @@ struct MORPH_VERTEX
 		return pDest;
 	}
 
+
 	/*/
 	{
 		bool correct =
