[Mutable] Removed unused code.

diff --git a/Engine/Plugins/Experimental/Mutable/Source/CustomizableObject/Classes/MuCO/CustomizableObjectSystem.h b/Engine/Plugins/Experimental/Mutable/Source/CustomizableObject/Classes/MuCO/CustomizableObjectSystem.h
--- a/Engine/Plugins/Experimental/Mutable/Source/CustomizableObject/Classes/MuCO/CustomizableObjectSystem.h
+++ b/Engine/Plugins/Experimental/Mutable/Source/CustomizableObject/Classes/MuCO/CustomizableObjectSystem.h
@@ -297,8 +297,6 @@ public:
 
 	bool IsProgressiveMipStreamingEnabled() const;
 
-public:
-
 	void AddPendingReleaseSkeletalMesh( USkeletalMesh* SkeletalMesh );
 
 	void PurgePendingReleaseSkeletalMesh();
@@ -434,7 +432,6 @@ private:
 
 	// For async material loading
 	FStreamableManager StreamableManager;
-	bool bMaterialsAlreadyLoaded = false;
 
 	// Most of the work in this plugin happens here.
 	bool Tick(float DeltaTime);
@@ -442,12 +439,6 @@ private:
 	// If there is an on-going operation, advance it.
 	void AdvanceCurrentOperation();
 
-	FTimerHandle TimerHandle_ApplyRandomValues;
-
-	bool ApplyRandomValuesToAllPawns = false;
-
-	bool RandomValuesCommandTimerWorking = false;
-
 	// TODO: Can we move this to the editor module?
 #if WITH_EDITOR
 
