Fix static analysis warning.

#rnx

diff --git a/Engine/Plugins/Experimental/Mutable/Source/CustomizableObjectEditor/Private/MuCOE/GenerateMutableSource/GenerateMutableSourceMesh.cpp b/Engine/Plugins/Experimental/Mutable/Source/CustomizableObjectEditor/Private/MuCOE/GenerateMutableSource/GenerateMutableSourceMesh.cpp
--- a/Engine/Plugins/Experimental/Mutable/Source/CustomizableObjectEditor/Private/MuCOE/GenerateMutableSource/GenerateMutableSourceMesh.cpp
+++ b/Engine/Plugins/Experimental/Mutable/Source/CustomizableObjectEditor/Private/MuCOE/GenerateMutableSource/GenerateMutableSourceMesh.cpp
@@ -2222,12 +2222,20 @@ bool GetAndValidateReshapePhysicsToDeform(
 
 	for (const USkeletalMesh* SkeletalMesh : SkeletalMeshes)
 	{
-		const FReferenceSkeleton& RefSkeleton = bIsReferenceSkeletalMeshMethod 
-				? GenerationContext.ComponentInfos[GenerationContext.CurrentMeshComponent].RefSkeletalMesh->GetRefSkeleton()
-				: SkeletalMesh->GetRefSkeleton();
+		check(SkeletalMesh)
+		check(GenerationContext.ComponentInfos[GenerationContext.CurrentMeshComponent].RefSkeletalMesh);
 
 		UPhysicsAsset* PhysicsAsset = SkeletalMesh->GetPhysicsAsset();
 
+		if (!PhysicsAsset)
+		{
+			continue;
+		}
+
+		const FReferenceSkeleton& RefSkeleton = bIsReferenceSkeletalMeshMethod 
+				? GenerationContext.ComponentInfos[GenerationContext.CurrentMeshComponent].RefSkeletalMesh->GetRefSkeleton()
+				: SkeletalMesh->GetRefSkeleton();
+	
 		TArray<uint8> BoneInclusionSet;
 		BoneInclusionSet.Init(0, PhysicsAsset->SkeletalBodySetups.Num());
 
@@ -2237,7 +2245,7 @@ bool GetAndValidateReshapePhysicsToDeform(
 			const FName& BodyBoneName = BoneNamesInUserSelection[IndexToDeform];
 			MissingBones[IndexToDeform].bMissingBone = RefSkeleton.FindBoneIndex(BodyBoneName) == INDEX_NONE;
 
-			if (PhysicsAsset && !MissingBones[IndexToDeform].bMissingBone)
+			if (!MissingBones[IndexToDeform].bMissingBone)
 			{
 				const int32 FoundIndex = PhysicsAsset->SkeletalBodySetups.IndexOfByPredicate(
 					[&BodyBoneName](const TObjectPtr<USkeletalBodySetup>& Setup) {  return Setup->BoneName == BodyBoneName; });
