#rnx

diff --git a/Engine/Plugins/Experimental/Mutable/Source/CustomizableObject/Private/MuCO/CustomizableObjectInstance.cpp b/Engine/Plugins/Experimental/Mutable/Source/CustomizableObject/Private/MuCO/CustomizableObjectInstance.cpp
--- a/Engine/Plugins/Experimental/Mutable/Source/CustomizableObject/Private/MuCO/CustomizableObjectInstance.cpp
+++ b/Engine/Plugins/Experimental/Mutable/Source/CustomizableObject/Private/MuCO/CustomizableObjectInstance.cpp
@@ -964,7 +964,6 @@ USkeleton* UCustomizableInstancePrivateData::MergeSkeletons(UCustomizableObjectI
 		const TObjectPtr<USkeleton> RefSkeleton = ReferencedSkeletons.SkeletonsToMerge[0];
 		ReferencedSkeletons.SkeletonsToMerge.Empty();
 
-		check(RefSkeleton);
 		return RefSkeleton;
 	}
 
