[mutable] Fixed wrong reconstruct node method of table nodes.

#rnx

diff --git a/Engine/Plugins/Experimental/Mutable/Source/CustomizableObjectEditor/Private/MuCOE/Nodes/CustomizableObjectNodeTable.cpp b/Engine/Plugins/Experimental/Mutable/Source/CustomizableObjectEditor/Private/MuCOE/Nodes/CustomizableObjectNodeTable.cpp
--- a/Engine/Plugins/Experimental/Mutable/Source/CustomizableObjectEditor/Private/MuCOE/Nodes/CustomizableObjectNodeTable.cpp
+++ b/Engine/Plugins/Experimental/Mutable/Source/CustomizableObjectEditor/Private/MuCOE/Nodes/CustomizableObjectNodeTable.cpp
@@ -535,10 +535,20 @@ void UCustomizableObjectNodeTable::RemapPinsData(const TMap<UEdGraphPin*, UEdGra
 
 			if (PinDataOldPin && PinDataNewPin)
 			{
-				PinDataOldPin->ColumnName = PinDataNewPin->ColumnName;
-				PinDataOldPin->MutableColumnName = PinDataNewPin->MutableColumnName;
-				PinDataOldPin->LOD = PinDataNewPin->LOD;
-				PinDataOldPin->Material = PinDataNewPin->Material;
+				if (Table->FindTableProperty(FName(*PinDataOldPin->AnimInstanceColumnName)))
+				{
+					PinDataNewPin->AnimInstanceColumnName = PinDataOldPin->AnimInstanceColumnName;
+				}
+
+				if (Table->FindTableProperty(FName(*PinDataOldPin->AnimSlotColumnName)))
+				{
+					PinDataNewPin->AnimSlotColumnName = PinDataOldPin->AnimSlotColumnName;
+				}
+
+				if (Table->FindTableProperty(FName(*PinDataOldPin->AnimTagColumnName)))
+				{
+					PinDataNewPin->AnimTagColumnName = PinDataOldPin->AnimTagColumnName;
+				}
 
 				// Keeping information added in layout editor if the layout is the same
 				for (TObjectPtr<UCustomizableObjectLayout>& NewLayout : PinDataNewPin->Layouts)
@@ -556,8 +566,6 @@ void UCustomizableObjectNodeTable::RemapPinsData(const TMap<UEdGraphPin*, UEdGra
 						}
 					}
 				}
-
-				PinDataOldPin->Layouts = PinDataNewPin->Layouts;
 			}
 		}
 	}
