[mutable] Code debugger UI improvements
- Fixed bug where the types of the root objects were not being taken into acount
- Simplified navigation panel
- Added visual hints to tell how expensive is each operation to run on tree and on navigation dropdown.

diff --git a/Engine/Plugins/Experimental/Mutable/Source/CustomizableObjectEditor/Private/MuCOE/SMutableCodeViewer.cpp b/Engine/Plugins/Experimental/Mutable/Source/CustomizableObjectEditor/Private/MuCOE/SMutableCodeViewer.cpp
--- a/Engine/Plugins/Experimental/Mutable/Source/CustomizableObjectEditor/Private/MuCOE/SMutableCodeViewer.cpp
+++ b/Engine/Plugins/Experimental/Mutable/Source/CustomizableObjectEditor/Private/MuCOE/SMutableCodeViewer.cpp
@@ -68,7 +68,7 @@
 #include "Widgets/Images/SImage.h"
 #include "Widgets/Input/SButton.h"
 #include "Widgets/Input/SNumericEntryBox.h"
-#include "Widgets/Input/STextComboBox.h"
+#include "Widgets/Input/SComboBox.h"
 #include "Widgets/Layout/SBorder.h"
 #include "Widgets/Layout/SScrollBox.h"
 #include "Widgets/Layout/SSplitter.h"
@@ -183,6 +183,7 @@ public:
 					[
 						SNew(STextBlock)
 						.Text(FText::FromString(MainLabel))
+						.ColorAndOpacity(RowItem->LabelColor)
 					]
 				]
 			];
@@ -319,7 +320,6 @@ private:
 
 	/** Color shown on the extra data column when the resource is found to be State Constant */
 	const FSlateColor StateConstantBoxColor = FSlateColor(FLinearColor(1,0,0,0.8));
-
 };
 
 
@@ -354,7 +354,10 @@ void SMutableCodeViewer::Construct(const FArguments& InArgs, const mu::ModelPtr&
 	{
 		const mu::FProgram::FState& State = ModelPrivate->m_program.m_states[StateIndex];
 		FString Caption = FString::Printf( TEXT("state [%s]"), ANSI_TO_TCHAR(State.m_name.c_str()) );
-		RootNodes.Add(MakeShareable(new FMutableCodeTreeElement(MutableModel, State.m_root, Caption)));
+
+		const FSlateColor LabelColor = ColorPerComputationalCost[StaticCast<uint8>(GetOperationTypeComputationalCost(
+			ModelPrivate->m_program.GetOpType(State.m_root)))];
+		RootNodes.Add(MakeShareable(new FMutableCodeTreeElement(MutableModel, State.m_root, Caption,LabelColor)));
 	}
 
 	// Setup Navigation system
@@ -364,9 +367,16 @@ void SMutableCodeViewer::Construct(const FArguments& InArgs, const mu::ModelPtr&
 
 		// Cache the operation types that are present on the model
 		CacheOperationTypesPresentOnModel();
-	
+
 		// Get an array of mutable types as an array of FStrings for the UI
 		GenerateNavigationOpTypeStrings();
+		
+		// Generate list elements for the found operation types so we are able to search over them on our type dropdown
+		GenerateNavigationDropdownElements();
+
+		// Check we did find types (witch should always happen in a normal run) and select the first option as default (NONE)
+		check (FoundModelOperationTypeElements.Num())
+		CurrentlySelectedOperationTypeElement = FoundModelOperationTypeElements[0];
 	}
 
 	
@@ -455,73 +465,74 @@ void SMutableCodeViewer::Construct(const FArguments& InArgs, const mu::ModelPtr&
 				.AutoHeight()
 				.HAlign(HAlign_Right)
 				[
-					// Border containing navigation elements
-					SNew(SBorder)
-					.HAlign(HAlign_Right)
-					[
-						SNew(SVerticalBox)
+					// Box containing navigation elements
+					SNew(SVerticalBox)
 
-						+ SVerticalBox::Slot()
-						.Padding(2,4)
-						.AutoHeight()
+					+ SVerticalBox::Slot()
+					.Padding(2,4)
+					.AutoHeight()
+					[
+						SNew(SHorizontalBox)
+	
+						+ SHorizontalBox::Slot()
+						.AutoWidth()
 						[
 							SNew(SHorizontalBox)
-						
-							+ SHorizontalBox::Slot()
-							.AutoWidth()
-							[
-								SNew(SHorizontalBox)
-
-								+ SHorizontalBox::Slot()
-								.AutoWidth()
-								[
-									SNew(STextBlock)
-									.Text(LOCTEXT("SelectedOperationTypeLabel","Search Operation Type :"))
-								]
-
-								+ SHorizontalBox::Slot()
-								.AutoWidth()
-								[
-									// TODO: Try to use a SComboButton so you are able to use a search box (Max)
-									SAssignNew(TargetedTypeSelector, STextComboBox)
-									.OptionsSource(&OperationTypesStrings)
-									.InitiallySelectedItem(OperationTypesStrings[0])	
-									.OnSelectionChanged(this,&SMutableCodeViewer::OnOptionTypeSelectionChanged)
-								]
-							]
 
 							+ SHorizontalBox::Slot()
-							.Padding(4,0)
 							.AutoWidth()
 							[
-								SNew(SButton)
-								.Text(LOCTEXT("GoToPreviousOperationButton"," < "))
-								.OnClicked(this,&SMutableCodeViewer::OnGoToPreviousOperationButtonPressed)
-								.IsEnabled(this,&SMutableCodeViewer::CanInteractWithPreviousOperationButton)
+								SNew(STextBlock)
+								.Text(LOCTEXT("SelectedOperationTypeLabel","Search Operation Type :"))
 							]
 
-							+ SHorizontalBox::Slot()
+							// ComboBox used to select one or another Op_Type for tree navigation purposes
+							+SHorizontalBox::Slot()
 							.AutoWidth()
 							[
-								SNew(SButton)
-								.Text(LOCTEXT("GoToNextOperationButton"," > "))
-								.OnClicked(this,&SMutableCodeViewer::OnGoToNextOperationButtonPressed)
-								.IsEnabled(this,&SMutableCodeViewer::CanInteractWithNextOperationButton)
+								SAssignNew(TargetedTypeSelector,SComboBox<TSharedPtr<FMutableOperationElement>>)
+								.OptionsSource(&FoundModelOperationTypeElements)
+								.InitiallySelectedItem(CurrentlySelectedOperationTypeElement)
+								[
+									SNew(STextBlock)
+									.Text(this,&SMutableCodeViewer::GetCurrentNavigationOpTypeText)
+									.ColorAndOpacity(this,&SMutableCodeViewer::GetCurrentNavigationOpTypeColor)
+								]
+								.OnGenerateWidget(this,&SMutableCodeViewer::OnGenerateOpNavigationDropDownWidget)
+								.OnSelectionChanged(this,&SMutableCodeViewer::OnNavigationSelectedOperationChanged)
 							]
 						]
 
-						+ SVerticalBox::Slot()
-						.AutoHeight()
-						.Padding(0,4)
+						+ SHorizontalBox::Slot()
+						.Padding(4,0)
+						.AutoWidth()
+						[
+							SNew(SButton)
+							.Text(LOCTEXT("GoToPreviousOperationButton"," < "))
+							.OnClicked(this,&SMutableCodeViewer::OnGoToPreviousOperationButtonPressed)
+							.IsEnabled(this,&SMutableCodeViewer::CanInteractWithPreviousOperationButton)
+						]
+
+						+ SHorizontalBox::Slot()
+						.Padding(4,0)
+						.AutoWidth()
 						[
 							SNew(STextBlock)
 							.Text(this,&SMutableCodeViewer::OnPrintNavigableObjectAddressesCount)
 							.Justification(ETextJustify::Right)
 						]
+		
+						+ SHorizontalBox::Slot()
+						.Padding(4,0)
+						.AutoWidth()
+						[
+							SNew(SButton)
+							.Text(LOCTEXT("GoToNextOperationButton"," > "))
+							.OnClicked(this,&SMutableCodeViewer::OnGoToNextOperationButtonPressed)
+							.IsEnabled(this,&SMutableCodeViewer::CanInteractWithNextOperationButton)
+						]
 					]
-					
 				]
-					
 				
 				// Tree operations slot
 				+ SVerticalBox::Slot()
@@ -658,52 +669,66 @@ void SMutableCodeViewer::OnCurrentMipSkipChanged(int32 NewValue)
 
 #pragma region CodeTree operation search
 
-void SMutableCodeViewer::GenerateNavigationOpTypeStrings()
+FText SMutableCodeViewer::GetCurrentNavigationOpTypeText() const
 {
-	// Grab only the names from the operation types located during the caching of operation types of the model
-	for (const mu::OP_TYPE LocatedOperationType : ModelOperationTypes)
-	{
-		// Find the name
-		const uint16 OperationIndex = static_cast<uint16>(LocatedOperationType);
-		const TCHAR* OpName = mu::s_opNames[OperationIndex];
-		
-		// Operation name to be saved
-		TSharedPtr<FString> OperationNamePointer = MakeShared<FString>(OpName);
-		OperationTypesStrings.Add( OperationNamePointer);
-	}
+	check (CurrentlySelectedOperationTypeElement);
+	
+	return CurrentlySelectedOperationTypeElement->OperationTypeText;
 }
 
-void SMutableCodeViewer::OnSelectedOperationTypeFromTree()
+FSlateColor SMutableCodeViewer::GetCurrentNavigationOpTypeColor() const
 {
-	// We require to have only 1 element selected to avoid having inconsistencies during operation
-	check(TreeView->GetNumItemsSelected() == 1);
+	check (CurrentlySelectedOperationTypeElement);
+
+	return CurrentlySelectedOperationTypeElement->OperationTextColor;
+}
+
+void SMutableCodeViewer::GenerateNavigationDropdownElements()
+{
+	const int32 OperationTypesCount = ModelOperationTypes.Num();
 	
-	const TSharedPtr<FMutableCodeTreeElement> ReferenceOperationElement = TreeView->GetSelectedItems()[0];
+	// It must have at least one type, if not may be because we are running this before filling ModelOperationTypes
+	check(OperationTypesCount);
+	FoundModelOperationTypeElements.Empty(OperationTypesCount);
 	
-	const mu::OP_TYPE OperationType =
-		MutableModel->GetPrivate()->m_program.GetOpType(ReferenceOperationElement->MutableOperation);
-	const int32 OperationIndex = ModelOperationTypes.IndexOfByKey(OperationType);
+	for	(int32 OperationTypeIndex = 0; OperationTypeIndex < OperationTypesCount;  OperationTypeIndex++)
+	{
+		// Get the type as a string to be able to print it on the UI
+		TSharedPtr<FString> OperationTypeAsString = ModelOperationTypeStrings[OperationTypeIndex];
+		FText OperationTypeName = FText::FromString(*OperationTypeAsString);
 
-	// Failing the next check would mean that we are not caching all the types present on the current operation's tree
-	check (OperationIndex != -1);
-	
-	const TSharedPtr<FString> NewOperationString = OperationTypesStrings[OperationIndex];
+		// Get the Color to be used on the text that will represent the operation on the dropdown
+		const mu::OP_TYPE RepresentedType = ModelOperationTypes[OperationTypeIndex];
+		
+		const FSlateColor OperationColor = ColorPerComputationalCost[StaticCast<uint8>(GetOperationTypeComputationalCost(RepresentedType))];
+		
+		// Generate an element to be used by the ComboBox handling the selection of the type to be used during navigation
+		TSharedPtr<FMutableOperationElement> OperationElement = MakeShared<FMutableOperationElement>(RepresentedType,OperationTypeName,OperationColor);
+		FoundModelOperationTypeElements.Add(OperationElement);
+	}
+}
 
-	// Set the type operation type to be looking for -> Will invoke OnOptionTypeSelectionChanged
-	TargetedTypeSelector->SetSelectedItem(NewOperationString);
+TSharedRef<SWidget> SMutableCodeViewer::OnGenerateOpNavigationDropDownWidget(
+	TSharedPtr<FMutableOperationElement> MutableOperationElement) const
+{
+	return SNew(STextBlock).
+			Text(MutableOperationElement->OperationTypeText).
+			ColorAndOpacity(MutableOperationElement->OperationTextColor);
 }
 
-void SMutableCodeViewer::OnOptionTypeSelectionChanged(TSharedPtr<FString, ESPMode::ThreadSafe> NesSelectedOperationString,
-	ESelectInfo::Type SelectionType)
+void SMutableCodeViewer::OnNavigationSelectedOperationChanged(
+	TSharedPtr<FMutableOperationElement, ESPMode::ThreadSafe> MutableOperationElement, ESelectInfo::Type Arg)
 {
+	check (MutableOperationElement.IsValid());
+
 	// Cache the currently selected operation set on the UI by the user
-	const int32 NewOperationIndex = OperationTypesStrings.IndexOfByKey(NesSelectedOperationString);
-	const mu::OP_TYPE NewOperationType = ModelOperationTypes[NewOperationIndex];
+	const mu::OP_TYPE NewOperationType = MutableOperationElement->OperationType;
 	if (NewOperationType == OperationTypeToSearch)
 	{
 		return;
 	}
 	OperationTypeToSearch = NewOperationType;
+	CurrentlySelectedOperationTypeElement = MutableOperationElement;
 	
 	// Only do the internal work if the type is one that makes sense searching
 	if (OperationTypeToSearch != mu::OP_TYPE::NONE)
@@ -728,6 +753,41 @@ void SMutableCodeViewer::OnOptionTypeSelectionChanged(TSharedPtr<FString, ESPMod
 	
 }
 
+
+void SMutableCodeViewer::GenerateNavigationOpTypeStrings()
+{
+	// Grab only the names from the operation types located during the caching of operation types of the model
+	for (const mu::OP_TYPE LocatedOperationType : ModelOperationTypes)
+	{
+		// Find the name
+		const uint16 OperationIndex = static_cast<uint16>(LocatedOperationType);
+		const TCHAR* OpName = mu::s_opNames[OperationIndex];
+		
+		// Operation name to be saved
+		TSharedPtr<FString> OperationNamePointer = MakeShared<FString>(OpName);
+		ModelOperationTypeStrings.Add( OperationNamePointer);
+	}
+}
+
+void SMutableCodeViewer::OnSelectedOperationTypeFromTree()
+{
+	// We require to have only 1 element selected to avoid having inconsistencies during operation
+	check(TreeView->GetNumItemsSelected() == 1);
+	
+	const TSharedPtr<FMutableCodeTreeElement> ReferenceOperationElement = TreeView->GetSelectedItems()[0];
+	
+	const mu::OP_TYPE OperationType =
+		MutableModel->GetPrivate()->m_program.GetOpType(ReferenceOperationElement->MutableOperation);
+	const int32 OperationIndex = ModelOperationTypes.IndexOfByKey(OperationType);
+
+	// Failing the next check would mean that we are not caching all the types present on the current operation's tree
+	check (OperationIndex != -1);
+	
+	// Set the type operation type to be looking for -> Will invoke OnOptionTypeSelectionChanged
+	TargetedTypeSelector->SetSelectedItem(FoundModelOperationTypeElements[OperationIndex]);
+}
+
+
 void SMutableCodeViewer::CacheAddressesOfOperationsOfType(mu::OP_TYPE TargetedOperationType)
 {
 	// Clear previous data
@@ -739,10 +799,7 @@ void SMutableCodeViewer::CacheAddressesOfOperationsOfType(mu::OP_TYPE TargetedOp
 
 	// Main update procedure run for the targeted state and the targeted parameter values
 	const mu::FProgram& Program = MutableModel->GetPrivate()->m_program;
-	for	(const mu::OP::ADDRESS& RootNodeAddress : RootNodeAddresses)
-	{
-		GetOperationsOfType(TargetedOperationType,RootNodeAddress,Program,OperationsWithTargetedType,ProcessedAddresses);
-	}
+	GetOperationsOfType(TargetedOperationType,RootNodeAddresses,Program,OperationsWithTargetedType,ProcessedAddresses);
 	
 	if (!OperationsWithTargetedType.IsEmpty())
 	{
@@ -755,29 +812,45 @@ void SMutableCodeViewer::CacheAddressesOfOperationsOfType(mu::OP_TYPE TargetedOp
 
 }
 
-void SMutableCodeViewer::GetOperationsOfType ( const mu::OP_TYPE& TargetOperationType, const mu::OP::ADDRESS& InParentAddress,const mu::FProgram& InProgram,
+void SMutableCodeViewer::GetOperationsOfType ( const mu::OP_TYPE& TargetOperationType,const TArray<mu::OP::ADDRESS>& InParentAddresses,const mu::FProgram& InProgram,
 	TSet<mu::OP::ADDRESS>& OutAddressesOfType,
 	TSet<mu::OP::ADDRESS>& AlreadyProcessedAddresses)
 {
-	// Generic case for unnamed children traversal.
-	mu::ForEachReference(InProgram, InParentAddress, [this, &InProgram, &OutAddressesOfType, &AlreadyProcessedAddresses, &TargetOperationType]( mu::OP::ADDRESS ChildAddress)
+	TArray<mu::OP::ADDRESS> NextBatch;
+	
+	for (const mu::OP::ADDRESS& ParentAddress : InParentAddresses)
 	{
-		// If the parent does have a child then process it 
-		if (ChildAddress && !AlreadyProcessedAddresses.Contains(ChildAddress))
+		// Process parent
+		if (!AlreadyProcessedAddresses.Contains(ParentAddress))
 		{
 			// Cache if same data type and we share the same address (means this op is pointing at the provided resource)
-			if (InProgram.GetOpType(ChildAddress) == TargetOperationType)
+			if (InProgram.GetOpType(ParentAddress) == TargetOperationType)
 			{
-				OutAddressesOfType.Add(ChildAddress);
+				OutAddressesOfType.Add(ParentAddress);
 			}
 		
 			// Cache to avoid processing it again later
-			AlreadyProcessedAddresses.Add(ChildAddress);
-		
-			// Process the children of this object
-			GetOperationsOfType(TargetOperationType,ChildAddress,InProgram,OutAddressesOfType,AlreadyProcessedAddresses);
+			AlreadyProcessedAddresses.Add(ParentAddress);
+
+			// Generic case for unnamed children traversal.
+			mu::ForEachReference(InProgram, ParentAddress, [this,&NextBatch]( mu::OP::ADDRESS ChildAddress)
+			{
+				// If the parent does have a child then process it 
+				if (ChildAddress)
+				{
+					// Add for expansion and search
+					NextBatch.Add(ChildAddress);
+				}
+			});
 		}
-	});
+	}
+
+	if (NextBatch.Num())
+	{
+		// Process the children of this object
+		GetOperationsOfType(TargetOperationType,NextBatch,InProgram,OutAddressesOfType,AlreadyProcessedAddresses);
+	}
+
 }
 
 
@@ -792,10 +865,7 @@ void SMutableCodeViewer::CacheOperationTypesPresentOnModel()
 	
 	const mu::FProgram& Program = MutableModel->GetPrivate()->m_program;
 	TSet<mu::OP::ADDRESS> AlreadyVisitedAddresses;
-	for	(const mu::OP::ADDRESS& RootOperationAddress : RootNodeAddresses)
-	{
-		GetOperationTypesPresentOnModel(RootOperationAddress,Program,OperationTypes,AlreadyVisitedAddresses);
-	}
+	GetOperationTypesPresentOnModel(RootNodeAddresses,Program,OperationTypes,AlreadyVisitedAddresses);
 
 	// After using the set object save it as an array for later reference
 	ModelOperationTypes = OperationTypes.Array();
@@ -833,29 +903,46 @@ void SMutableCodeViewer::CacheOperationTypesPresentOnModel()
 
 
 void SMutableCodeViewer::GetOperationTypesPresentOnModel(
-	const mu::OP::ADDRESS& InParentAddress, const mu::FProgram& InProgram, TSet<mu::OP_TYPE>& OutLocatedOperations,
+	const TArray< mu::OP::ADDRESS>& InParentAddresses, const mu::FProgram& InProgram, TSet<mu::OP_TYPE>& OutLocatedOperations,
 	TSet<mu::OP::ADDRESS>& AlreadyProcessedAddresses)
 {
-	// Generic case for unnamed children traversal.
-	mu::ForEachReference(InProgram, InParentAddress, [this, &InProgram, &OutLocatedOperations, &AlreadyProcessedAddresses]( mu::OP::ADDRESS ChildAddress)
+	TArray<mu::OP::ADDRESS> NextBatch;
+	
+	for (const mu::OP::ADDRESS& ParentAddress : InParentAddresses)
 	{
-		// Avoid processing items more than once
-		if (ChildAddress && !AlreadyProcessedAddresses.Contains(ChildAddress))
+		// Process parent 
+		if ( !AlreadyProcessedAddresses.Contains(ParentAddress))
 		{
 			// Cache the operation type for later usage
-			OutLocatedOperations.Add(InProgram.GetOpType(ChildAddress));
+			OutLocatedOperations.Add(InProgram.GetOpType(ParentAddress));
 			
 			// Cache to avoid processing it again later
-			AlreadyProcessedAddresses.Add(ChildAddress);
-		
-			// Process the children of this object
-			GetOperationTypesPresentOnModel(ChildAddress,InProgram,OutLocatedOperations,AlreadyProcessedAddresses);
+			AlreadyProcessedAddresses.Add(ParentAddress);
+
+			// Generic case for unnamed children traversal.
+			mu::ForEachReference(InProgram, ParentAddress, [this, &InProgram, &NextBatch, &OutLocatedOperations, &AlreadyProcessedAddresses]( mu::OP::ADDRESS ChildAddress)
+			{
+				// Avoid processing items more than once
+				if (ChildAddress)
+				{
+					// Add for expansion
+					NextBatch.Add(ChildAddress);
+				}
+			});
 		}
-	});
+		
+	}
+
+	if (NextBatch.Num())
+	{
+		// Process the children of this object
+		GetOperationTypesPresentOnModel(NextBatch,InProgram,OutLocatedOperations,AlreadyProcessedAddresses);
+	}
 }
 
 
-void SMutableCodeViewer::LocateNavigationElements()
+
+void SMutableCodeViewer::LocateNavigationElementsOnTree()
 {
 	NavigationFoundElements.Empty();
 	NavigationFoundElements.Reserve(TreeElements.Num());
@@ -887,7 +974,7 @@ void SMutableCodeViewer::LocateNavigationElements()
 FText SMutableCodeViewer::OnPrintNavigableObjectAddressesCount() const
 {
 	// Depending on the amount of navigable objects (addresses, not actual elements) display the amount there are
-	return FText::Format(LOCTEXT("AmountOfNavigableOps"," Found Operations : {0} "), NavigationOPAddresses.Num());
+	return FText::AsNumber(NavigationOPAddresses.Num());
 }
 
 void SMutableCodeViewer::SetCurrentNavigationElement(TSharedPtr<FMutableCodeTreeElement> NewelyFoundElement)
@@ -1268,6 +1355,28 @@ bool SMutableCodeViewer::HaveNewElementsBeenFound() const
 
 #pragma endregion 
 
+#pragma region Operation Cost Color Hints
+
+SMutableCodeViewer::EOperationComputationalCost SMutableCodeViewer::GetOperationTypeComputationalCost(mu::OP_TYPE OperationType) const
+{
+	if (VeryExpensiveOperationTypes.Contains(OperationType))
+	{
+		return EOperationComputationalCost::VeryExpensive;
+	}
+	else if (ExpensiveOperationTypes.Contains(OperationType))
+	{
+		return EOperationComputationalCost::Expensive;
+	}
+	else
+	{
+		return EOperationComputationalCost::Standard;
+	}
+}
+
+#pragma endregion 
+
+
+
 #pragma region CodeTree Callbacks
 
 
@@ -1322,7 +1431,7 @@ void SMutableCodeViewer::GetChildrenForInfo(TSharedPtr<FMutableCodeTreeElement>
 
 	// Generic case for unnamed children traversal.
 	uint32 ChildIndex = 0;
-	mu::ForEachReference(Program, InInfo->MutableOperation, [this, ParentAddress, &ChildIndex, &OutChildren](mu::OP::ADDRESS ChildAddress)
+	mu::ForEachReference(Program, InInfo->MutableOperation, [this, ParentAddress, &ChildIndex, &OutChildren, &Program](mu::OP::ADDRESS ChildAddress)
 	{
 		if (ChildAddress)
 		{
@@ -1336,8 +1445,11 @@ void SMutableCodeViewer::GetChildrenForInfo(TSharedPtr<FMutableCodeTreeElement>
 			else
 			{
 				const TSharedPtr<FMutableCodeTreeElement>* MainItemPtr = MainItemPerOp.Find(ChildAddress);
-				const TSharedPtr<FMutableCodeTreeElement> Item = MakeShareable(new FMutableCodeTreeElement(MutableModel, ChildAddress, TEXT(""), MainItemPtr));
 
+				const FSlateColor LabelColor = ColorPerComputationalCost[ StaticCast<uint8> (GetOperationTypeComputationalCost(Program.GetOpType(ChildAddress))) ];
+				
+				const TSharedPtr<FMutableCodeTreeElement> Item = MakeShareable(new FMutableCodeTreeElement(MutableModel, ChildAddress, TEXT(""),LabelColor, MainItemPtr));
+				
 				OutChildren.Add(Item);
 				ItemCache.Add(Key, Item);
 
@@ -1851,14 +1963,8 @@ void SMutableCodeViewer::CacheRootNodeAddresses()
 		// SetCurrentNavigationElement(nullptr);
 		
 		// Reset operation selection object to show no element to navigate to since we are not navigating over op types
-		{
-			// When setting it to none we are telling the system to not navigate over the operations
-			const int32 OperationIndex =  StaticCast<int32>(mu::OP_TYPE::NONE);
-			const TSharedPtr<FString> NewOperationString = OperationTypesStrings[OperationIndex];
-
-			// Set the type operation type to be looking for to be none.
-			TargetedTypeSelector->SetSelectedItem(NewOperationString);
-		}
+		// Set the type operation type to NONE (we know it is always the first one)
+		TargetedTypeSelector->SetSelectedItem(FoundModelOperationTypeElements[0]);
 
 		// Dump the located resources array onto the navigation array since we have content to navigate over
 		NavigationOPAddresses = ConstantResourcesAddresses.Array();
@@ -2117,7 +2223,7 @@ void SMutableCodeViewer::Tick(const FGeometry& AllottedGeometry, const double In
 			bUpdatedNavigationAddresses = false;
 
 			// Locate the elements that can be currently navigated over.
-			LocateNavigationElements();
+			LocateNavigationElementsOnTree();
 
 			// Simulate a button press to go to the next operation. At this point we should have no operation set so it
 			// will find and select the first one
diff --git a/Engine/Plugins/Experimental/Mutable/Source/CustomizableObjectEditor/Private/MuCOE/SMutableCodeViewer.h b/Engine/Plugins/Experimental/Mutable/Source/CustomizableObjectEditor/Private/MuCOE/SMutableCodeViewer.h
--- a/Engine/Plugins/Experimental/Mutable/Source/CustomizableObjectEditor/Private/MuCOE/SMutableCodeViewer.h
+++ b/Engine/Plugins/Experimental/Mutable/Source/CustomizableObjectEditor/Private/MuCOE/SMutableCodeViewer.h
@@ -31,8 +31,11 @@
 #include "Widgets/SCompoundWidget.h"
 #include "Widgets/Views/STableViewBase.h"
 #include "Widgets/Views/STreeView.h"
+#include "Widgets/Input/SComboBox.h"
+
 
 class FMutableCodeTreeElement;
+class FMutableOperationElement;
 class FReferenceCollector;
 class ITableRow;
 class SBorder;
@@ -188,34 +191,116 @@ private:
 	/** Operation that is currently being highlighted */
 	mu::OP::ADDRESS HighlightedOperation{};
 
+	/*
+	 * Operations computational cost reference collections
+	 */
+
+	/** Collection with all very expensive to run operation types */
+	const TArray<mu::OP_TYPE> VeryExpensiveOperationTypes
+	{
+		mu::OP_TYPE::ME_BINDSHAPE,
+		mu::OP_TYPE::ME_MASKCLIPMESH,
+	};
+
+	/** Collection with all expensive to run operation types */
+	const TArray<mu::OP_TYPE> ExpensiveOperationTypes
+	{
+		// none
+	};
+	
+	/** Enum designed to be able to notify the row generation of the type of operation being generated */
+	enum class EOperationComputationalCost : uint8
+	{
+		Standard =			0,			// All other operation types
+		Expensive =			1,			// The ones located on ExpensiveOperationTypes
+		VeryExpensive =		2			// The ones located on VeryExpensiveOperationTypes
+	};
+
+	/** Array holding the relation between each computational cost category and the color to be used to display elements related
+	 * with it. 
+	 */
+	const TArray<FSlateColor> ColorPerComputationalCost
+	{
+		FSlateColor(FLinearColor(1,1,1,1)),		// Standard cost color
+		FSlateColor(FLinearColor(1,0.4,0.2,1)),	// Expensive cost color
+		FSlateColor(FLinearColor(1,1,1,1))			// Very Expensive cost color
+	};
+	
+	/** Provided an operation type it returns the category representing how much costs to run an operation of this type
+	 * @param OperationType The operation type you want to know the computational cost of
+	 * @return The enum value representing the computational cost of the operation type provided.
+	 */
+	EOperationComputationalCost GetOperationTypeComputationalCost(mu::OP_TYPE OperationType) const;
+	
 	/*
 	 * Navigation : Operation type navigation Type selection object
 	 */
 	
 	/** Slate object that provides the user a way of selecting what kind of operation it wants to navigate.*/
-	TSharedPtr<STextComboBox> TargetedTypeSelector;
+	TSharedPtr<SComboBox<TSharedPtr<FMutableOperationElement>>> TargetedTypeSelector;
+
+	/** Data backend for the list displayed for the navigation type selection */
+	TArray<TSharedPtr<FMutableOperationElement>> FoundModelOperationTypeElements;
+
+	/** Currently selected element on the TargetedTypeSelector slate. Actively used by the ui */
+	TSharedPtr<FMutableOperationElement> CurrentlySelectedOperationTypeElement;
 	
 	/** Operation type we are using to search for tree nodes. Driven primarily by the UI */
 	mu::OP_TYPE OperationTypeToSearch = mu::OP_TYPE::NONE;
 
+	/** Operation types present on the currently set mutable model. */
+	TArray<mu::OP_TYPE> ModelOperationTypes;
+	
 	/** Array with all the names for each of the operations available on mutable. Used by the UI  */
-	TArray<TSharedPtr<FString>> OperationTypesStrings;
+	TArray<TSharedPtr<FString>> ModelOperationTypeStrings;
 	
 	/** Stores a list of strings based on the possible types of operations mutable define to be used by the UI */
 	void GenerateNavigationOpTypeStrings();
 
-	/** Method called when the user changes the targeted type using the UI. It is only half part of the operation since
-	 * "OnOptionTypeSelectionChangedAfterRefresh" will be the one that at the end will perform the operations over the tree.
-	 */
-	void OnOptionTypeSelectionChanged(TSharedPtr<FString, ESPMode::ThreadSafe> NesSelectedOperationString, ESelectInfo::Type SelectionType);
+	/** Generate a list of elements that will be used as backend for the navigation type selection dropdown */
+	void GenerateNavigationDropdownElements();
 	
 	/** Update the array with all the elements using the current operation. Required when changing the targeted operation type */
-	void LocateNavigationElements();
+	void LocateNavigationElementsOnTree();
+
+	/** Fills ModelOperationTypes with all the types present on the current model.
+	 * It also makes sure we have a NONE operation and that the operations are sorted alphabetically.
+	 * @note Method designed to be called once per model.*/
+	void CacheOperationTypesPresentOnModel();
+	
+	/** Method to scan over all the operations performed on the current model to produce a set of operation types present on it.
+	 * @param InParentAddresses Addresses of the parent objects. Required to be able to perform recursive calls to this method (to get the data from the children)
+	 * @param  InProgram Mutable program.
+	 * @param OutLocatedOperations Output of the method : Set with all the operation types present on the model.
+	 * @param AlreadyProcessedAddresses Addresses already processed. Used to avoid processing the same operation twice.
+	 */
+	void GetOperationTypesPresentOnModel(
+		const TArray< mu::OP::ADDRESS>& InParentAddresses, const mu::FProgram& InProgram, TSet<mu::OP_TYPE>& OutLocatedOperations,
+		TSet<mu::OP::ADDRESS>& AlreadyProcessedAddresses);
 	
 	/*
 	 * Navigation : UI callback methods
 	 */
 
+	/** Generate the text to be used by the navigation operation selector */
+	FText GetCurrentNavigationOpTypeText() const;
+
+	/** Returns the color to be used by the text being currently displayed as selected on the operation selector*/
+	FSlateColor GetCurrentNavigationOpTypeColor() const;
+
+	/** Callback invoked by the ComboBox used for displaying and selecting operation types for navigation. it gets invoked
+	 * each time the slate object requires to draw a line representing one of the elements set on FoundModelOperationTypeElements
+	 */
+	TSharedRef<SWidget> OnGenerateOpNavigationDropDownWidget(TSharedPtr<FMutableOperationElement> MutableOperationElement) const;
+
+	/** Callback invoked each time the selected operation on our navigation slate changes. It can change due to UI interaction
+	 * or also due to direct change by invoking the SetSelectedOption on the SComboBox TargetedTypeSelector
+	 */
+	void OnNavigationSelectedOperationChanged(TSharedPtr<FMutableOperationElement, ESPMode::ThreadSafe> MutableOperationElement, ESelectInfo::Type Arg);
+
+	/** Callback used to print on screen the amount of operations found on the tree that share the same operation type that
+	 * the one currently selected on the navigation system.
+	 */
 	FText OnPrintNavigableObjectAddressesCount() const;
 	
 	/** Method that tells the UI if the bottom to go to the previous element can be interacted
@@ -276,40 +361,16 @@ private:
 
 	/** Fills an array with all the addresses of operations that do have in common the same targeted operation type
 	 * @param TargetOperationType The operation type used to discriminate what operation addresses we want to retrieve.
-	 * @param InParentAddress The address of the parent object. Used to get the children and then process them too.
+	 * @param InParentAddresses The addresses of the parent object. Used to get the children and then process them too.
 	 * @param OutAddressesOfType Output with all the addresses that are of the same type as TargetOperationType.
 	 * @param  AlreadyProcessedAddresses Used during recursive calls. Cache the already processed operations to avoid
 	 * processing them more than once.
 	 */
-	void GetOperationsOfType ( const mu::OP_TYPE& TargetOperationType, const mu::OP::ADDRESS& InParentAddress, const mu::FProgram& InProgram,
+	void GetOperationsOfType ( const mu::OP_TYPE& TargetOperationType,const TArray<mu::OP::ADDRESS>& InParentAddresses,const mu::FProgram& InProgram,
 		TSet<mu::OP::ADDRESS>& OutAddressesOfType,
 		TSet<mu::OP::ADDRESS>& AlreadyProcessedAddresses);
 
 
-	/*
-	 * Navigation : Caching of operation types being used on the model
-	 */
-
-	/** Operation types present on the currently set mutable model. */
-	TArray<mu::OP_TYPE> ModelOperationTypes;
-
-	
-	/** Fills ModelOperationTypes with all the types present on the current model.
-	 * It also makes sure we have a NONE operation and that the operations are sorted alphabetically.
-	 * @note Method designed to be called once per model.*/
-	void CacheOperationTypesPresentOnModel();
-
-	
-	/** Method to scan over all the operations performed on the current model to produce a set of operation types present on it.
-	 * @param InParentAddress Address of the parent object. Required to be able to perform recursive calls to this method (to get the data from the children)
-	 * @param  InProgram Mutable program.
-	 * @param OutLocatedOperations Output of the method : Set with all the operation types present on the model.
-	 * @param AlreadyProcessedAddresses Addresses already processed. Used to avoid processing the same operation twice.
-	 */
-	void GetOperationTypesPresentOnModel( const mu::OP::ADDRESS& InParentAddress, const mu::FProgram& InProgram,
-		TSet<mu::OP_TYPE>& OutLocatedOperations,
-		TSet<mu::OP::ADDRESS>& AlreadyProcessedAddresses);
-
 	/*
 	 * Navigation : Cache operation instances based on constant resource
 	 */
@@ -371,7 +432,7 @@ private:
 	
 	
 	/*
-	 * Navigation : operation methods and variables
+	 * Navigation : operation methods and variables from Tree
 	 */
 	
 	/** The element on FoundElementsOfType that we are inspecting and using as origin for our search.
@@ -576,17 +637,34 @@ private:
 	void PrepareProjectorViewer();
 };
 
+class FMutableOperationElement : public TSharedFromThis<FMutableOperationElement>
+{
+public:
+	FMutableOperationElement(mu::OP_TYPE InOperationType, FText OperationTypeBeingReprepresented,FSlateColor OperationColor)
+	{
+		OperationType = InOperationType;
+		OperationTypeText = OperationTypeBeingReprepresented;
+		OperationTextColor = OperationColor;
+	}
+
+public:
+	mu::OP_TYPE OperationType;
+	FText OperationTypeText;
+	FSlateColor OperationTextColor;
+};
+
 
 
 /** An row of the code tree in the SMutableCodeViewer. */
 class FMutableCodeTreeElement : public TSharedFromThis<FMutableCodeTreeElement>
 {
 public:
-	FMutableCodeTreeElement(const mu::ModelPtr& InModel, mu::OP::ADDRESS InOperation, const FString& InCaption, const TSharedPtr<FMutableCodeTreeElement>* InDuplicatedOf = nullptr)
+	FMutableCodeTreeElement(const mu::ModelPtr& InModel, mu::OP::ADDRESS InOperation, const FString& InCaption,const FSlateColor InLabelColor, const TSharedPtr<FMutableCodeTreeElement>* InDuplicatedOf = nullptr)
 	{
 		MutableModel = InModel;
 		MutableOperation = InOperation;
 		Caption = InCaption;
+		LabelColor = InLabelColor;
 		if (InDuplicatedOf)
 		{
 			DuplicatedOf = *InDuplicatedOf;
@@ -651,4 +729,7 @@ public:
 
 	/** If this tree element is a duplicated of another op, this is the op. */
 	TSharedPtr<FMutableCodeTreeElement> DuplicatedOf;
+
+	/** The color to be used by the row representing this object */
+	FSlateColor LabelColor;
 };
