[mutable] Fixed morph stack application and definition crash.
+ Added stack nodes support for table nodes, variation nodes...
+ Now when a morph stack node is compiled, it is stored and applied at the end (after the generation of the skeletal mesh)

#rnx

diff --git a/Engine/Plugins/Experimental/Mutable/Source/CustomizableObjectEditor/Private/MuCOE/GenerateMutableSource/GenerateMutableSource.h b/Engine/Plugins/Experimental/Mutable/Source/CustomizableObjectEditor/Private/MuCOE/GenerateMutableSource/GenerateMutableSource.h
--- a/Engine/Plugins/Experimental/Mutable/Source/CustomizableObjectEditor/Private/MuCOE/GenerateMutableSource/GenerateMutableSource.h
+++ b/Engine/Plugins/Experimental/Mutable/Source/CustomizableObjectEditor/Private/MuCOE/GenerateMutableSource/GenerateMutableSource.h
@@ -97,6 +97,32 @@ enum class EMutableMeshConversionFlags : uint32
 };
 
 
+/** 
+	Struct to store the necessary data to generate the morphs of a skeletal mesh 
+	This struct allows the stack morph nodes to use the same functions as the mesh morph nodes
+*/
+struct FMorphNodeData
+{
+	// Pointer to the node that owns this morph data
+	UCustomizableObjectNode* OwningNode;
+
+	// Name of the morph that will be applied
+	FString MorphTargetName;
+
+	// Pin to the node that generates the factor of the morph
+	const UEdGraphPin* FactorPin;
+
+	// Pin of the mesh where the morphs will ble apllied
+	const UEdGraphPin* MeshPin;
+
+	bool operator==(const FMorphNodeData& Other) const
+	{
+		return OwningNode == Other.OwningNode && MorphTargetName == Other.MorphTargetName
+			&& FactorPin == Other.FactorPin && MeshPin == Other.MeshPin;
+	}
+};
+
+
 // Key for the data stored for each processed unreal graph node.
 class FGeneratedKey
 {
@@ -119,7 +145,7 @@ private:
 	uint32 Flags = 0;
 
 	/** Active morphs at the time of mesh generation. */
-	TArray<UCustomizableObjectNodeMeshMorph*> MeshMorphStack;
+	TArray<FMorphNodeData> MeshMorphStack;
 };
 
 
@@ -702,7 +728,7 @@ struct FMutableGraphGenerationContext
 	FPinData PinData; // DEPRECATED. DO NOT USE!
 
 	// Stores all morphs to apply them directly to a skeletal mesh node
-	TArray<UCustomizableObjectNodeMeshMorph*> MeshMorphStack;
+	TArray<FMorphNodeData> MeshMorphStack;
 
 	// Current material parameter name to find the corresponding column in a mutable table
 	FString CurrentMaterialTableParameter;
diff --git a/Engine/Plugins/Experimental/Mutable/Source/CustomizableObjectEditor/Private/MuCOE/GenerateMutableSource/GenerateMutableSourceMesh.cpp b/Engine/Plugins/Experimental/Mutable/Source/CustomizableObjectEditor/Private/MuCOE/GenerateMutableSource/GenerateMutableSourceMesh.cpp
--- a/Engine/Plugins/Experimental/Mutable/Source/CustomizableObjectEditor/Private/MuCOE/GenerateMutableSource/GenerateMutableSourceMesh.cpp
+++ b/Engine/Plugins/Experimental/Mutable/Source/CustomizableObjectEditor/Private/MuCOE/GenerateMutableSource/GenerateMutableSourceMesh.cpp
@@ -58,6 +58,7 @@
 #include "MuCOE/GenerateMutableSource/GenerateMutableSourceFloat.h"
 #include "MuCOE/GenerateMutableSource/GenerateMutableSourceLayout.h"
 #include "MuCOE/GenerateMutableSource/GenerateMutableSourceTable.h"
+#include "MuCOE/EdGraphSchema_CustomizableObject.h"
 #include "MuCOE/GraphTraversal.h"
 #include "MuCOE/Nodes/CustomizableObjectNode.h"
 #include "MuCOE/Nodes/CustomizableObjectNodeAnimationPose.h"
@@ -2317,7 +2318,7 @@ bool GetAndValidateReshapePhysicsToDeform(
 
 
 mu::NodeMeshPtr GenerateMorphMesh(const UEdGraphPin* Pin,
-	TArray<UCustomizableObjectNodeMeshMorph*> TypedNodeMorphs,
+	TArray<FMorphNodeData> TypedNodeMorphs,
 	int32 MorphIndex,
 	mu::NodeMeshPtr SourceNode,
 	FMutableGraphGenerationContext & GenerationContext,
@@ -2331,14 +2332,14 @@ mu::NodeMeshPtr GenerateMorphMesh(const UEdGraphPin* Pin,
 	check(MeshNode);
 	
 	// Current morph node
-	const UCustomizableObjectNode* MorphNode = TypedNodeMorphs[MorphIndex];
+	const UCustomizableObjectNode* MorphNode = TypedNodeMorphs[MorphIndex].OwningNode;
 	check(MorphNode);
 	
 	mu::NodeMeshMorphPtr Result = new mu::NodeMeshMorph();
 	Result->SetMorphCount(2);
 	
 	// Factor
-	GenerateMorphFactor(MorphNode, *TypedNodeMorphs[MorphIndex]->FactorPin(), GenerationContext, Result);
+	GenerateMorphFactor(MorphNode, *TypedNodeMorphs[MorphIndex].FactorPin, GenerationContext, Result);
 	
 	// Base
 	if (MorphIndex == TypedNodeMorphs.Num() - 1)
@@ -2366,7 +2367,7 @@ mu::NodeMeshPtr GenerateMorphMesh(const UEdGraphPin* Pin,
 		// Should exist
 		mu::TablePtr Table = GenerationContext.GeneratedTables[TypedNodeTable->Table->GetName()];
 
-		FString ColumnName = TableColumnName + TypedNodeMorphs[MorphIndex]->MorphTargetName;
+		FString ColumnName = TableColumnName + TypedNodeMorphs[MorphIndex].MorphTargetName;
 		int32 ColumnIndex = -1;
 
 		for (int32 RowIndex = 0; RowIndex < NumRows; ++RowIndex)
@@ -2380,7 +2381,7 @@ mu::NodeMeshPtr GenerateMorphMesh(const UEdGraphPin* Pin,
 				ColumnIndex = Table->AddColumn(TCHAR_TO_ANSI(*ColumnName), mu::TABLE_COLUMN_TYPE::TCT_MESH);
 			}
 
-			mu::MeshPtr MorphedSourceTableMesh = BuildMorphedMutableMesh(Pin, TypedNodeMorphs[MorphIndex]->MorphTargetName, GenerationContext, RowName);
+			mu::MeshPtr MorphedSourceTableMesh = BuildMorphedMutableMesh(Pin, TypedNodeMorphs[MorphIndex].MorphTargetName, GenerationContext, RowName);
 			Table->SetCell(ColumnIndex, RowIndex, MorphedSourceTableMesh.get());
 		}
 
@@ -2411,7 +2412,7 @@ mu::NodeMeshPtr GenerateMorphMesh(const UEdGraphPin* Pin,
 	}
 	else
 	{
-		MorphedSourceMesh = BuildMorphedMutableMesh(Pin, TypedNodeMorphs[MorphIndex]->MorphTargetName, GenerationContext);
+		MorphedSourceMesh = BuildMorphedMutableMesh(Pin, TypedNodeMorphs[MorphIndex].MorphTargetName, GenerationContext);
 
 		if (MorphedSourceMesh)
 		{
@@ -2435,49 +2436,50 @@ mu::NodeMeshPtr GenerateMorphMesh(const UEdGraphPin* Pin,
 
 			Result->SetMorph(1, Morph);
 
-			UCustomizableObjectNodeMeshMorph* TypedMorphNode = TypedNodeMorphs[MorphIndex];
-
-			Result->SetReshapeSkeleton(TypedMorphNode->bReshapeSkeleton);
-			Result->SetReshapePhysicsVolumes(TypedMorphNode->bReshapePhysicsVolumes);	
+			if (UCustomizableObjectNodeMeshMorph* TypedMorphNode = Cast<UCustomizableObjectNodeMeshMorph>(TypedNodeMorphs[MorphIndex].OwningNode))
 			{
-				const UEdGraphPin* ConnectedPin = FollowInputPin(*TypedMorphNode->MeshPin());
-				const UEdGraphPin* SourceMeshPin = ConnectedPin ? FindMeshBaseSource(*ConnectedPin, false) : nullptr;
-				const UEdGraphNode* SkeletalMeshNode = SourceMeshPin ? SourceMeshPin->GetOwningNode() : nullptr;
-
-				TArray<USkeletalMesh*> SkeletalMeshesToDeform = GetSkeletalMeshesForReshapeSelection(SkeletalMeshNode, SourceMeshPin);
-				
-				bool bWarningFound = false;
-				if (TypedMorphNode->bReshapeSkeleton)
+				Result->SetReshapeSkeleton(TypedMorphNode->bReshapeSkeleton);
+				Result->SetReshapePhysicsVolumes(TypedMorphNode->bReshapePhysicsVolumes);
 				{
-					TArray<FString> BonesToDeform;
-					bWarningFound = GetAndValidateReshapeBonesToDeform(
-						BonesToDeform, TypedMorphNode->BonesToDeform, SkeletalMeshesToDeform, TypedMorphNode, TypedMorphNode->SelectionMethod, GenerationContext);
-					
-					for (const FString& BoneName : BonesToDeform)
+					const UEdGraphPin* ConnectedPin = FollowInputPin(*TypedMorphNode->MeshPin());
+					const UEdGraphPin* SourceMeshPin = ConnectedPin ? FindMeshBaseSource(*ConnectedPin, false) : nullptr;
+					const UEdGraphNode* SkeletalMeshNode = SourceMeshPin ? SourceMeshPin->GetOwningNode() : nullptr;
+
+					TArray<USkeletalMesh*> SkeletalMeshesToDeform = GetSkeletalMeshesForReshapeSelection(SkeletalMeshNode, SourceMeshPin);
+
+					bool bWarningFound = false;
+					if (TypedMorphNode->bReshapeSkeleton)
 					{
-						Result->AddBoneToDeform(TCHAR_TO_ANSI(*BoneName));
+						TArray<FString> BonesToDeform;
+						bWarningFound = GetAndValidateReshapeBonesToDeform(
+							BonesToDeform, TypedMorphNode->BonesToDeform, SkeletalMeshesToDeform, TypedMorphNode, TypedMorphNode->SelectionMethod, GenerationContext);
+
+						for (const FString& BoneName : BonesToDeform)
+						{
+							Result->AddBoneToDeform(TCHAR_TO_ANSI(*BoneName));
+						}
 					}
-				}
 
-				if (TypedMorphNode->bReshapePhysicsVolumes)
-				{
-					TArray<FString> PhysicsToDeform;
+					if (TypedMorphNode->bReshapePhysicsVolumes)
+					{
+						TArray<FString> PhysicsToDeform;
 
-					const EBoneDeformSelectionMethod SelectionMethod = TypedMorphNode->PhysicsSelectionMethod;
-					bWarningFound = bWarningFound || GetAndValidateReshapePhysicsToDeform(
-						PhysicsToDeform, 
-						TypedMorphNode->PhysicsBodiesToDeform, SkeletalMeshesToDeform, SelectionMethod, 
-						TypedMorphNode, GenerationContext);
+						const EBoneDeformSelectionMethod SelectionMethod = TypedMorphNode->PhysicsSelectionMethod;
+						bWarningFound = bWarningFound || GetAndValidateReshapePhysicsToDeform(
+							PhysicsToDeform, 
+							TypedMorphNode->PhysicsBodiesToDeform, SkeletalMeshesToDeform, SelectionMethod, 
+							TypedMorphNode, GenerationContext);
 	
-					for (const FString& PhysicsBoneName : PhysicsToDeform)
+						for (const FString& PhysicsBoneName : PhysicsToDeform)
+						{
+							Result->AddPhysicsBodyToDeform(TCHAR_TO_ANSI(*PhysicsBoneName));
+						}
+					}
+						
+					if (bWarningFound)
 					{
-						Result->AddPhysicsBodyToDeform(TCHAR_TO_ANSI(*PhysicsBoneName));
-					}	
-				}
-				
-				if (bWarningFound)
-				{
-					TypedMorphNode->SetRefreshNodeWarning();
+						TypedMorphNode->SetRefreshNodeWarning();
+					}
 				}
 			}
 		}
@@ -2834,7 +2836,8 @@ mu::NodeMeshPtr GenerateMutableSourceMesh(const UEdGraphPin * Pin,
 		if (const UEdGraphPin* ConnectedPin = FollowInputPin(*TypedNodeMorph->MeshPin()))
 		{
 			// Mesh Morph Stack Management
-			GenerationContext.MeshMorphStack.Push(TypedNodeMorph);
+			FMorphNodeData NewMorphData = { TypedNodeMorph, TypedNodeMorph->MorphTargetName ,TypedNodeMorph->FactorPin(), TypedNodeMorph->MeshPin() };
+			GenerationContext.MeshMorphStack.Push(NewMorphData);
 			Result = GenerateMutableSourceMesh(ConnectedPin, GenerationContext, MeshData);
 			GenerationContext.MeshMorphStack.Pop(true);
 		}
@@ -2863,73 +2866,44 @@ mu::NodeMeshPtr GenerateMutableSourceMesh(const UEdGraphPin * Pin,
 
 				MeshNode->SetMorphCount(2);
 
-				int32 NextPin = -1;
-				bool bAddNewMorph = true;
+				TArray<UEdGraphPin*> MorphPins = TypedNodeMeshMorphStackDef->GetAllNonOrphanPins();
+
+				int32 AddedMorphs = 0;
 
-				while (bAddNewMorph)
+				for (int32 PinIndex = 0; PinIndex < MorphPins.Num(); ++PinIndex)
 				{
-					//Getting next connected pin index
-					NextPin = TypedNodeMeshMorphStackDef->NextConnectedPin(NextPin, TypedNodeMeshMorphStackApp->MorphNames);
+					UEdGraphPin* MorphPin = MorphPins[PinIndex];
 
-					UEdGraphPin* FactorPin = nullptr;
+					const UEdGraphSchema_CustomizableObject* Schema = GetDefault<UEdGraphSchema_CustomizableObject>();
 
-					// If Next Pin is -1 then there is no pin connected, which is a warning
-					if (NextPin != -1)
+					// Checking if it's a valid pin
+					if (MorphPin->Direction == EEdGraphPinDirection::EGPD_Output 
+						|| MorphPin->PinType.PinCategory != Helper_GetPinCategory(Schema->PC_Float) 
+						|| !MorphPins[PinIndex]->LinkedTo.Num())
 					{
-						FactorPin = TypedNodeMeshMorphStackDef->GetMorphPin(NextPin);
+						continue;
 					}
 
-					// Checking if there is another pin connected to add a new morph
-					bAddNewMorph = (TypedNodeMeshMorphStackDef->NextConnectedPin(NextPin, TypedNodeMeshMorphStackApp->MorphNames) != -1);
-
-					// Factor
-					GenerateMorphFactor(Node, *FactorPin, GenerationContext, MeshNode);
-
-					// Base
-					if (const UEdGraphPin* ConnectedMehsPin = FollowInputPin(*TypedNodeMeshMorphStackApp->GetMeshPin()))
+					// Cheking if the morph exists in the application node
+					FString MorphName = MorphPin->PinFriendlyName.ToString();
+					if (!TypedNodeMeshMorphStackApp->MorphNames.Contains(MorphName))
 					{
-						mu::NodeMeshPtr SourceNode = nullptr;
-
-						if (bAddNewMorph)
-						{
-							mu::NodeMeshMorphPtr NextMorph = new mu::NodeMeshMorph();
-
-							SourceNode = NextMorph;
-							MeshNode->SetBase(SourceNode);
+						continue;
+					}
 
-							if (const UEdGraphPin* BaseSourcePin = FindMeshBaseSource(*ConnectedMehsPin, false))
-							{
-								FString MorphName = "";
-								if (NextPin != -1)
-								{
-									MorphName = TypedNodeMeshMorphStackDef->MorphNames[NextPin];
-								}
-
-								//Target
-								GenerateMorphTarget(Node, BaseSourcePin, GenerationContext, MeshNode, MorphName);
-							}
+					// Mesh Morph Stack Management. TODO(Max): should we add the stack application node here instead of the def? Or both?
+					FMorphNodeData NewMorphData = { TypedNodeMeshMorphStackDef, MorphName, MorphPin, TypedNodeMeshMorphStackApp->GetMeshPin() };
+					GenerationContext.MeshMorphStack.Push(NewMorphData);
+				}
 
-							MeshNode = NextMorph;
-							MeshNode->SetMorphCount(2);
-						}
-						else
-						{
-							SourceNode = GenerateMutableSourceMesh(ConnectedMehsPin, GenerationContext, MeshData);
-							MeshNode->SetBase(SourceNode);
+				if (const UEdGraphPin* MeshConnectedPin = FollowInputPin(*TypedNodeMeshMorphStackApp->GetMeshPin()))
+				{
+					Result = GenerateMutableSourceMesh(MeshConnectedPin, GenerationContext, MeshData);
+				}
 
-							if (const UEdGraphPin* BaseSourcePin = FindMeshBaseSource(*ConnectedMehsPin, false))
-							{
-								FString MorphName = "";
-								if (NextPin != -1)
-								{
-									MorphName = TypedNodeMeshMorphStackDef->MorphNames[NextPin];
-								}
-
-								//Target
-								GenerateMorphTarget(Node, BaseSourcePin, GenerationContext, MeshNode, MorphName);
-							}
-						}
-					}
+				for (int32 MorphIndex = 0; MorphIndex < AddedMorphs; ++MorphIndex)
+				{
+					GenerationContext.MeshMorphStack.Pop(true);
 				}
 			}
 			else
@@ -3436,5 +3410,4 @@ mu::NodeMeshPtr GenerateMutableSourceMesh(const UEdGraphPin * Pin,
 	return Result;
 }
 
-#undef LOCTEXT_NAMESPACE
-
+#undef LOCTEXT_NAMESPACE
\ No newline at end of file
diff --git a/Engine/Plugins/Experimental/Mutable/Source/CustomizableObjectEditor/Private/MuCOE/Nodes/CustomizableObjectNodeMeshMorphStackApplication.cpp b/Engine/Plugins/Experimental/Mutable/Source/CustomizableObjectEditor/Private/MuCOE/Nodes/CustomizableObjectNodeMeshMorphStackApplication.cpp
--- a/Engine/Plugins/Experimental/Mutable/Source/CustomizableObjectEditor/Private/MuCOE/Nodes/CustomizableObjectNodeMeshMorphStackApplication.cpp
+++ b/Engine/Plugins/Experimental/Mutable/Source/CustomizableObjectEditor/Private/MuCOE/Nodes/CustomizableObjectNodeMeshMorphStackApplication.cpp
@@ -12,6 +12,7 @@
 #include "MuCOE/GraphTraversal.h"
 #include "MuCOE/ICustomizableObjectEditor.h"
 #include "MuCOE/Nodes/CustomizableObjectNodeSkeletalMesh.h"
+#include "MuCOE/Nodes/CustomizableObjectNodeTable.h"
 #include "Templates/Casts.h"
 #include "Templates/SharedPointer.h"
 #include "UObject/NameTypes.h"
@@ -98,21 +99,40 @@ void UCustomizableObjectNodeMeshMorphStackApplication::UpdateMorphList()
 {
 	MorphNames.Empty();
 
-	if (UEdGraphPin* MeshPin = GetMeshPin())
+	UEdGraphPin* MeshPin = GetMeshPin();
+
+	if (!MeshPin)
+	{
+		return;
+	}
+
+	UEdGraphPin* OutputMeshPin = FollowInputPin(*MeshPin);
+
+	if (!OutputMeshPin)
 	{
-		if (const UEdGraphPin* ConnectedPin = FollowInputPin(*MeshPin))
+		return;
+	}
+
+	const UEdGraphPin* MeshNodePin = FindMeshBaseSource(*OutputMeshPin, false);
+
+	if (MeshNodePin && MeshNodePin->GetOwningNode())
+	{
+		USkeletalMesh* SkeletalMesh = nullptr;
+
+		if (const UCustomizableObjectNodeSkeletalMesh* SkeletalMeshNode = Cast< UCustomizableObjectNodeSkeletalMesh >(MeshNodePin->GetOwningNode()))
 		{
-			if (UCustomizableObjectNodeSkeletalMesh* SkeletalMeshNode = Cast< UCustomizableObjectNodeSkeletalMesh >(ConnectedPin->GetOuter()))
-			{
-				USkeletalMesh* SkeletalMesh = SkeletalMeshNode->SkeletalMesh;
+			SkeletalMesh = SkeletalMeshNode->SkeletalMesh;
+		}
+		else if (const UCustomizableObjectNodeTable* TableNode = Cast< UCustomizableObjectNodeTable >(MeshNodePin->GetOwningNode()))
+		{
+			SkeletalMesh = TableNode->GetColumnDefaultAssetByType<USkeletalMesh>(MeshNodePin);
+		}
 
-				if (SkeletalMesh)
-				{
-					for (int32 i = 0; i < SkeletalMesh->GetMorphTargets().Num(); ++i)
-					{
-						MorphNames.Add(SkeletalMesh->GetMorphTargets()[i]->GetName());
-					}
-				}
+		if (SkeletalMesh)
+		{
+			for (int32 i = 0; i < SkeletalMesh->GetMorphTargets().Num(); ++i)
+			{
+				MorphNames.Add(SkeletalMesh->GetMorphTargets()[i]->GetName());
 			}
 		}
 	}
@@ -121,26 +141,52 @@ void UCustomizableObjectNodeMeshMorphStackApplication::UpdateMorphList()
 
 bool UCustomizableObjectNodeMeshMorphStackApplication::IsNodeOutDatedAndNeedsRefresh()
 {
-	if (UEdGraphPin* MeshPin = GetMeshPin())
+	UEdGraphPin* MeshPin = GetMeshPin();
+
+	if (!MeshPin)
 	{
-		if (const UEdGraphPin* ConnectedPin = FollowInputPin(*MeshPin))
+		return false;
+	}
+
+	UEdGraphPin* OutputMeshPin = FollowInputPin(*MeshPin);
+
+	if (!OutputMeshPin)
+	{
+		if (MorphNames.Num())
 		{
-			if (const UCustomizableObjectNodeSkeletalMesh* SkeletalMeshNode = Cast< UCustomizableObjectNodeSkeletalMesh >(ConnectedPin->GetOuter()))
+			return true;
+		}
+
+		return false;
+	}
+
+	const UEdGraphPin* MeshNodePin = FindMeshBaseSource(*OutputMeshPin, false);
+
+	if (MeshNodePin && MeshNodePin->GetOwningNode())
+	{
+		USkeletalMesh* SkeletalMesh = nullptr;
+
+		if (const UCustomizableObjectNodeSkeletalMesh* SkeletalMeshNode = Cast< UCustomizableObjectNodeSkeletalMesh >(MeshNodePin->GetOwningNode()))
+		{
+			SkeletalMesh = SkeletalMeshNode->SkeletalMesh;
+		}
+		else if (const UCustomizableObjectNodeTable* TableNode = Cast< UCustomizableObjectNodeTable >(MeshNodePin->GetOwningNode()))
+		{
+			SkeletalMesh = TableNode->GetColumnDefaultAssetByType<USkeletalMesh>(MeshNodePin);
+		}
+
+		if (SkeletalMesh)
+		{
+			if (SkeletalMesh->GetMorphTargets().Num() != MorphNames.Num())
+			{
+				return true;
+			}
+
+			for (int32 i = 0; i < SkeletalMesh->GetMorphTargets().Num(); ++i)
 			{
-				if (USkeletalMesh* SkeletalMesh = SkeletalMeshNode->SkeletalMesh)
+				if (!MorphNames.Contains(SkeletalMesh->GetMorphTargets()[i]->GetName()))
 				{
-					if (SkeletalMesh->GetMorphTargets().Num() != MorphNames.Num())
-					{
-						return true;
-					}
-
-					for (int32 i = 0; i < SkeletalMesh->GetMorphTargets().Num(); ++i)
-					{
-						if (!MorphNames.Contains(SkeletalMesh->GetMorphTargets()[i]->GetName()))
-						{
-							return true;
-						}
-					}
+					return true;
 				}
 			}
 		}
diff --git a/Engine/Plugins/Experimental/Mutable/Source/CustomizableObjectEditor/Private/MuCOE/Nodes/CustomizableObjectNodeMeshMorphStackDefinition.cpp b/Engine/Plugins/Experimental/Mutable/Source/CustomizableObjectEditor/Private/MuCOE/Nodes/CustomizableObjectNodeMeshMorphStackDefinition.cpp
--- a/Engine/Plugins/Experimental/Mutable/Source/CustomizableObjectEditor/Private/MuCOE/Nodes/CustomizableObjectNodeMeshMorphStackDefinition.cpp
+++ b/Engine/Plugins/Experimental/Mutable/Source/CustomizableObjectEditor/Private/MuCOE/Nodes/CustomizableObjectNodeMeshMorphStackDefinition.cpp
@@ -103,27 +103,40 @@ void UCustomizableObjectNodeMeshMorphStackDefinition::UpdateMorphList()
 {
 	MorphNames.Empty();
 
-	if (UEdGraphPin* MeshPin = GetMeshPin())
+	UEdGraphPin* MeshPin = GetMeshPin();
+
+	if (!MeshPin)
 	{
-		if (const UEdGraphPin* ConnectedPin = FollowInputPin(*MeshPin))
-		{
-			USkeletalMesh* SkeletalMesh = nullptr;
+		return;
+	}
 
-			if (const UCustomizableObjectNodeSkeletalMesh* SkeletalMeshNode = Cast< UCustomizableObjectNodeSkeletalMesh >(ConnectedPin->GetOuter()))
-			{
-				SkeletalMesh = SkeletalMeshNode->SkeletalMesh;
-			}
-			else if (const UCustomizableObjectNodeTable* TableNode = Cast< UCustomizableObjectNodeTable >(ConnectedPin->GetOuter()))
-			{
-				SkeletalMesh = TableNode->GetColumnDefaultAssetByType<USkeletalMesh>(MeshPin->LinkedTo[0]);
-			}
+	UEdGraphPin* OutputMeshPin = FollowInputPin(*MeshPin);
+
+	if (!OutputMeshPin)
+	{
+		return;
+	}
+
+	const UEdGraphPin* MeshNodePin = FindMeshBaseSource(*OutputMeshPin, false);
 
-			if (SkeletalMesh)
+	if (MeshNodePin && MeshNodePin->GetOwningNode())
+	{
+		USkeletalMesh* SkeletalMesh = nullptr;
+
+		if (const UCustomizableObjectNodeSkeletalMesh* SkeletalMeshNode = Cast< UCustomizableObjectNodeSkeletalMesh >(MeshNodePin->GetOwningNode()))
+		{
+			SkeletalMesh = SkeletalMeshNode->SkeletalMesh;
+		}
+		else if (const UCustomizableObjectNodeTable* TableNode = Cast< UCustomizableObjectNodeTable >(MeshNodePin->GetOwningNode()))
+		{
+			SkeletalMesh = TableNode->GetColumnDefaultAssetByType<USkeletalMesh>(MeshNodePin);
+		}
+
+		if (SkeletalMesh)
+		{
+			for (int32 i = 0; i < SkeletalMesh->GetMorphTargets().Num(); ++i)
 			{
-				for (int32 i = 0; i < SkeletalMesh->GetMorphTargets().Num(); ++i)
-				{
-					MorphNames.Add(SkeletalMesh->GetMorphTargets()[i]->GetName());
-				}
+				MorphNames.Add(SkeletalMesh->GetMorphTargets()[i]->GetName());
 			}
 		}
 	}
@@ -162,26 +175,52 @@ int32 UCustomizableObjectNodeMeshMorphStackDefinition::NextConnectedPin(int32 In
 
 bool UCustomizableObjectNodeMeshMorphStackDefinition::IsNodeOutDatedAndNeedsRefresh()
 {
-	if (UEdGraphPin* MeshPin = GetMeshPin())
+	UEdGraphPin* MeshPin = GetMeshPin();
+
+	if (!MeshPin)
+	{
+		return false;
+	}
+
+	UEdGraphPin* OutputMeshPin = FollowInputPin(*MeshPin);
+
+	if (!OutputMeshPin)
 	{
-		if (const UEdGraphPin* ConnectedPin = FollowInputPin(*MeshPin))
+		if (MorphNames.Num())
 		{
-			if (const UCustomizableObjectNodeSkeletalMesh* SkeletalMeshNode = Cast< UCustomizableObjectNodeSkeletalMesh >(ConnectedPin->GetOuter()))
+			return true;
+		}
+
+		return false;
+	}
+
+	const UEdGraphPin* MeshNodePin = FindMeshBaseSource(*OutputMeshPin, false);
+
+	if (MeshNodePin && MeshNodePin->GetOwningNode())
+	{
+		USkeletalMesh* SkeletalMesh = nullptr;
+
+		if (const UCustomizableObjectNodeSkeletalMesh* SkeletalMeshNode = Cast< UCustomizableObjectNodeSkeletalMesh >(MeshNodePin->GetOwningNode()))
+		{
+			SkeletalMesh = SkeletalMeshNode->SkeletalMesh;
+		}
+		else if (const UCustomizableObjectNodeTable* TableNode = Cast< UCustomizableObjectNodeTable >(MeshNodePin->GetOwningNode()))
+		{
+			SkeletalMesh = TableNode->GetColumnDefaultAssetByType<USkeletalMesh>(MeshNodePin);
+		}
+
+		if (SkeletalMesh)
+		{
+			if (SkeletalMesh->GetMorphTargets().Num() != MorphNames.Num())
+			{
+				return true;
+			}
+
+			for (int32 i = 0; i < SkeletalMesh->GetMorphTargets().Num(); ++i)
 			{
-				if (USkeletalMesh* SkeletalMesh = SkeletalMeshNode->SkeletalMesh)
+				if (!MorphNames.Contains(SkeletalMesh->GetMorphTargets()[i]->GetName()))
 				{
-					if(SkeletalMesh->GetMorphTargets().Num() != MorphNames.Num())
-					{
-						return true;
-					}
-
-					for (int32 i = 0; i < SkeletalMesh->GetMorphTargets().Num(); ++i)
-					{
-						if (!MorphNames.Contains(SkeletalMesh->GetMorphTargets()[i]->GetName()))
-						{
-							return true;
-						}
-					}
+					return true;
 				}
 			}
 		}
