#rnx

diff --git a/Engine/Plugins/Experimental/Mutable/Source/CustomizableObject/Private/MuCO/CustomizableObjectInstance.cpp b/Engine/Plugins/Experimental/Mutable/Source/CustomizableObject/Private/MuCO/CustomizableObjectInstance.cpp
--- a/Engine/Plugins/Experimental/Mutable/Source/CustomizableObject/Private/MuCO/CustomizableObjectInstance.cpp
+++ b/Engine/Plugins/Experimental/Mutable/Source/CustomizableObject/Private/MuCO/CustomizableObjectInstance.cpp
@@ -4881,7 +4881,11 @@ void UCustomizableInstancePrivateData::BuildMaterials(const TSharedPtr<FMutableO
 					}
 					else
 					{
-						ReuseMaterialCache.Add(MaterialPlaceholderSerialization, &MutableMaterialPlaceholder);
+						if (MutableMaterialPlaceholder.ParentMaterial)
+						{
+							check(MaterialPlaceholderSerialization != FString("null"));
+							ReuseMaterialCache.Add(MaterialPlaceholderSerialization, &MutableMaterialPlaceholder);
+						}
 
 						FGeneratedMaterial Material;
 						UMaterialInstanceDynamic* MaterialInstance = nullptr;
