[mutable] Fixed morph stack and an always false loop condition.

#rnx

diff --git a/Engine/Plugins/Experimental/Mutable/Source/CustomizableObjectEditor/Private/MuCOE/GenerateMutableSource/GenerateMutableSourceMesh.cpp b/Engine/Plugins/Experimental/Mutable/Source/CustomizableObjectEditor/Private/MuCOE/GenerateMutableSource/GenerateMutableSourceMesh.cpp
--- a/Engine/Plugins/Experimental/Mutable/Source/CustomizableObjectEditor/Private/MuCOE/GenerateMutableSource/GenerateMutableSourceMesh.cpp
+++ b/Engine/Plugins/Experimental/Mutable/Source/CustomizableObjectEditor/Private/MuCOE/GenerateMutableSource/GenerateMutableSourceMesh.cpp
@@ -2894,6 +2894,8 @@ mu::NodeMeshPtr GenerateMutableSourceMesh(const UEdGraphPin * Pin,
 					// Mesh Morph Stack Management. TODO(Max): should we add the stack application node here instead of the def? Or both?
 					FMorphNodeData NewMorphData = { TypedNodeMeshMorphStackDef, MorphName, MorphPin, TypedNodeMeshMorphStackApp->GetMeshPin() };
 					GenerationContext.MeshMorphStack.Push(NewMorphData);
+
+					AddedMorphs++;
 				}
 
 				if (const UEdGraphPin* MeshConnectedPin = FollowInputPin(*TypedNodeMeshMorphStackApp->GetMeshPin()))
