[mutable] Add mutable mesh metadata tags uniquely.

- Mesh tags holding metadata could be added multiple times if the mesh was accessed from multiple branches. The mesh was cached but the metadata tags where added a again.

- The repetition of tags caused false positives when checking for AnimBP slot collisions.

#rnx

diff --git a/Engine/Plugins/Experimental/Mutable/Source/CustomizableObjectEditor/Private/MuCOE/GenerateMutableSource/GenerateMutableSource.cpp b/Engine/Plugins/Experimental/Mutable/Source/CustomizableObjectEditor/Private/MuCOE/GenerateMutableSource/GenerateMutableSource.cpp
--- a/Engine/Plugins/Experimental/Mutable/Source/CustomizableObjectEditor/Private/MuCOE/GenerateMutableSource/GenerateMutableSource.cpp
+++ b/Engine/Plugins/Experimental/Mutable/Source/CustomizableObjectEditor/Private/MuCOE/GenerateMutableSource/GenerateMutableSource.cpp
@@ -1206,6 +1206,23 @@ bool AffectsCurrentComponent(const UEdGraphPin* Pin, FMutableGraphGenerationCont
 	return ComponentIndex == GenerationContext.CurrentMeshComponent;
 }
 
+int32 AddTagToMutableMeshUnique(mu::Mesh& MutableMesh, const FString& Tag)
+{
+	const int32 TagCount = MutableMesh.GetTagCount();
+
+	for (int32 TagIndex = TagCount - 1; TagIndex >= 0; --TagIndex)
+	{
+		if (FString(MutableMesh.GetTag(TagIndex)) == Tag)
+		{
+			return TagIndex;
+		}
+	}
+
+	MutableMesh.SetTagCount(TagCount + 1);
+	MutableMesh.SetTag(TagCount, StringCast<ANSICHAR>(*Tag).Get());
+
+	return TagCount;
+}
 
 FString GenerateAnimationInstanceTag(const FString& AnimInstance, int32 SlotIndex)
 {
diff --git a/Engine/Plugins/Experimental/Mutable/Source/CustomizableObjectEditor/Private/MuCOE/GenerateMutableSource/GenerateMutableSource.h b/Engine/Plugins/Experimental/Mutable/Source/CustomizableObjectEditor/Private/MuCOE/GenerateMutableSource/GenerateMutableSource.h
--- a/Engine/Plugins/Experimental/Mutable/Source/CustomizableObjectEditor/Private/MuCOE/GenerateMutableSource/GenerateMutableSource.h
+++ b/Engine/Plugins/Experimental/Mutable/Source/CustomizableObjectEditor/Private/MuCOE/GenerateMutableSource/GenerateMutableSource.h
@@ -787,6 +787,9 @@ UTexture2D* FindReferenceImage(const UEdGraphPin* Pin, FMutableGraphGenerationCo
 mu::NodeMeshApplyPosePtr CreateNodeMeshApplyPose(mu::NodeMeshPtr InputMeshNode, UCustomizableObject * CustomizableObject, TArray<FString> ArrayBoneName, TArray<FTransform> ArrayTransform);
 
 
+/** Adds Tag to MutableMesh uniquely, returns the index were the tag has been inserted or the index where an intance of the tag has been found */
+int32 AddTagToMutableMeshUnique(mu::Mesh& MutableMesh, const FString& Tag);
+
 // Generates the tag for an animation instance
 FString GenerateAnimationInstanceTag(const FString& AnimInstance, int32 SlotIndex);
 
diff --git a/Engine/Plugins/Experimental/Mutable/Source/CustomizableObjectEditor/Private/MuCOE/GenerateMutableSource/GenerateMutableSourceMesh.cpp b/Engine/Plugins/Experimental/Mutable/Source/CustomizableObjectEditor/Private/MuCOE/GenerateMutableSource/GenerateMutableSourceMesh.cpp
--- a/Engine/Plugins/Experimental/Mutable/Source/CustomizableObjectEditor/Private/MuCOE/GenerateMutableSource/GenerateMutableSourceMesh.cpp
+++ b/Engine/Plugins/Experimental/Mutable/Source/CustomizableObjectEditor/Private/MuCOE/GenerateMutableSource/GenerateMutableSourceMesh.cpp
@@ -2616,8 +2616,8 @@ mu::NodeMeshPtr GenerateMutableSourceMesh(const UEdGraphPin * Pin,
 					TSoftObjectPtr<UPhysicsAsset> PhysicsAsset = TypedNodeSkel->SkeletalMesh->GetPhysicsAsset();
 					GenerationContext.PhysicsAssetMap.Add(PhysicsAsset.ToString(), PhysicsAsset);
 					FString PhysicsAssetTag = FString("__PhysicsAsset:") + PhysicsAsset.ToString();
-					MutableMesh->SetTagCount(1);
-					MutableMesh->SetTag(0, TCHAR_TO_ANSI(*PhysicsAssetTag));
+
+					AddTagToMutableMeshUnique(*MutableMesh, PhysicsAssetTag);
 				}
 
 				if (GenerationContext.Options.bClothingEnabled)
@@ -2641,9 +2641,7 @@ mu::NodeMeshPtr GenerateMutableSourceMesh(const UEdGraphPin * Pin,
 						
 						GenerationContext.PhysicsAssetMap.Add(PhysicsAsset.ToString(), ClothingAssetCommon->PhysicsAsset.Get());
 
-						int32 CurrentTagCount = MutableMesh->GetTagCount();
-						MutableMesh->SetTagCount(CurrentTagCount + 1);
-						MutableMesh->SetTag(CurrentTagCount, TCHAR_TO_ANSI(*ClothPhysicsAssetTag));
+						AddTagToMutableMeshUnique(*MutableMesh, ClothPhysicsAssetTag);
 					}
 				}
 
@@ -2668,17 +2666,16 @@ mu::NodeMeshPtr GenerateMutableSourceMesh(const UEdGraphPin * Pin,
 					GenerationContext.AnimBPAssetsMap.Add(TypedNodeSkel->AnimInstance.ToString(), TypedNodeSkel->AnimInstance);
 
 					FString AnimBPAssetTag = GenerateAnimationInstanceTag(TypedNodeSkel->AnimInstance.ToString(), SlotIndex);
-					int32 CurrentTagCount = MutableMesh->GetTagCount();
-					MutableMesh->SetTagCount(CurrentTagCount + 1);
-					MutableMesh->SetTag(CurrentTagCount, TCHAR_TO_ANSI(*AnimBPAssetTag));
+
+					AddTagToMutableMeshUnique(*MutableMesh, AnimBPAssetTag);
 				}
 
 				for (const FGameplayTag& GamePlayTag : TypedNodeSkel->AnimationGameplayTags)
 				{
 					const FString AnimBPTag = GenerateGameplayTag(GamePlayTag.ToString());
 					const int32 CurrentTagCount = MutableMesh->GetTagCount();
-					MutableMesh->SetTagCount(CurrentTagCount + 1);
-					MutableMesh->SetTag(CurrentTagCount, TCHAR_TO_ANSI(*AnimBPTag));
+					
+					AddTagToMutableMeshUnique(*MutableMesh, AnimBPTag);
 				}
 
 				MeshData.bHasVertexColors = TypedNodeSkel->SkeletalMesh->GetHasVertexColors();
diff --git a/Engine/Plugins/Experimental/Mutable/Source/CustomizableObjectEditor/Private/MuCOE/GenerateMutableSource/GenerateMutableSourceTable.cpp b/Engine/Plugins/Experimental/Mutable/Source/CustomizableObjectEditor/Private/MuCOE/GenerateMutableSource/GenerateMutableSourceTable.cpp
--- a/Engine/Plugins/Experimental/Mutable/Source/CustomizableObjectEditor/Private/MuCOE/GenerateMutableSource/GenerateMutableSourceTable.cpp
+++ b/Engine/Plugins/Experimental/Mutable/Source/CustomizableObjectEditor/Private/MuCOE/GenerateMutableSource/GenerateMutableSourceTable.cpp
@@ -226,26 +226,20 @@ void FillTableColumn(const UEdGraphPin* Pin, mu::TablePtr MutableTable, FString
  						if (SkeletalMesh->GetPhysicsAsset() && MutableMesh->GetPhysicsBody() && MutableMesh->GetPhysicsBody()->GetBodyCount())
 						{	
 							TSoftObjectPtr<UPhysicsAsset> PhysicsAsset = SkeletalMesh->GetPhysicsAsset();
-							int32 CurrentTagCount = MutableMesh->GetTagCount();
-							MutableMesh->SetTagCount(CurrentTagCount + 1);
 							GenerationContext.PhysicsAssetMap.Add(PhysicsAsset.ToString(), PhysicsAsset);
 							FString PhysicsAssetTag = FString("__PhysicsAsset:") + PhysicsAsset.ToString();
-							MutableMesh->SetTag(CurrentTagCount, TCHAR_TO_ANSI(*PhysicsAssetTag));						
+
+							AddTagToMutableMeshUnique(*MutableMesh, PhysicsAssetTag);
 						}
 
 						if (!AnimBPAssetTag.IsEmpty())
 						{
-							int32 CurrentTagCount = MutableMesh->GetTagCount();
-							MutableMesh->SetTagCount(CurrentTagCount + 1);
-							MutableMesh->SetTag(CurrentTagCount, TCHAR_TO_ANSI(*AnimBPAssetTag));
+							AddTagToMutableMeshUnique(*MutableMesh, AnimBPAssetTag);
 						}
 
 						for (const FGameplayTag& Tag : GameplayTags)
 						{
-							const FString AnimBPTag = GenerateGameplayTag(Tag.ToString());
-							const int32 CurrentTagCount = MutableMesh->GetTagCount();
-							MutableMesh->SetTagCount(CurrentTagCount + 1);
-							MutableMesh->SetTag(CurrentTagCount, TCHAR_TO_ANSI(*AnimBPTag));
+							AddTagToMutableMeshUnique(*MutableMesh, GenerateGameplayTag(Tag.ToString()));
 						}
 
 						MutableTable->SetCell(CurrentColumn, RowIdx, MutableMesh.get());
