[mutable] fix multiple shared cloth config trigering check at simulation creation.

- Multiple copies of the same shared config were been added to different clothing assets but only one is allowed per skeletalmesh.

#rnx

diff --git a/Engine/Plugins/Experimental/Mutable/Source/CustomizableObject/Private/MuCO/CustomizableObjectInstance.cpp b/Engine/Plugins/Experimental/Mutable/Source/CustomizableObject/Private/MuCO/CustomizableObjectInstance.cpp
--- a/Engine/Plugins/Experimental/Mutable/Source/CustomizableObject/Private/MuCO/CustomizableObjectInstance.cpp
+++ b/Engine/Plugins/Experimental/Mutable/Source/CustomizableObject/Private/MuCO/CustomizableObjectInstance.cpp
@@ -3547,40 +3547,43 @@ void UCustomizableInstancePrivateData::BuildClothingData(const TSharedPtr<FMutab
 			NewSelfCollisionIndices.SetNum(SrcSelfCollisionIndices.Num());
 			NewSelfCollisionIndices.SetNum(TrimAndRemapIndices(NewSelfCollisionIndices, SrcSelfCollisionIndices), false);
 						
-			// TODO: Possible implementation, not checked thoroughly or benchmarked.
-			auto TrimAndRemapTethers = [&IndexMap](FClothTetherData& Dst, const FClothTetherData& Src)
 			{
-				if (!IndexMap.Num())
-				{
-					Dst.Tethers = Src.Tethers;
-					return;
-				}
+				MUTABLE_CPUPROFILER_SCOPE(BuildClothTetherData)
 
-				Dst.Tethers.Reserve(Src.Tethers.Num());
-				for ( const TArray<TTuple<int32, int32, float>>& SrcTetherCluster : Src.Tethers )
+				auto TrimAndRemapTethers = [&IndexMap](FClothTetherData& Dst, const FClothTetherData& Src)
 				{
-					TArray<TTuple<int32, int32, float>>& DstTetherCluster = Dst.Tethers.Emplace_GetRef();
-					DstTetherCluster.Reserve(SrcTetherCluster.Num());
-					for ( const TTuple<int32, int32, float>& Tether : SrcTetherCluster )
+					if (!IndexMap.Num())
 					{
-						const int32 Index0 = IndexMap[Tether.Get<0>()];
-						const int32 Index1 = IndexMap[Tether.Get<1>()];
-						if ((Index0 >= 0) & (Index1 >= 0))
-						{
-							DstTetherCluster.Emplace(Index0, Index1, Tether.Get<2>());
-						}
+						Dst.Tethers = Src.Tethers;
+						return;
 					}
 
-					if (!DstTetherCluster.Num())
+					Dst.Tethers.Reserve(Src.Tethers.Num());
+					for ( const TArray<TTuple<int32, int32, float>>& SrcTetherCluster : Src.Tethers )
 					{
-						constexpr bool bAllowShrinking = false; 
-						Dst.Tethers.RemoveAt( Dst.Tethers.Num() - 1, 1, bAllowShrinking );
+						TArray<TTuple<int32, int32, float>>& DstTetherCluster = Dst.Tethers.Emplace_GetRef();
+						DstTetherCluster.Reserve(SrcTetherCluster.Num());
+						for ( const TTuple<int32, int32, float>& Tether : SrcTetherCluster )
+						{
+							const int32 Index0 = IndexMap[Tether.Get<0>()];
+							const int32 Index1 = IndexMap[Tether.Get<1>()];
+							if ((Index0 >= 0) & (Index1 >= 0))
+							{
+								DstTetherCluster.Emplace(Index0, Index1, Tether.Get<2>());
+							}
+						}
+
+						if (!DstTetherCluster.Num())
+						{
+							constexpr bool bAllowShrinking = false; 
+							Dst.Tethers.RemoveAt( Dst.Tethers.Num() - 1, 1, bAllowShrinking );
+						}
 					}
-				}
-			};
+				};
 
-			TrimAndRemapTethers( NewLodData.PhysicalMeshData.GeodesicTethers, SrcLodData.PhysicalMeshData.GeodesicTethers );
-			TrimAndRemapTethers( NewLodData.PhysicalMeshData.EuclideanTethers, SrcLodData.PhysicalMeshData.EuclideanTethers );
+				TrimAndRemapTethers(NewLodData.PhysicalMeshData.GeodesicTethers, SrcLodData.PhysicalMeshData.GeodesicTethers);
+				TrimAndRemapTethers(NewLodData.PhysicalMeshData.EuclideanTethers, SrcLodData.PhysicalMeshData.EuclideanTethers);
+			}
 		}
 	}
 
@@ -3857,7 +3860,37 @@ void UCustomizableInstancePrivateData::BuildClothingData(const TSharedPtr<FMutab
 	NewClothingAssets.Init(nullptr, ContributingClothingAssetsData.Num());
 
 	{ 
-		MUTABLE_CPUPROFILER_SCOPE(CreateClothingConfigs)
+		MUTABLE_CPUPROFILER_SCOPE(CreateClothingAssets)
+
+		auto CreateNewClothConfigFromData = [](UObject* Outer, const FCustomizableObjectClothConfigData& ConfigData) -> UClothConfigCommon* 
+		{
+			UClass* ClothConfigClass = FindObject<UClass>(nullptr, *ConfigData.ClassPath);
+			if (ClothConfigClass)
+			{
+				UClothConfigCommon* ClothConfig = NewObject<UClothConfigCommon>(Outer, ClothConfigClass);
+				if (ClothConfig)
+				{
+					FMemoryReaderView MemoryReader(ConfigData.ConfigBytes);
+					ClothConfig->Serialize(MemoryReader);
+
+					return ClothConfig;
+				}
+			}
+
+			return nullptr;
+		};
+
+		TArray<TTuple<FName, UClothConfigCommon*>> SharedConfigs;
+		SharedConfigs.Reserve(ClothSharedConfigsData.Num());
+ 
+		for (const FCustomizableObjectClothConfigData& ConfigData : ClothSharedConfigsData)
+		{
+			UClothConfigCommon* ClothConfig = CreateNewClothConfigFromData(SkeletalMesh, ConfigData);
+			if (ClothConfig)
+			{
+				SharedConfigs.Emplace(ConfigData.ConfigName, ClothConfig);
+			}
+		}
 
 		check(NewClothingAssets.Num() == ClothingPhysicsAssets.Num());
 		for (int32 I = 0; I < NewClothingAssetsData.Num(); ++I)
@@ -3881,25 +3914,6 @@ void UCustomizableInstancePrivateData::BuildClothingData(const TSharedPtr<FMutab
 			NewClothingAssets[I]->CalculateReferenceBoneIndex();	
 			NewClothingAssets[I]->PhysicsAsset = ClothingPhysicsAssets[I];
 
-			auto CreateNewClothConfigFromData = [](UObject* Outer, const FCustomizableObjectClothConfigData& ConfigData) 
-			-> UClothConfigCommon* 
-			{
-				UClass* ClothConfigClass = FindObject<UClass>(nullptr, *ConfigData.ClassPath);
-				if (ClothConfigClass)
-				{
-					UClothConfigCommon* ClothConfig = NewObject<UClothConfigCommon>(Outer, ClothConfigClass);
-					if (ClothConfig)
-					{
-						FMemoryReaderView MemoryReader(ConfigData.ConfigBytes);
-						ClothConfig->Serialize(MemoryReader);
-
-						return ClothConfig;
-					}
-				}
-
-				return nullptr;
-			};
-
 			for (const FCustomizableObjectClothConfigData& ConfigData : ContributingClothingAssetsData[I].ConfigsData)
 			{
 				UClothConfigCommon* ClothConfig = CreateNewClothConfigFromData(NewClothingAssets[I], ConfigData);
@@ -3909,16 +3923,12 @@ void UCustomizableInstancePrivateData::BuildClothingData(const TSharedPtr<FMutab
 				}
 			}
 
-			for (const FCustomizableObjectClothConfigData& ConfigData : ClothSharedConfigsData)
-			{	
-				UClothConfigCommon* ClothConfig = CreateNewClothConfigFromData(NewClothingAssets[I], ConfigData);
-				if (ClothConfig)
-				{
-					NewClothingAssets[I]->ClothConfigs.Add(ConfigData.ConfigName, ClothConfig);
-				}
+			for (const TTuple<FName, UClothConfigCommon*>& SharedConfig : SharedConfigs)
+			{
+				NewClothingAssets[I]->ClothConfigs.Add(SharedConfig);
 			}
 
-			SkeletalMesh->GetMeshClothingAssets().AddUnique(NewClothingAssets[I]);	
+			SkeletalMesh->GetMeshClothingAssets().AddUnique(NewClothingAssets[I]);
 		}	
 	}
 
