[mutable] Improve code optimization to prevent unnecessary image format conversions in multilayer projectors.
Improve blending of multilayer operations avoiding unnecessary memory allocations for temp images.
#rnx

diff --git a/Engine/Plugins/Experimental/Mutable/Source/CustomizableObjectEditor/Private/MuCOE/SCustomizableInstanceProperties.cpp b/Engine/Plugins/Experimental/Mutable/Source/CustomizableObjectEditor/Private/MuCOE/SCustomizableInstanceProperties.cpp
--- a/Engine/Plugins/Experimental/Mutable/Source/CustomizableObjectEditor/Private/MuCOE/SCustomizableInstanceProperties.cpp
+++ b/Engine/Plugins/Experimental/Mutable/Source/CustomizableObjectEditor/Private/MuCOE/SCustomizableInstanceProperties.cpp
@@ -2019,9 +2019,8 @@ void SCustomizableInstanceProperties::OnProjectorTextureParameterComboBoxChanged
 					FString ValueName = CustomObject->GetIntParameterAvailableOption(ParamIndexInObject, v);
 					if (ValueName == *Selection)
 					{
-						if (RangeIndex >= 0)
+						if (IntParameters[i].ParameterRangeValueNames.IsValidIndex(RangeIndex))
 						{
-							check(IntParameters[i].ParameterRangeValueNames.IsValidIndex(RangeIndex));
 							IntParameters[i].ParameterRangeValueNames[RangeIndex] = ValueName;
 						}
 						else
diff --git a/Engine/Plugins/Experimental/Mutable/Source/MutableRuntime/Private/MuR/CodeRunner.cpp b/Engine/Plugins/Experimental/Mutable/Source/MutableRuntime/Private/MuR/CodeRunner.cpp
--- a/Engine/Plugins/Experimental/Mutable/Source/MutableRuntime/Private/MuR/CodeRunner.cpp
+++ b/Engine/Plugins/Experimental/Mutable/Source/MutableRuntime/Private/MuR/CodeRunner.cpp
@@ -2504,9 +2504,6 @@ namespace mu
             {
                 MUTABLE_CPUPROFILER_SCOPE(ImageMultiLayer);
 
-                Ptr<const Image> pResult;
-                Ptr<Image> pNew;
-
 				int32 iterations = item.customState;
 
                 // TODO: Reuse base if possible
@@ -2514,16 +2511,13 @@ namespace mu
 
                 if (iterations==0)
                 {
-                    pResult = pBase;
-                }
-                else
-                {
-                    EImageFormat resultFormat = GetUncompressedFormat(pBase->GetFormat());
-                    pNew = new Image( pBase->GetSizeX(), pBase->GetSizeY(),
-                                      pBase->GetLODCount(), resultFormat );
-					pResult = pNew;
+					// Early out
+					GetMemory().SetImage(item, pBase);
+					break;
                 }
 
+				Ptr<const Image> pResult;
+
                 // This shouldn't happen in optimised models, but it could happen in editors, etc.
                 // \todo: raise a performance warning?
                 EImageFormat baseFormat = GetUncompressedFormat(pBase->GetFormat());
@@ -2533,6 +2527,9 @@ namespace mu
                                               pBase.get(), baseFormat );
                 }
 
+				Ptr<Image> pNew = mu::CloneOrTakeOver(pBase.get());
+				pResult = pNew;
+
                 SCHEDULED_OP itemCopy = item;
                 ExecutionIndex index = GetMemory().GetRageIndex( item.executionIndex );
 
@@ -2612,16 +2609,22 @@ namespace mu
 
                         switch (EBlendType(args.blendType))
                         {
-						case EBlendType::BT_NORMAL_COMBINE: ImageNormalCombine(pNew.get(), pBase.get(), pMask.get(), pBlended.get()); break; 
-                        case EBlendType::BT_SOFTLIGHT: ImageSoftLight( pNew.get(), pBase.get(), pMask.get(), pBlended.get(), false ); break;
-                        case EBlendType::BT_HARDLIGHT: ImageHardLight( pNew.get(), pBase.get(), pMask.get(), pBlended.get(), false ); break;
-                        case EBlendType::BT_BURN: ImageBurn( pNew.get(), pBase.get(), pMask.get(), pBlended.get(), false ); break;
-                        case EBlendType::BT_DODGE: ImageDodge( pNew.get(), pBase.get(), pMask.get(), pBlended.get(), false ); break;
-                        case EBlendType::BT_SCREEN: ImageScreen( pNew.get(), pBase.get(), pMask.get(), pBlended.get(), false ); break;
-                        case EBlendType::BT_OVERLAY: ImageOverlay( pNew.get(), pBase.get(), pMask.get(), pBlended.get(), false ); break;
-                        case EBlendType::BT_ALPHA_OVERLAY: ImageAlphaOverlay( pNew.get(), pBase.get(), pMask.get(), pBlended.get(), false ); break;
-                        case EBlendType::BT_MULTIPLY: ImageMultiply( pNew.get(), pBase.get(), pMask.get(), pBlended.get(), false ); break;
-                        case EBlendType::BT_BLEND: ImageBlend( pNew.get(), pBase.get(), pMask.get(), pBlended.get(), false ); break;
+						case EBlendType::BT_NORMAL_COMBINE: 
+							// \todo: optimise with ping pong buffers or blend-in-place
+							ImageNormalCombine(pNew.get(), pBase.get(), pMask.get(), pBlended.get());
+							pResult = pNew;
+							pBase = pNew;
+							pNew = new Image(pBase->GetSizeX(), pBase->GetSizeY(), pBase->GetLODCount(), pBase->GetFormat());
+							break;
+                        case EBlendType::BT_SOFTLIGHT: ImageLayerOnBase<SoftLightChannelMasked, SoftLightChannel, false>( pNew.get(), pMask.get(), pBlended.get(), false ); break;
+                        case EBlendType::BT_HARDLIGHT: ImageLayerOnBase<HardLightChannelMasked, HardLightChannel, false>( pNew.get(), pMask.get(), pBlended.get(), false ); break;
+                        case EBlendType::BT_BURN: ImageLayerOnBase<BurnChannelMasked, BurnChannel, false>( pNew.get(), pMask.get(), pBlended.get(), false ); break;
+                        case EBlendType::BT_DODGE: ImageLayerOnBase<DodgeChannelMasked, DodgeChannel, false>( pNew.get(), pMask.get(), pBlended.get(), false ); break;
+                        case EBlendType::BT_SCREEN: ImageLayerOnBase<ScreenChannelMasked, ScreenChannel, false>( pNew.get(), pMask.get(), pBlended.get(), false ); break;
+                        case EBlendType::BT_OVERLAY: ImageLayerOnBase<OverlayChannelMasked, OverlayChannel, false>( pNew.get(), pMask.get(), pBlended.get(), false ); break;
+                        case EBlendType::BT_ALPHA_OVERLAY: ImageLayerOnBase<AlphaOverlayChannelMasked, AlphaOverlayChannel, false>( pNew.get(), pMask.get(), pBlended.get(), false ); break;
+                        case EBlendType::BT_MULTIPLY: ImageLayerOnBase<MultiplyChannelMasked, MultiplyChannel, false>( pNew.get(), pMask.get(), pBlended.get(), false ); break;
+                        case EBlendType::BT_BLEND: ImageLayerOnBase<BlendChannelMasked, BlendChannel, false>( pNew.get(), pMask.get(), pBlended.get(), false ); break;
                         default: check(false);
                         }
 
@@ -2630,25 +2633,25 @@ namespace mu
                     {
                         switch (EBlendType(args.blendType))
                         {
-						case EBlendType::BT_NORMAL_COMBINE: ImageNormalCombine(pNew.get(), pBase.get(), pBlended.get()); break;
-                        case EBlendType::BT_SOFTLIGHT: ImageSoftLight( pNew.get(), pBase.get(), pBlended.get(), false ); break;
-                        case EBlendType::BT_HARDLIGHT: ImageHardLight( pNew.get(), pBase.get(), pBlended.get(), false ); break;
-                        case EBlendType::BT_BURN: ImageBurn( pNew.get(), pBase.get(), pBlended.get(), false ); break;
-                        case EBlendType::BT_DODGE: ImageDodge( pNew.get(), pBase.get(), pBlended.get(), false ); break;
-                        case EBlendType::BT_SCREEN: ImageScreen( pNew.get(), pBase.get(), pBlended.get(), false ); break;
-                        case EBlendType::BT_OVERLAY: ImageOverlay( pNew.get(), pBase.get(), pBlended.get(), false ); break;
-                        case EBlendType::BT_ALPHA_OVERLAY: ImageAlphaOverlay( pNew.get(), pBase.get(), pBlended.get(), false ); break;
-                        case EBlendType::BT_MULTIPLY: ImageMultiply( pNew.get(), pBase.get(), pBlended.get(), false ); break;
-                        case EBlendType::BT_BLEND: pNew = pBase->Clone(); break;
+						case EBlendType::BT_NORMAL_COMBINE: 
+							// \todo: optimise with ping pong buffers or blend-in-place
+							ImageNormalCombine(pNew.get(), pBase.get(), pBlended.get());
+							pResult = pNew;
+							pBase = pNew;
+							pNew = new Image(pBase->GetSizeX(), pBase->GetSizeY(), pBase->GetLODCount(), pBase->GetFormat());
+							break;
+                        case EBlendType::BT_SOFTLIGHT: ImageLayerOnBase<SoftLightChannel, false>( pNew.get(), pBlended.get(), false ); break;
+                        case EBlendType::BT_HARDLIGHT: ImageLayerOnBase<HardLightChannel, false>( pNew.get(), pBlended.get(), false ); break;
+                        case EBlendType::BT_BURN: ImageLayerOnBase<BurnChannel, false>( pNew.get(), pBlended.get(), false ); break;
+                        case EBlendType::BT_DODGE: ImageLayerOnBase<DodgeChannel, false>( pNew.get(), pBlended.get(), false ); break;
+                        case EBlendType::BT_SCREEN: ImageLayerOnBase<ScreenChannel, false>( pNew.get(),  pBlended.get(), false ); break;
+                        case EBlendType::BT_OVERLAY: ImageLayerOnBase<OverlayChannel, false>( pNew.get(), pBlended.get(), false ); break;
+                        case EBlendType::BT_ALPHA_OVERLAY: ImageLayerOnBase<AlphaOverlayChannel, false>( pNew.get(), pBlended.get(), false ); break;
+                        case EBlendType::BT_MULTIPLY: ImageLayerOnBase<MultiplyChannel, false>( pNew.get(), pBlended.get(), false ); break;
+                        case EBlendType::BT_BLEND: pNew = pBlended->Clone(); break;
                         default: check(false);
                         }
                     }
-
-                    pResult = pNew;
-                    pBase = pNew;
-                    // \todo: optimise with ping pong buffers or blend-in-place
-                    pNew = new Image( pBase->GetSizeX(), pBase->GetSizeY(),
-                                      pBase->GetLODCount(), pBase->GetFormat() );
 				}
 
                 GetMemory().SetImage( item, pResult );
diff --git a/Engine/Plugins/Experimental/Mutable/Source/MutableTools/Private/MuT/AST.cpp b/Engine/Plugins/Experimental/Mutable/Source/MutableTools/Private/MuT/AST.cpp
--- a/Engine/Plugins/Experimental/Mutable/Source/MutableTools/Private/MuT/AST.cpp
+++ b/Engine/Plugins/Experimental/Mutable/Source/MutableTools/Private/MuT/AST.cpp
@@ -11,9 +11,11 @@
 #include "MuR/Platform.h"
 #include "MuT/ASTOpConstantResource.h"
 #include "MuT/ASTOpMeshRemoveMask.h"
+#include "MuT/ASTOpImageMultiLayer.h"
 #include "MuT/ASTOpParameter.h"
 #include "MuT/StreamsPrivate.h"
 #include "MuT/Platform.h"
+#include "MuT/ErrorLogPrivate.h"
 #include "Trace/Detail/Channel.h"
 
 #include <atomic>
@@ -266,40 +268,69 @@ void ASTOp::LogHistogram( ASTOpList& roots )
 {
     (void)roots;
 
-#if 0
-    uint64_t countPerType[(int)OP_TYPE::COUNT];
-    FMemory::Memzero(countPerType,sizeof(countPerType));
+ //   uint64 CountPerType[(int)OP_TYPE::COUNT];
+ //   FMemory::Memzero(CountPerType,sizeof(CountPerType));
+
+ //   size_t count=0;
+
+ //   Traverse_TopRandom_Unique_NonReentrant( roots,
+ //                            [&](const Ptr<ASTOp>& n)
+ //   {
+ //       ++count;
+	//	CountPerType[(int)n->GetOpType()]++;
+ //       return true;
+ //   });
+
+	//TArray< TPair<uint64, OP_TYPE> > Sorted;
+	//Sorted.SetNum((int)OP_TYPE::COUNT);
+ //   for (int i=0; i<(int)OP_TYPE::COUNT; ++i)
+ //   {
+ //       Sorted[i].Value = (OP_TYPE)i;
+ //       Sorted[i].Key = CountPerType[i];
+ //   }
+
+
+	//Sorted.Sort( [](const TPair<uint64, OP_TYPE>& a, const TPair<uint64, OP_TYPE>& b)
+ //   {
+ //       return a.Key >b.Key;
+ //   });
+
+ //   UE_LOG(LogMutableCore,Log, TEXT("Op histogram (%llu ops):"), count);
+ //   for(int i=0; i<8; ++i)
+ //   {
+ //       float p = Sorted[i].Key/float(count)*100.0f;
+ //       int OpType = (int)Sorted[i].Value;
+ //       UE_LOG(LogMutableCore, Log, TEXT("  %5.2f%% : %s"), p, s_opNames[OpType] );
+ //   }
+
+	//// Debug log part of the tree
+	//Traverse_TopDown_Unique( roots,
+	//	[&](const Ptr<ASTOp>& n)
+	//	{
+	//		if (n->GetOpType() == OP_TYPE::IM_MULTILAYER)
+	//		{
+	//			Ptr<const ASTOpImageMultiLayer> Typed = dynamic_cast<const ASTOpImageMultiLayer*>(n.get());
+	//			Ptr<const ASTOp> Base = Typed->base.child();
+	//			Ptr<const ASTOpConstantResource> Constant = dynamic_cast<const ASTOpConstantResource*>(Base.get());
+
+	//			// Log the op
+	//			UE_LOG(LogMutableCore, Log, TEXT("Multilayer at %x:"), n.get());
+	//			UE_LOG(LogMutableCore, Log, TEXT("    base is type %s:"), s_opNames[(int)Base->GetOpType()]);
+	//			if (Constant)
+	//			{
+	//				int Format = int( ((mu::Image*)Constant->GetValue().get())->GetFormat() );
+	//				UE_LOG(LogMutableCore, Log, TEXT("    constant image format is %d:"), Format);
+	//			}
+	//			else
+	//			{
+	//				FGetImageDescContext Context;
+	//				FImageDesc Desc = Base->GetImageDesc(true, &Context);
+	//				UE_LOG(LogMutableCore, Log, TEXT("    estimated image format is %d:"), int(Desc.m_format) );
+	//			}
+	//		}
+	//		return true;
+	//	});
 
-    size_t count=0;
-
-    Traverse_TopRandom_Unique_NonReentrant( roots,
-                             [&](const Ptr<ASTOp>& n)
-    {
-        ++count;
-        countPerType[(int)n->GetOpType()]++;
-        return true;
-    });
-
-    vector< pair<uint64_t,OP_TYPE> > sorted((int)OP_TYPE::COUNT);
-    for (int i=0; i<(int)OP_TYPE::COUNT; ++i)
-    {
-        sorted[i].second = (OP_TYPE)i;
-        sorted[i].first = countPerType[i];
-    }
-
-    std::sort(sorted.begin(),sorted.end(), []( const pair<uint64_t,OP_TYPE>& a, const pair<uint64_t,OP_TYPE>& b )
-    {
-        return a.first>b.first;
-    });
-
-    UE_LOG(LogMutableCore,Log, TEXT("Op histogram (%llu ops):"), count);
-    for(int i=0; i<8; ++i)
-    {
-        float p = sorted[i].first/float(count)*100.0f;
-        int op = (int)sorted[i].second;
-        UE_LOG(LogMutableCore,Log, TEXT("  %5.2f%% : %s"), p, s_opNames[op] );
-    }
-#endif
 }
 
 
@@ -864,7 +895,7 @@ void ASTOp::Replace( const Ptr<ASTOp>& node, const Ptr<ASTOp>& other )
 }
 
 
-FImageDesc ASTOp::GetImageDesc( bool, GetImageDescContext* )
+FImageDesc ASTOp::GetImageDesc( bool, FGetImageDescContext* ) const
 {
     check(false);
     return FImageDesc();
@@ -1063,12 +1094,12 @@ uint64 ASTOpFixed::Hash() const
 
 
 //!
-FImageDesc ASTOpFixed::GetImageDesc( bool returnBestOption, GetImageDescContext* context )
+FImageDesc ASTOpFixed::GetImageDesc( bool returnBestOption, FGetImageDescContext* context ) const
 {
     FImageDesc res;
 
     // Local context in case it is necessary
-    GetImageDescContext localContext;
+    FGetImageDescContext localContext;
     if (!context)
     {
       context = &localContext;
diff --git a/Engine/Plugins/Experimental/Mutable/Source/MutableTools/Private/MuT/AST.h b/Engine/Plugins/Experimental/Mutable/Source/MutableTools/Private/MuT/AST.h
--- a/Engine/Plugins/Experimental/Mutable/Source/MutableTools/Private/MuT/AST.h
+++ b/Engine/Plugins/Experimental/Mutable/Source/MutableTools/Private/MuT/AST.h
@@ -661,14 +661,14 @@ namespace mu
 
         //! This class contains the support data to accelerate the GetImageDesc recursive function.
         //! If none is provided in the call, one will be created at that level and used from there on.
-        class GetImageDescContext
+        class FGetImageDescContext
         {
         public:
             TMap<const ASTOp*, FImageDesc> m_results;
         };
 
         //!
-        virtual FImageDesc GetImageDesc( bool returnBestOption=false, GetImageDescContext* context=nullptr );
+        virtual FImageDesc GetImageDesc( bool returnBestOption=false, FGetImageDescContext* context=nullptr ) const;
 
         //! Optional cache struct to use int he method below.
         using BLOCK_LAYOUT_SIZE_CACHE=TMap< const TPair<ASTOp*,int>, TPair<int,int>>;
@@ -969,7 +969,7 @@ namespace mu
         void Link( FProgram& program, const FLinkerOptions* Options) override;
         bool IsEqual(const ASTOp& otherUntyped) const override;
 		uint64 Hash() const override;
-		FImageDesc GetImageDesc( bool returnBestOption, GetImageDescContext* context ) override;
+		FImageDesc GetImageDesc( bool returnBestOption, FGetImageDescContext* context ) const override;
         void GetBlockLayoutSize( int blockIndex, int* pBlockX, int* pBlockY,
                                  BLOCK_LAYOUT_SIZE_CACHE* cache ) override;
         void GetLayoutBlockSize( int* pBlockX, int* pBlockY ) override;
@@ -1019,7 +1019,7 @@ namespace mu
     private:
 
         //
-        inline FImageDesc GetImageDesc( OP::ADDRESS at, bool returnBestOption=false, GetImageDescContext* context=nullptr )
+        inline FImageDesc GetImageDesc( OP::ADDRESS at, bool returnBestOption=false, FGetImageDescContext* context=nullptr ) const
         {
 			FImageDesc res;
 
diff --git a/Engine/Plugins/Experimental/Mutable/Source/MutableTools/Private/MuT/ASTOpConditional.cpp b/Engine/Plugins/Experimental/Mutable/Source/MutableTools/Private/MuT/ASTOpConditional.cpp
--- a/Engine/Plugins/Experimental/Mutable/Source/MutableTools/Private/MuT/ASTOpConditional.cpp
+++ b/Engine/Plugins/Experimental/Mutable/Source/MutableTools/Private/MuT/ASTOpConditional.cpp
@@ -127,12 +127,12 @@ namespace mu
 
 
 	//-------------------------------------------------------------------------------------------------
-	FImageDesc ASTOpConditional::GetImageDesc(bool returnBestOption, class GetImageDescContext* context)
+	FImageDesc ASTOpConditional::GetImageDesc(bool returnBestOption, class FGetImageDescContext* context) const
 	{
 		FImageDesc res;
 
 		// Local context in case it is necessary
-		GetImageDescContext localContext;
+		FGetImageDescContext localContext;
 		if (!context)
 		{
 			context = &localContext;
diff --git a/Engine/Plugins/Experimental/Mutable/Source/MutableTools/Private/MuT/ASTOpConditional.h b/Engine/Plugins/Experimental/Mutable/Source/MutableTools/Private/MuT/ASTOpConditional.h
--- a/Engine/Plugins/Experimental/Mutable/Source/MutableTools/Private/MuT/ASTOpConditional.h
+++ b/Engine/Plugins/Experimental/Mutable/Source/MutableTools/Private/MuT/ASTOpConditional.h
@@ -44,7 +44,7 @@ struct FProgram;
 		void Assert() override;
 		void ForEachChild(const TFunctionRef<void(ASTChild&)> f) override;
 		void Link(FProgram& program, const FLinkerOptions* Options) override;
-		FImageDesc GetImageDesc(bool returnBestOption, class GetImageDescContext* context) override;
+		FImageDesc GetImageDesc(bool returnBestOption, class FGetImageDescContext* context) const override;
 		void GetLayoutBlockSize(int* pBlockX, int* pBlockY) override;
 		void GetBlockLayoutSize(int blockIndex, int* pBlockX, int* pBlockY,
 			BLOCK_LAYOUT_SIZE_CACHE* cache) override;
diff --git a/Engine/Plugins/Experimental/Mutable/Source/MutableTools/Private/MuT/ASTOpConstantResource.cpp b/Engine/Plugins/Experimental/Mutable/Source/MutableTools/Private/MuT/ASTOpConstantResource.cpp
--- a/Engine/Plugins/Experimental/Mutable/Source/MutableTools/Private/MuT/ASTOpConstantResource.cpp
+++ b/Engine/Plugins/Experimental/Mutable/Source/MutableTools/Private/MuT/ASTOpConstantResource.cpp
@@ -177,7 +177,7 @@ namespace mu
 
 
 	//-------------------------------------------------------------------------------------------------
-	FImageDesc ASTOpConstantResource::GetImageDesc(bool, class GetImageDescContext*)
+	FImageDesc ASTOpConstantResource::GetImageDesc(bool, class FGetImageDescContext*) const
 	{
 		FImageDesc res;
 
diff --git a/Engine/Plugins/Experimental/Mutable/Source/MutableTools/Private/MuT/ASTOpConstantResource.h b/Engine/Plugins/Experimental/Mutable/Source/MutableTools/Private/MuT/ASTOpConstantResource.h
--- a/Engine/Plugins/Experimental/Mutable/Source/MutableTools/Private/MuT/ASTOpConstantResource.h
+++ b/Engine/Plugins/Experimental/Mutable/Source/MutableTools/Private/MuT/ASTOpConstantResource.h
@@ -65,7 +65,7 @@ template <class SCALAR> class vec4;
 		Ptr<ASTOp> Clone(MapChildFuncRef mapChild) const override;
 		uint64 Hash() const override;
 		void Link(FProgram& program, const FLinkerOptions*) override;
-		FImageDesc GetImageDesc(bool, class GetImageDescContext*) override;
+		FImageDesc GetImageDesc(bool, class FGetImageDescContext*) const override;
 		void GetBlockLayoutSize(int blockIndex, int* pBlockX, int* pBlockY,
 			BLOCK_LAYOUT_SIZE_CACHE* cache) override;
 		void GetLayoutBlockSize(int* pBlockX, int* pBlockY) override;
diff --git a/Engine/Plugins/Experimental/Mutable/Source/MutableTools/Private/MuT/ASTOpImageCompose.cpp b/Engine/Plugins/Experimental/Mutable/Source/MutableTools/Private/MuT/ASTOpImageCompose.cpp
--- a/Engine/Plugins/Experimental/Mutable/Source/MutableTools/Private/MuT/ASTOpImageCompose.cpp
+++ b/Engine/Plugins/Experimental/Mutable/Source/MutableTools/Private/MuT/ASTOpImageCompose.cpp
@@ -119,12 +119,12 @@ void ASTOpImageCompose::Link( FProgram& program, const FLinkerOptions*)
 }
 
 
-FImageDesc ASTOpImageCompose::GetImageDesc( bool returnBestOption, GetImageDescContext* context )
+FImageDesc ASTOpImageCompose::GetImageDesc( bool returnBestOption, FGetImageDescContext* context ) const
 {
     FImageDesc res;
 
     // Local context in case it is necessary
-    GetImageDescContext localContext;
+    FGetImageDescContext localContext;
     if (!context)
     {
       context = &localContext;
diff --git a/Engine/Plugins/Experimental/Mutable/Source/MutableTools/Private/MuT/ASTOpImageCompose.h b/Engine/Plugins/Experimental/Mutable/Source/MutableTools/Private/MuT/ASTOpImageCompose.h
--- a/Engine/Plugins/Experimental/Mutable/Source/MutableTools/Private/MuT/ASTOpImageCompose.h
+++ b/Engine/Plugins/Experimental/Mutable/Source/MutableTools/Private/MuT/ASTOpImageCompose.h
@@ -38,7 +38,7 @@ template <class SCALAR> class vec4;
 		void Link(FProgram& program, const FLinkerOptions* Options) override;
 		//TODO: Ptr<ASTOp> OptimiseSink(const MODEL_OPTIMIZATION_OPTIONS& options, OPTIMIZE_SINK_CONTEXT& context) const override;
 		Ptr<ASTOp> OptimiseSemantic(const MODEL_OPTIMIZATION_OPTIONS& options) const;
-		FImageDesc GetImageDesc(bool returnBestOption, GetImageDescContext* context) override;
+		FImageDesc GetImageDesc(bool returnBestOption, FGetImageDescContext* context) const override;
 		void GetLayoutBlockSize(int* pBlockX, int* pBlockY) override;
 		bool IsImagePlainConstant(FVector4f& colour) const override;
 		Ptr<ImageSizeExpression> GetImageSizeExpression() const override;
diff --git a/Engine/Plugins/Experimental/Mutable/Source/MutableTools/Private/MuT/ASTOpImageMipmap.cpp b/Engine/Plugins/Experimental/Mutable/Source/MutableTools/Private/MuT/ASTOpImageMipmap.cpp
--- a/Engine/Plugins/Experimental/Mutable/Source/MutableTools/Private/MuT/ASTOpImageMipmap.cpp
+++ b/Engine/Plugins/Experimental/Mutable/Source/MutableTools/Private/MuT/ASTOpImageMipmap.cpp
@@ -163,12 +163,12 @@ mu::Ptr<ASTOp> ASTOpImageMipmap::OptimiseSink(const MODEL_OPTIMIZATION_OPTIONS&
 
 
 //!
-FImageDesc ASTOpImageMipmap::GetImageDesc(bool returnBestOption, GetImageDescContext* context)
+FImageDesc ASTOpImageMipmap::GetImageDesc(bool returnBestOption, FGetImageDescContext* context) const
 {
 	FImageDesc res;
 
 	// Local context in case it is necessary
-	GetImageDescContext localContext;
+	FGetImageDescContext localContext;
 	if (!context)
 	{
 		context = &localContext;
@@ -392,7 +392,7 @@ mu::Ptr<ASTOp> Sink_ImageMipmapAST::Visit(mu::Ptr<ASTOp> at, const ASTOpImageMip
 
 			const ASTOpImagePatch* typedSource = dynamic_cast<const ASTOpImagePatch*>(at.get());
 			auto rectOp = typedSource->patch.child();
-			ASTOp::GetImageDescContext context;
+			ASTOp::FGetImageDescContext context;
 			auto patchDesc = rectOp->GetImageDesc(false, &context);
 
 			uint8_t blockLevels = 0;
diff --git a/Engine/Plugins/Experimental/Mutable/Source/MutableTools/Private/MuT/ASTOpImageMipmap.h b/Engine/Plugins/Experimental/Mutable/Source/MutableTools/Private/MuT/ASTOpImageMipmap.h
--- a/Engine/Plugins/Experimental/Mutable/Source/MutableTools/Private/MuT/ASTOpImageMipmap.h
+++ b/Engine/Plugins/Experimental/Mutable/Source/MutableTools/Private/MuT/ASTOpImageMipmap.h
@@ -48,7 +48,7 @@ template <class SCALAR> class vec4;
 		void ForEachChild(const TFunctionRef<void(ASTChild&)>) override;
 		void Link(FProgram& program, const FLinkerOptions* Options) override;
 		Ptr<ASTOp> OptimiseSink(const MODEL_OPTIMIZATION_OPTIONS& options, OPTIMIZE_SINK_CONTEXT& context) const override;
-		FImageDesc GetImageDesc(bool returnBestOption, GetImageDescContext* context) override;
+		FImageDesc GetImageDesc(bool returnBestOption, FGetImageDescContext* context) const override;
 		void GetLayoutBlockSize(int* pBlockX, int* pBlockY) override;
 		bool IsImagePlainConstant(FVector4f& colour) const override;
 		Ptr<ImageSizeExpression> GetImageSizeExpression() const override;
diff --git a/Engine/Plugins/Experimental/Mutable/Source/MutableTools/Private/MuT/ASTOpImageMultiLayer.cpp b/Engine/Plugins/Experimental/Mutable/Source/MutableTools/Private/MuT/ASTOpImageMultiLayer.cpp
--- a/Engine/Plugins/Experimental/Mutable/Source/MutableTools/Private/MuT/ASTOpImageMultiLayer.cpp
+++ b/Engine/Plugins/Experimental/Mutable/Source/MutableTools/Private/MuT/ASTOpImageMultiLayer.cpp
@@ -115,12 +115,12 @@ namespace mu
 
 
 	//-------------------------------------------------------------------------------------------------
-	FImageDesc ASTOpImageMultiLayer::GetImageDesc(bool returnBestOption, GetImageDescContext* context)
+	FImageDesc ASTOpImageMultiLayer::GetImageDesc(bool returnBestOption, FGetImageDescContext* context) const 
 	{
 		FImageDesc res;
 
 		// Local context in case it is necessary
-		GetImageDescContext localContext;
+		FGetImageDescContext localContext;
 		if (!context)
 		{
 			context = &localContext;
diff --git a/Engine/Plugins/Experimental/Mutable/Source/MutableTools/Private/MuT/ASTOpImageMultiLayer.h b/Engine/Plugins/Experimental/Mutable/Source/MutableTools/Private/MuT/ASTOpImageMultiLayer.h
--- a/Engine/Plugins/Experimental/Mutable/Source/MutableTools/Private/MuT/ASTOpImageMultiLayer.h
+++ b/Engine/Plugins/Experimental/Mutable/Source/MutableTools/Private/MuT/ASTOpImageMultiLayer.h
@@ -39,7 +39,7 @@ struct FProgram;
 		Ptr<ASTOp> Clone(MapChildFuncRef mapChild) const override;
 		void ForEachChild(const TFunctionRef<void(ASTChild&)>) override;
 		void Link(FProgram& program, const FLinkerOptions* Options) override;
-		FImageDesc GetImageDesc(bool returnBestOption, GetImageDescContext* context) override;
+		FImageDesc GetImageDesc(bool returnBestOption, FGetImageDescContext* context) const override;
 		void GetLayoutBlockSize(int* pBlockX, int* pBlockY) override;
 		Ptr<ImageSizeExpression> GetImageSizeExpression() const override;
 	};
diff --git a/Engine/Plugins/Experimental/Mutable/Source/MutableTools/Private/MuT/ASTOpImageNormalComposite.cpp b/Engine/Plugins/Experimental/Mutable/Source/MutableTools/Private/MuT/ASTOpImageNormalComposite.cpp
--- a/Engine/Plugins/Experimental/Mutable/Source/MutableTools/Private/MuT/ASTOpImageNormalComposite.cpp
+++ b/Engine/Plugins/Experimental/Mutable/Source/MutableTools/Private/MuT/ASTOpImageNormalComposite.cpp
@@ -101,12 +101,12 @@ namespace mu
 	}
 
 
-	FImageDesc ASTOpImageNormalComposite::GetImageDesc(bool returnBestOption, GetImageDescContext* context)
+	FImageDesc ASTOpImageNormalComposite::GetImageDesc(bool returnBestOption, FGetImageDescContext* context) const
 	{
 		FImageDesc res;
 
 		// Local context in case it is necessary
-		GetImageDescContext localContext;
+		FGetImageDescContext localContext;
 		if (!context)
 		{
 			context = &localContext;
diff --git a/Engine/Plugins/Experimental/Mutable/Source/MutableTools/Private/MuT/ASTOpImageNormalComposite.h b/Engine/Plugins/Experimental/Mutable/Source/MutableTools/Private/MuT/ASTOpImageNormalComposite.h
--- a/Engine/Plugins/Experimental/Mutable/Source/MutableTools/Private/MuT/ASTOpImageNormalComposite.h
+++ b/Engine/Plugins/Experimental/Mutable/Source/MutableTools/Private/MuT/ASTOpImageNormalComposite.h
@@ -33,7 +33,7 @@ struct FProgram;
 		Ptr<ASTOp> Clone(MapChildFuncRef mapChild) const override;
 		void ForEachChild(const TFunctionRef<void(ASTChild&)>) override;
 		void Link(FProgram& program, const FLinkerOptions* Options) override;
-		FImageDesc GetImageDesc(bool returnBestOption, GetImageDescContext* context) override;
+		FImageDesc GetImageDesc(bool returnBestOption, FGetImageDescContext* context) const override;
 		void GetLayoutBlockSize(int* pBlockX, int* pBlockY) override;
 		Ptr<ImageSizeExpression> GetImageSizeExpression() const override;
 	};
diff --git a/Engine/Plugins/Experimental/Mutable/Source/MutableTools/Private/MuT/ASTOpImagePatch.cpp b/Engine/Plugins/Experimental/Mutable/Source/MutableTools/Private/MuT/ASTOpImagePatch.cpp
--- a/Engine/Plugins/Experimental/Mutable/Source/MutableTools/Private/MuT/ASTOpImagePatch.cpp
+++ b/Engine/Plugins/Experimental/Mutable/Source/MutableTools/Private/MuT/ASTOpImagePatch.cpp
@@ -94,12 +94,12 @@ namespace mu
 
 
 	//-------------------------------------------------------------------------------------------------
-	FImageDesc ASTOpImagePatch::GetImageDesc(bool returnBestOption, GetImageDescContext* context)
+	FImageDesc ASTOpImagePatch::GetImageDesc(bool returnBestOption, FGetImageDescContext* context) const
 	{
 		FImageDesc res;
 
 		// Local context in case it is necessary
-		GetImageDescContext localContext;
+		FGetImageDescContext localContext;
 		if (!context)
 		{
 			context = &localContext;
diff --git a/Engine/Plugins/Experimental/Mutable/Source/MutableTools/Private/MuT/ASTOpImagePatch.h b/Engine/Plugins/Experimental/Mutable/Source/MutableTools/Private/MuT/ASTOpImagePatch.h
--- a/Engine/Plugins/Experimental/Mutable/Source/MutableTools/Private/MuT/ASTOpImagePatch.h
+++ b/Engine/Plugins/Experimental/Mutable/Source/MutableTools/Private/MuT/ASTOpImagePatch.h
@@ -34,7 +34,7 @@ struct FProgram;
 		Ptr<ASTOp> Clone(MapChildFuncRef mapChild) const override;
 		void ForEachChild(const TFunctionRef<void(ASTChild&)>) override;
 		void Link(FProgram& program, const FLinkerOptions* Options) override;
-		FImageDesc GetImageDesc(bool returnBestOption, GetImageDescContext* context) override;
+		FImageDesc GetImageDesc(bool returnBestOption, FGetImageDescContext* context) const override;
 		Ptr<ImageSizeExpression> GetImageSizeExpression() const override;
 		//TODO: void GetLayoutBlockSize(int* pBlockX, int* pBlockY) override;
 		//TODO: bool IsImagePlainConstant(vec4<float>& colour) const override;
diff --git a/Engine/Plugins/Experimental/Mutable/Source/MutableTools/Private/MuT/ASTOpImagePixelFormat.cpp b/Engine/Plugins/Experimental/Mutable/Source/MutableTools/Private/MuT/ASTOpImagePixelFormat.cpp
--- a/Engine/Plugins/Experimental/Mutable/Source/MutableTools/Private/MuT/ASTOpImagePixelFormat.cpp
+++ b/Engine/Plugins/Experimental/Mutable/Source/MutableTools/Private/MuT/ASTOpImagePixelFormat.cpp
@@ -190,12 +190,12 @@ mu::Ptr<ASTOp> ASTOpImagePixelFormat::OptimiseSink(const MODEL_OPTIMIZATION_OPTI
 
 
 //!
-FImageDesc ASTOpImagePixelFormat::GetImageDesc(bool returnBestOption, GetImageDescContext* context)
+FImageDesc ASTOpImagePixelFormat::GetImageDesc(bool returnBestOption, FGetImageDescContext* context) const
 {
 	FImageDesc res;
 
 	// Local context in case it is necessary
-	GetImageDescContext localContext;
+	FGetImageDescContext localContext;
 	if (!context)
 	{
 		context = &localContext;
diff --git a/Engine/Plugins/Experimental/Mutable/Source/MutableTools/Private/MuT/ASTOpImagePixelFormat.h b/Engine/Plugins/Experimental/Mutable/Source/MutableTools/Private/MuT/ASTOpImagePixelFormat.h
--- a/Engine/Plugins/Experimental/Mutable/Source/MutableTools/Private/MuT/ASTOpImagePixelFormat.h
+++ b/Engine/Plugins/Experimental/Mutable/Source/MutableTools/Private/MuT/ASTOpImagePixelFormat.h
@@ -35,7 +35,7 @@ template <class SCALAR> class vec4;
 		void ForEachChild(const TFunctionRef<void(ASTChild&)>) override;
 		void Link(FProgram& program, const FLinkerOptions* Options) override;
 		Ptr<ASTOp> OptimiseSink(const MODEL_OPTIMIZATION_OPTIONS& options, OPTIMIZE_SINK_CONTEXT& context) const override;
-		FImageDesc GetImageDesc(bool returnBestOption, GetImageDescContext* context) override;
+		FImageDesc GetImageDesc(bool returnBestOption, FGetImageDescContext* context) const override;
 		void GetLayoutBlockSize(int* pBlockX, int* pBlockY) override;
 		bool IsImagePlainConstant(FVector4f& colour) const override;
 		Ptr<ImageSizeExpression> GetImageSizeExpression() const override;
diff --git a/Engine/Plugins/Experimental/Mutable/Source/MutableTools/Private/MuT/ASTOpImageTransform.cpp b/Engine/Plugins/Experimental/Mutable/Source/MutableTools/Private/MuT/ASTOpImageTransform.cpp
--- a/Engine/Plugins/Experimental/Mutable/Source/MutableTools/Private/MuT/ASTOpImageTransform.cpp
+++ b/Engine/Plugins/Experimental/Mutable/Source/MutableTools/Private/MuT/ASTOpImageTransform.cpp
@@ -56,7 +56,7 @@ namespace mu
 	//-------------------------------------------------------------------------------------------------
 	uint64 ASTOpImageTransform::Hash() const
 	{
-		uint64 res = std::hash<OP_TYPE>()(OP_TYPE::IM_MULTILAYER);
+		uint64 res = std::hash<OP_TYPE>()(OP_TYPE::IM_TRANSFORM);
 		hash_combine(res, base.child().get());
 		hash_combine(res, offsetX.child().get());
 		hash_combine(res, offsetY.child().get());
@@ -116,12 +116,12 @@ namespace mu
 	}
 
 	//-------------------------------------------------------------------------------------------------
-	FImageDesc ASTOpImageTransform::GetImageDesc(bool returnBestOption, GetImageDescContext* context)
+	FImageDesc ASTOpImageTransform::GetImageDesc(bool returnBestOption, FGetImageDescContext* context) const
 	{
 		FImageDesc res;
 
 		// Local context in case it is necessary
-		GetImageDescContext localContext;
+		FGetImageDescContext localContext;
 		if (!context)
 		{
 			context = &localContext;
diff --git a/Engine/Plugins/Experimental/Mutable/Source/MutableTools/Private/MuT/ASTOpImageTransform.h b/Engine/Plugins/Experimental/Mutable/Source/MutableTools/Private/MuT/ASTOpImageTransform.h
--- a/Engine/Plugins/Experimental/Mutable/Source/MutableTools/Private/MuT/ASTOpImageTransform.h
+++ b/Engine/Plugins/Experimental/Mutable/Source/MutableTools/Private/MuT/ASTOpImageTransform.h
@@ -40,7 +40,7 @@ struct FProgram;
 		Ptr<ASTOp> Clone(MapChildFuncRef mapChild) const override;
 		void ForEachChild(const TFunctionRef<void(ASTChild&)>) override;
 		void Link(FProgram& program, const FLinkerOptions* Options) override;
-		FImageDesc GetImageDesc(bool returnBestOption, GetImageDescContext* context) override;
+		FImageDesc GetImageDesc(bool returnBestOption, FGetImageDescContext* context) const override;
 		void GetLayoutBlockSize(int* pBlockX, int* pBlockY) override;
 		Ptr<ImageSizeExpression> GetImageSizeExpression() const override;
 	};
diff --git a/Engine/Plugins/Experimental/Mutable/Source/MutableTools/Private/MuT/ASTOpParameter.cpp b/Engine/Plugins/Experimental/Mutable/Source/MutableTools/Private/MuT/ASTOpParameter.cpp
--- a/Engine/Plugins/Experimental/Mutable/Source/MutableTools/Private/MuT/ASTOpParameter.cpp
+++ b/Engine/Plugins/Experimental/Mutable/Source/MutableTools/Private/MuT/ASTOpParameter.cpp
@@ -155,7 +155,7 @@ namespace mu
 	}
 
 
-	FImageDesc ASTOpParameter::GetImageDesc(bool, GetImageDescContext*)
+	FImageDesc ASTOpParameter::GetImageDesc(bool, FGetImageDescContext*) const
 	{
 		check(type == OP_TYPE::IM_PARAMETER);
 		return FImageDesc();
diff --git a/Engine/Plugins/Experimental/Mutable/Source/MutableTools/Private/MuT/ASTOpParameter.h b/Engine/Plugins/Experimental/Mutable/Source/MutableTools/Private/MuT/ASTOpParameter.h
--- a/Engine/Plugins/Experimental/Mutable/Source/MutableTools/Private/MuT/ASTOpParameter.h
+++ b/Engine/Plugins/Experimental/Mutable/Source/MutableTools/Private/MuT/ASTOpParameter.h
@@ -51,7 +51,7 @@ namespace mu
 		void Link(FProgram& program, const FLinkerOptions*) override;
 		int EvaluateInt(ASTOpList& facts, bool& unknown) const override;
 		FBoolEvalResult EvaluateBool(ASTOpList& /*facts*/, FEvaluateBoolCache* = nullptr) const override;
-		FImageDesc GetImageDesc(bool, GetImageDescContext*) override;
+		FImageDesc GetImageDesc(bool, FGetImageDescContext*) const override;
 
 	};
 
diff --git a/Engine/Plugins/Experimental/Mutable/Source/MutableTools/Private/MuT/ASTOpSwitch.cpp b/Engine/Plugins/Experimental/Mutable/Source/MutableTools/Private/MuT/ASTOpSwitch.cpp
--- a/Engine/Plugins/Experimental/Mutable/Source/MutableTools/Private/MuT/ASTOpSwitch.cpp
+++ b/Engine/Plugins/Experimental/Mutable/Source/MutableTools/Private/MuT/ASTOpSwitch.cpp
@@ -193,12 +193,12 @@ namespace mu
 	}
 
 	//-------------------------------------------------------------------------------------------------
-	FImageDesc ASTOpSwitch::GetImageDesc(bool returnBestOption, class GetImageDescContext* context)
+	FImageDesc ASTOpSwitch::GetImageDesc(bool returnBestOption, class FGetImageDescContext* context) const
 	{
 		FImageDesc res;
 
 		// Local context in case it is necessary
-		GetImageDescContext localContext;
+		FGetImageDescContext localContext;
 		if (!context)
 		{
 			context = &localContext;
diff --git a/Engine/Plugins/Experimental/Mutable/Source/MutableTools/Private/MuT/ASTOpSwitch.h b/Engine/Plugins/Experimental/Mutable/Source/MutableTools/Private/MuT/ASTOpSwitch.h
--- a/Engine/Plugins/Experimental/Mutable/Source/MutableTools/Private/MuT/ASTOpSwitch.h
+++ b/Engine/Plugins/Experimental/Mutable/Source/MutableTools/Private/MuT/ASTOpSwitch.h
@@ -67,10 +67,8 @@ template <class SCALAR> class vec4;
 		Ptr<ASTOp> Clone(MapChildFuncRef mapChild) const override;
 		void Assert() override;
 		void Link(FProgram& program, const FLinkerOptions* Options) override;
-		FImageDesc GetImageDesc(bool returnBestOption,
-			class GetImageDescContext* context) override;
-		void GetBlockLayoutSize(int blockIndex, int* pBlockX, int* pBlockY,
-			BLOCK_LAYOUT_SIZE_CACHE* cache) override;
+		FImageDesc GetImageDesc(bool returnBestOption, class FGetImageDescContext* context) const  override;
+		void GetBlockLayoutSize(int blockIndex, int* pBlockX, int* pBlockY, BLOCK_LAYOUT_SIZE_CACHE* cache) override;
 		void GetLayoutBlockSize(int* pBlockX, int* pBlockY) override;
 		Ptr<ASTOp> OptimiseSemantic(const MODEL_OPTIMIZATION_OPTIONS&) const override;
 		bool GetNonBlackRect(FImageRect& maskUsage) const override;
diff --git a/Engine/Plugins/Experimental/Mutable/Source/MutableTools/Private/MuT/CodeGenerator.h b/Engine/Plugins/Experimental/Mutable/Source/MutableTools/Private/MuT/CodeGenerator.h
--- a/Engine/Plugins/Experimental/Mutable/Source/MutableTools/Private/MuT/CodeGenerator.h
+++ b/Engine/Plugins/Experimental/Mutable/Source/MutableTools/Private/MuT/CodeGenerator.h
@@ -684,7 +684,7 @@ namespace mu
     //---------------------------------------------------------------------------------------------
     extern FImageDesc GetImageDesc( const FProgram& program, OP::ADDRESS at,
                                     bool returnBestOption = false,
-                                    class GetImageDescContext* context=nullptr );
+                                    class FGetImageDescContext* context=nullptr );
 
     //!
     extern void PartialOptimise( Ptr<ASTOp>& op, const CompilerOptions* options );
diff --git a/Engine/Plugins/Experimental/Mutable/Source/MutableTools/Private/MuT/CodeGenerator_Image.cpp b/Engine/Plugins/Experimental/Mutable/Source/MutableTools/Private/MuT/CodeGenerator_Image.cpp
--- a/Engine/Plugins/Experimental/Mutable/Source/MutableTools/Private/MuT/CodeGenerator_Image.cpp
+++ b/Engine/Plugins/Experimental/Mutable/Source/MutableTools/Private/MuT/CodeGenerator_Image.cpp
@@ -1276,7 +1276,7 @@ namespace mu
         Ptr<ASTOpFixed> op = new ASTOpFixed();
         op->op.type = OP_TYPE::IM_PLAINCOLOUR;
         op->SetChild( op->op.args.ImagePlainColour.colour, base);
-        op->op.args.ImagePlainColour.format = EImageFormat::IF_RGB_UBYTE;
+		op->op.args.ImagePlainColour.format = node.Format;
         op->op.args.ImagePlainColour.size[0] = uint16(node.m_sizeX);
         op->op.args.ImagePlainColour.size[1] = uint16(node.m_sizeY);
 
@@ -2049,7 +2049,7 @@ namespace mu
     //---------------------------------------------------------------------------------------------
     //! This class contains the support data to accelerate the GetImageDesc recursive function.
     //! _If none is provided in the call, one will be created at that level and used from there on.
-    class GetImageDescContext
+    class FGetImageDescContext
     {
     public:
         vector<bool> m_visited;
diff --git a/Engine/Plugins/Experimental/Mutable/Source/MutableTools/Private/MuT/CodeOptimiser.cpp b/Engine/Plugins/Experimental/Mutable/Source/MutableTools/Private/MuT/CodeOptimiser.cpp
--- a/Engine/Plugins/Experimental/Mutable/Source/MutableTools/Private/MuT/CodeOptimiser.cpp
+++ b/Engine/Plugins/Experimental/Mutable/Source/MutableTools/Private/MuT/CodeOptimiser.cpp
@@ -1197,8 +1197,7 @@ namespace mu
 		{
 			// We use 4 times the count because at the time we moved to sharing this count it
 			// was being used 4 times, and we want to keep the tests consistent.
-			m_optimizeIterationsLeft =
-				m_options->GetPrivate()->m_optimisationOptions.m_maxOptimisationLoopCount * 4;
+			m_optimizeIterationsLeft = m_options->GetPrivate()->m_optimisationOptions.m_maxOptimisationLoopCount * 4;
 			m_optimizeIterationsMax = m_optimizeIterationsLeft;
 
 			// The first duplicated data remover has the special mission of removing
@@ -1247,12 +1246,14 @@ namespace mu
 			}
 
 			UE_LOG(LogMutableCore, Verbose, TEXT("(int) %s : %ld"), TEXT("ast size"), int64(ASTOp::CountNodes(roots)));
+			ASTOp::LogHistogram(roots);
 
 			// Main optimisation stage
 			{
 				MUTABLE_CPUPROFILER_SCOPE(MainStage);
 				FullOptimiseAST( roots, pTaskManager );
 				UE_LOG(LogMutableCore, Verbose, TEXT("(int) %s : %ld"), TEXT("ast size"), int64(ASTOp::CountNodes(roots)));
+				ASTOp::LogHistogram(roots);
 			}
 
 	//            // Third optimisation stage: introduce image compose conditions
@@ -1266,15 +1267,15 @@ namespace mu
 	//            }
 
 
-			{
-				MUTABLE_CPUPROFILER_SCOPE(ConditionalsStage);
+			//{
+			//	MUTABLE_CPUPROFILER_SCOPE(ConditionalsStage);
 
-				//program.Check();
+			//	//program.Check();
 
-				// Optimise the conditionals added by the image compose condition generator
-				// also the constants generator at the end of previous fulloptimiseAST.
-				FullOptimiseAST( roots, pTaskManager );
-			}
+			//	// Optimise the conditionals added by the image compose condition generator
+			//	// also the constants generator at the end of previous fulloptimiseAST.
+			//	FullOptimiseAST( roots, pTaskManager );
+			//}
 
 			// Analyse mesh constants to see which of them are in optimised mesh formats, and set the flags.
 			ASTOp::Traverse_BottomUp_Unique_NonReentrant( roots, [&](Ptr<ASTOp>& n)
diff --git a/Engine/Plugins/Experimental/Mutable/Source/MutableTools/Private/MuT/CodeOptimiser_Semantic.cpp b/Engine/Plugins/Experimental/Mutable/Source/MutableTools/Private/MuT/CodeOptimiser_Semantic.cpp
--- a/Engine/Plugins/Experimental/Mutable/Source/MutableTools/Private/MuT/CodeOptimiser_Semantic.cpp
+++ b/Engine/Plugins/Experimental/Mutable/Source/MutableTools/Private/MuT/CodeOptimiser_Semantic.cpp
@@ -302,55 +302,57 @@ namespace mu
         case OP_TYPE::IM_SWIZZLE:
         {
             // If children channels are also swizzle ops, recurse them
+            Ptr<ASTOpFixed> sat;
+
+            for ( int c=0; c<MUTABLE_OP_MAX_SWIZZLE_CHANNELS; ++c )
             {
-                Ptr<ASTOpFixed> sat;
+                Ptr<ASTOp> candidate = children[op.args.ImageSwizzle.sources[c]].child();
+				if (!candidate)
+				{
+					continue;
+				}
 
-                for ( int c=0; c<MUTABLE_OP_MAX_SWIZZLE_CHANNELS; ++c )
+                // Swizzle
+				if (candidate->GetOpType()==OP_TYPE::IM_SWIZZLE)
                 {
-                    auto candidate = children[op.args.ImageSwizzle.sources[c]].child();
-
-                    // Swizzle
-                    if (candidate && candidate->GetOpType()==OP_TYPE::IM_SWIZZLE)
+                    if (!sat)
                     {
-                        if (!sat)
-                        {
-                            sat = mu::Clone<ASTOpFixed>(this);
-                        }
-                        ASTOpFixed* typedCandidate = dynamic_cast<ASTOpFixed*>(candidate.get());
-                        int candidateChannel = op.args.ImageSwizzle.sourceChannels[c];
-
-                        sat->SetChild( sat->op.args.ImageSwizzle.sources[c],
-                                       typedCandidate->children[typedCandidate->op.args.ImageSwizzle.sources[candidateChannel]] );
-                        sat->op.args.ImageSwizzle.sourceChannels[c] =  typedCandidate->op.args.ImageSwizzle.sourceChannels[candidateChannel];
+                        sat = mu::Clone<ASTOpFixed>(this);
                     }
+                    ASTOpFixed* typedCandidate = dynamic_cast<ASTOpFixed*>(candidate.get());
+                    int candidateChannel = op.args.ImageSwizzle.sourceChannels[c];
 
-                    // Format
-                    if (candidate && candidate->GetOpType()==OP_TYPE::IM_PIXELFORMAT)
-                    {
-                        // We can remove the format if its source is already an uncompressed format
-                        ASTOpImagePixelFormat* typedCandidate = dynamic_cast<ASTOpImagePixelFormat*>(candidate.get());
-                        Ptr<ASTOp> formatSource = typedCandidate->Source.child();
+                    sat->SetChild( sat->op.args.ImageSwizzle.sources[c],
+                                    typedCandidate->children[typedCandidate->op.args.ImageSwizzle.sources[candidateChannel]] );
+                    sat->op.args.ImageSwizzle.sourceChannels[c] =  typedCandidate->op.args.ImageSwizzle.sourceChannels[candidateChannel];
+                }
 
-                        if (formatSource)
+                // Format
+                if (candidate->GetOpType()==OP_TYPE::IM_PIXELFORMAT)
+                {
+                    // We can remove the format if its source is already an uncompressed format
+                    ASTOpImagePixelFormat* typedCandidate = dynamic_cast<ASTOpImagePixelFormat*>(candidate.get());
+                    Ptr<ASTOp> formatSource = typedCandidate->Source.child();
+
+                    if (formatSource)
+                    {
+                        auto desc = formatSource->GetImageDesc();
+                        if ( desc.m_format!= EImageFormat::IF_NONE && !IsCompressedFormat( desc.m_format ) )
                         {
-                            auto desc = formatSource->GetImageDesc();
-                            if ( desc.m_format!= EImageFormat::IF_NONE && !IsCompressedFormat( desc.m_format ) )
+                            if (!sat)
                             {
-                                if (!sat)
-                                {
-                                    sat = mu::Clone<ASTOpFixed>(this);
-                                }
-                                sat->SetChild( sat->op.args.ImageSwizzle.sources[c], formatSource );
+                                sat = mu::Clone<ASTOpFixed>(this);
                             }
+                            sat->SetChild( sat->op.args.ImageSwizzle.sources[c], formatSource );
                         }
                     }
-
                 }
 
-                if (sat)
-                {
-                    at = sat;
-                }
+            }
+
+            if (sat)
+            {
+                at = sat;
             }
 
             break;
@@ -398,7 +400,7 @@ namespace mu
 					check(sourceMaskUsage.size[0] > 0);
 					check(sourceMaskUsage.size[1] > 0);
 					
-					GetImageDescContext context;
+					FGetImageDescContext context;
                     maskDesc = maskAt->GetImageDesc( false, &context );
                 }
             }
@@ -530,6 +532,7 @@ namespace mu
 					NewSwizzle->op.args.ColourSwizzle.sourceChannels[i] = op.args.ImageSwizzle.sourceChannels[i];
 				}
 				NewPlain->SetChild(NewPlain->op.args.ImagePlainColour.colour, NewSwizzle);
+				NewPlain->op.args.ImagePlainColour.format = this->op.args.ImageSwizzle.format;
 				at = NewPlain;
 				break;
 			}
@@ -1111,7 +1114,7 @@ namespace mu
             const ASTOpImagePatch* typedPatch = dynamic_cast<const ASTOpImagePatch*>(at.get());
 
             auto rectOp = typedPatch->patch.child();
-            ASTOp::GetImageDescContext context;
+            ASTOp::FGetImageDescContext context;
             auto patchDesc = rectOp->GetImageDesc( false, &context );
             box<vec2<int16_t>> patchBox;
             patchBox.min = typedPatch->location;
diff --git a/Engine/Plugins/Experimental/Mutable/Source/MutableTools/Private/MuT/NodeImagePlainColourPrivate.h b/Engine/Plugins/Experimental/Mutable/Source/MutableTools/Private/MuT/NodeImagePlainColourPrivate.h
--- a/Engine/Plugins/Experimental/Mutable/Source/MutableTools/Private/MuT/NodeImagePlainColourPrivate.h
+++ b/Engine/Plugins/Experimental/Mutable/Source/MutableTools/Private/MuT/NodeImagePlainColourPrivate.h
@@ -29,30 +29,44 @@ namespace mu
 
 		static NODE_TYPE s_type;
 
-		NodeColourPtr m_pColour;
-		int m_sizeX, m_sizeY;
+		Ptr<NodeColour> m_pColour;
+		int32 m_sizeX, m_sizeY;
+		EImageFormat Format = EImageFormat::IF_RGB_UBYTE;
+
 
 		//!
 		void Serialise( OutputArchive& arch ) const
 		{
-            uint32_t ver = 0;
+            uint32 ver = 1;
 			arch << ver;
 
 			arch << m_pColour;
 			arch << m_sizeX;
 			arch << m_sizeY;
+			arch << (uint32)Format;
 		}
 
 		//!
 		void Unserialise( InputArchive& arch )
 		{
-            uint32_t ver;
+            uint32 ver;
 			arch >> ver;
-			check(ver==0);
+			check(ver<=1);
 
 			arch >> m_pColour;
 			arch >> m_sizeX;
 			arch >> m_sizeY;
+
+			if (ver >= 1)
+			{
+				uint32 Temp = 0;
+				arch >> Temp;
+				Format = EImageFormat(Temp);
+			}
+			else
+			{
+				Format = EImageFormat::IF_RGB_UBYTE;
+			}
 		}
 	};
 
