[mutable] Added a warning for reshape's meshes that do not coincide in number of triangles (mesh shape and target shape pins).

#rnx

diff --git a/Engine/Plugins/Experimental/Mutable/Source/MutableRuntime/Private/MuR/OpMeshReshape.h b/Engine/Plugins/Experimental/Mutable/Source/MutableRuntime/Private/MuR/OpMeshReshape.h
--- a/Engine/Plugins/Experimental/Mutable/Source/MutableRuntime/Private/MuR/OpMeshReshape.h
+++ b/Engine/Plugins/Experimental/Mutable/Source/MutableRuntime/Private/MuR/OpMeshReshape.h
@@ -296,6 +296,18 @@ namespace mu
 		UntypedMeshBufferIterator ItNormal(Mesh->GetVertexBuffers(), MBS_NORMAL);
 		UntypedMeshBufferIterator ItTangent(Mesh->GetVertexBuffers(), MBS_TANGENT);
 
+#if DO_CHECK
+		// checking if the Base shape has more triangles than the target shape
+		for (const FReshapeVertexBindingData& Binding : BindingData)
+		{
+			if (Binding.Triangle >= Shape.Triangles.Num())
+			{
+				UE_LOG(LogMutableCore, Warning, TEXT("Performing a Mesh Reshape where base shape and target shape do not have the same number of triangles."));
+				break;
+			}
+		}
+#endif
+
 		const int32 MeshVertexCount = BindingData.Num();
 		for (int32 MeshVertexIndex = 0; MeshVertexIndex < MeshVertexCount; ++MeshVertexIndex)
 		{
