[mutable] Remove some restrictions on which bones can be reshaped based on usage.

#rnx

diff --git a/Engine/Plugins/Experimental/Mutable/Source/MutableRuntime/Private/MuR/OpMeshBind.h b/Engine/Plugins/Experimental/Mutable/Source/MutableRuntime/Private/MuR/OpMeshBind.h
--- a/Engine/Plugins/Experimental/Mutable/Source/MutableRuntime/Private/MuR/OpMeshBind.h
+++ b/Engine/Plugins/Experimental/Mutable/Source/MutableRuntime/Private/MuR/OpMeshBind.h
@@ -743,8 +743,13 @@ namespace mu
 			const EBoneUsageFlags BoneUsageFlags = Mesh->BonePoses[BoneIndex].BoneUsageFlags;
 
 			// Only deform skinned bones.
-			if (EnumHasAnyFlags(BoneUsageFlags, EBoneUsageFlags::SkinningParent | EBoneUsageFlags::Root) &&
-				!EnumHasAnyFlags(BoneUsageFlags, EBoneUsageFlags::Skinning))
+			//if (EnumHasAnyFlags(BoneUsageFlags, EBoneUsageFlags::SkinningParent | EBoneUsageFlags::Root) &&
+			//	!EnumHasAnyFlags(BoneUsageFlags, EBoneUsageFlags::Skinning))
+			//{
+			//	continue;
+			//}
+
+			if (EnumHasAnyFlags(BoneUsageFlags, EBoneUsageFlags::Root))
 			{
 				continue;
 			}
