[mutable] Additional optimization to multilayer image blending by avoiding unecessary mipmapping.
Fix bug with projectors not working on morphed meshes.
#rnx

diff --git a/Engine/Plugins/Experimental/Mutable/Source/MutableRuntime/Private/MuR/CodeRunner.cpp b/Engine/Plugins/Experimental/Mutable/Source/MutableRuntime/Private/MuR/CodeRunner.cpp
--- a/Engine/Plugins/Experimental/Mutable/Source/MutableRuntime/Private/MuR/CodeRunner.cpp
+++ b/Engine/Plugins/Experimental/Mutable/Source/MutableRuntime/Private/MuR/CodeRunner.cpp
@@ -2532,6 +2532,10 @@ namespace mu
 
                 SCHEDULED_OP itemCopy = item;
                 ExecutionIndex index = GetMemory().GetRageIndex( item.executionIndex );
+				
+				// This becomes true if we need to update the mips of the resulting image
+				// This could happen in the base image has mips, but one of the blended one doesn't.
+				bool bBlendOnlyOneMip = false;
 
                 for (int i=0; i<iterations; ++i)
                 {
@@ -2545,36 +2549,21 @@ namespace mu
                     // \todo: raise a performance warning?
                     if (pBlended && pBlended->GetFormat()!=baseFormat )
                     {
-                        pBlended =
-                            ImagePixelFormat( m_pSettings->GetPrivate()->m_imageCompressionQuality,
-                                              pBlended.get(), baseFormat );
+						MUTABLE_CPUPROFILER_SCOPE(ImageResize_BlendedReformat);
+						pBlended = ImagePixelFormat( m_pSettings->GetPrivate()->m_imageCompressionQuality, pBlended.get(), baseFormat );
                     }
 
 					// TODO: This shouldn't happen, but be defensive.
 					FImageSize ResultSize = pBase->GetSize();
 					if (pBlended && pBlended->GetSize() != ResultSize)
 					{
-						MUTABLE_CPUPROFILER_SCOPE(ImageResize_MaskFixForMultilayer);
+						MUTABLE_CPUPROFILER_SCOPE(ImageResize_BlendedFixForMultilayer);
 						pBlended = ImageResizeLinear(0, pBlended.get(), ResultSize);
 					}
 
 					if (pBlended->GetLODCount() < pBase->GetLODCount())
 					{
-						MUTABLE_CPUPROFILER_SCOPE(ImageLayer_MipBlendedEmergencyFix);
-
-						int levelCount = pBase->GetLODCount();
-						ImagePtr pDest = new Image(pBlended->GetSizeX(), pBlended->GetSizeY(),
-							levelCount,
-							pBlended->GetFormat());
-
-						SCRATCH_IMAGE_MIPMAP scratch;
-						FMipmapGenerationSettings settings{};
-
-						ImageMipmap_PrepareScratch(pDest.get(), pBlended.get(), levelCount, &scratch);
-						ImageMipmap(m_pSettings->GetPrivate()->m_imageCompressionQuality, pDest.get(), pBlended.get(), levelCount,
-							&scratch, settings);
-
-						pBlended = pDest;
+						bBlendOnlyOneMip = true;
 					}
 
                     if (args.mask)
@@ -2588,43 +2577,29 @@ namespace mu
 							pMask = ImageResizeLinear(0, pMask.get(), ResultSize);
 						}
 
-                        // emergy fix c36adf47-e40d-490f-b709-41142bafad78
                         if (pMask->GetLODCount()<pBase->GetLODCount())
                         {
-                            MUTABLE_CPUPROFILER_SCOPE(ImageLayer_EmergencyFix);
-
-                            int levelCount = pBase->GetLODCount();
-                            ImagePtr pDest = new Image( pMask->GetSizeX(), pMask->GetSizeY(),
-                                                        levelCount,
-                                                        pMask->GetFormat() );
-							
-                            SCRATCH_IMAGE_MIPMAP scratch;
-							FMipmapGenerationSettings mipSettings{};
-                        
-                            ImageMipmap_PrepareScratch( pDest.get(), pMask.get(), levelCount, &scratch );
-                            ImageMipmap( m_pSettings->GetPrivate()->m_imageCompressionQuality,
-                                         pDest.get(), pMask.get(), levelCount, &scratch, mipSettings );
-                            pMask = pDest;
-                        }
+							bBlendOnlyOneMip = true;
+						}
 
                         switch (EBlendType(args.blendType))
                         {
 						case EBlendType::BT_NORMAL_COMBINE: 
 							// \todo: optimise with ping pong buffers or blend-in-place
-							ImageNormalCombine(pNew.get(), pBase.get(), pMask.get(), pBlended.get());
+							ImageNormalCombine(pNew.get(), pBase.get(), pMask.get(), pBlended.get(), bBlendOnlyOneMip);
 							pResult = pNew;
 							pBase = pNew;
 							pNew = new Image(pBase->GetSizeX(), pBase->GetSizeY(), pBase->GetLODCount(), pBase->GetFormat());
 							break;
-                        case EBlendType::BT_SOFTLIGHT: ImageLayerOnBase<SoftLightChannelMasked, SoftLightChannel, false>( pNew.get(), pMask.get(), pBlended.get(), false ); break;
-                        case EBlendType::BT_HARDLIGHT: ImageLayerOnBase<HardLightChannelMasked, HardLightChannel, false>( pNew.get(), pMask.get(), pBlended.get(), false ); break;
-                        case EBlendType::BT_BURN: ImageLayerOnBase<BurnChannelMasked, BurnChannel, false>( pNew.get(), pMask.get(), pBlended.get(), false ); break;
-                        case EBlendType::BT_DODGE: ImageLayerOnBase<DodgeChannelMasked, DodgeChannel, false>( pNew.get(), pMask.get(), pBlended.get(), false ); break;
-                        case EBlendType::BT_SCREEN: ImageLayerOnBase<ScreenChannelMasked, ScreenChannel, false>( pNew.get(), pMask.get(), pBlended.get(), false ); break;
-                        case EBlendType::BT_OVERLAY: ImageLayerOnBase<OverlayChannelMasked, OverlayChannel, false>( pNew.get(), pMask.get(), pBlended.get(), false ); break;
-                        case EBlendType::BT_ALPHA_OVERLAY: ImageLayerOnBase<AlphaOverlayChannelMasked, AlphaOverlayChannel, false>( pNew.get(), pMask.get(), pBlended.get(), false ); break;
-                        case EBlendType::BT_MULTIPLY: ImageLayerOnBase<MultiplyChannelMasked, MultiplyChannel, false>( pNew.get(), pMask.get(), pBlended.get(), false ); break;
-                        case EBlendType::BT_BLEND: ImageLayerOnBase<BlendChannelMasked, BlendChannel, false>( pNew.get(), pMask.get(), pBlended.get(), false ); break;
+                        case EBlendType::BT_SOFTLIGHT: ImageLayerOnBase<SoftLightChannelMasked, SoftLightChannel, false>( pNew.get(), pMask.get(), pBlended.get(), false, bBlendOnlyOneMip); break;
+                        case EBlendType::BT_HARDLIGHT: ImageLayerOnBase<HardLightChannelMasked, HardLightChannel, false>( pNew.get(), pMask.get(), pBlended.get(), false, bBlendOnlyOneMip); break;
+                        case EBlendType::BT_BURN: ImageLayerOnBase<BurnChannelMasked, BurnChannel, false>( pNew.get(), pMask.get(), pBlended.get(), false, bBlendOnlyOneMip); break;
+                        case EBlendType::BT_DODGE: ImageLayerOnBase<DodgeChannelMasked, DodgeChannel, false>( pNew.get(), pMask.get(), pBlended.get(), false, bBlendOnlyOneMip); break;
+                        case EBlendType::BT_SCREEN: ImageLayerOnBase<ScreenChannelMasked, ScreenChannel, false>( pNew.get(), pMask.get(), pBlended.get(), false, bBlendOnlyOneMip); break;
+                        case EBlendType::BT_OVERLAY: ImageLayerOnBase<OverlayChannelMasked, OverlayChannel, false>( pNew.get(), pMask.get(), pBlended.get(), false, bBlendOnlyOneMip); break;
+                        case EBlendType::BT_ALPHA_OVERLAY: ImageLayerOnBase<AlphaOverlayChannelMasked, AlphaOverlayChannel, false>( pNew.get(), pMask.get(), pBlended.get(), false, bBlendOnlyOneMip); break;
+                        case EBlendType::BT_MULTIPLY: ImageLayerOnBase<MultiplyChannelMasked, MultiplyChannel, false>( pNew.get(), pMask.get(), pBlended.get(), false, bBlendOnlyOneMip); break;
+                        case EBlendType::BT_BLEND: ImageLayerOnBase<BlendChannelMasked, BlendChannel, false>( pNew.get(), pMask.get(), pBlended.get(), false, bBlendOnlyOneMip); break;
                         default: check(false);
                         }
 
@@ -2635,25 +2610,34 @@ namespace mu
                         {
 						case EBlendType::BT_NORMAL_COMBINE: 
 							// \todo: optimise with ping pong buffers or blend-in-place
-							ImageNormalCombine(pNew.get(), pBase.get(), pBlended.get());
+							ImageNormalCombine(pNew.get(), pBase.get(), pBlended.get(), bBlendOnlyOneMip);
 							pResult = pNew;
 							pBase = pNew;
 							pNew = new Image(pBase->GetSizeX(), pBase->GetSizeY(), pBase->GetLODCount(), pBase->GetFormat());
 							break;
-                        case EBlendType::BT_SOFTLIGHT: ImageLayerOnBase<SoftLightChannel, false>( pNew.get(), pBlended.get(), false ); break;
-                        case EBlendType::BT_HARDLIGHT: ImageLayerOnBase<HardLightChannel, false>( pNew.get(), pBlended.get(), false ); break;
-                        case EBlendType::BT_BURN: ImageLayerOnBase<BurnChannel, false>( pNew.get(), pBlended.get(), false ); break;
-                        case EBlendType::BT_DODGE: ImageLayerOnBase<DodgeChannel, false>( pNew.get(), pBlended.get(), false ); break;
-                        case EBlendType::BT_SCREEN: ImageLayerOnBase<ScreenChannel, false>( pNew.get(),  pBlended.get(), false ); break;
-                        case EBlendType::BT_OVERLAY: ImageLayerOnBase<OverlayChannel, false>( pNew.get(), pBlended.get(), false ); break;
-                        case EBlendType::BT_ALPHA_OVERLAY: ImageLayerOnBase<AlphaOverlayChannel, false>( pNew.get(), pBlended.get(), false ); break;
-                        case EBlendType::BT_MULTIPLY: ImageLayerOnBase<MultiplyChannel, false>( pNew.get(), pBlended.get(), false ); break;
+                        case EBlendType::BT_SOFTLIGHT: ImageLayerOnBase<SoftLightChannel, false>( pNew.get(), pBlended.get(), false, bBlendOnlyOneMip); break;
+                        case EBlendType::BT_HARDLIGHT: ImageLayerOnBase<HardLightChannel, false>( pNew.get(), pBlended.get(), false, bBlendOnlyOneMip); break;
+                        case EBlendType::BT_BURN: ImageLayerOnBase<BurnChannel, false>( pNew.get(), pBlended.get(), false, bBlendOnlyOneMip); break;
+                        case EBlendType::BT_DODGE: ImageLayerOnBase<DodgeChannel, false>( pNew.get(), pBlended.get(), false, bBlendOnlyOneMip); break;
+                        case EBlendType::BT_SCREEN: ImageLayerOnBase<ScreenChannel, false>( pNew.get(),  pBlended.get(), false, bBlendOnlyOneMip); break;
+                        case EBlendType::BT_OVERLAY: ImageLayerOnBase<OverlayChannel, false>( pNew.get(), pBlended.get(), false, bBlendOnlyOneMip); break;
+                        case EBlendType::BT_ALPHA_OVERLAY: ImageLayerOnBase<AlphaOverlayChannel, false>( pNew.get(), pBlended.get(), false, bBlendOnlyOneMip); break;
+                        case EBlendType::BT_MULTIPLY: ImageLayerOnBase<MultiplyChannel, false>( pNew.get(), pBlended.get(), false, bBlendOnlyOneMip); break;
                         case EBlendType::BT_BLEND: pNew = pBlended->Clone(); break;
                         default: check(false);
                         }
                     }
 				}
 
+				if (bBlendOnlyOneMip)
+				{
+					MUTABLE_CPUPROFILER_SCOPE(ImageLayer_MipFix);
+					FMipmapGenerationSettings DummyMipSettings{};
+					ImageMipmapInPlace(m_pSettings->GetPrivate()->m_imageCompressionQuality, pNew.get(), DummyMipSettings);
+				}
+
+				// TODO: Reconvert to baseFormat?
+
                 GetMemory().SetImage( item, pResult );
                 break;
             }
@@ -3942,9 +3926,11 @@ namespace mu
 					pSource = ImageResizeLinear(0, pSource.get(), pMap->GetSize());
 				}
 
-                ImagePtr pNew = new Image( pSource->GetSizeX(), pSource->GetSizeY(), 1, pSource->GetFormat() );
+				// This works based on the assumption that displacement maps never read from a position they actually write to.
+				// Since they are used for UV border expansion, this should always be the case.
+				//ImagePtr pNew = new Image(pSource->GetSizeX(), pSource->GetSizeY(), 1, pSource->GetFormat());
+				Ptr<Image> pNew = mu::CloneOrTakeOver(pSource.get());
 
-                // TODO: Reuse source if possible
                 ImageDisplace( pNew.get(), pSource.get(), pMap.get() );
 
 				if (OriginalSourceScale != pNew->GetSize())
diff --git a/Engine/Plugins/Experimental/Mutable/Source/MutableRuntime/Private/MuR/CodeRunnerTasked.cpp b/Engine/Plugins/Experimental/Mutable/Source/MutableRuntime/Private/MuR/CodeRunnerTasked.cpp
--- a/Engine/Plugins/Experimental/Mutable/Source/MutableRuntime/Private/MuR/CodeRunnerTasked.cpp
+++ b/Engine/Plugins/Experimental/Mutable/Source/MutableRuntime/Private/MuR/CodeRunnerTasked.cpp
@@ -502,7 +502,7 @@ namespace mu
 
 				switch (EBlendType(Args.blendType))
 				{
-				case EBlendType::BT_NORMAL_COMBINE: ImageNormalCombine(pNew.get(), m_base.get(), m_mask.get(), m_blended.get()); break;
+				case EBlendType::BT_NORMAL_COMBINE: ImageNormalCombine(pNew.get(), m_base.get(), m_mask.get(), m_blended.get(), false); break;
 				case EBlendType::BT_SOFTLIGHT: ImageSoftLight(pNew.get(), m_base.get(), m_mask.get(), m_blended.get(), applyToAlpha); break;
 				case EBlendType::BT_HARDLIGHT: ImageHardLight(pNew.get(), m_base.get(), m_mask.get(), m_blended.get(), applyToAlpha); break;
 				case EBlendType::BT_BURN: ImageBurn(pNew.get(), m_base.get(), m_mask.get(), m_blended.get(), applyToAlpha); break;
@@ -520,7 +520,7 @@ namespace mu
 			{
 				switch (EBlendType(Args.blendType))
 				{
-				case EBlendType::BT_NORMAL_COMBINE: ImageNormalCombine(pNew.get(), m_base.get(), m_blended.get()); break;
+				case EBlendType::BT_NORMAL_COMBINE: ImageNormalCombine(pNew.get(), m_base.get(), m_blended.get(), false); break;
 				case EBlendType::BT_SOFTLIGHT: ImageSoftLight(pNew.get(), m_base.get(), m_blended.get(), applyToAlpha); break;
 				case EBlendType::BT_HARDLIGHT: ImageHardLight(pNew.get(), m_base.get(), m_blended.get(), applyToAlpha); break;
 				case EBlendType::BT_BURN: ImageBurn(pNew.get(), m_base.get(), m_blended.get(), applyToAlpha); break;
diff --git a/Engine/Plugins/Experimental/Mutable/Source/MutableRuntime/Private/MuR/OpImageBlend.h b/Engine/Plugins/Experimental/Mutable/Source/MutableRuntime/Private/MuR/OpImageBlend.h
--- a/Engine/Plugins/Experimental/Mutable/Source/MutableRuntime/Private/MuR/OpImageBlend.h
+++ b/Engine/Plugins/Experimental/Mutable/Source/MutableRuntime/Private/MuR/OpImageBlend.h
@@ -52,10 +52,10 @@ namespace mu
 
 
     inline void ImageBlendOnBase( Image* pBase, const Image* pMask, const Image* pBlended,
-                                  bool applyToAlpha )
+                                  bool applyToAlpha, bool bOnlyFirstLOD )
 	{
         ImageLayerOnBase<BlendChannelMasked, BlendChannel, false>
-                 ( pBase, pMask, pBlended, applyToAlpha );
+                 ( pBase, pMask, pBlended, applyToAlpha, bOnlyFirstLOD );
 	}
 
 	//! Blend a subimage on the base using a mask.
diff --git a/Engine/Plugins/Experimental/Mutable/Source/MutableRuntime/Private/MuR/OpImageLayer.h b/Engine/Plugins/Experimental/Mutable/Source/MutableRuntime/Private/MuR/OpImageLayer.h
--- a/Engine/Plugins/Experimental/Mutable/Source/MutableRuntime/Private/MuR/OpImageLayer.h
+++ b/Engine/Plugins/Experimental/Mutable/Source/MutableRuntime/Private/MuR/OpImageLayer.h
@@ -475,7 +475,7 @@ namespace mu
 	//---------------------------------------------------------------------------------------------
 	template< unsigned (*BLEND_FUNC)(unsigned,unsigned), bool CLAMP, int NC >
     inline void BufferLayer(uint8* pDestBuf, const Image* pBase, const Image* pBlended,
-                            bool applyToAlpha )
+                            bool applyToAlpha, bool bOnlyFirstLOD )
 	{
 		check( pBase->GetSizeX() == pBlended->GetSizeX() );
 		check( pBase->GetSizeY() == pBlended->GetSizeY() );
@@ -485,9 +485,17 @@ namespace mu
         const uint8* pBlendedBuf = pBlended->GetData();
 
 		// Generic implementation
-		int32 pixelCount = pBase->CalculatePixelCount();
+		int32 PixelCount = 0;
+		if (bOnlyFirstLOD)
+		{
+			PixelCount = pBase->GetSizeX() * pBase->GetSizeY();
+		}
+		else
+		{
+			PixelCount = pBase->CalculatePixelCount();
+		}
 
-		ParallelFor(pixelCount,
+		ParallelFor(PixelCount,
 			[
 				pBaseBuf, pBlendedBuf, pDestBuf, applyToAlpha
 			] (uint32 i)
@@ -554,15 +562,15 @@ namespace mu
 
 		if ( baseFormat==EImageFormat::IF_RGB_UBYTE )
 		{
-            BufferLayer< BLEND_FUNC, CLAMP, 3 > ( pDestBuf, pBase, pBlended, applyToAlpha );
+            BufferLayer< BLEND_FUNC, CLAMP, 3 > ( pDestBuf, pBase, pBlended, applyToAlpha, false );
 		}
         else if ( baseFormat==EImageFormat::IF_RGBA_UBYTE || baseFormat==EImageFormat::IF_BGRA_UBYTE )
 		{
-            BufferLayer< BLEND_FUNC, CLAMP, 4 > ( pDestBuf, pBase, pBlended, applyToAlpha );
+            BufferLayer< BLEND_FUNC, CLAMP, 4 > ( pDestBuf, pBase, pBlended, applyToAlpha, false);
 		}
 		else if ( baseFormat==EImageFormat::IF_L_UBYTE )
 		{
-            BufferLayer< BLEND_FUNC, CLAMP, 1 > ( pDestBuf, pBase, pBlended, applyToAlpha );
+            BufferLayer< BLEND_FUNC, CLAMP, 1 > ( pDestBuf, pBase, pBlended, applyToAlpha, false);
 		}
 		else
 		{
@@ -576,7 +584,7 @@ namespace mu
 	//---------------------------------------------------------------------------------------------
 	template< unsigned (*BLEND_FUNC)(unsigned,unsigned), bool CLAMP >
     inline void ImageLayerOnBase( Image* pBase, const Image* pBlended,
-                                  bool applyToAlpha )
+                                  bool applyToAlpha, bool bOnlyFirstLOD=false )
 	{
 		check( pBase->GetSizeX() == pBlended->GetSizeX() );
 		check( pBase->GetSizeY() == pBlended->GetSizeY() );
@@ -587,15 +595,15 @@ namespace mu
 
 		if ( baseFormat==EImageFormat::IF_RGB_UBYTE )
 		{
-            BufferLayer< BLEND_FUNC, CLAMP, 3 > ( pDestBuf, pBase, pBlended, applyToAlpha );
+            BufferLayer< BLEND_FUNC, CLAMP, 3 > ( pDestBuf, pBase, pBlended, applyToAlpha, bOnlyFirstLOD);
 		}
         else if ( baseFormat==EImageFormat::IF_RGBA_UBYTE || baseFormat==EImageFormat::IF_BGRA_UBYTE )
 		{
-            BufferLayer< BLEND_FUNC, CLAMP, 4 > ( pDestBuf, pBase, pBlended, applyToAlpha );
+            BufferLayer< BLEND_FUNC, CLAMP, 4 > ( pDestBuf, pBase, pBlended, applyToAlpha, bOnlyFirstLOD);
 		}
 		else if ( baseFormat==EImageFormat::IF_L_UBYTE )
 		{
-            BufferLayer< BLEND_FUNC, CLAMP, 1 > ( pDestBuf, pBase, pBlended, applyToAlpha );
+            BufferLayer< BLEND_FUNC, CLAMP, 1 > ( pDestBuf, pBase, pBlended, applyToAlpha, bOnlyFirstLOD);
 		}
 		else
 		{
@@ -613,12 +621,13 @@ namespace mu
                                    const Image* pBase,
                                    const Image* pMask,
                                    const Image* pBlend,
-                                   bool applyToAlpha )
+                                   bool applyToAlpha,
+								   bool bOnlyFirstLOD )
 	{
         check(pBase->GetSizeX()==pMask->GetSizeX() && pBase->GetSizeY()==pMask->GetSizeY());
         check(pBase->GetSizeX()==pBlend->GetSizeX() && pBase->GetSizeY()==pBlend->GetSizeY());
-        check(pBase->GetLODCount() <= pMask->GetLODCount() );
-        check(pBase->GetLODCount() <= pBlend->GetLODCount() );
+        check(bOnlyFirstLOD || pBase->GetLODCount() <= pMask->GetLODCount() );
+        check(bOnlyFirstLOD || pBase->GetLODCount() <= pBlend->GetLODCount() );
 
 		uint8* pDestBuf = pStartDestBuf;
         const uint8* pMaskBuf = pMask->GetData();
@@ -629,15 +638,23 @@ namespace mu
         bool isUncompressed = ( maskFormat == EImageFormat::IF_L_UBYTE );
 
 		// The base determines the number of lods to process.
-		int32 LODCount = pBase->GetLODCount();
+		int32 LODCount = bOnlyFirstLOD ? 1 : pBase->GetLODCount();
 
         if ( isUncompressed )
 		{
-			int32 pixelCount = pBase->CalculatePixelCount();
+			int32 PixelCount = 0;
+			if (bOnlyFirstLOD)
+			{
+				PixelCount = pBase->GetSizeX() * pBase->GetSizeY();
+			}
+			else
+			{
+				PixelCount = pBase->CalculatePixelCount();
+			}
 
             if (!applyToAlpha && NC>3)
             {
-				ParallelFor(pixelCount,
+				ParallelFor(PixelCount,
 					[
 						pBaseBuf, pBlendedBuf, pMaskBuf, pDestBuf
 					] (uint32 i)
@@ -664,7 +681,7 @@ namespace mu
             }
             else
             {
-				ParallelFor(pixelCount,
+				ParallelFor(PixelCount,
 					[
 						pBaseBuf, pBlendedBuf, pMaskBuf, pDestBuf
 					] (uint32 i)
@@ -988,22 +1005,23 @@ namespace mu
 							 const Image* pBase,
 							 const Image* pMask,
                              const Image* pBlend,
-                             bool applyToAlpha )
+                             bool applyToAlpha,
+							 bool bOnlyFirstLOD )
 	{
 		if ( pBase->GetFormat()==EImageFormat::IF_RGB_UBYTE )
 		{
             BufferLayerFormat< BLEND_FUNC_MASKED, BLEND_FUNC, CLAMP, 3 >
-                    ( pDestBuf, pBase, pMask, pBlend, applyToAlpha ) ;
+                    ( pDestBuf, pBase, pMask, pBlend, applyToAlpha, bOnlyFirstLOD) ;
 		}
         else if ( pBase->GetFormat()==EImageFormat::IF_RGBA_UBYTE || pBase->GetFormat()==EImageFormat::IF_BGRA_UBYTE )
 		{
             BufferLayerFormat< BLEND_FUNC_MASKED, BLEND_FUNC, CLAMP, 4 >
-                    ( pDestBuf, pBase, pMask, pBlend, applyToAlpha ) ;
+                    ( pDestBuf, pBase, pMask, pBlend, applyToAlpha, bOnlyFirstLOD) ;
 		}
 		else if ( pBase->GetFormat()==EImageFormat::IF_L_UBYTE )
 		{
             BufferLayerFormat< BLEND_FUNC_MASKED, BLEND_FUNC, CLAMP, 1 >
-                    ( pDestBuf, pBase, pMask, pBlend, applyToAlpha ) ;
+                    ( pDestBuf, pBase, pMask, pBlend, applyToAlpha, bOnlyFirstLOD) ;
 		}
         else
 		{
@@ -1067,7 +1085,7 @@ namespace mu
         check( pBase->GetLODCount() <= pBlended->GetLODCount() );
 
         BufferLayer<BLEND_FUNC_MASKED,BLEND_FUNC,CLAMP>( pResult->GetData(), pBase, pMask,
-                                                         pBlended, applyToAlpha );
+                                                         pBlended, applyToAlpha, false);
 	}
 
 
@@ -1079,7 +1097,7 @@ namespace mu
 			  unsigned (*BLEND_FUNC)(unsigned,unsigned),
 			  bool CLAMP >
     inline void ImageLayerOnBase( Image* pBase, const Image* pMask, const Image* pBlended,
-                                  bool applyToAlpha )
+                                  bool applyToAlpha, bool bOnlyFirstLOD )
 	{
 		check( pBase->GetSizeX() == pMask->GetSizeX() );
 		check( pBase->GetSizeY() == pMask->GetSizeY() );
@@ -1089,11 +1107,11 @@ namespace mu
 		check( pMask->GetFormat() == EImageFormat::IF_L_UBYTE
 						||
 						pMask->GetFormat() == EImageFormat::IF_L_UBYTE_RLE );
-        check( pBase->GetLODCount() <= pMask->GetLODCount() );
-        check( pBase->GetLODCount() <= pBlended->GetLODCount() );
+        check(bOnlyFirstLOD || pBase->GetLODCount() <= pMask->GetLODCount() );
+        check(bOnlyFirstLOD || pBase->GetLODCount() <= pBlended->GetLODCount() );
 
         BufferLayer<BLEND_FUNC_MASKED,BLEND_FUNC,CLAMP>( pBase->GetData(), pBase, pMask,
-                                                         pBlended, applyToAlpha );
+                                                         pBlended, applyToAlpha, bOnlyFirstLOD);
 	}
 
 
@@ -1386,7 +1404,7 @@ namespace mu
 	template< 
 			uint32 (*BLEND_FUNC)(uint32, uint32), 
 			size_t NC >
-	inline void BufferLayerCombine(uint8* pDestBuf, const Image* pBase, const Image* pBlended)
+	inline void BufferLayerCombine(uint8* pDestBuf, const Image* pBase, const Image* pBlended, bool bOnlyFirstLOD)
 	{
 		//static_assert(NC > 0 && NC <= 4);
 
@@ -1397,9 +1415,17 @@ namespace mu
 		const uint8* pBaseBuf = pBase->GetData();
 		const uint8* pBlendedBuf = pBlended->GetData();
 
-		int32 pixelCount = pBase->CalculatePixelCount();
+		int32 PixelCount = 0;
+		if (bOnlyFirstLOD)
+		{
+			PixelCount = pBase->GetSizeX() * pBase->GetSizeX();
+		}
+		else
+		{
+			PixelCount = pBase->CalculatePixelCount();
+		}		
 
-		for (int i = 0; i < pixelCount; ++i)
+		for (int i = 0; i < PixelCount; ++i)
 		{
 			const uint32 base = PackPixel<NC>(&pBaseBuf[NC * i]);
 			const uint32 blend = PackPixel<NC>(&pBlendedBuf[NC * i]);
@@ -1413,7 +1439,7 @@ namespace mu
 	template< 
 			uint32 (*BLEND_FUNC_MASKED)(uint32, uint32, uint32), 
 			size_t NC >
-	inline void BufferLayerCombine(uint8* pDestBuf, const Image* pBase, const Image* pMask, const Image* pBlended)
+	inline void BufferLayerCombine(uint8* pDestBuf, const Image* pBase, const Image* pMask, const Image* pBlended, bool bOnlyFirstLOD)
 	{
 		//static_assert(NC > 0 && NC <= 4);
 
@@ -1427,9 +1453,17 @@ namespace mu
 		const uint8* pMaskBuf = pMask->GetData();
 		const uint8* pBlendedBuf = pBlended->GetData();
 
-		int32 pixelCount = pBase->CalculatePixelCount();
+		int32 PixelCount = 0;
+		if (bOnlyFirstLOD)
+		{
+			PixelCount = pBase->GetSizeX() * pBase->GetSizeX();
+		}
+		else
+		{
+			PixelCount = pBase->CalculatePixelCount();
+		}
 
-		for (int i = 0; i < pixelCount; ++i)
+		for (int i = 0; i < PixelCount; ++i)
 		{
 			const uint32 base = PackPixel<NC>(&pBaseBuf[NC * i]);
 			const uint32 blend = PackPixel<NC>(&pBlendedBuf[NC * i]);
@@ -1475,32 +1509,32 @@ namespace mu
 	}
 
 	template< uint32 (*BLEND_FUNC)(uint32, uint32) >
-	inline void ImageLayerCombine(Image* pResult, const Image* pBase, const Image* pBlended)
+	inline void ImageLayerCombine(Image* pResult, const Image* pBase, const Image* pBlended, bool bOnlyFirstLOD)
 	{
 		check(pResult->GetFormat() == pBase->GetFormat());
 		check(pResult->GetSizeX() == pBase->GetSizeX());
 		check(pResult->GetSizeY() == pBase->GetSizeY());
-		check(pResult->GetLODCount() == pBase->GetLODCount());
+		check(bOnlyFirstLOD || pResult->GetLODCount() == pBase->GetLODCount());
 		check(pBase->GetSizeX() == pBlended->GetSizeX());
 		check(pBase->GetSizeY() == pBlended->GetSizeY());
 		check(pBase->GetFormat() == pBlended->GetFormat());
-		check(pResult->GetLODCount() <= pBlended->GetLODCount());
+		check(bOnlyFirstLOD || pResult->GetLODCount() <= pBlended->GetLODCount());
 
 		const EImageFormat baseFormat = pBase->GetFormat();
 
 		uint8* pDestBuf = pResult->GetData();
 		if (baseFormat == EImageFormat::IF_L_UBYTE)
 		{
-			BufferLayerCombine<BLEND_FUNC, 1 >(pDestBuf, pBase, pBlended);
+			BufferLayerCombine<BLEND_FUNC, 1 >(pDestBuf, pBase, pBlended, bOnlyFirstLOD);
 		}
 		else if (baseFormat == EImageFormat::IF_RGB_UBYTE)
 		{
-			BufferLayerCombine< BLEND_FUNC, 3 >(pDestBuf, pBase, pBlended);
+			BufferLayerCombine< BLEND_FUNC, 3 >(pDestBuf, pBase, pBlended, bOnlyFirstLOD);
 		}
 		else if (baseFormat == EImageFormat::IF_RGBA_UBYTE || baseFormat == EImageFormat::IF_BGRA_UBYTE)
 		{
 			// \todo: pass swizzle template argument if BGRA_UBYTE, not yet supported.
-			BufferLayerCombine< BLEND_FUNC, 4 >(pDestBuf, pBase, pBlended);
+			BufferLayerCombine< BLEND_FUNC, 4 >(pDestBuf, pBase, pBlended, bOnlyFirstLOD);
 		}
 		else
 		{
@@ -1511,16 +1545,16 @@ namespace mu
 	template< 
 		uint32 (*BLEND_FUNC)(uint32, uint32), 
 		uint32 (*BLEND_FUNC_MASKED)(uint32, uint32, uint32) >
-	inline void ImageLayerCombine(Image* pResult, const Image* pBase, const Image* pMask, const Image* pBlended)
+	inline void ImageLayerCombine(Image* pResult, const Image* pBase, const Image* pMask, const Image* pBlended, bool bOnlyFirstLOD)
 	{
 		check(pResult->GetFormat() == pBase->GetFormat());
 		check(pResult->GetSizeX() == pBase->GetSizeX());
 		check(pResult->GetSizeY() == pBase->GetSizeY());
-		check(pResult->GetLODCount() == pBase->GetLODCount());
+		check(bOnlyFirstLOD || pResult->GetLODCount() == pBase->GetLODCount());
 		check(pBase->GetSizeX() == pBlended->GetSizeX());
 		check(pBase->GetSizeY() == pBlended->GetSizeY());
 		check(pBase->GetFormat() == pBlended->GetFormat());
-		check(pResult->GetLODCount() <= pBlended->GetLODCount());
+		check(bOnlyFirstLOD || pResult->GetLODCount() <= pBlended->GetLODCount());
 
 		const EImageFormat baseFormat = pBase->GetFormat();
 
@@ -1530,21 +1564,21 @@ namespace mu
 		{
 			checkf(false, TEXT("Unsupported mask format."));
 
-			BufferLayerCombine< BLEND_FUNC, 1 >(pDestBuf, pBase, pBlended);
+			BufferLayerCombine< BLEND_FUNC, 1 >(pDestBuf, pBase, pBlended, bOnlyFirstLOD);
 		}
 
 		if (baseFormat == EImageFormat::IF_L_UBYTE)
 		{
-			BufferLayerCombine< BLEND_FUNC_MASKED, 1 >(pDestBuf, pBase, pMask, pBlended);
+			BufferLayerCombine< BLEND_FUNC_MASKED, 1 >(pDestBuf, pBase, pMask, pBlended, bOnlyFirstLOD);
 		}
 		else if (baseFormat == EImageFormat::IF_RGB_UBYTE)
 		{
-			BufferLayerCombine< BLEND_FUNC_MASKED, 3 >(pDestBuf, pBase, pMask, pBlended);
+			BufferLayerCombine< BLEND_FUNC_MASKED, 3 >(pDestBuf, pBase, pMask, pBlended, bOnlyFirstLOD);
 		}
 		else if (baseFormat == EImageFormat::IF_RGBA_UBYTE || baseFormat == EImageFormat::IF_BGRA_UBYTE)
 		{
 			// \todo: pass swizzle template argument if BGRA_UBYTE, not yet supported.
-			BufferLayerCombine< BLEND_FUNC_MASKED, 4 >(pDestBuf, pBase, pMask, pBlended);
+			BufferLayerCombine< BLEND_FUNC_MASKED, 4 >(pDestBuf, pBase, pMask, pBlended, bOnlyFirstLOD);
 		}
 		else
 		{
diff --git a/Engine/Plugins/Experimental/Mutable/Source/MutableRuntime/Private/MuR/OpImageMipmap.h b/Engine/Plugins/Experimental/Mutable/Source/MutableRuntime/Private/MuR/OpImageMipmap.h
--- a/Engine/Plugins/Experimental/Mutable/Source/MutableRuntime/Private/MuR/OpImageMipmap.h
+++ b/Engine/Plugins/Experimental/Mutable/Source/MutableRuntime/Private/MuR/OpImageMipmap.h
@@ -613,7 +613,7 @@ namespace mu
 
     //---------------------------------------------------------------------------------------------
     //! Generate the mipmaps for images.
-	//! if bGenerateOnlyTail is true, generates the mips missing form pBase to levelCount and sets
+	//! if bGenerateOnlyTail is true, generates the mips missing from pBase to levelCount and sets
 	//! them in pDest (the full chain is spit in two images). Otherwise generate the mips missing 
 	//! from pBase up to levelCount and append them in pDest to the already generated pBase's mips.
     //---------------------------------------------------------------------------------------------
@@ -739,4 +739,46 @@ namespace mu
         }
     }
 
+
+	/** Update all the mipmaps in the image from the data in the base one. 
+	* Only the mipmaps already existing in the image are updated.
+	*/
+	inline void ImageMipmapInPlace(int InImageCompressionQuality, Image* InBase, const FMipmapGenerationSettings& InSettings)
+	{
+		int32 StartLevel = 0;
+		int32 LevelCount = InBase->GetLODCount();
+
+		check(!(InBase->m_flags & Image::IF_CANNOT_BE_SCALED));
+
+		int32 MipsToBuild = InBase->GetLODCount()-1;
+		if (!MipsToBuild)
+		{
+			return;
+		}
+
+		const uint8_t* pSourceBuf = InBase->GetMipData(StartLevel);
+		uint8_t* pDestBuf = InBase->GetMipData(StartLevel + 1);
+
+		vec2<int> sourceSize = InBase->CalculateMipSize(StartLevel);
+
+		switch (InBase->GetFormat())
+		{
+		case EImageFormat::IF_L_UBYTE:
+			GenerateMipmapsUint8<1>(MipsToBuild, pSourceBuf, pDestBuf, sourceSize, InSettings);
+			break;
+
+		case EImageFormat::IF_RGB_UBYTE:
+			GenerateMipmapsUint8<3>(MipsToBuild, pSourceBuf, pDestBuf, sourceSize, InSettings);
+			break;
+
+		case EImageFormat::IF_BGRA_UBYTE:
+		case EImageFormat::IF_RGBA_UBYTE:
+			GenerateMipmapsUint8<4>(MipsToBuild, pSourceBuf, pDestBuf, sourceSize, InSettings);
+			break;
+
+		default:
+			checkf(false, TEXT("Format not implemented in mipmap generation."));
+		}
+	}
+
 }
diff --git a/Engine/Plugins/Experimental/Mutable/Source/MutableRuntime/Private/MuR/OpImageMultiply.h b/Engine/Plugins/Experimental/Mutable/Source/MutableRuntime/Private/MuR/OpImageMultiply.h
--- a/Engine/Plugins/Experimental/Mutable/Source/MutableRuntime/Private/MuR/OpImageMultiply.h
+++ b/Engine/Plugins/Experimental/Mutable/Source/MutableRuntime/Private/MuR/OpImageMultiply.h
@@ -60,7 +60,7 @@ namespace mu
                                      bool applyToAlpha )
 	{
         ImageLayerOnBase<MultiplyChannelMasked,MultiplyChannel, false>
-                ( pBase, pMask, pBlended, applyToAlpha );
+                ( pBase, pMask, pBlended, applyToAlpha, false );
 	}
 
     inline void ImageMultiplyOnBase( Image* pBase, const Image* pBlended,
diff --git a/Engine/Plugins/Experimental/Mutable/Source/MutableRuntime/Private/MuR/OpImageNormalCombine.h b/Engine/Plugins/Experimental/Mutable/Source/MutableRuntime/Private/MuR/OpImageNormalCombine.h
--- a/Engine/Plugins/Experimental/Mutable/Source/MutableRuntime/Private/MuR/OpImageNormalCombine.h
+++ b/Engine/Plugins/Experimental/Mutable/Source/MutableRuntime/Private/MuR/OpImageNormalCombine.h
@@ -66,9 +66,9 @@ namespace mu
 		ImageLayerCombineColour<CombineNormal>(pResult, pBase, col);
 	}
 
-	inline void ImageNormalCombine(Image* pResult, const Image* pBase, const Image* pBlended)
+	inline void ImageNormalCombine(Image* pResult, const Image* pBase, const Image* pBlended, bool bOnlyFirstLOD)
 	{
-		ImageLayerCombine<CombineNormal>(pResult, pBase, pBlended);
+		ImageLayerCombine<CombineNormal>(pResult, pBase, pBlended, bOnlyFirstLOD);
 	}
 
 	inline void ImageNormalCombine(Image* pResult, const Image* pBase, const Image* pMask, vec3<float> col)
@@ -76,8 +76,8 @@ namespace mu
 		ImageLayerCombineColour<CombineNormal, CombineNormalMasked>(pResult, pBase, pMask, col);
 	}
 
-	inline void ImageNormalCombine(Image* pResult, const Image* pBase, const Image* pMask, const Image* pBlended)
+	inline void ImageNormalCombine(Image* pResult, const Image* pBase, const Image* pMask, const Image* pBlended, bool bOnlyFirstLOD)
 	{
-		ImageLayerCombine<CombineNormal, CombineNormalMasked>(pResult, pBase, pMask, pBlended);
+		ImageLayerCombine<CombineNormal, CombineNormalMasked>(pResult, pBase, pMask, pBlended, bOnlyFirstLOD);
 	}
 }
diff --git a/Engine/Plugins/Experimental/Mutable/Source/MutableRuntime/Private/MuR/OpLayoutRemoveBlocks.h b/Engine/Plugins/Experimental/Mutable/Source/MutableRuntime/Private/MuR/OpLayoutRemoveBlocks.h
--- a/Engine/Plugins/Experimental/Mutable/Source/MutableRuntime/Private/MuR/OpLayoutRemoveBlocks.h
+++ b/Engine/Plugins/Experimental/Mutable/Source/MutableRuntime/Private/MuR/OpLayoutRemoveBlocks.h
@@ -77,7 +77,7 @@ namespace mu
 		//	pB->GetPrivate()->m_blocks.end()
 		//);
 
-		for ( const Layout::BLOCK& block: pB->m_blocks )
+		for ( const Layout::FBlock& block: pB->m_blocks )
 		{
 			if ( pResult->FindBlock(block.m_id)<0 )
 			{
diff --git a/Engine/Plugins/Experimental/Mutable/Source/MutableRuntime/Public/MuR/Layout.h b/Engine/Plugins/Experimental/Mutable/Source/MutableRuntime/Public/MuR/Layout.h
--- a/Engine/Plugins/Experimental/Mutable/Source/MutableRuntime/Public/MuR/Layout.h
+++ b/Engine/Plugins/Experimental/Mutable/Source/MutableRuntime/Public/MuR/Layout.h
@@ -134,9 +134,9 @@ namespace mu
 
 		// This used to be the data in the private implementation of the image interface
 		//-----------------------------------------------------------------------------------------
-		struct BLOCK
+		struct FBlock
 		{
-			BLOCK(vec2<uint16> min = vec2<uint16>(), vec2<uint16> size = vec2<uint16>())
+			FBlock(vec2<uint16> min = vec2<uint16>(), vec2<uint16> size = vec2<uint16>())
 			{
 				m_min = min;
 				m_size = size;
@@ -173,7 +173,7 @@ namespace mu
 			}
 
 			//!
-			inline bool operator==(const BLOCK& o) const
+			inline bool operator==(const FBlock& o) const
 			{
 				return (m_min == o.m_min) &&
 					(m_size == o.m_size) &&
@@ -181,7 +181,7 @@ namespace mu
 					(m_priority == o.m_priority);
 			}
 
-			inline bool IsSimilar(const BLOCK& o) const
+			inline bool IsSimilar(const FBlock& o) const
 			{
 				// All but ids
 				return (m_min == o.m_min) &&
@@ -197,7 +197,7 @@ namespace mu
 		vec2<uint16> m_maxsize;
 
 		//!
-		TArray<BLOCK> m_blocks;
+		TArray<FBlock> m_blocks;
 
 		//! Packing strategy
 		EPackStrategy m_strategy = EPackStrategy::RESIZABLE_LAYOUT;
diff --git a/Engine/Plugins/Experimental/Mutable/Source/MutableTools/Private/MuT/CodeGenerator.cpp b/Engine/Plugins/Experimental/Mutable/Source/MutableTools/Private/MuT/CodeGenerator.cpp
--- a/Engine/Plugins/Experimental/Mutable/Source/MutableTools/Private/MuT/CodeGenerator.cpp
+++ b/Engine/Plugins/Experimental/Mutable/Source/MutableTools/Private/MuT/CodeGenerator.cpp
@@ -1351,7 +1351,7 @@ namespace mu
                     IMAGE_STATE newState;
                     newState.m_imageSize = desc.m_size;
                     newState.m_imageRect = box<vec2<int>>::FromMinSize( vec2<int>(0,0), desc.m_size );
-                    newState.m_layoutBlock = -1;
+                    newState.m_layoutBlockId = -1;
                     newState.m_pLayout = nullptr;
                     m_imageState.Add( newState );
 
@@ -1525,7 +1525,7 @@ namespace mu
                             IMAGE_STATE newState;
                             newState.m_imageSize = desc.m_size;
                             newState.m_imageRect = rect;
-                            newState.m_layoutBlock = pLayout->m_blocks[b].m_id;
+                            newState.m_layoutBlockId = pLayout->m_blocks[b].m_id;
                             newState.m_pLayout = pLayout;
                             m_imageState.Add( newState );
 
@@ -1626,7 +1626,7 @@ namespace mu
                                             IMAGE_STATE newState;
                                             newState.m_imageSize = extendDesc.m_size;
                                             newState.m_imageRect = rect;
-                                            newState.m_layoutBlock = pExtendLayout->m_blocks[b].m_id;
+                                            newState.m_layoutBlockId = pExtendLayout->m_blocks[b].m_id;
                                             newState.m_pLayout = pExtendLayout;
                                             m_imageState.Add( newState );
 
diff --git a/Engine/Plugins/Experimental/Mutable/Source/MutableTools/Private/MuT/CodeGenerator.h b/Engine/Plugins/Experimental/Mutable/Source/MutableTools/Private/MuT/CodeGenerator.h
--- a/Engine/Plugins/Experimental/Mutable/Source/MutableTools/Private/MuT/CodeGenerator.h
+++ b/Engine/Plugins/Experimental/Mutable/Source/MutableTools/Private/MuT/CodeGenerator.h
@@ -246,7 +246,7 @@ namespace mu
         {
             vec2<int> m_imageSize;
             box< vec2<int> > m_imageRect;
-            int32 m_layoutBlock;
+            int32 m_layoutBlockId;
             LayoutPtrConst m_pLayout;
         };
 		TArray<IMAGE_STATE> m_imageState;
diff --git a/Engine/Plugins/Experimental/Mutable/Source/MutableTools/Private/MuT/CodeGenerator_Image.cpp b/Engine/Plugins/Experimental/Mutable/Source/MutableTools/Private/MuT/CodeGenerator_Image.cpp
--- a/Engine/Plugins/Experimental/Mutable/Source/MutableTools/Private/MuT/CodeGenerator_Image.cpp
+++ b/Engine/Plugins/Experimental/Mutable/Source/MutableTools/Private/MuT/CodeGenerator_Image.cpp
@@ -135,7 +135,7 @@ namespace mu
 			newState.m_imageRect.size = desc.m_size;
 			newState.m_imageRect.min[0] = 0;
 			newState.m_imageRect.min[1] = 0;
-			newState.m_layoutBlock = -1;
+			newState.m_layoutBlockId = -1;
 			m_imageState.Add(newState);
 			addedImageState = true;
 		}
@@ -1440,6 +1440,8 @@ namespace mu
 
         pop->SetChild( pop->op.args.MeshProject.projector, projectorResult.op );
 
+		int32 LayoutBlockIndex = m_imageState.Last().m_pLayout->m_blocks.IndexOfByPredicate([&](const Layout::FBlock& Block) { return Block.m_id == m_imageState.Last().m_layoutBlockId; });
+
         // Mesh
         if ( node.m_pMesh )
         {
@@ -1477,12 +1479,16 @@ namespace mu
             else
             {
                 // Extract the mesh layout block
-                if ( m_imageState.Num() && m_imageState.Last().m_layoutBlock>=0 )
+                if ( m_imageState.Num() && m_imageState.Last().m_layoutBlockId >=0 )
                 {
                     Ptr<ASTOpMeshExtractLayoutBlocks> eop = new ASTOpMeshExtractLayoutBlocks();
                     eop->source = pop->children[pop->op.args.MeshProject.mesh].child();
                     eop->layout = node.m_layout;
-                    eop->blocks.Add( m_imageState.Last().m_layoutBlock );
+					
+					// Match the block id of the block we are generating with the id that resulted in the generated mesh
+					int32 GeneratedLayoutBlockId = MeshResult.GeneratedLayouts[node.m_layout]->m_blocks[LayoutBlockIndex].m_id;
+
+                    eop->blocks.Add(GeneratedLayoutBlockId);
 
                     pop->SetChild( pop->op.args.MeshProject.mesh, eop );
                 }
@@ -1506,13 +1512,12 @@ namespace mu
         else
         {
             // This argument is required
-            MeshPtrConst pMesh = new Mesh();
+            Ptr<const Mesh> pMesh = new Mesh();
             Ptr<ASTOpConstantResource> cop = new ASTOpConstantResource();
             cop->type = OP_TYPE::ME_CONSTANT;
             cop->SetValue( pMesh, m_compilerOptions->m_optimisationOptions.m_useDiskCache );
             pop->SetChild( pop->op.args.MeshProject.mesh, cop );
-            m_pErrorLog->GetPrivate()->Add( "Projector mesh not set.",
-                                            ELMT_ERROR, node.m_errorContext );
+            m_pErrorLog->GetPrivate()->Add( "Projector mesh not set.", ELMT_ERROR, node.m_errorContext );
         }
 
 
@@ -1535,7 +1540,7 @@ namespace mu
             newState.m_imageRect.min[0] = 0;
             newState.m_imageRect.min[1] = 0;
             newState.m_imageRect.size = desc.m_size;
-            newState.m_layoutBlock = -1;
+            newState.m_layoutBlockId = -1;
             m_imageState.Add( newState );
 
             // Generate
@@ -1555,7 +1560,7 @@ namespace mu
         // Image size, from the current block being generated
         op->op.args.ImageRasterMesh.sizeX = (uint16)m_imageState.Last().m_imageRect.size[0];
         op->op.args.ImageRasterMesh.sizeY = (uint16)m_imageState.Last().m_imageRect.size[1];
-        op->op.args.ImageRasterMesh.blockIndex = m_imageState.Last().m_layoutBlock;
+        op->op.args.ImageRasterMesh.blockIndex = LayoutBlockIndex;
 
         // Fading properties are optional, and stored in a colour
         if (node.m_pAngleFadeStart||node.m_pAngleFadeEnd)
diff --git a/Engine/Plugins/Experimental/Mutable/Source/MutableTools/Private/MuT/CodeGenerator_Scalar.cpp b/Engine/Plugins/Experimental/Mutable/Source/MutableTools/Private/MuT/CodeGenerator_Scalar.cpp
--- a/Engine/Plugins/Experimental/Mutable/Source/MutableTools/Private/MuT/CodeGenerator_Scalar.cpp
+++ b/Engine/Plugins/Experimental/Mutable/Source/MutableTools/Private/MuT/CodeGenerator_Scalar.cpp
@@ -172,7 +172,7 @@ namespace mu
 					newState.m_imageRect.min[0] = 0;
 					newState.m_imageRect.min[1] = 0;
 					newState.m_imageRect.size = desc.m_size;
-					newState.m_layoutBlock = -1;
+					newState.m_layoutBlockId = -1;
 					m_imageState.Add(newState);
 
 					// Generate
