[mutable] Fix wrapping projectors.
#rnx

diff --git a/Engine/Plugins/Experimental/Mutable/Source/MutableTools/Private/MuT/CodeGenerator_Image.cpp b/Engine/Plugins/Experimental/Mutable/Source/MutableTools/Private/MuT/CodeGenerator_Image.cpp
--- a/Engine/Plugins/Experimental/Mutable/Source/MutableTools/Private/MuT/CodeGenerator_Image.cpp
+++ b/Engine/Plugins/Experimental/Mutable/Source/MutableTools/Private/MuT/CodeGenerator_Image.cpp
@@ -1441,6 +1441,7 @@ namespace mu
         pop->SetChild( pop->op.args.MeshProject.projector, projectorResult.op );
 
 		int32 LayoutBlockIndex = m_imageState.Last().m_pLayout->m_blocks.IndexOfByPredicate([&](const Layout::FBlock& Block) { return Block.m_id == m_imageState.Last().m_layoutBlockId; });
+		int32 GeneratedLayoutBlockId = -1;
 
         // Mesh
         if ( node.m_pMesh )
@@ -1456,6 +1457,9 @@ namespace mu
 			MeshOptions.bUniqueVertexIDs = false;	// We don't need the IDs at this point.
             GenerateMesh( MeshOptions, MeshResult, node.m_pMesh );
 
+			// Match the block id of the block we are generating with the id that resulted in the generated mesh
+			GeneratedLayoutBlockId = MeshResult.GeneratedLayouts[node.m_layout]->m_blocks[LayoutBlockIndex].m_id;
+
             pop->SetChild( pop->op.args.MeshProject.mesh, MeshResult.meshOp );
 
             if (projectorResult.type == PROJECTOR_TYPE::WRAPPING)
@@ -1484,9 +1488,6 @@ namespace mu
                     Ptr<ASTOpMeshExtractLayoutBlocks> eop = new ASTOpMeshExtractLayoutBlocks();
                     eop->source = pop->children[pop->op.args.MeshProject.mesh].child();
                     eop->layout = node.m_layout;
-					
-					// Match the block id of the block we are generating with the id that resulted in the generated mesh
-					int32 GeneratedLayoutBlockId = MeshResult.GeneratedLayouts[node.m_layout]->m_blocks[LayoutBlockIndex].m_id;
 
                     eop->blocks.Add(GeneratedLayoutBlockId);
 
@@ -1560,7 +1561,8 @@ namespace mu
         // Image size, from the current block being generated
         op->op.args.ImageRasterMesh.sizeX = (uint16)m_imageState.Last().m_imageRect.size[0];
         op->op.args.ImageRasterMesh.sizeY = (uint16)m_imageState.Last().m_imageRect.size[1];
-        op->op.args.ImageRasterMesh.blockIndex = LayoutBlockIndex;
+		// \TODO: Review naming of arg
+        op->op.args.ImageRasterMesh.blockIndex = GeneratedLayoutBlockId;
 
         // Fading properties are optional, and stored in a colour
         if (node.m_pAngleFadeStart||node.m_pAngleFadeEnd)
