[mutable] Now, if there is a layout without any layout block, one with the grid size is automatically added. A message warns the user of that also.

#rnx

diff --git a/Engine/Plugins/Experimental/Mutable/Source/CustomizableObjectEditor/Private/MuCOE/GenerateMutableSource/GenerateMutableSourceLayout.cpp b/Engine/Plugins/Experimental/Mutable/Source/CustomizableObjectEditor/Private/MuCOE/GenerateMutableSource/GenerateMutableSourceLayout.cpp
--- a/Engine/Plugins/Experimental/Mutable/Source/CustomizableObjectEditor/Private/MuCOE/GenerateMutableSource/GenerateMutableSourceLayout.cpp
+++ b/Engine/Plugins/Experimental/Mutable/Source/CustomizableObjectEditor/Private/MuCOE/GenerateMutableSource/GenerateMutableSourceLayout.cpp
@@ -77,7 +77,7 @@ mu::NodeLayoutPtr GenerateMutableSourceLayout(const UEdGraphPin * Pin, FMutableG
 
 		LayoutNode->SetGridSize(TypedNodeBlocks->Layout->GetGridSize().X, TypedNodeBlocks->Layout->GetGridSize().Y);
 		LayoutNode->SetMaxGridSize(TypedNodeBlocks->Layout->GetMaxGridSize().X, TypedNodeBlocks->Layout->GetMaxGridSize().Y);
-		LayoutNode->SetBlockCount(TypedNodeBlocks->Layout->Blocks.Num());
+		LayoutNode->SetBlockCount(TypedNodeBlocks->Layout->Blocks.Num() ? TypedNodeBlocks->Layout->Blocks.Num() : 1);
 
 		mu::EPackStrategy strategy = mu::EPackStrategy::RESIZABLE_LAYOUT;
 
@@ -95,15 +95,26 @@ mu::NodeLayoutPtr GenerateMutableSourceLayout(const UEdGraphPin * Pin, FMutableG
 
 		LayoutNode->SetLayoutPackingStrategy(strategy);
 
-		for (int BlockIndex = 0; BlockIndex < TypedNodeBlocks->Layout->Blocks.Num(); ++BlockIndex)
+		if (TypedNodeBlocks->Layout->Blocks.Num())
 		{
-			LayoutNode->SetBlock(BlockIndex,
-				TypedNodeBlocks->Layout->Blocks[BlockIndex].Min.X,
-				TypedNodeBlocks->Layout->Blocks[BlockIndex].Min.Y,
-				TypedNodeBlocks->Layout->Blocks[BlockIndex].Max.X - TypedNodeBlocks->Layout->Blocks[BlockIndex].Min.X,
-				TypedNodeBlocks->Layout->Blocks[BlockIndex].Max.Y - TypedNodeBlocks->Layout->Blocks[BlockIndex].Min.Y);
+			for (int BlockIndex = 0; BlockIndex < TypedNodeBlocks->Layout->Blocks.Num(); ++BlockIndex)
+			{
+				LayoutNode->SetBlock(BlockIndex,
+					TypedNodeBlocks->Layout->Blocks[BlockIndex].Min.X,
+					TypedNodeBlocks->Layout->Blocks[BlockIndex].Min.Y,
+					TypedNodeBlocks->Layout->Blocks[BlockIndex].Max.X - TypedNodeBlocks->Layout->Blocks[BlockIndex].Min.X,
+					TypedNodeBlocks->Layout->Blocks[BlockIndex].Max.Y - TypedNodeBlocks->Layout->Blocks[BlockIndex].Min.Y);
+
+				LayoutNode->SetBlockPriority(BlockIndex, TypedNodeBlocks->Layout->Blocks[BlockIndex].Priority);
+			}
+		}
+		else
+		{
+			FString msg = "Layout without any block found. A grid sized block will be used instead.";
+			GenerationContext.Compiler->CompilerLog(FText::FromString(msg), Node, EMessageSeverity::Warning);
 
-			LayoutNode->SetBlockPriority(BlockIndex, TypedNodeBlocks->Layout->Blocks[BlockIndex].Priority);
+			LayoutNode->SetBlock(0, 0, 0, TypedNodeBlocks->Layout->GetGridSize().X, TypedNodeBlocks->Layout->GetGridSize().Y);
+			LayoutNode->SetBlockPriority(0, 0);
 		}
 	}
 	
diff --git a/Engine/Plugins/Experimental/Mutable/Source/CustomizableObjectEditor/Private/MuCOE/GenerateMutableSource/GenerateMutableSourceMesh.cpp b/Engine/Plugins/Experimental/Mutable/Source/CustomizableObjectEditor/Private/MuCOE/GenerateMutableSource/GenerateMutableSourceMesh.cpp
--- a/Engine/Plugins/Experimental/Mutable/Source/CustomizableObjectEditor/Private/MuCOE/GenerateMutableSource/GenerateMutableSourceMesh.cpp
+++ b/Engine/Plugins/Experimental/Mutable/Source/CustomizableObjectEditor/Private/MuCOE/GenerateMutableSource/GenerateMutableSourceMesh.cpp
@@ -3326,7 +3326,7 @@ mu::NodeMeshPtr GenerateMutableSourceMesh(const UEdGraphPin * Pin,
 
 					LayoutNode->SetGridSize(Layouts[i]->GetGridSize().X, Layouts[i]->GetGridSize().Y);
 					LayoutNode->SetMaxGridSize(Layouts[i]->GetMaxGridSize().X, Layouts[i]->GetMaxGridSize().Y);
-					LayoutNode->SetBlockCount(Layouts[i]->Blocks.Num());
+					LayoutNode->SetBlockCount(Layouts[i]->Blocks.Num() ? Layouts[i]->Blocks.Num() : 1);
 
 					mu::EPackStrategy strategy = mu::EPackStrategy::RESIZABLE_LAYOUT;
 
@@ -3359,7 +3359,11 @@ mu::NodeMeshPtr GenerateMutableSourceMesh(const UEdGraphPin * Pin,
 					}
 					else
 					{
+						FString msg = "Mesh Column [" + ColumnName + "] Layout doesn't has any block. A grid sized block will be used instead.";
+						GenerationContext.Compiler->CompilerLog(FText::FromString(msg), Node, EMessageSeverity::Warning);
+
 						LayoutNode->SetBlock(0, 0, 0, Layouts[i]->GetGridSize().X, Layouts[i]->GetGridSize().Y);
+						LayoutNode->SetBlockPriority(0, 0);
 					}
 
 					MeshTableNode->SetLayout(i, LayoutNode);
