[Mutable] Split ReloadParametersFromObject.

- Split into two funcitons:
    - GetParameters: Creates a Core Parameters object.
    - ReloadParameters: Copies parametres from the core.
- Fixed an issue where SetCurrentState was setting an incorrecte state index.
- Minor code improvements.
- Fix UCustomizableObjectInstance::PostLoad not calling ReloadParameters.

#rnx

diff --git a/Engine/Plugins/Experimental/Mutable/Source/CustomizableObject/Private/MuCO/CustomizableInstancePrivateData.h b/Engine/Plugins/Experimental/Mutable/Source/CustomizableObject/Private/MuCO/CustomizableInstancePrivateData.h
--- a/Engine/Plugins/Experimental/Mutable/Source/CustomizableObject/Private/MuCO/CustomizableInstancePrivateData.h
+++ b/Engine/Plugins/Experimental/Mutable/Source/CustomizableObject/Private/MuCO/CustomizableInstancePrivateData.h
@@ -190,8 +190,15 @@ public:
 	// Return an event that will be fired when the assets  have been loaded. It returns null if no asset needs loading.
 	FGraphEventRef LoadAdditionalAssetsAsync(const TSharedPtr<FMutableOperationData>& OperationData, UCustomizableObjectInstance* Public, struct FStreamableManager &StreamableManager);
 	void AdditionalAssetsAsyncLoaded(UCustomizableObjectInstance* Public);
-	
-	mu::ParametersPtr ReloadParametersFromObject(UCustomizableObjectInstance* Public, bool ClearLastMeshIds = false);
+
+	/** Temporal function. Used by GetParameters and ReloadParameters but its body should be moved somewhere else (probably on some other Instance update function). */
+	void InstanceUpdateFlags(const UCustomizableObjectInstance& Public);
+
+	/** See FCustomizableObjectInstanceDescriptor::GetParameters(...). */
+	mu::ParametersPtr GetParameters(UCustomizableObjectInstance* Public);
+
+	/** See FCustomizableObjectInstanceDescriptor::ReloadParameters(...). */
+	void ReloadParameters(UCustomizableObjectInstance* Public);
 
 	void TickUpdateCloseCustomizableObjects(UCustomizableObjectInstance* Public);
 	void UpdateInstanceIfNotGenerated(UCustomizableObjectInstance* Public, bool bAsync);
diff --git a/Engine/Plugins/Experimental/Mutable/Source/CustomizableObject/Private/MuCO/CustomizableObjectInstance.cpp b/Engine/Plugins/Experimental/Mutable/Source/CustomizableObject/Private/MuCO/CustomizableObjectInstance.cpp
--- a/Engine/Plugins/Experimental/Mutable/Source/CustomizableObject/Private/MuCO/CustomizableObjectInstance.cpp
+++ b/Engine/Plugins/Experimental/Mutable/Source/CustomizableObject/Private/MuCO/CustomizableObjectInstance.cpp
@@ -215,7 +215,7 @@ const FCustomizableObjectInstanceDescriptor& UCustomizableObjectInstance::GetDes
 void UCustomizableObjectInstance::SetDescriptor(const FCustomizableObjectInstanceDescriptor& InDescriptor)
 {
 	Descriptor = InDescriptor;
-	PrivateData->ReloadParametersFromObject(this, true);
+	PrivateData->ReloadParameters(this);
 }
 
 
@@ -502,7 +502,7 @@ void UCustomizableObjectInstance::PostLoad()
 
 	Super::PostLoad();
 
-	Descriptor.CreateParametersLookupTable();
+	Descriptor.ReloadParameters();
 }
 
 
@@ -567,7 +567,26 @@ FCustomizableObjectProjector UCustomizableObjectInstance::GetProjectorDefaultVal
 }
 
 
-mu::ParametersPtr UCustomizableInstancePrivateData::ReloadParametersFromObject(UCustomizableObjectInstance* Public, bool ClearLastMeshIds)
+void UCustomizableInstancePrivateData::InstanceUpdateFlags(const UCustomizableObjectInstance& Public)
+{
+	const UCustomizableObject* CustomizableObject = Public.GetCustomizableObject();
+	
+	FirstLODAvailable = CustomizableObject->LODSettings.FirstLODAvailable;
+
+	if (CustomizableObject->LODSettings.bLODStreamingEnabled || !UCustomizableObjectSystem::GetInstance()->IsProgressiveMipStreamingEnabled())
+	{
+		SetCOInstanceFlags(LODsStreamingEnabled);
+		NumMaxLODsToStream = CustomizableObject->LODSettings.NumLODsToStream;
+	}
+	else
+	{
+		ClearCOInstanceFlags(LODsStreamingEnabled);
+		NumMaxLODsToStream = 0;
+	}
+}
+
+
+mu::ParametersPtr UCustomizableInstancePrivateData::GetParameters(UCustomizableObjectInstance* Public) 
 {
 	const UCustomizableObject* CustomizableObject = Public->GetCustomizableObject();
 	if (!CustomizableObject)
@@ -581,19 +600,28 @@ mu::ParametersPtr UCustomizableInstancePrivateData::ReloadParametersFromObject(U
 		return nullptr;
 	}
 
-	FirstLODAvailable = CustomizableObject->LODSettings.FirstLODAvailable;
+	InstanceUpdateFlags(*Public); // TODO Move somewhere else.
+		
+	return Public->GetDescriptor().GetParameters();
+}
 
-	if (CustomizableObject->LODSettings.bLODStreamingEnabled || !UCustomizableObjectSystem::GetInstance()->IsProgressiveMipStreamingEnabled())
+
+void UCustomizableInstancePrivateData::ReloadParameters(UCustomizableObjectInstance* Public)
+{
+	const UCustomizableObject* CustomizableObject = Public->GetCustomizableObject();
+	if (!CustomizableObject)
 	{
-		SetCOInstanceFlags(LODsStreamingEnabled);
-		NumMaxLODsToStream = CustomizableObject->LODSettings.NumLODsToStream;
+		// May happen when deleting assets in editor, while they are still open.
+		return;
 	}
-	else
-	{
-		ClearCOInstanceFlags(LODsStreamingEnabled);
-		NumMaxLODsToStream = 0;
+
+	if (!CustomizableObject->IsCompiled())
+	{	
+		return;
 	}
-	
+
+	InstanceUpdateFlags(*Public); // TODO Move somewhere else.
+
 #if WITH_EDITOR
 	if(!Public->ProjectorAlphaChange)
 #endif
@@ -601,19 +629,16 @@ mu::ParametersPtr UCustomizableInstancePrivateData::ReloadParametersFromObject(U
 		ProjectorStates.Reset();
 	}
 	
-	if (ClearLastMeshIds)
-	{
-		ClearAllLastMeshIds();
-	}
+	ClearAllLastMeshIds();
 	
-	return Public->GetDescriptor().ReloadParametersFromObject();
+	Public->GetDescriptor().ReloadParameters();
 }
 
 
 void UCustomizableObjectInstance::SetObject(UCustomizableObject* InObject)
 {
-	Descriptor.SetCustomizableObject(InObject);
-	PrivateData->ReloadParametersFromObject(this, true);
+	Descriptor.SetCustomizableObject(*InObject);
+	PrivateData->ReloadParameters(this);
 }
 
 
@@ -659,13 +684,13 @@ void UCustomizableObjectInstance::SetState(const int32 InState)
 
 FString UCustomizableObjectInstance::GetCurrentState() const
 {
-	return GetCustomizableObject()->GetStateName(GetState());
+	return Descriptor.GetCurrentState();
 }
 
 
 void UCustomizableObjectInstance::SetCurrentState(const FString& StateName)
 {
-	SetState(GetCustomizableObject()->FindState(StateName));
+	Descriptor.SetCurrentState(StateName);
 }
 
 
diff --git a/Engine/Plugins/Experimental/Mutable/Source/CustomizableObject/Private/MuCO/CustomizableObjectInstanceDescriptor.cpp b/Engine/Plugins/Experimental/Mutable/Source/CustomizableObject/Private/MuCO/CustomizableObjectInstanceDescriptor.cpp
--- a/Engine/Plugins/Experimental/Mutable/Source/CustomizableObject/Private/MuCO/CustomizableObjectInstanceDescriptor.cpp
+++ b/Engine/Plugins/Experimental/Mutable/Source/CustomizableObject/Private/MuCO/CustomizableObjectInstanceDescriptor.cpp
@@ -43,8 +43,7 @@ FString GetAvailableOptionsString(const UCustomizableObject& CustomizableObject,
 
 FCustomizableObjectInstanceDescriptor::FCustomizableObjectInstanceDescriptor(UCustomizableObject &Object)
 {
-	CustomizableObject = &Object;	
-	Init();
+	SetCustomizableObject(Object);
 }
 
 
@@ -446,14 +445,14 @@ void FCustomizableObjectInstanceDescriptor::SetBuildParameterDecorations(const b
 }
 
 
-void FCustomizableObjectInstanceDescriptor::SetCustomizableObject(UCustomizableObject* InCustomizableObject)
+void FCustomizableObjectInstanceDescriptor::SetCustomizableObject(UCustomizableObject& InCustomizableObject)
 {
-	CustomizableObject = InCustomizableObject;
-	Init();
+	CustomizableObject = &InCustomizableObject;
+	ReloadParameters();
 }
 
 
-mu::ParametersPtr FCustomizableObjectInstanceDescriptor::ReloadParametersFromObject()
+mu::ParametersPtr FCustomizableObjectInstanceDescriptor::GetParameters() const
 {
 	if (!CustomizableObject)
 	{
@@ -465,494 +464,513 @@ mu::ParametersPtr FCustomizableObjectInstanceDescriptor::ReloadParametersFromObj
 		return nullptr;
 	}
 	
-	int32 OldState = GetState();
-	TArray<FCustomizableObjectBoolParameterValue> OldBoolParameters = GetBoolParameters();
-	TArray<FCustomizableObjectIntParameterValue> OldIntParameters = GetIntParameters();
-	TArray<FCustomizableObjectFloatParameterValue> OldFloatParameters = GetFloatParameters();
-	TArray<FCustomizableObjectTextureParameterValue> OldTextureParameters = GetTextureParameters();
-	TArray<FCustomizableObjectVectorParameterValue> OldVectorParameters = GetVectorParameters();
-	TArray<FCustomizableObjectProjectorParameterValue> OldProjectorParameters = GetProjectorParameters();
-
-	mu::ParametersPtr MutableParameters;
-	SetState(FMath::Clamp(OldState, 0, CustomizableObject->GetStateCount() - 1));
-	GetBoolParameters().Reset();
-	GetIntParameters().Reset();
-	GetFloatParameters().Reset();
-	GetTextureParameters().Reset();
-	GetVectorParameters().Reset();
-	GetProjectorParameters().Reset();
-	
 	if (!CustomizableObject->GetPrivate()->GetModel())
 	{
-		UE_LOG(LogMutable, Warning, TEXT("[ReloadParametersFromObject] No model in object [%s], generated empty parameters for [%s] "), *CustomizableObject->GetName());
 		return nullptr;
 	}
 
-	TArray<bool> OldIntParametersUsed;
-	OldIntParametersUsed.SetNumUninitialized(OldIntParameters.Num());
-	FMemory::Memzero(OldIntParametersUsed.GetData(), OldIntParametersUsed.GetAllocatedSize());
-
-	MutableParameters = CustomizableObject->GetPrivate()->GetModel()->NewParameters();
-	int32 paramCount = MutableParameters->GetCount();
+	mu::ParametersPtr MutableParameters = CustomizableObject->GetPrivate()->GetModel()->NewParameters();
 
-	for (int32 paramIndex = 0; paramIndex<paramCount; ++paramIndex)
+	const int32 ParamCount = MutableParameters->GetCount();
+	for (int32 ParamIndex = 0; ParamIndex < ParamCount; ++ParamIndex)
 	{
-		FString Name = MutableParameters->GetName(paramIndex);
-		FString Uid = MutableParameters->GetUid(paramIndex);
-		mu::PARAMETER_TYPE mutableType = MutableParameters->GetType(paramIndex);
+		const FString& Name = MutableParameters->GetName(ParamIndex);
+		const FString& Uid = MutableParameters->GetUid(ParamIndex);
+		const mu::PARAMETER_TYPE MutableType = MutableParameters->GetType(ParamIndex);
 
-		switch (mutableType)
+		switch (MutableType)
 		{
 		case mu::PARAMETER_TYPE::T_BOOL:
 		{
-			FCustomizableObjectBoolParameterValue param;
-			param.ParameterName = Name;
-			param.Uid = Uid;
-
-			bool found = false;
-			for (int32 i = 0; i < OldBoolParameters.Num(); ++i)
+			for (const FCustomizableObjectBoolParameterValue& BoolParameter : BoolParameters)
 			{
-				if (OldBoolParameters[i].ParameterName == Name)
+				if (BoolParameter.ParameterName == Name || (!Uid.IsEmpty() && BoolParameter.Uid == Uid))
 				{
-					found = true;
-					param.ParameterValue = OldBoolParameters[i].ParameterValue;
-					MutableParameters->SetBoolValue(paramIndex, param.ParameterValue);
+					MutableParameters->SetBoolValue(ParamIndex,  BoolParameter.ParameterValue);
 					break;
 				}
 			}
 
-			if (!found)
+			break;
+		}
+
+		case mu::PARAMETER_TYPE::T_INT:
+		{
+			for (const FCustomizableObjectIntParameterValue& IntParameter : IntParameters)
 			{
-				for (int32 i = 0; i < OldBoolParameters.Num(); ++i)
+				if (IntParameter.ParameterName.Equals(Name, ESearchCase::CaseSensitive) || (!Uid.IsEmpty() && IntParameter.Uid == Uid))
 				{
-					if (!Uid.IsEmpty() && OldBoolParameters[i].Uid == Uid)
+					if (mu::RangeIndexPtr RangeIdxPtr = MutableParameters->NewRangeIndex(ParamIndex))
 					{
-						found = true;
-						param.ParameterValue = OldBoolParameters[i].ParameterValue;
-						MutableParameters->SetBoolValue(paramIndex, param.ParameterValue);
-						break;
+						for (int RangeIndex = 0; RangeIndex < IntParameter.ParameterRangeValueNames.Num(); ++RangeIndex)
+						{
+							RangeIdxPtr->SetPosition(0, RangeIndex);
+
+							const FString& RangeValueName = IntParameter.ParameterRangeValueNames[RangeIndex];
+							const int32 Value = CustomizableObject->FindIntParameterValue(ParamIndex, RangeValueName);
+							MutableParameters->SetIntValue(ParamIndex, Value, RangeIdxPtr);
+						}
+					}
+					else
+					{
+						const int32 Value = CustomizableObject->FindIntParameterValue(ParamIndex, IntParameter.ParameterValueName);
+						MutableParameters->SetIntValue(ParamIndex, Value);
 					}
-				}
-			}
 
-			if (!found)
-			{
-				param.ParameterValue = MutableParameters->GetBoolValue(paramIndex);
+					break;
+				}
 			}
 
-			GetBoolParameters().Add(param);
 			break;
 		}
 
-		case mu::PARAMETER_TYPE::T_INT:
+		case mu::PARAMETER_TYPE::T_FLOAT:
 		{
-			FString ParameterValueName;
-			TArray<FString> ParameterRangeValueNames;
-
-			bool found = false;
-			for (int32 i = 0; i < OldIntParameters.Num(); ++i)
+			for (const FCustomizableObjectFloatParameterValue& FloatParameter : FloatParameters)
 			{
-				if (OldIntParametersUsed[i] == false && OldIntParameters[i].ParameterName.Equals(Name, ESearchCase::CaseSensitive))
+				if (FloatParameter.ParameterName == Name || (!Uid.IsEmpty() && FloatParameter.Uid == Uid))
 				{
-					found = true;
-					OldIntParametersUsed[i] = true;
-
-					mu::RangeIndexPtr RangeIdxPtr = MutableParameters->NewRangeIndex(paramIndex);
-
-					if (!RangeIdxPtr)
-					{
-					ParameterValueName = OldIntParameters[i].ParameterValueName;
-					int32 Value = CustomizableObject->FindIntParameterValue(paramIndex, ParameterValueName);
-					MutableParameters->SetIntValue(paramIndex, Value);
-					}
-					else
+					if (mu::RangeIndexPtr RangeIdxPtr = MutableParameters->NewRangeIndex(ParamIndex))
 					{
-						for (int RangeIndex = 0; RangeIndex < OldIntParameters[i].ParameterRangeValueNames.Num(); ++RangeIndex)
+						for (int RangeIndex = 0; RangeIndex < FloatParameter.ParameterRangeValues.Num(); ++RangeIndex)
 						{
-							FString AuxParameterValueName = OldIntParameters[i].ParameterRangeValueNames[RangeIndex];
-							ParameterRangeValueNames.Add(AuxParameterValueName);
-							int32 Value = CustomizableObject->FindIntParameterValue(paramIndex, AuxParameterValueName);
 							RangeIdxPtr->SetPosition(0, RangeIndex);
-							MutableParameters->SetIntValue(paramIndex, Value, RangeIdxPtr);
+							MutableParameters->SetFloatValue(ParamIndex, FloatParameter.ParameterRangeValues[RangeIndex], RangeIdxPtr);
 						}
 					}
+					else
+					{
+						MutableParameters->SetFloatValue(ParamIndex, FloatParameter.ParameterValue);
+					}
 
 					break;
 				}
 			}
 
-			if (!found)
+			break;
+		}
+
+		case mu::PARAMETER_TYPE::T_COLOUR:
+		{
+			for (const FCustomizableObjectVectorParameterValue& VectorParameter : VectorParameters)
 			{
-				for (int32 i = 0; i < OldIntParameters.Num(); ++i)
+				if (VectorParameter.ParameterName == Name || (!Uid.IsEmpty() && VectorParameter.Uid == Uid))
 				{
-					if (!Uid.IsEmpty() && OldIntParametersUsed[i] == false && OldIntParameters[i].Uid == Uid)
-					{
-						found = true;
-						OldIntParametersUsed[i] = true;
+					MutableParameters->SetColourValue(ParamIndex, VectorParameter.ParameterValue.R, VectorParameter.ParameterValue.G, VectorParameter.ParameterValue.B);
 
-						mu::RangeIndexPtr RangeIdxPtr = MutableParameters->NewRangeIndex(paramIndex);
+					break;
+				}
+			}
 
-						if (!RangeIdxPtr)
+			break;
+		}
+
+		case mu::PARAMETER_TYPE::T_PROJECTOR:
+		{
+			for (const auto& ProjectorParameter : ProjectorParameters)
+			{
+				if (ProjectorParameter.ParameterName == Name || (!Uid.IsEmpty() && ProjectorParameter.Uid == Uid))
+				{
+					auto CopyProjector = [&MutableParameters, ParamIndex](const FCustomizableObjectProjector& Value, const mu::RangeIndexPtr RangeIdxPtr = nullptr)
+					{
+						switch (Value.ProjectionType)
 						{
-						ParameterValueName = OldIntParameters[i].ParameterValueName;
-						int32 Value = CustomizableObject->FindIntParameterValue(paramIndex, ParameterValueName);
-						MutableParameters->SetIntValue(paramIndex, Value);
-						}
-						else
+						case ECustomizableObjectProjectorType::Planar:
+						case ECustomizableObjectProjectorType::Wrapping:
+							MutableParameters->SetProjectorValue(ParamIndex,
+								Value.Position,
+								Value.Direction,
+								Value.Up,
+								Value.Scale,
+								Value.Angle,
+								RangeIdxPtr);
+							break;
+
+						case ECustomizableObjectProjectorType::Cylindrical:
 						{
-							for (int RangeIndex = 0; RangeIndex < OldIntParameters[i].ParameterRangeValueNames.Num(); ++RangeIndex)
-							{
-								FString AuxParameterValueName = OldIntParameters[i].ParameterRangeValueNames[RangeIndex];
-								ParameterRangeValueNames.Add(AuxParameterValueName);
-								int32 Value = CustomizableObject->FindIntParameterValue(paramIndex, AuxParameterValueName);
-								RangeIdxPtr->SetPosition(0, RangeIndex);
-								MutableParameters->SetIntValue(paramIndex, Value, RangeIdxPtr);
-							}
+							// Apply strange swizzle for scales
+							// TODO: try to avoid this
+							const float Radius = FMath::Abs(Value.Scale[0] / 2.0f);
+							const float Height = Value.Scale[2];
+							// TODO: try to avoid this
+							MutableParameters->SetProjectorValue(ParamIndex,
+									Value.Position,
+									-Value.Direction,
+									-Value.Up,
+									FVector3f(-Height, Radius, Radius),
+									Value.Angle,
+									RangeIdxPtr);
+							break;
 						}
 
-						break;
+						default:
+							check(false); // Not implemented.
+						}
+					};
+
+					if (const mu::RangeIndexPtr RangeIdxPtr = MutableParameters->NewRangeIndex(ParamIndex))
+					{
+						for (int RangeIndex = 0; RangeIndex < ProjectorParameter.RangeValues.Num(); ++RangeIndex)
+						{
+							RangeIdxPtr->SetPosition(0, RangeIndex);
+							CopyProjector(ProjectorParameter.RangeValues[RangeIndex], RangeIdxPtr);
+						}
+					}
+					else
+					{
+						CopyProjector(ProjectorParameter.Value);
 					}
 				}
 			}
 
-			if (!found)
-			{
-				mu::RangeIndexPtr RangeIdxPtr = MutableParameters->NewRangeIndex(paramIndex);
+			break;
+		}
 
-				if (!RangeIdxPtr.get())
+		case mu::PARAMETER_TYPE::T_IMAGE:
+		{
+			for (const FCustomizableObjectTextureParameterValue& TextureParameter : TextureParameters)
+			{
+				if (TextureParameter.ParameterName == Name || (!Uid.IsEmpty() && TextureParameter.Uid == Uid))
 				{
-				int32 Value = MutableParameters->GetIntValue(paramIndex);
-				ParameterValueName = CustomizableObject->FindIntParameterValueName(paramIndex, Value);
+					MutableParameters->SetImageValue(ParamIndex, TextureParameter.ParameterValue);
 				}
-				else
-				{
-					int32 ValueCount = MutableParameters->GetValueCount(paramIndex);
+	
+				break;
+			}
 
-					for (int32 ValueIterator = 0; ValueIterator < ValueCount; ++ValueIterator)
-					{
-						mu::RangeIndexPtr RangeValueIdxPtr = MutableParameters->GetValueIndex(paramIndex, ValueIterator);
-						int32 ValueIndex = RangeValueIdxPtr->GetPosition(0);
+			break;
+		}
 
-						//if (!param.RangeValues.IsValidIndex(ValueIndex))
-						//{
-						//	param.RangeValues.AddDefaulted(ValueIndex + 1 - param.RangeValues.Num());
-						//}
+		default:
+			check(false); // Missing case.
+			break;
+		}
+	}
+	
+	return MutableParameters;
+}
 
-						int32 Value = MutableParameters->GetIntValue(paramIndex, RangeValueIdxPtr);
-						FString AuxParameterValueName = CustomizableObject->FindIntParameterValueName(paramIndex, Value);
-						//param.RangeValues[ValueIndex] = AuxParameterValueName;
 
-						if (!ParameterRangeValueNames.IsValidIndex(ValueIndex))
-						{
-							ParameterRangeValueNames.AddDefaulted(ValueIndex + 1 - ParameterRangeValueNames.Num());
-						}
+void FCustomizableObjectInstanceDescriptor::ReloadParameters()
+{
+	if (!CustomizableObject)
+	{
+		return;
+	}
 
-						ParameterRangeValueNames[ValueIndex] = AuxParameterValueName;
-					}
-				}
-			}
+	if (!CustomizableObject->IsCompiled())
+	{	
+		return;
+	}
 
-			GetIntParameters().Emplace(Name, ParameterValueName, Uid, ParameterRangeValueNames);
+	SetState(FMath::Clamp(GetState(), 0, CustomizableObject->GetStateCount() - 1));
+
+	TArray<FCustomizableObjectBoolParameterValue> OldBoolParameters = BoolParameters;
+	TArray<FCustomizableObjectIntParameterValue> OldIntParameters = IntParameters;
+	TArray<FCustomizableObjectFloatParameterValue> OldFloatParameters = FloatParameters;
+	TArray<FCustomizableObjectTextureParameterValue> OldTextureParameters = TextureParameters;
+	TArray<FCustomizableObjectVectorParameterValue> OldVectorParameters = VectorParameters;
+	TArray<FCustomizableObjectProjectorParameterValue> OldProjectorParameters = ProjectorParameters;
+
+	BoolParameters.Reset();
+	IntParameters.Reset();
+	FloatParameters.Reset();
+	TextureParameters.Reset();
+	VectorParameters.Reset();
+	ProjectorParameters.Reset();
+	
+	if (!CustomizableObject->GetPrivate()->GetModel())
+	{
+		UE_LOG(LogMutable, Warning, TEXT("[ReloadParametersFromObject] No model in object [%s], generated empty parameters for [%s] "), *CustomizableObject->GetName());
+		return;
+	}
+
+	mu::ParametersPtr MutableParameters = CustomizableObject->GetPrivate()->GetModel()->NewParameters();
+
+	int32 ParamCount = MutableParameters->GetCount();
+	for (int32 ParamIndex = 0; ParamIndex < ParamCount; ++ParamIndex)
+	{
+		const FString& Name = MutableParameters->GetName(ParamIndex);
+		const FString& Uid = MutableParameters->GetUid(ParamIndex);
+		const mu::PARAMETER_TYPE MutableType = MutableParameters->GetType(ParamIndex);
+
+		switch (MutableType)
+		{
+		case mu::PARAMETER_TYPE::T_BOOL:
+		{
+			FCustomizableObjectBoolParameterValue Param;
+			Param.ParameterName = Name;
+			Param.Uid = Uid;
+
+			auto FindByNameAndUid = [&](const FCustomizableObjectBoolParameterValue& P)
+			{
+				return P.ParameterName == Name || (!Uid.IsEmpty() && P.Uid == Uid);
+			};
+
+			if (FCustomizableObjectBoolParameterValue* Result = OldBoolParameters.FindByPredicate(FindByNameAndUid))
+			{	
+				Param.ParameterValue = Result->ParameterValue;
+			} 
+			else // Not found in Instance Parameters. Use Mutable Parameters.
+			{
+				Param.ParameterValue = MutableParameters->GetBoolValue(ParamIndex);
+			}
 
+			BoolParameters.Add(Param);
 			break;
 		}
 
-		case mu::PARAMETER_TYPE::T_FLOAT:
+		case mu::PARAMETER_TYPE::T_INT:
 		{
-			FCustomizableObjectFloatParameterValue param;
-			param.ParameterName = Name;
-			param.Uid = Uid;
+			FCustomizableObjectIntParameterValue Param;
+			Param.ParameterName = Name;
+			Param.Uid = Uid;
 
-			bool found = false;
-			for (int32 i = 0; i < OldFloatParameters.Num(); ++i)
+			auto FindByNameAndUid = [&](const FCustomizableObjectIntParameterValue& P)
+			{
+				return P.ParameterName == Name || (!Uid.IsEmpty() && P.Uid == Uid);
+			};
+			
+			if (FCustomizableObjectIntParameterValue* Result = OldIntParameters.FindByPredicate(FindByNameAndUid))
+			{	
+				if (mu::RangeIndexPtr RangeIdxPtr = MutableParameters->NewRangeIndex(ParamIndex))
+				{
+					Param.ParameterRangeValueNames = Result->ParameterRangeValueNames;
+				}
+				else
+				{
+					Param.ParameterValueName = Result->ParameterValueName;
+				}
+			} 
+			else // Not found in Instance Parameters. Use Mutable Parameters.
 			{
-				if (OldFloatParameters[i].ParameterName == Name || (!Uid.IsEmpty() && OldFloatParameters[i].Uid == Uid))
+				if (mu::RangeIndexPtr RangeIdxPtr = MutableParameters->NewRangeIndex(ParamIndex))
 				{
-					found = true;
-					param.ParameterValue = OldFloatParameters[i].ParameterValue;
-					param.ParameterRangeValues = OldFloatParameters[i].ParameterRangeValues;
+					const int32 ValueCount = MutableParameters->GetValueCount(ParamIndex);
 
-					mu::RangeIndexPtr RangeIdxPtr = MutableParameters->NewRangeIndex(paramIndex);
+					for (int ValueIndex = 0; ValueIndex < ValueCount; ++ValueIndex)
+					{	
+						const mu::RangeIndexPtr RangeValueIdxPtr = MutableParameters->GetValueIndex(ParamIndex, ValueIndex);
+						const int32 RangeIndex = RangeValueIdxPtr->GetPosition(0);
 
-					if (!RangeIdxPtr)
-					{
-						MutableParameters->SetFloatValue(paramIndex, param.ParameterValue);
-					}
-					else
-					{
-						for (int RangeIndex = 0; RangeIndex < OldFloatParameters[i].ParameterRangeValues.Num(); ++RangeIndex)
+						if (!Param.ParameterRangeValueNames.IsValidIndex(ValueIndex))
 						{
-							RangeIdxPtr->SetPosition(0, RangeIndex);
-							MutableParameters->SetFloatValue(paramIndex, param.ParameterRangeValues[RangeIndex], RangeIdxPtr);
+							Param.ParameterRangeValueNames.AddDefaulted(ValueIndex + 1 - Param.ParameterRangeValueNames.Num());
 						}
+
+						const int32 Value = MutableParameters->GetIntValue(ParamIndex, RangeValueIdxPtr);
+						const FString AuxParameterValueName = CustomizableObject->FindIntParameterValueName(ParamIndex, Value);
+						Param.ParameterRangeValueNames[RangeIndex] = AuxParameterValueName;
 					}
 				}
+				else
+				{
+					const int32 ParamValue = MutableParameters->GetIntValue(ParamIndex);
+					Param.ParameterValueName = CustomizableObject->FindIntParameterValueName(ParamIndex, ParamValue);
+				}
 			}
 
-			if (!found)
-			{
-				mu::RangeIndexPtr RangeIdxPtr = MutableParameters->NewRangeIndex(paramIndex);
+			IntParameters.Add(Param);
+
+			break;
+		}
+
+		case mu::PARAMETER_TYPE::T_FLOAT:
+		{
+			FCustomizableObjectFloatParameterValue Param;
+			Param.ParameterName = Name;
+			Param.Uid = Uid;
 
-				if (!RangeIdxPtr.get())
+			auto FindByNameAndUid = [&](const FCustomizableObjectFloatParameterValue& P)
+			{
+				return P.ParameterName == Name || (!Uid.IsEmpty() && P.Uid == Uid);
+			};
+			
+			if (FCustomizableObjectFloatParameterValue* Result = OldFloatParameters.FindByPredicate(FindByNameAndUid))
+			{	
+				if (mu::RangeIndexPtr RangeIdxPtr = MutableParameters->NewRangeIndex(ParamIndex))
 				{
-				param.ParameterValue = MutableParameters->GetFloatValue(paramIndex);
-			}
+					Param.ParameterRangeValues = Result->ParameterRangeValues;
+				}
 				else
 				{
-				param.ParameterValue = MutableParameters->GetFloatValue(paramIndex);
-
-					int32 ValueCount = MutableParameters->GetValueCount(paramIndex);
-
-					for (int32 ValueIterator = 0; ValueIterator < ValueCount; ++ValueIterator)
-					{
-						mu::RangeIndexPtr RangeValueIdxPtr = MutableParameters->GetValueIndex(paramIndex, ValueIterator);
-						int32 ValueIndex = RangeValueIdxPtr->GetPosition(0);
+					Param.ParameterValue = Result->ParameterValue;
+				}
+			} 
+			else // Not found in Instance Parameters. Use Mutable Parameters.
+			{
+				if (mu::RangeIndexPtr RangeIdxPtr = MutableParameters->NewRangeIndex(ParamIndex))
+				{
+					const int32 ValueCount = MutableParameters->GetValueCount(ParamIndex);
 
-						float Value = MutableParameters->GetFloatValue(paramIndex, RangeValueIdxPtr);
+					for (int ValueIndex = 0; ValueIndex < ValueCount; ++ValueIndex)
+					{	
+						mu::RangeIndexPtr RangeValueIdxPtr = MutableParameters->GetValueIndex(ParamIndex, ValueIndex);
+						int32 RangeIndex = RangeValueIdxPtr->GetPosition(0);
 
-						if (!param.ParameterRangeValues.IsValidIndex(ValueIndex))
+						if (!Param.ParameterRangeValues.IsValidIndex(ValueIndex))
 						{
-							param.ParameterRangeValues.AddDefaulted(ValueIndex + 1 - param.ParameterRangeValues.Num());
+							Param.ParameterRangeValues.AddDefaulted(ValueIndex + 1 - Param.ParameterRangeValues.Num());
 						}
 
-						param.ParameterRangeValues[ValueIndex] = Value;
+						float Value = MutableParameters->GetFloatValue(ParamIndex, RangeValueIdxPtr);
+						Param.ParameterRangeValues[RangeIndex] = Value;
 					}
 				}
+				else
+				{
+					Param.ParameterValue = MutableParameters->GetFloatValue(ParamIndex);
+				}
 			}
 
-			GetFloatParameters().Add(param);
+			FloatParameters.Add(Param);
 			break;
 		}
 
 		case mu::PARAMETER_TYPE::T_COLOUR:
 		{
-			FCustomizableObjectVectorParameterValue param;
-			param.ParameterName = Name;
-			param.Uid = Uid;
+			FCustomizableObjectVectorParameterValue Param;
+			Param.ParameterName = Name;
+			Param.Uid = Uid;
 
-			bool found = false;
-			for (int32 i = 0; i < OldVectorParameters.Num(); ++i)
+			auto FindByNameAndUid = [&](const FCustomizableObjectVectorParameterValue& P)
 			{
-				if (OldVectorParameters[i].ParameterName == Name || (!Uid.IsEmpty() && OldVectorParameters[i].Uid == Uid))
-				{
-					found = true;
-					param.ParameterValue = OldVectorParameters[i].ParameterValue;
-					MutableParameters->SetColourValue(paramIndex, param.ParameterValue.R, param.ParameterValue.G, param.ParameterValue.B);
-				}
-			}
+				return P.ParameterName == Name || (!Uid.IsEmpty() && P.Uid == Uid);
+			};
 
-			if (!found)
+			if (FCustomizableObjectVectorParameterValue* Result = OldVectorParameters.FindByPredicate(FindByNameAndUid))
+			{	
+				Param.ParameterValue = Result->ParameterValue;
+			}
+			else // Not found in Instance Parameters. Use Mutable Parameters.
 			{
-				MutableParameters->GetColourValue(paramIndex, &param.ParameterValue.R, &param.ParameterValue.G, &param.ParameterValue.B);
-				param.ParameterValue.A = 1.0f;
+				MutableParameters->GetColourValue(ParamIndex, &Param.ParameterValue.R, &Param.ParameterValue.G, &Param.ParameterValue.B);
+				Param.ParameterValue.A = 1.0f;
 			}
 
-			GetVectorParameters().Add(param);
+			VectorParameters.Add(Param);
 			break;
 		}
 
 		case mu::PARAMETER_TYPE::T_PROJECTOR:
 		{
-			FCustomizableObjectProjectorParameterValue param;
-			param.ParameterName = Name;
-			param.Uid = Uid;
+			FCustomizableObjectProjectorParameterValue Param;
+			Param.ParameterName = Name;
+			Param.Uid = Uid;
 
-			bool found = false;
-			for (int32 i = 0; i < OldProjectorParameters.Num(); ++i)
+			auto FindByNameAndUid = [&](const FCustomizableObjectProjectorParameterValue& P)
 			{
-				if (OldProjectorParameters[i].ParameterName == Name || (!Uid.IsEmpty() && OldProjectorParameters[i].Uid == Uid))
+				return P.ParameterName == Name || (!Uid.IsEmpty() && P.Uid == Uid);
+			};
+			
+			if (FCustomizableObjectProjectorParameterValue* Result = OldProjectorParameters.FindByPredicate(FindByNameAndUid))
+			{	
+				if (mu::RangeIndexPtr RangeIdxPtr = MutableParameters->NewRangeIndex(ParamIndex))
 				{
-					found = true;
-
-					param.Value = OldProjectorParameters[i].Value;
-					param.RangeValues = OldProjectorParameters[i].RangeValues;
-
-					mu::RangeIndexPtr RangeIdxPtr = MutableParameters->NewRangeIndex(paramIndex);
-
-					TArray<FCustomizableObjectProjector> AuxValuesArray;
-					AuxValuesArray.Add(param.Value);
-					TArray<FCustomizableObjectProjector>& Values = RangeIdxPtr.get() ? param.RangeValues : AuxValuesArray;
-					int32 ArrayIndex = 0;
-
-					for (FCustomizableObjectProjector& Value : Values)
+					Param.RangeValues = Result->RangeValues;
+				}
+				else
+				{
+					Param.Value = Result->Value;
+				}
+			} 
+			else // Not found in Instance Parameters. Use Mutable Parameters.
+			{
+				auto GetProjector = [&MutableParameters, ParamIndex](FCustomizableObjectProjector &Value, const mu::RangeIndexPtr RangeIndex = nullptr)
+				{
+					mu::PROJECTOR_TYPE Type;
+					MutableParameters->GetProjectorValue(ParamIndex,
+						&Type,
+						&Value.Position,
+						&Value.Direction,
+						&Value.Up,
+						&Value.Scale,
+						&Value.Angle,
+						RangeIndex);
+					
+					switch (Type)
 					{
-						if (RangeIdxPtr.get())
-						{
-							check(RangeIdxPtr->GetRangeCount() == 1);
-							RangeIdxPtr->SetPosition(0, ArrayIndex);
-						}
+					case mu::PROJECTOR_TYPE::PLANAR:
+						Value.ProjectionType = ECustomizableObjectProjectorType::Planar;
+						break;
 
-						switch (Value.ProjectionType)
-						{
-						case ECustomizableObjectProjectorType::Planar:
-						case ECustomizableObjectProjectorType::Wrapping:
-						{
-							MutableParameters->SetProjectorValue(paramIndex,
-								Value.Position,
-								Value.Direction,
-								Value.Up,
-								Value.Scale,
-								Value.Angle,
-								RangeIdxPtr);
-							break;
-						}
+					case mu::PROJECTOR_TYPE::CYLINDRICAL:
+						// Unapply strange swizzle for scales.
+						// TODO: try to avoid this
+						Value.ProjectionType = ECustomizableObjectProjectorType::Cylindrical;
+						Value.Direction = -Value.Direction;
+						Value.Up = -Value.Up;
+						Value.Scale[2] = -Value.Scale[0];
+						Value.Scale[0] = Value.Scale[1] = Value.Scale[1] * 2.0f;
+						break;
 
-						case ECustomizableObjectProjectorType::Cylindrical:
-						{
-							// Apply strange swizzle for scales
-							// TODO: try to avoid this
-							float Radius = FMath::Abs(Value.Scale[0] / 2.0f);
-							float Height = Value.Scale[2];
-							// TODO: try to avoid this
-							MutableParameters->SetProjectorValue(paramIndex,
-									Value.Position,
-									-Value.Direction,
-									-Value.Up,
-									FVector3f(-Height, Radius, Radius),
-									Value.Angle,
-									RangeIdxPtr);
-							break;
-						}
+					case mu::PROJECTOR_TYPE::WRAPPING:
+						Value.ProjectionType = ECustomizableObjectProjectorType::Wrapping;
+						break;
 
-						default:
-							// Not implemented.
-							check(false);
-						}
-						
-						ArrayIndex++;
+					default:
+						check(false); // Not implemented.
 					}
-				}
-			}
+				}; 
 
-			if (!found)
-			{
-				mu::RangeIndexPtr RangeIdxPtr = MutableParameters->NewRangeIndex(paramIndex);
-
-				if (!RangeIdxPtr.get())
-				{
-				mu::PROJECTOR_TYPE type;
-				MutableParameters->GetProjectorValue(paramIndex,
-					&type,
-					&param.Value.Position,
-					&param.Value.Direction,
-					&param.Value.Up,
-					&param.Value.Scale,
-					&param.Value.Angle );
-				switch (type)
+				if (mu::RangeIndexPtr RangeIdxPtr = MutableParameters->NewRangeIndex(ParamIndex))
 				{
-				case mu::PROJECTOR_TYPE::PLANAR:
-					param.Value.ProjectionType = ECustomizableObjectProjectorType::Planar;
-					break;
+					const int32 ValueCount = MutableParameters->GetValueCount(ParamIndex);
 
-				case mu::PROJECTOR_TYPE::CYLINDRICAL:
-					// Unapply strange swizzle for scales.
-					// TODO: try to avoid this
-					param.Value.ProjectionType = ECustomizableObjectProjectorType::Cylindrical;
-					param.Value.Direction = -param.Value.Direction;
-					param.Value.Up = -param.Value.Up;
-					param.Value.Scale[2] = -param.Value.Scale[0];
-					param.Value.Scale[0] = param.Value.Scale[1] = param.Value.Scale[1] * 2.0f;
-					break;
+					for (int ValueIndex = 0; ValueIndex < ValueCount; ++ValueIndex)
+					{	
+						mu::RangeIndexPtr RangeValueIdxPtr = MutableParameters->GetValueIndex(ParamIndex, ValueIndex);
+						int32 RangeIndex = RangeValueIdxPtr->GetPosition(0);
 
-				case mu::PROJECTOR_TYPE::WRAPPING:
-					param.Value.ProjectionType = ECustomizableObjectProjectorType::Wrapping;
-					break;
-				default:
-					// not implemented
-					check(false);
-				}
-			}
-				else
-				{
-					int32 ValueCount = MutableParameters->GetValueCount(paramIndex);
-
-					for (int32 ValueIterator = 0; ValueIterator < ValueCount; ++ValueIterator)
-					{
-						mu::RangeIndexPtr RangeValueIdxPtr = MutableParameters->GetValueIndex(paramIndex, ValueIterator);
-						int32 ValueIndex = RangeValueIdxPtr->GetPosition(0);
-
-						if (!param.RangeValues.IsValidIndex(ValueIndex))
+						if (!Param.RangeValues.IsValidIndex(ValueIndex))
 						{
-							param.RangeValues.AddDefaulted(ValueIndex + 1 - param.RangeValues.Num());
+							Param.RangeValues.AddDefaulted(ValueIndex + 1 - Param.RangeValues.Num());
 						}
 
-						mu::PROJECTOR_TYPE type;
-						MutableParameters->GetProjectorValue(paramIndex,
-							&type,
-							&param.RangeValues[ValueIndex].Position,
-							&param.RangeValues[ValueIndex].Direction,
-							&param.RangeValues[ValueIndex].Up,
-							&param.RangeValues[ValueIndex].Scale,
-							&param.RangeValues[ValueIndex].Angle,
-							RangeValueIdxPtr);
-
-						switch (type)
-						{
-						case mu::PROJECTOR_TYPE::PLANAR:
-							param.RangeValues[ValueIndex].ProjectionType = ECustomizableObjectProjectorType::Planar;
-							break;
-
-						case mu::PROJECTOR_TYPE::CYLINDRICAL:
-							// Unapply strange swizzle for scales.
-							// TODO: try to avoid this
-							param.RangeValues[ValueIndex].ProjectionType = ECustomizableObjectProjectorType::Cylindrical;
-							param.RangeValues[ValueIndex].Direction = -param.RangeValues[ValueIndex].Direction;
-							param.RangeValues[ValueIndex].Up = -param.RangeValues[ValueIndex].Up;
-							param.RangeValues[ValueIndex].Scale[2] = -param.RangeValues[ValueIndex].Scale[0];
-							param.RangeValues[ValueIndex].Scale[0] = param.RangeValues[ValueIndex].Scale[1] = param.RangeValues[ValueIndex].Scale[1] * 2.0f;
-							break;
-
-						case mu::PROJECTOR_TYPE::WRAPPING:
-							param.RangeValues[ValueIndex].ProjectionType = ECustomizableObjectProjectorType::Wrapping;
-							break;
-						default:
-							// not implemented
-							check(false);
-						}
+						GetProjector(Param.RangeValues[RangeIndex], RangeValueIdxPtr);
 					}
 				}
+				else
+				{
+					GetProjector(Param.Value);
+				}
 			}
 
-			GetProjectorParameters().Add(param);
+			ProjectorParameters.Add(Param);
 			break;
 		}
 
 		case mu::PARAMETER_TYPE::T_IMAGE:
 		{
-			FCustomizableObjectTextureParameterValue param;
-			param.ParameterName = Name;
-			param.Uid = Uid;
+			FCustomizableObjectTextureParameterValue Param;
+			Param.ParameterName = Name;
+			Param.Uid = Uid;
 
-			bool found = false;
-			for (int32 i = 0; i < OldTextureParameters.Num(); ++i)
+			auto FindByNameAndUid = [&](const FCustomizableObjectTextureParameterValue& P)
 			{
-				if (OldTextureParameters[i].ParameterName == Name || (!Uid.IsEmpty() && OldTextureParameters[i].Uid == Uid))
-				{
-					found = true;
-					param.ParameterValue = OldTextureParameters[i].ParameterValue;
-					MutableParameters->SetImageValue(paramIndex, param.ParameterValue);
-				}
-			}
+				return P.ParameterName == Name || (!Uid.IsEmpty() && P.Uid == Uid);
+			};
 
-			if (!found)
+			if (FCustomizableObjectTextureParameterValue* Result = OldTextureParameters.FindByPredicate(FindByNameAndUid))
+			{
+				Param.ParameterValue = Result->ParameterValue;
+			}
+			else // Not found in Instance Parameters. Use Mutable Parameters.
 			{
-				param.ParameterValue = MutableParameters->GetImageValue(paramIndex);
+				Param.ParameterValue = MutableParameters->GetImageValue(ParamIndex);
 			}
 
-			GetTextureParameters().Add(param);
+			TextureParameters.Add(Param);
 			break;
 		}
 
 		default:
-			// TODO
+			check(false); // Missing case.
 			break;
-
 		}
 	}
 	
 	CreateParametersLookupTable();
-
-	return MutableParameters;
 }
 
 
@@ -1612,6 +1630,20 @@ int32 FCustomizableObjectInstanceDescriptor::FindFloatParameterNameIndex(const F
 }
 
 
+int32 FCustomizableObjectInstanceDescriptor::FindTextureParameterNameIndex(const FString& ParamName) const
+{
+	for (int32 i = 0; i < TextureParameters.Num(); ++i)
+	{
+		if (TextureParameters[i].ParameterName == ParamName)
+		{
+			return i;
+		}
+	}
+
+	return -1;
+}
+
+
 int32 FCustomizableObjectInstanceDescriptor::FindBoolParameterNameIndex(const FString& ParamName) const
 {
 	for (int32 i = 0; i < BoolParameters.Num(); ++i)
@@ -1706,8 +1738,7 @@ int32 FCustomizableObjectInstanceDescriptor::AddValueToIntRange(const FString& P
 		const int32 ParamIndexInObject = CustomizableObject->FindParameter(intParameter.ParameterName);
 		// TODO: Define the default option in the editor instead of taking the first available, like it's currently defined for GetProjectorDefaultValue()
 		const FString DefaultValue = CustomizableObject->GetIntParameterAvailableOption(ParamIndexInObject, 0);
-		intParameter.ParameterRangeValueNames.Add(DefaultValue);
-		return intParameter.ParameterRangeValueNames.Num() - 1;
+		return intParameter.ParameterRangeValueNames.Add(DefaultValue);
 	}
 	return -1;
 }
@@ -1720,8 +1751,7 @@ int32 FCustomizableObjectInstanceDescriptor::AddValueToFloatRange(const FString&
 	{
 		FCustomizableObjectFloatParameterValue& floatParameter = FloatParameters[floatParameterIndex];
 		// TODO: Define the default float in the editor instead of [0.5f], like it's currently defined for GetProjectorDefaultValue()
-		floatParameter.ParameterRangeValues.Add(0.5f);
-		return floatParameter.ParameterRangeValues.Num() - 1;
+		return floatParameter.ParameterRangeValues.Add(0.5f);
 	}
 	return -1;
 }
@@ -1736,8 +1766,7 @@ int32 FCustomizableObjectInstanceDescriptor::AddValueToProjectorRange(const FStr
 	{
 		FCustomizableObjectProjectorParameterValue& ProjectorParameter = ProjectorParameters[projectorParameterIndex];
 		const FCustomizableObjectProjector Projector = GetProjectorDefaultValue(CustomizableObject->FindParameter(ParamName));
-		ProjectorParameter.RangeValues.Add(Projector);
-		return ProjectorParameter.RangeValues.Num() - 1;
+		return ProjectorParameter.RangeValues.Add(Projector);
 	}
 
 	return -1;
@@ -1911,7 +1940,11 @@ void FCustomizableObjectInstanceDescriptor::SetCurrentState(const FString& State
 {
 	check(CustomizableObject);
 
-	SetState(CustomizableObject->FindState(StateName));
+	if (const int32 Result = CustomizableObject->FindState(StateName);
+		ensureAlwaysMsgf(Result != -1, TEXT("Unknown %s state."), *StateName))
+	{
+		SetState(Result);
+	}	
 }
 
 
@@ -2063,12 +2096,6 @@ void FCustomizableObjectInstanceDescriptor::MultilayerProjectorUpdateVirtualLaye
 }
 
 
-void FCustomizableObjectInstanceDescriptor::Init()
-{
-	ReloadParametersFromObject();
-}
-
-
 void FCustomizableObjectInstanceDescriptor::CreateParametersLookupTable()
 {
 	check(CustomizableObject);
diff --git a/Engine/Plugins/Experimental/Mutable/Source/CustomizableObject/Private/MuCO/CustomizableObjectSystem.cpp b/Engine/Plugins/Experimental/Mutable/Source/CustomizableObject/Private/MuCO/CustomizableObjectSystem.cpp
--- a/Engine/Plugins/Experimental/Mutable/Source/CustomizableObject/Private/MuCO/CustomizableObjectSystem.cpp
+++ b/Engine/Plugins/Experimental/Mutable/Source/CustomizableObject/Private/MuCO/CustomizableObjectSystem.cpp
@@ -2424,7 +2424,7 @@ FMutableOperation FMutableOperation::CreateInstanceUpdate(UCustomizableObjectIns
 	Op.InstanceDescriptorRuntimeHash = InCustomizableObjectInstance->GetUpdateDescriptorRuntimeHash();
 	Op.bStarted = false;
 	Op.bBuildParameterDecorations = InCustomizableObjectInstance->GetBuildParameterDecorations();
-	Op.Parameters = InCustomizableObjectInstance->GetPrivate()->ReloadParametersFromObject(InCustomizableObjectInstance, false);
+	Op.Parameters = InCustomizableObjectInstance->GetPrivate()->GetParameters(InCustomizableObjectInstance);
 
 	InCustomizableObjectInstance->GetCustomizableObject()->ApplyStateForcedValuesToParameters(InCustomizableObjectInstance->GetState(), Op.Parameters.get());
 
diff --git a/Engine/Plugins/Experimental/Mutable/Source/CustomizableObject/Public/MuCO/CustomizableObjectInstanceDescriptor.h b/Engine/Plugins/Experimental/Mutable/Source/CustomizableObject/Public/MuCO/CustomizableObjectInstanceDescriptor.h
--- a/Engine/Plugins/Experimental/Mutable/Source/CustomizableObject/Public/MuCO/CustomizableObjectInstanceDescriptor.h
+++ b/Engine/Plugins/Experimental/Mutable/Source/CustomizableObject/Public/MuCO/CustomizableObjectInstanceDescriptor.h
@@ -44,10 +44,10 @@ struct CUSTOMIZABLEOBJECT_API FCustomizableObjectInstanceDescriptor
 
 	FCustomizableObjectInstanceDescriptor(const FCustomizableObjectInstanceDescriptor& Other) = default;
 
-	/** Serialize this object. */
+	/** Serialize this object. Does not support Multilayer Projectors! */
 	void SaveDescriptor(FArchive &Ar);
 
-	/** Deserialize this object. */
+	/** Deserialize this object. Does not support Multilayer Projectors! */
 	void LoadDescriptor(FArchive &Ar);
 
 	UCustomizableObject* GetCustomizableObject() const;
@@ -56,9 +56,10 @@ struct CUSTOMIZABLEOBJECT_API FCustomizableObjectInstanceDescriptor
 	
 	void SetBuildParameterDecorations(bool Value);
 	
-	void SetCustomizableObject(UCustomizableObject* InCustomizableObject);
+	void SetCustomizableObject(UCustomizableObject& InCustomizableObject);
 
-	mu::Ptr<mu::Parameters> ReloadParametersFromObject(); // TODO Split function in 2: 1. Load parameters from model, 2. Get mu::ParametersPtr
+	/** Update all parameters to be up to date with the Mutable Core parameters. */
+	void ReloadParameters();
     
 	// ------------------------------------------------------------
 	// Parameters
@@ -169,6 +170,9 @@ struct CUSTOMIZABLEOBJECT_API FCustomizableObjectInstanceDescriptor
 	/** Finds in FloatParameters a parameter with name ParamName, returns the index if found, -1 otherwise. */
 	int32 FindFloatParameterNameIndex(const FString& ParamName) const;
 
+	/** Finds in TextureParameters a parameter with name ParamName, returns the index if found, -1 otherwise. */
+	int32 FindTextureParameterNameIndex(const FString& ParamName) const;
+
 	/** Finds in BoolParameters a parameter with name ParamName, returns the index if found, -1 otherwise. */
 	int32 FindBoolParameterNameIndex(const FString& ParamName) const;
 
@@ -331,8 +335,9 @@ private:
 	/** Multilayer Projector helpers. See FMultilayerProjector.*/
 	UPROPERTY()
 	TMap<FName, FMultilayerProjector> MultilayerProjectors;
-	
-	void Init();
+
+	/** Return a Mutable Core object containing all parameters. */
+	mu::Ptr<mu::Parameters> GetParameters() const;
 
 	void CreateParametersLookupTable();
 	
