[Mutable] Instance update minor improvements.

- Fixed issue where the Instance Preview Viewport was not showing the correct message overlay.
- Merged DoUpdateSkeletalMesh and DoUpdateSkeletalMeshAsync.
- Converted known alwyas valid pointers to references.
- Removed DoUpdateSkeletalMesh bAsync parameter (always was checked against true).

diff --git a/Engine/Plugins/Experimental/Mutable/Source/CustomizableObject/Classes/MuCO/CustomizableObjectInstance.h b/Engine/Plugins/Experimental/Mutable/Source/CustomizableObject/Classes/MuCO/CustomizableObjectInstance.h
--- a/Engine/Plugins/Experimental/Mutable/Source/CustomizableObject/Classes/MuCO/CustomizableObjectInstance.h
+++ b/Engine/Plugins/Experimental/Mutable/Source/CustomizableObject/Classes/MuCO/CustomizableObjectInstance.h
@@ -108,9 +108,13 @@ enum class EUpdateResult : uint8
  * - Native delegates names should end with "NativeDelegate".
  * - Dynamic delegates broadcast before native delegates. */
 
+/** Broadcast at the end of an Instance update request (e.g., before returning from UpdateSkeletalMeshAsync).
+ * Notice that Mutable internally can also request an Instance update. */
 DECLARE_DYNAMIC_MULTICAST_DELEGATE_OneParam(FBeginUpdateDelegate, UCustomizableObjectInstance*, Instance);
 DECLARE_MULTICAST_DELEGATE_OneParam(FBeginUpdateNativeDelegate, UCustomizableObjectInstance*);
 
+/** Broadcast when an Instance update has completed.
+ * Notice that Mutable internally can also start an Instance update. */
 DECLARE_DYNAMIC_MULTICAST_DELEGATE_OneParam(FObjectInstanceUpdatedDelegate, UCustomizableObjectInstance*, Instance);
 DECLARE_MULTICAST_DELEGATE_OneParam(FObjectInstanceUpdatedNativeDelegate, UCustomizableObjectInstance*);
 
diff --git a/Engine/Plugins/Experimental/Mutable/Source/CustomizableObject/Private/MuCO/CustomizableInstancePrivateData.h b/Engine/Plugins/Experimental/Mutable/Source/CustomizableObject/Private/MuCO/CustomizableInstancePrivateData.h
--- a/Engine/Plugins/Experimental/Mutable/Source/CustomizableObject/Private/MuCO/CustomizableInstancePrivateData.h
+++ b/Engine/Plugins/Experimental/Mutable/Source/CustomizableObject/Private/MuCO/CustomizableInstancePrivateData.h
@@ -200,8 +200,8 @@ public:
 	/** See FCustomizableObjectInstanceDescriptor::ReloadParameters(...). */
 	void ReloadParameters(UCustomizableObjectInstance* Public);
 
-	void TickUpdateCloseCustomizableObjects(UCustomizableObjectInstance* Public);
-	void UpdateInstanceIfNotGenerated(UCustomizableObjectInstance* Public, bool bAsync);
+	void TickUpdateCloseCustomizableObjects(UCustomizableObjectInstance& Public);
+	void UpdateInstanceIfNotGenerated(UCustomizableObjectInstance& Public);
 
 	// Returns true if success (?)
 	bool UpdateSkeletalMesh_PostBeginUpdate0(UCustomizableObjectInstance* Public, const TSharedPtr<FMutableOperationData>& OperationData);
@@ -230,8 +230,7 @@ public:
 	// The following method is basically copied from PostEditChangeProperty and/or SkeletalMesh.cpp to be able to replicate PostEditChangeProperty without the editor
 	void PostEditChangePropertyWithoutEditor(USkeletalMesh* SkeletalMesh);
 	
-	void DoUpdateSkeletalMeshAsync(UCustomizableObjectInstance* Instance, FMutableQueueElem::EQueuePriorityType Priority = FMutableQueueElem::EQueuePriorityType::Low);
-	void DoUpdateSkeletalMesh(UCustomizableObjectInstance* Instance, bool bAsync, bool bIsCloseDistTick=false, bool bOnlyUpdateIfNotGenerated=false, bool bIgnoreCloseDist=false, bool bForceHighPriority = false);
+	void DoUpdateSkeletalMesh(UCustomizableObjectInstance& Instance, bool bIsCloseDistTick=false, bool bOnlyUpdateIfNotGenerated=false, bool bIgnoreCloseDist=false, bool bForceHighPriority = false);
 
 	void DiscardResourcesAndSetReferenceSkeletalMesh(UCustomizableObjectInstance* Public);
 
diff --git a/Engine/Plugins/Experimental/Mutable/Source/CustomizableObject/Private/MuCO/CustomizableObjectInstance.cpp b/Engine/Plugins/Experimental/Mutable/Source/CustomizableObject/Private/MuCO/CustomizableObjectInstance.cpp
--- a/Engine/Plugins/Experimental/Mutable/Source/CustomizableObject/Private/MuCO/CustomizableObjectInstance.cpp
+++ b/Engine/Plugins/Experimental/Mutable/Source/CustomizableObject/Private/MuCO/CustomizableObjectInstance.cpp
@@ -808,51 +808,23 @@ void UCustomizableInstancePrivateData::PostEditChangePropertyWithoutEditor(USkel
 
 void UCustomizableObjectInstance::UpdateSkeletalMeshAsync(bool bIgnoreCloseDist, bool bForceHighPriority)
 {
-	PrivateData->DoUpdateSkeletalMesh(this, true, false, false, bIgnoreCloseDist, bForceHighPriority);
+	PrivateData->DoUpdateSkeletalMesh(*this, false, false, bIgnoreCloseDist, bForceHighPriority);
 }
 
 
-void UCustomizableInstancePrivateData::DoUpdateSkeletalMeshAsync(UCustomizableObjectInstance* Instance, FMutableQueueElem::EQueuePriorityType Priority)
+void UCustomizableInstancePrivateData::DoUpdateSkeletalMesh(UCustomizableObjectInstance& CustomizableInstance, bool bIsCloseDistTick, bool bOnlyUpdateIfNotGenerated, bool bIgnoreCloseDist, bool bForceHighPriority)
 {
+	MUTABLE_CPUPROFILER_SCOPE(UCustomizableInstancePrivateData::DoUpdateSkeletalMesh);
 	check(IsInGameThread());
 
-	UE_LOG(LogMutable, Verbose, TEXT("Started UpdateSkeletalMesh Async. of Instance %d with priority %d at dist %f bIsPlayerOrNearIt=%d, frame=%d"), Instance->GetUniqueID(), static_cast<int32>(Priority), FMath::Sqrt(Instance->GetPrivate()->LastMinSquareDistFromComponentToPlayer), Instance->GetPrivate()->HasCOInstanceFlags(UsedByPlayerOrNearIt), GFrameNumber);
-
-	if (HasCOInstanceFlags(PendingLODsUpdateSecondStage))
-	{
-		UE_LOG(LogMutable, Verbose, TEXT("LOD change: %d, %d -> %d, %d"), LastUpdateMinLOD, LastUpdateMaxLOD, Instance->Descriptor.MinLODToLoad, Instance->Descriptor.MaxLODToLoad);
-	}
-
-	Instance->UpdateDescriptorRuntimeHash = FDescriptorRuntimeHash(Instance->Descriptor);
-
-	UCustomizableObjectSystem::GetInstance()->GetPrivate()->InitUpdateSkeletalMesh(Instance, Priority);
-
-#if WITH_EDITOR
-	if (!UCustomizableObjectSystem::GetInstance()->IsCompilationDisabled())
+	if (!CustomizableInstance.IsValidLowLevel())
 	{
-		if (Instance->SkeletalMeshStatus != ESkeletalMeshState::AsyncUpdatePending)
-		{
-			Instance->PreUpdateSkeletalMeshStatus = Instance->SkeletalMeshStatus;
-		}
-
-		Instance->SkeletalMeshStatus = ESkeletalMeshState::AsyncUpdatePending;
+		return;
 	}
-#else
-	Instance->SkeletalMeshStatus = ESkeletalMeshState::AsyncUpdatePending;
-#endif
-}
-
-
-void UCustomizableInstancePrivateData::DoUpdateSkeletalMesh(UCustomizableObjectInstance* CustomizableInstance, bool bAsync, bool bIsCloseDistTick, bool bOnlyUpdateIfNotGenerated, bool bIgnoreCloseDist, bool bForceHighPriority)
-{
-	MUTABLE_CPUPROFILER_SCOPE(UCustomizableInstancePrivateData::DoUpdateSkeletalMesh);
-
-	check(IsInGameThread());
-	check(bAsync);
-
+	
 	UCustomizableObjectSystem* System = UCustomizableObjectSystem::GetInstance();
 	
-	UCustomizableObject* CustomizableObject = CustomizableInstance->GetCustomizableObject();
+	UCustomizableObject* CustomizableObject = CustomizableInstance.GetCustomizableObject();
 
 	if (CustomizableObject && !CustomizableObject->IsCompiled())
 	{
@@ -870,9 +842,9 @@ void UCustomizableInstancePrivateData::DoUpdateSkeletalMesh(UCustomizableObjectI
 
 	const bool bIsGenerated = HasCOInstanceFlags(Generated);
 
-	int32 NumGeneratedInstancesLimit = System->GetInstanceLODManagement()->GetNumGeneratedInstancesLimitFullLODs();
-	int32 NumGeneratedInstancesLimitLOD1 = System->GetInstanceLODManagement()->GetNumGeneratedInstancesLimitLOD1();
-	int32 NumGeneratedInstancesLimitLOD2 = System->GetInstanceLODManagement()->GetNumGeneratedInstancesLimitLOD2();
+	const int32 NumGeneratedInstancesLimit = System->GetInstanceLODManagement()->GetNumGeneratedInstancesLimitFullLODs();
+	const int32 NumGeneratedInstancesLimitLOD1 = System->GetInstanceLODManagement()->GetNumGeneratedInstancesLimitLOD1();
+	const int32 NumGeneratedInstancesLimitLOD2 = System->GetInstanceLODManagement()->GetNumGeneratedInstancesLimitLOD2();
 
 	if (NumGeneratedInstancesLimit > 0 && System->GetPrivate()->GetCountAllocatedSkeletalMesh() > NumGeneratedInstancesLimit + NumGeneratedInstancesLimitLOD1 + NumGeneratedInstancesLimitLOD2)
 	{
@@ -900,6 +872,7 @@ void UCustomizableInstancePrivateData::DoUpdateSkeletalMesh(UCustomizableObjectI
 		if (bShouldUpdateLODs)
 		{
 			SetCOInstanceFlags(PendingLODsUpdateSecondStage);
+			UE_LOG(LogMutable, Verbose, TEXT("LOD change: %d, %d -> %d, %d"), LastUpdateMinLOD, LastUpdateMaxLOD, CustomizableInstance.Descriptor.MinLODToLoad, CustomizableInstance.Descriptor.MaxLODToLoad);
 		}
 
 		if ((!(bIsGenerated && bOnlyUpdateIfNotGenerated) && !(bIsCloseDistTick && bIsGenerated)) || bShouldUpdateLODs)
@@ -913,7 +886,7 @@ void UCustomizableInstancePrivateData::DoUpdateSkeletalMesh(UCustomizableObjectI
 			{
 				Priority = FMutableQueueElem::EQueuePriorityType::High;
 			}
-			else if (!bIsGenerated || !CustomizableInstance->HasAnySkeletalMesh())
+			else if (!bIsGenerated || !CustomizableInstance.HasAnySkeletalMesh())
 			{
 				Priority = FMutableQueueElem::EQueuePriorityType::Med;
 			}
@@ -934,11 +907,33 @@ void UCustomizableInstancePrivateData::DoUpdateSkeletalMesh(UCustomizableObjectI
 				Priority = FMutableQueueElem::EQueuePriorityType::High;
 			}
 
-			DoUpdateSkeletalMeshAsync(CustomizableInstance, Priority);
+			{
+				const uint32 InstanceId = CustomizableInstance.GetUniqueID();
+				const float Distance = FMath::Sqrt(CustomizableInstance.GetPrivate()->LastMinSquareDistFromComponentToPlayer);
+				UE_LOG(LogMutable, Verbose, TEXT("Started UpdateSkeletalMesh Async. of Instance %d with priority %d at dist %f bIsPlayerOrNearIt=%d, frame=%d"), InstanceId, static_cast<int32>(Priority), Distance, bIsPlayerOrNearIt, GFrameNumber);				
+			}
+
+			CustomizableInstance.UpdateDescriptorRuntimeHash = FDescriptorRuntimeHash(CustomizableInstance.Descriptor);
+
+#if WITH_EDITOR
+			if (!UCustomizableObjectSystem::GetInstance()->IsCompilationDisabled())
+			{
+				if (CustomizableInstance.SkeletalMeshStatus != ESkeletalMeshState::AsyncUpdatePending)
+				{
+					CustomizableInstance.PreUpdateSkeletalMeshStatus = CustomizableInstance.SkeletalMeshStatus;
+				}
+
+				CustomizableInstance.SkeletalMeshStatus = ESkeletalMeshState::AsyncUpdatePending;
+			}
+#else
+			CustomizableInstance.SkeletalMeshStatus = ESkeletalMeshState::AsyncUpdatePending;
+#endif
 
 			SetCOInstanceFlags(Generated); // Will be done in UpdateSkeletalMesh_PostBeginUpdate
 			LastUpdateTime = FPlatformTime::Seconds();
-
+			
+			// Do not do work after calling InitUpdateSkeletalMesh. This function can optimize an update and fully complete it before even exiting its scope.
+			System->GetPrivate()->InitUpdateSkeletalMesh(CustomizableInstance, Priority);
 		}
 
 		ClearCOInstanceFlags(PendingLODsUpdate);
@@ -947,24 +942,24 @@ void UCustomizableInstancePrivateData::DoUpdateSkeletalMesh(UCustomizableObjectI
 	{
 		if (HasCOInstanceFlags(Generated) && !HasCOInstanceFlags(Updating))
 		{
-			System->GetPrivate()->InitDiscardResourcesSkeletalMesh(CustomizableInstance);
+			System->GetPrivate()->InitDiscardResourcesSkeletalMesh(&CustomizableInstance);
 			ClearCOInstanceFlags((ECOInstanceFlags)(PendingLODsUpdate | PendingLODsUpdateSecondStage));
 		}
 	}
 }
 
 
-void UCustomizableInstancePrivateData::TickUpdateCloseCustomizableObjects( UCustomizableObjectInstance* Public )
+void UCustomizableInstancePrivateData::TickUpdateCloseCustomizableObjects( UCustomizableObjectInstance& Public )
 {
-	DoUpdateSkeletalMesh(Public, true, true, true, false);
+	DoUpdateSkeletalMesh(Public, true, true, false);
 }
 
 
-void UCustomizableInstancePrivateData::UpdateInstanceIfNotGenerated(UCustomizableObjectInstance* Public, bool bAsync)
+void UCustomizableInstancePrivateData::UpdateInstanceIfNotGenerated(UCustomizableObjectInstance& Public)
 {
 	if (!HasCOInstanceFlags(Generated))
 	{
-		DoUpdateSkeletalMesh(Public, bAsync, true, false, true);
+		DoUpdateSkeletalMesh(Public, false, true);
 	}
 }
 
diff --git a/Engine/Plugins/Experimental/Mutable/Source/CustomizableObject/Private/MuCO/CustomizableObjectSystem.cpp b/Engine/Plugins/Experimental/Mutable/Source/CustomizableObject/Private/MuCO/CustomizableObjectSystem.cpp
--- a/Engine/Plugins/Experimental/Mutable/Source/CustomizableObject/Private/MuCO/CustomizableObjectSystem.cpp
+++ b/Engine/Plugins/Experimental/Mutable/Source/CustomizableObject/Private/MuCO/CustomizableObjectSystem.cpp
@@ -744,7 +744,7 @@ void UpdateSkeletalMesh(UCustomizableObjectInstance* CustomizableObjectInstance)
 		{
 			MUTABLE_CPUPROFILER_SCOPE(UpdateSkeletalMesh_UpdatedDelegate);
 
-			CustomizableObjectInstance->Updated(EUpdateResult::Success);
+			CustomizableObjectInstance->Updated(CustomizableObjectInstance->SkeletalMeshStatus == ESkeletalMeshState::Correct ? EUpdateResult::Success : EUpdateResult::Error);
 
 #if WITH_EDITOR
 			CustomizableObjectInstance->InstanceUpdated = true;
@@ -776,76 +776,59 @@ void UpdateSkeletalMesh(UCustomizableObjectInstance* CustomizableObjectInstance)
 }
 
 
-void FCustomizableObjectSystemPrivate::InitUpdateSkeletalMesh(UCustomizableObjectInstance* Public, FMutableQueueElem::EQueuePriorityType Priority)
+void FCustomizableObjectSystemPrivate::InitUpdateSkeletalMesh(UCustomizableObjectInstance& Instance, FMutableQueueElem::EQueuePriorityType Priority)
 {
 	MUTABLE_CPUPROFILER_SCOPE(FCustomizableObjectSystemPrivate::InitUpdateSkeletalMesh);
 
 	check(IsInGameThread());
+	
+	const FString CurrentState = Instance.GetCurrentState();
+	const FParameterUIData* State = Instance.GetCustomizableObject()->StateUIDataMap.Find(CurrentState);
+	const bool bNeverStream = State ? State->bDontCompressRuntimeTextures : false;
+	const bool bUseMipmapStreaming = !bNeverStream;
+	int32 MipsToSkip = 0; // 0 means all mips
 
-	if (!Public ||
-		!Public->IsValidLowLevel())
+	if (bUseMipmapStreaming && EnableMutableProgressiveMipStreaming)
 	{
-		Public->Updated(EUpdateResult::Error);
-		return;
+		MipsToSkip = 255; // This means skip all possible mips until only UTexture::GetStaticMinTextureResidentMipCount() are left
 	}
 	
-	if (UCustomizableObject* CustomizableObject = Public->GetCustomizableObject();
-		CustomizableObject &&
-		!CustomizableObject->IsLocked())
-	{
-		FString CurrentState = Public->GetCurrentState();
-		FParameterUIData* State = CustomizableObject->StateUIDataMap.Find(CurrentState);
-		bool bNeverStream = State ? State->bDontCompressRuntimeTextures : false;
-		const bool bUseMipmapStreaming = !bNeverStream;
-		int32 MipsToSkip = 0; // 0 means all mips
+	const TSharedPtr<FMutableOperation> Operation = MakeShared<FMutableOperation>(FMutableOperation::CreateInstanceUpdate(&Instance, bNeverStream, MipsToSkip));
 
-		if (bUseMipmapStreaming && EnableMutableProgressiveMipStreaming)
-		{
-			MipsToSkip = 255; // This means skip all possible mips until only UTexture::GetStaticMinTextureResidentMipCount() are left
-		}
-		
-		const TSharedPtr<FMutableOperation> Operation = MakeShared<FMutableOperation>(FMutableOperation::CreateInstanceUpdate(Public, bNeverStream, MipsToSkip));
+	Instance.GetPrivate()->SaveMinMaxLODToLoad(&Instance);
 
-		Public->GetPrivate()->SaveMinMaxLODToLoad(Public);
+	const FDescriptorRuntimeHash UpdateDescriptorHash = Instance.GetUpdateDescriptorRuntimeHash();
 
-		// Skip the update if the Instance has the same descriptor (pending or applied).
-		// If there is a hash collision, at worse, we would add an unnecessary update.
+	if (const FMutableQueueElem* QueueElem = MutableOperationQueue.Get(&Instance))
+	{
+		if (const TSharedPtr<FMutableOperation>& QueuedOperation = QueueElem->Operation;
+			QueuedOperation &&
+			QueuedOperation->Type == FMutableOperation::EOperationType::Update &&
+			UpdateDescriptorHash.IsSubset(QueuedOperation->InstanceDescriptorRuntimeHash))
 		{
-			const FDescriptorRuntimeHash UpdateDescriptorHash = Public->GetUpdateDescriptorRuntimeHash();
-
-			if (const FMutableQueueElem* QueueElem = MutableOperationQueue.Get(Public))
-			{
-				if (const TSharedPtr<FMutableOperation>& QueuedOperation = QueueElem->Operation;
-					QueuedOperation &&
-					QueuedOperation->Type == FMutableOperation::EOperationType::Update &&
-					UpdateDescriptorHash.IsSubset(QueuedOperation->InstanceDescriptorRuntimeHash))
-				{
-					return;
-				}
-			}
-		
-			if (CurrentMutableOperation &&
-				CurrentMutableOperation->Type == FMutableOperation::EOperationType::Update &&
-				UpdateDescriptorHash.IsSubset(CurrentMutableOperation->InstanceDescriptorRuntimeHash))
-			{
-				return;
-			}
-			
-			if (UpdateDescriptorHash.IsSubset(Public->GetDescriptorRuntimeHash()))
-			{
-				UpdateSkeletalMesh(Public);
-				return;
-			}
+			return; // The the requested update is equal to the last enqueued update.
 		}
+	}
+
+	if (CurrentMutableOperation &&
+		CurrentMutableOperation->Type == FMutableOperation::EOperationType::Update &&
+		UpdateDescriptorHash.IsSubset(CurrentMutableOperation->InstanceDescriptorRuntimeHash))
+	{
+		return; // The the requested update is equal to the running update.
+	}
 
-		MutableOperationQueue.Enqueue(FMutableQueueElem::Create(Operation, Priority, Public->GetPrivate()->MinSquareDistFromComponentToPlayer));
+	// These delegates must be called at the end of the begin update.
+	Instance.BeginUpdateDelegate.Broadcast(&Instance);
+	Instance.BeginUpdateNativeDelegate.Broadcast(&Instance);
 
-		Public->BeginUpdateDelegate.Broadcast(Public);
-		Public->BeginUpdateNativeDelegate.Broadcast(Public);
+	if (UpdateDescriptorHash.IsSubset(Instance.GetDescriptorRuntimeHash()))
+	{
+		Instance.SkeletalMeshStatus = ESkeletalMeshState::Correct; // TODO GMT MTBL-1033 should not be here. Move to UCustomizableObjectInstance::Updated
+		UpdateSkeletalMesh(&Instance);
 	}
 	else
 	{
-		Public->Updated(EUpdateResult::Error);
+		MutableOperationQueue.Enqueue(FMutableQueueElem::Create(Operation, Priority, Instance.GetPrivate()->MinSquareDistFromComponentToPlayer));
 	}
 }
 
@@ -2328,11 +2311,11 @@ bool UCustomizableObjectSystem::Tick(float DeltaTime)
 		{
 			if (CustomizableObjectInstance->GetPrivate()->HasCOInstanceFlags(UsedByComponentInPlay))
 			{
-				CustomizableObjectInstance->GetPrivate()->TickUpdateCloseCustomizableObjects(*CustomizableObjectInstance);
+				CustomizableObjectInstance->GetPrivate()->TickUpdateCloseCustomizableObjects(**CustomizableObjectInstance);
 			}
 			else if (CustomizableObjectInstance->GetPrivate()->HasCOInstanceFlags(UsedByComponent))
 			{
-				CustomizableObjectInstance->GetPrivate()->UpdateInstanceIfNotGenerated(*CustomizableObjectInstance, true);
+				CustomizableObjectInstance->GetPrivate()->UpdateInstanceIfNotGenerated(**CustomizableObjectInstance);
 			}
 
 			CustomizableObjectInstance->GetPrivate()->ClearCOInstanceFlags((ECOInstanceFlags)(UsedByComponent | UsedByComponentInPlay));
diff --git a/Engine/Plugins/Experimental/Mutable/Source/CustomizableObject/Private/MuCO/CustomizableObjectSystemPrivate.h b/Engine/Plugins/Experimental/Mutable/Source/CustomizableObject/Private/MuCO/CustomizableObjectSystemPrivate.h
--- a/Engine/Plugins/Experimental/Mutable/Source/CustomizableObject/Private/MuCO/CustomizableObjectSystemPrivate.h
+++ b/Engine/Plugins/Experimental/Mutable/Source/CustomizableObject/Private/MuCO/CustomizableObjectSystemPrivate.h
@@ -25,7 +25,7 @@ class FMutableOperation
 {
 	/** Instance parameters at the time of the operation request. */
 	mu::ParametersPtr Parameters; 
-
+	
 	bool bBuildParameterDecorations = false;
 	bool bMeshNeedsUpdate = false;
 
@@ -70,6 +70,9 @@ public:
 	//! This is used to calculate stats.
 	double StartUpdateTime = 0.0;
 
+	/** Instance optimization state. */
+	int32 State;
+	
 	//!
 	bool IsBuildParameterDecorations() const
 	{
@@ -678,7 +681,7 @@ public:
 
 
 	// Init the async Skeletal Mesh creation/update
-	void InitUpdateSkeletalMesh(UCustomizableObjectInstance* Public, FMutableQueueElem::EQueuePriorityType Priority);
+	void InitUpdateSkeletalMesh(UCustomizableObjectInstance& Public, FMutableQueueElem::EQueuePriorityType Priority);
 		
 	// Init an async and safe release of the UE4 and Mutable resources used by the instance without actually destroying the instance, for example if it's very far away
 	void InitDiscardResourcesSkeletalMesh(UCustomizableObjectInstance* InCustomizableObjectInstance);
diff --git a/Engine/Plugins/Experimental/Mutable/Source/CustomizableObjectEditor/Private/MuCOE/CustomizableObjectInstanceEditor.cpp b/Engine/Plugins/Experimental/Mutable/Source/CustomizableObjectEditor/Private/MuCOE/CustomizableObjectInstanceEditor.cpp
--- a/Engine/Plugins/Experimental/Mutable/Source/CustomizableObjectEditor/Private/MuCOE/CustomizableObjectInstanceEditor.cpp
+++ b/Engine/Plugins/Experimental/Mutable/Source/CustomizableObjectEditor/Private/MuCOE/CustomizableObjectInstanceEditor.cpp
@@ -398,7 +398,7 @@ void FCustomizableObjectInstanceEditor::UpdatePreviewVisibility()
 {
 	for (UDebugSkelMeshComponent* PreviewSkeletalMeshComponent : PreviewSkeletalMeshComponents)
 	{
-		if (CustomizableObjectInstance->SkeletalMeshStatus != ESkeletalMeshState::UpdateError)
+		if (CustomizableObjectInstance->SkeletalMeshStatus == ESkeletalMeshState::Correct)
 		{
 			PreviewSkeletalMeshComponent->SetVisibility(true, true);
 		}
