[mutable] Add skeleton assets as part of a baked instance.

#rnx

diff --git a/Engine/Plugins/Experimental/Mutable/Source/CustomizableObjectEditor/Private/MuCOE/CustomizableObjectEditorViewportClient.cpp b/Engine/Plugins/Experimental/Mutable/Source/CustomizableObjectEditor/Private/MuCOE/CustomizableObjectEditorViewportClient.cpp
--- a/Engine/Plugins/Experimental/Mutable/Source/CustomizableObjectEditor/Private/MuCOE/CustomizableObjectEditorViewportClient.cpp
+++ b/Engine/Plugins/Experimental/Mutable/Source/CustomizableObjectEditor/Private/MuCOE/CustomizableObjectEditorViewportClient.cpp
@@ -2246,11 +2246,6 @@ public:
 	SLATE_ARGUMENT(FText, DefaultFileName)
 	SLATE_END_ARGS()
 
-	SMutableSelectFolderDlg() : UserResponse(EAppReturnType::Cancel)
-	{
-		AddAllMaterialTextures = false;
-	}
-
 	void Construct(const FArguments& InArgs);
 
 public:
@@ -2263,8 +2258,7 @@ public:
 	/** FileName getter */
 	FString GetFileName();
 
-	/** Getter for AddAllMaterialTextures */
-	bool GetAddAllMaterialTextures();
+	bool GetExportAllResources();
 
 protected:
 	void OnPathChange(const FString& NewPath);
@@ -2272,10 +2266,10 @@ protected:
 	void OnNameChange(const FText& NewName, ETextCommit::Type CommitInfo);
 	void OnBoolParameterChanged(ECheckBoxState InCheckboxState);
 
-	EAppReturnType::Type UserResponse;
+	EAppReturnType::Type UserResponse = EAppReturnType::Cancel; 
 	FText AssetPath;
 	FText FileName;
-	bool AddAllMaterialTextures;
+	bool bExportAllResources = false;
 };
 
 
@@ -2335,7 +2329,7 @@ void FCustomizableObjectEditorViewportClient::BakeInstance()
 	UCustomizableObjectInstance* Instance = CustomizableObjectEditorPtr.Pin()->GetPreviewInstance();
 	FString ObjectName = Instance->GetCustomizableObject()->GetName();
 	FText DefaultFileName = FText::Format(LOCTEXT("DefaultFileNameForBakeInstance", "{0}"), FText::AsCultureInvariant(ObjectName));
-	bool AddAllMaterialTextures = false;
+	bool bExportAllResources = false;
 
 	TSharedRef<SMutableSelectFolderDlg> FolderDlg =
 		SNew(SMutableSelectFolderDlg)
@@ -2368,7 +2362,8 @@ void FCustomizableObjectEditorViewportClient::BakeInstance()
 			}
 		}
 
-		AddAllMaterialTextures = FolderDlg->GetAddAllMaterialTextures();
+		bExportAllResources = FolderDlg->GetExportAllResources();
+
 		BakingOverwritePermission = false;
 		FString CustomObjectPath = Instance->GetCustomizableObject()->GetPathName();
 		FString AssetPath = FolderDlg->GetAssetPath();
@@ -2406,7 +2401,7 @@ void FCustomizableObjectEditorViewportClient::BakeInstance()
 				TArray<FString> ArrayCachedElement;
 				TArray<UObject*> ArrayCachedObject;
 
-				if (AddAllMaterialTextures)
+				if (bExportAllResources)
 				{
 					UMaterialInstance* Inst;
 					UMaterial* Material;
@@ -2709,6 +2704,29 @@ void FCustomizableObjectEditorViewportClient::BakeInstance()
 						}
 					}
 				}
+				
+				// Skeletal Mesh's Skeleton
+				if (Mesh->GetSkeleton())
+				{
+					const bool bTransient = Mesh->GetSkeleton()->GetPackage() == GetTransientPackage();
+
+					// Don't duplicate if not transient or export all assets.
+					if (bTransient || bExportAllResources)
+					{
+						FString SkeletonName = ObjectName + "_Skeleton";
+						if (!ManageBakingAction(AssetPath, SkeletonName))
+						{
+							return;
+						}
+
+						FString SkeletonPkgName = FolderDlg->GetAssetPath() + FString("/") + SkeletonName;
+						UObject* DuplicatedSkeleton = BakeHelper_DuplicateAsset(Mesh->GetSkeleton(), SkeletonName, SkeletonPkgName, false, ReplacementMap, BakingOverwritePermission);
+
+						ArrayCachedObject.Add(DuplicatedSkeleton);
+						PackagesToSave.Add(DuplicatedSkeleton->GetPackage());
+						ReplacementMap.Add(Mesh->GetSkeleton(), DuplicatedSkeleton);
+					}
+				}
 
 				// Make sure source data is present in the mesh before we duplciate:
 				FUnrealBakeHelpers::BakeHelper_RegenerateImportedModel(Mesh);
@@ -2718,6 +2736,7 @@ void FCustomizableObjectEditorViewportClient::BakeInstance()
 				{
 					return;
 				}
+
 				FString PkgName = FolderDlg->GetAssetPath() + FString("/") + ObjectName;
 				UObject* DupObject = BakeHelper_DuplicateAsset(Mesh, ObjectName, PkgName, false, ReplacementMap, BakingOverwritePermission);
 				ArrayCachedObject.Add(DupObject);
@@ -3305,7 +3324,7 @@ void SMutableSelectFolderDlg::Construct(const FArguments& InArgs)
 	AssetPath = FText::FromString(FPackageName::GetLongPackagePath(InArgs._DefaultAssetPath.ToString()));
 	FileName = InArgs._DefaultFileName;
 
-	AddAllMaterialTextures = false;
+	bExportAllResources = false;
 
 	if (AssetPath.IsEmpty())
 	{
@@ -3382,16 +3401,16 @@ void SMutableSelectFolderDlg::Construct(const FArguments& InArgs)
 				SNew(STextBlock)
 				.Text(LOCTEXT("ExportAllUsedResources", "Export all used resources  "))
 				.Font(FSlateFontInfo(FPaths::EngineContentDir() / TEXT("Slate/Fonts/Roboto-Regular.ttf"), 14))
-				.ToolTipText(LOCTEXT("Export all used Resources", "All the materials and textures used by the object will be baked/stored in the target folder. Otherwise, only the assets that Mutable modifies will be baked/stored."))
+				.ToolTipText(LOCTEXT("Export all used Resources", "All the resources used by the object will be baked/stored in the target folder. Otherwise, only the assets that Mutable modifies will be baked/stored."))
 			]
 			+ SHorizontalBox::Slot()
 			.HAlign(HAlign_Left)
 			.AutoWidth()
 			[
 				SNew(SCheckBox)
-				.ToolTipText(LOCTEXT("ExportAllTextures", "Export all material's textures"))
+				.ToolTipText(LOCTEXT("ExportAllResources", "Export all resources"))
 				.HAlign(HAlign_Right)
-				.IsChecked(AddAllMaterialTextures ? ECheckBoxState::Checked : ECheckBoxState::Unchecked)
+				.IsChecked(bExportAllResources ? ECheckBoxState::Checked : ECheckBoxState::Unchecked)
 				.OnCheckStateChanged(this, &SMutableSelectFolderDlg::OnBoolParameterChanged)
 			]
 		]
@@ -3448,7 +3467,7 @@ void SMutableSelectFolderDlg::OnNameChange(const FText& NewName, ETextCommit::Ty
 
 void SMutableSelectFolderDlg::OnBoolParameterChanged(ECheckBoxState InCheckboxState)
 {
-	AddAllMaterialTextures = !AddAllMaterialTextures;
+	bExportAllResources = InCheckboxState == ECheckBoxState::Checked;
 }
 
 
@@ -3470,9 +3489,9 @@ FString SMutableSelectFolderDlg::GetFileName()
 }
 
 
-bool SMutableSelectFolderDlg::GetAddAllMaterialTextures()
+bool SMutableSelectFolderDlg::GetExportAllResources()
 {
-	return AddAllMaterialTextures;
+	return bExportAllResources;
 }
 
 #undef LOCTEXT_NAMESPACE 
