[mutable] Fix morph bug when the same mesh and morph are reused in a customizable object.
Remove unnecessary processing to generate empty morphs.
#rnx

diff --git a/Engine/Plugins/Experimental/Mutable/Source/CustomizableObjectEditor/Private/MuCOE/GenerateMutableSource/GenerateMutableSourceMesh.cpp b/Engine/Plugins/Experimental/Mutable/Source/CustomizableObjectEditor/Private/MuCOE/GenerateMutableSource/GenerateMutableSourceMesh.cpp
--- a/Engine/Plugins/Experimental/Mutable/Source/CustomizableObjectEditor/Private/MuCOE/GenerateMutableSource/GenerateMutableSourceMesh.cpp
+++ b/Engine/Plugins/Experimental/Mutable/Source/CustomizableObjectEditor/Private/MuCOE/GenerateMutableSource/GenerateMutableSourceMesh.cpp
@@ -2395,11 +2395,8 @@ mu::NodeMeshPtr GenerateMorphMesh(const UEdGraphPin* Pin,
 			MorphedSourceMeshNodeTable->SetParameterName(TCHAR_TO_ANSI(*TypedNodeTable->ParameterName));
 			MorphedSourceMeshNodeTable->SetMessageContext(MorphNode);
 
-			mu::NodeMeshMakeMorphPtr IdentityMorph = new mu::NodeMeshMakeMorph;
-			IdentityMorph->SetBase(BaseSourceMesh.get());
-			IdentityMorph->SetTarget(BaseSourceMesh.get());
-			IdentityMorph->SetMessageContext(MorphNode);
-
+			// A null target will leave the base unchanged
+			mu::NodeMeshPtr IdentityMorph = nullptr;
 			Result->SetMorph(0, IdentityMorph);
 
 			mu::NodeMeshMakeMorphPtr Morph = new mu::NodeMeshMakeMorph;
@@ -2508,11 +2505,8 @@ void GenerateMorphTarget(const UCustomizableObjectNode* Node, const UEdGraphPin*
 		MorphedSourceMeshNode->SetValue(MorphedSourceMesh);
 		MorphedSourceMeshNode->SetMessageContext(Node);
 
-		mu::NodeMeshMakeMorphPtr IdentityMorph = new mu::NodeMeshMakeMorph;
-		IdentityMorph->SetBase(BaseSourceMesh.get());
-		IdentityMorph->SetTarget(BaseSourceMesh.get());
-		IdentityMorph->SetMessageContext(Node);
-
+		// A null target will leave the base unchanged
+		mu::NodeMeshPtr IdentityMorph = nullptr;
 		MeshNode->SetMorph(0, IdentityMorph);
 
 		mu::NodeMeshMakeMorphPtr Morph = new mu::NodeMeshMakeMorph;
diff --git a/Engine/Plugins/Experimental/Mutable/Source/MutableTools/Private/MuT/CodeGenerator_Image.cpp b/Engine/Plugins/Experimental/Mutable/Source/MutableTools/Private/MuT/CodeGenerator_Image.cpp
--- a/Engine/Plugins/Experimental/Mutable/Source/MutableTools/Private/MuT/CodeGenerator_Image.cpp
+++ b/Engine/Plugins/Experimental/Mutable/Source/MutableTools/Private/MuT/CodeGenerator_Image.cpp
@@ -1446,6 +1446,8 @@ namespace mu
         // Mesh
         if ( node.m_pMesh )
         {
+			// TODO: This will probably result in a duplicated mesh subgraph, with the original mesh but new layout block ids.
+			// See if it can be optimized and try to reuse the existing layout block ids instead of generating new ones.
             FMeshGenerationResult MeshResult;
 			FMeshGenerationOptions MeshOptions;
 			MeshOptions.State = m_currentStateIndex;
@@ -1603,7 +1605,6 @@ namespace mu
         rasterop->op.args.ImageRasterMesh.blockIndex = op->op.args.ImageRasterMesh.blockIndex;
         rasterop->op.args.ImageRasterMesh.sizeX = op->op.args.ImageRasterMesh.sizeX;
         rasterop->op.args.ImageRasterMesh.sizeY = op->op.args.ImageRasterMesh.sizeY;
-//        rasterop->op.args.ImageRasterMesh.growBorder = 0;
 
         Ptr<ASTOpFixed> mapop = new ASTOpFixed();
         mapop->op.type = OP_TYPE::IM_MAKEGROWMAP;
@@ -1626,7 +1627,7 @@ namespace mu
 
 		MUTABLE_CPUPROFILER_SCOPE(NodeImageMipmap);
 
-        Ptr<ASTOp> res = 0;
+        Ptr<ASTOp> res;
 
         Ptr<ASTOpImageMipmap> op = new ASTOpImageMipmap();
 
@@ -1635,7 +1636,7 @@ namespace mu
         op->Levels = 0;
 
         // Source image
-        Ptr<ASTOp> base = 0;
+        Ptr<ASTOp> base;
         if ( Node* pSource = node.m_pSource.get() )
         {
             MUTABLE_CPUPROFILER_SCOPE(Base);
@@ -1651,8 +1652,8 @@ namespace mu
 
         // The number of tail mipmaps depends on the cell size. We need to know it for some code
         // optimisation operations. Scan the source image code looking for this info
-        int blockX = 0;
-        int blockY = 0;
+        int32 blockX = 0;
+        int32 blockY = 0;
         if ( m_compilerOptions->m_textureLayoutStrategy
              !=
              CompilerOptions::TextureLayoutStrategy::None )
@@ -1663,9 +1664,9 @@ namespace mu
 
         if ( blockX && blockY )
         {
-            int mipsX = (int)ceilf( logf( (float)blockX )/logf(2.0f) );
-            int mipsY = (int)ceilf( logf( (float)blockY )/logf(2.0f) );
-            op->BlockLevels = (uint8_t)FMath::Max( mipsX, mipsY );
+            int32 mipsX = (int)ceilf( logf( (float)blockX )/logf(2.0f) );
+            int32 mipsY = (int)ceilf( logf( (float)blockY )/logf(2.0f) );
+            op->BlockLevels = (uint8)FMath::Max( mipsX, mipsY );
         }
         else
         {
diff --git a/Engine/Plugins/Experimental/Mutable/Source/MutableTools/Private/MuT/CodeGenerator_Mesh.cpp b/Engine/Plugins/Experimental/Mutable/Source/MutableTools/Private/MuT/CodeGenerator_Mesh.cpp
--- a/Engine/Plugins/Experimental/Mutable/Source/MutableTools/Private/MuT/CodeGenerator_Mesh.cpp
+++ b/Engine/Plugins/Experimental/Mutable/Source/MutableTools/Private/MuT/CodeGenerator_Mesh.cpp
@@ -382,22 +382,15 @@ class Node;
 			FMeshGenerationOptions TargetOptions = InOptions;
 			TargetOptions.bUniqueVertexIDs = false;
 			TargetOptions.bLayouts = false;
-			TargetOptions.OverrideLayouts.Empty();
+			// We need to override the layouts with the layouts that were generated for the base to make
+			// sure that we get the correct mesh when generating the target
+			TargetOptions.OverrideLayouts = BaseResult.GeneratedLayouts;
 			TargetOptions.ActiveTags.Empty();
             GenerateMesh(TargetOptions, TargetResult, node.m_morphs[t]);
 
             // TODO: Make sure that the target is a mesh with the morph format
             Ptr<ASTOp> target = TargetResult.meshOp;
 
-            // If the vertex indices are supposed to be relative in the targets, adjust them
-            //if (node.m_vertexIndicesAreRelative)
-            //{
-            //    Ptr<ASTOpMeshRemapIndices> remapIndices = new ASTOpMeshRemapIndices;
-            //    remapIndices->source = target;
-            //    remapIndices->reference = baseResult.baseMeshOp;
-            //    target = remapIndices;
-            //}
-
             OpMorph->AddTarget( target );
         }
  
@@ -910,7 +903,7 @@ class Node;
 
 				// If it is similar and we are overriding the layouts, we must compare the layouts of the candidate with the ones
 				// we are using to override.
-				if (InOptions.bLayouts && bIsOverridingLayouts)
+				if (bIsOverridingLayouts)
 				{
 					if (Candidate->GetLayoutCount() != InOptions.OverrideLayouts.Num())
 					{
diff --git a/Engine/Plugins/Experimental/Mutable/Source/MutableTools/Private/MuT/NodeMeshMorph.cpp b/Engine/Plugins/Experimental/Mutable/Source/MutableTools/Private/MuT/NodeMeshMorph.cpp
--- a/Engine/Plugins/Experimental/Mutable/Source/MutableTools/Private/MuT/NodeMeshMorph.cpp
+++ b/Engine/Plugins/Experimental/Mutable/Source/MutableTools/Private/MuT/NodeMeshMorph.cpp
@@ -149,12 +149,6 @@ namespace mu
 		return m_pD->m_morphs.Num();
 	}
 
-	//---------------------------------------------------------------------------------------------
-	//void NodeMeshMorph::SetMorphIndicesAreRelative( bool relative )
-	//{
-	//	m_pD->m_vertexIndicesAreRelative = relative;
-	//}
-
 	//---------------------------------------------------------------------------------------------
 	void NodeMeshMorph::SetReshapeSkeleton(bool bEnable)
 	{
