[Mutable] Fix duplicated parameter ids.

#rnx

diff --git a/Engine/Plugins/Experimental/Mutable/Source/CustomizableObjectEditor/Private/MuCOE/GenerateMutableSource/GenerateMutableSource.cpp b/Engine/Plugins/Experimental/Mutable/Source/CustomizableObjectEditor/Private/MuCOE/GenerateMutableSource/GenerateMutableSource.cpp
--- a/Engine/Plugins/Experimental/Mutable/Source/CustomizableObjectEditor/Private/MuCOE/GenerateMutableSource/GenerateMutableSource.cpp
+++ b/Engine/Plugins/Experimental/Mutable/Source/CustomizableObjectEditor/Private/MuCOE/GenerateMutableSource/GenerateMutableSource.cpp
@@ -642,8 +642,7 @@ mu::NodeObjectPtr GenerateMutableSource(const UEdGraphPin * Pin, FMutableGraphGe
 		Result = ObjectNode;
 
 		ObjectNode->SetName(TCHAR_TO_ANSI(*TypedNodeObj->ObjectName));
-		ObjectNode->SetUid(TCHAR_TO_ANSI(*TypedNodeObj->NodeGuid.ToString()));
-
+		ObjectNode->SetUid(TCHAR_TO_ANSI(*GenerationContext.GetNodeIdUnique(TypedNodeObj).ToString()));
 
 		// LOD
 		int NumLODs = TypedNodeObj->GetNumLODPins();
