[Mutable] Added support for Multidimensional Texture Parameters.

- Extended the Instance/Descriptor api to support setting a UTexture2D directly as a Texture Parameter.

#rnx

diff --git a/Engine/Plugins/Experimental/Mutable/Source/CustomizableObject/Classes/MuCO/CustomizableObjectInstance.h b/Engine/Plugins/Experimental/Mutable/Source/CustomizableObject/Classes/MuCO/CustomizableObjectInstance.h
--- a/Engine/Plugins/Experimental/Mutable/Source/CustomizableObject/Classes/MuCO/CustomizableObjectInstance.h
+++ b/Engine/Plugins/Experimental/Mutable/Source/CustomizableObject/Classes/MuCO/CustomizableObjectInstance.h
@@ -187,6 +187,7 @@ public:
 
 	// UObject interface.
 #if WITH_EDITOR
+	virtual void PreEditChange(FProperty* PropertyAboutToChange) override;
 	virtual void PostEditChangeProperty(FPropertyChangedEvent& PropertyChangedEvent) override;
 	virtual bool CanEditChange( const FProperty* InProperty ) const override;
 	bool InstanceUpdated; // Flag for the editor, to know when the instance's skeletal mesh has been updated
@@ -350,7 +351,21 @@ public:
 	// Sets the float value "FloatValue" of a float parameter with index "FloatParamIndex"
 	UFUNCTION(BlueprintCallable, Category = CustomizableObjectInstance)
 	void SetFloatParameterSelectedOption(const FString& FloatParamName, float FloatValue, int32 RangeIndex = -1);
+	
+	/** Gets the value of a texture parameter with name "TextureParamName". */
+	uint64 GetTextureParameterSelectedOption(const FString& TextureParamName, int32 RangeIndex) const;
+
+	/** Gets the texture of a texture parameter with name "TextureParamName". */
+	UFUNCTION(BlueprintCallable, Category = CustomizableObjectInstance)
+	UTexture2D* GetTextureParameterSelectedOptionT(const FString& TextureParamName, int32 RangeIndex) const;
+
+	/** Sets the texture value "TextureValue" of a texture parameter with index "TextureParamIndex". */
+	void SetTextureParameterSelectedOption(const FString& TextureParamName, uint64 TextureValue, int32 RangeIndex);
 
+	/** Sets the texture "Texture" of a texture parameter with index "TextureParamIndex". */
+	UFUNCTION(BlueprintCallable, Category = CustomizableObjectInstance)
+	void SetTextureParameterSelectedOptionT(const FString& TextureParamName, UTexture2D* TextureValue, int32 RangeIndex);
+	
 	// Gets the value of a color parameter with name "ColorParamName"
 	UFUNCTION(BlueprintCallable, Category = CustomizableObjectInstance)
 	FLinearColor GetColorParameterSelectedOption(const FString& ColorParamName) const;
@@ -673,6 +688,12 @@ private:
 	/** Hash of the UCustomizableObjectInstance::Descriptor on the last successful update. */
 	FDescriptorRuntimeHash DescriptorRuntimeHash;
 
+#if WITH_EDITORONLY_DATA
+	/** Textures used in the Texture Parameters. */
+	UPROPERTY(EditAnywhere, Category = TextureParameter)
+	TArray<TObjectPtr<UTexture2D>> TextureParameterDeclarations;
+#endif
+	
 	// Deprecated properties
 	
 	UPROPERTY()
diff --git a/Engine/Plugins/Experimental/Mutable/Source/CustomizableObject/Classes/MuCO/CustomizableObjectSystem.h b/Engine/Plugins/Experimental/Mutable/Source/CustomizableObject/Classes/MuCO/CustomizableObjectSystem.h
--- a/Engine/Plugins/Experimental/Mutable/Source/CustomizableObject/Classes/MuCO/CustomizableObjectSystem.h
+++ b/Engine/Plugins/Experimental/Mutable/Source/CustomizableObject/Classes/MuCO/CustomizableObjectSystem.h
@@ -28,6 +28,7 @@ class FCustomizableObjectCompilerBase;
 class ITargetPlatform;
 class SNotificationItem;
 class UCustomizableObject;
+class UDefaultImageProvider;
 class USkeletalMesh;
 class UTexture2D;
 struct FFrame;
@@ -175,7 +176,10 @@ public:
 		Unreal,
 
 		// Data will be provided from an unreal texture, and will only be loaded when actually needed in the Mutable thread
-		Unreal_Deferred
+		Unreal_Deferred,
+
+		// Number of elements of this enum.
+		Count
 	};
 
 	// Query that Mutable will run to find out if a texture will be provided as an Unreal UTexture2D,
@@ -198,40 +202,6 @@ public:
 };
 
 
-// Example implementation of a ICustomizableSystemImageProvider that just uses a predefined array of textures.
-// It is also used by editors to set some preview images.
-UCLASS()
-class CUSTOMIZABLEOBJECT_API UCustomizableObjectImageProviderArray : public UCustomizableSystemImageProvider
-{
-public:
-	GENERATED_BODY()
-
-	UPROPERTY(EditAnywhere, Category = Preview)
-	TArray< TObjectPtr<UTexture2D> > Textures;
-
-	// UObject interface
-#if WITH_EDITOR
-	void PostEditChangeProperty(struct FPropertyChangedEvent& PropertyChangedEvent) override;
-#endif
-
-	// UCustomizableSystemImageProvider interface
-	ValueType HasTextureParameterValue(int64 ID) override;
-	UTexture2D* GetTextureParameterValue(int64 ID) override;
-	void GetTextureParameterValues(TArray<FCustomizableObjectExternalTexture>& OutValues) override;
-
-	// Own interface
-	void InvalidateIds();
-
-	DECLARE_MULTICAST_DELEGATE(FOnTexturesChanged);
-	FOnTexturesChanged TexturesChangeDelegate;
-
-private:
-
-	// The preview texture values will start at this ID, to avoid clashing with other game-specific custom providers.
-	int FirstId = 100000;
-};
-
-
 class UCustomizableInstanceLODManagementBase;
 
 
@@ -333,21 +303,12 @@ public:
 	/** [Texture Parameters] Remove all images from the cache. */
 	void ClearImageCache();
 
-
-#if WITH_EDITOR
-	/** [Texture Parameters] Get the editor-side preview external texture provider. 
-	 * \TODO: Move to the editor? 
-	 */
-	UCustomizableObjectImageProviderArray* GetEditorExternalImageProvider();
-#endif
+	UDefaultImageProvider& GetDefaultImageProvider();
 
 private:
-
-	/** [Texture Parameters] If in editor, this will hold a reference to the image provider used to give examples of external images to preview in the customizable objects that use this functionality.
-	 * \TODO : Move to the editor ?
-	 */
+	
 	UPROPERTY()
-	TObjectPtr<UCustomizableObjectImageProviderArray> PreviewExternalImageProvider = nullptr;
+	TObjectPtr<UDefaultImageProvider> DefaultImageProvider = nullptr;
 
 public:
     
@@ -467,5 +428,8 @@ private:
 	TMap<FString, int64> PlatformMaxChunkSize;
 	
 #endif
+
+	// Friends
+	friend class FCustomizableObjectSystemPrivate;
 };
 
diff --git a/Engine/Plugins/Experimental/Mutable/Source/CustomizableObject/Private/MuCO/CustomizableObjectInstance.cpp b/Engine/Plugins/Experimental/Mutable/Source/CustomizableObject/Private/MuCO/CustomizableObjectInstance.cpp
--- a/Engine/Plugins/Experimental/Mutable/Source/CustomizableObject/Private/MuCO/CustomizableObjectInstance.cpp
+++ b/Engine/Plugins/Experimental/Mutable/Source/CustomizableObject/Private/MuCO/CustomizableObjectInstance.cpp
@@ -62,6 +62,7 @@
 #include "MuCO/CustomizableObjectSystem.h"
 #include "MuCO/CustomizableObjectSystemPrivate.h"
 #include "MuCO/CustomizableSkeletalComponent.h"
+#include "MuCO/DefaultImageProvider.h"
 #include "MuCO/MultilayerProjector.h"
 #include "MuCO/UnrealConversionUtils.h"
 #include "MuR/MeshBufferSet.h"
@@ -326,11 +327,51 @@ void UCustomizableInstancePrivateData::PrepareForUpdate(const TSharedPtr<FMutabl
 
 #if WITH_EDITOR
 
+void UCustomizableObjectInstance::PreEditChange(FProperty* PropertyAboutToChange)
+{
+	UObject::PreEditChange(PropertyAboutToChange);
+	
+	const FName PropertyName = PropertyAboutToChange ? PropertyAboutToChange->GetFName() : NAME_None;
+
+	if (PropertyName == GET_MEMBER_NAME_CHECKED(UCustomizableObjectInstance, TextureParameterDeclarations))
+	{
+		UDefaultImageProvider& DefaultImageProvider = UCustomizableObjectSystem::GetInstance()->GetDefaultImageProvider();
+		
+		for (TObjectPtr<UTexture2D> Texture : TextureParameterDeclarations)
+		{
+			if (Texture)
+			{
+				const int32 TextureId = DefaultImageProvider.Get(Texture);
+				DefaultImageProvider.Keep(TextureId, false);
+			}
+		}
+	}
+}
+
+
 void UCustomizableObjectInstance::PostEditChangeProperty(FPropertyChangedEvent& PropertyChangedEvent)
 {
+	Super::PostEditChangeProperty(PropertyChangedEvent);
+
 	bEditorPropertyChanged = true;
 
-	Super::PostEditChangeProperty(PropertyChangedEvent);
+	const FName PropertyName = PropertyChangedEvent.MemberProperty ? PropertyChangedEvent.MemberProperty->GetFName() : NAME_None;
+
+	if (PropertyName == GET_MEMBER_NAME_CHECKED(UCustomizableObjectInstance, TextureParameterDeclarations))
+	{
+		UDefaultImageProvider& DefaultImageProvider = UCustomizableObjectSystem::GetInstance()->GetDefaultImageProvider();
+		
+		for (TObjectPtr<UTexture2D> Texture : TextureParameterDeclarations)
+		{
+			if (Texture)
+			{
+				const int32 TextureId = DefaultImageProvider.GetOrAdd(Texture);
+				DefaultImageProvider.Keep(TextureId, true);
+			}
+		}
+
+		UpdateSkeletalMeshAsync(true, true);
+	}
 }
 
 
@@ -1939,6 +1980,30 @@ void UCustomizableObjectInstance::SetFloatParameterSelectedOption(const FString&
 }
 
 
+uint64 UCustomizableObjectInstance::GetTextureParameterSelectedOption(const FString& TextureParamName, const int32 RangeIndex) const
+{
+	return Descriptor.GetTextureParameterSelectedOption(TextureParamName, RangeIndex);
+}
+
+
+UTexture2D* UCustomizableObjectInstance::GetTextureParameterSelectedOptionT(const FString& TextureParamName, const int32 RangeIndex) const
+{
+	return Descriptor.GetTextureParameterSelectedOptionT(TextureParamName, RangeIndex);
+}
+
+
+void UCustomizableObjectInstance::SetTextureParameterSelectedOption(const FString& TextureParamName, const uint64 TextureValue, const int32 RangeIndex)
+{
+	Descriptor.SetTextureParameterSelectedOption(TextureParamName, TextureValue, RangeIndex);
+}
+
+
+void UCustomizableObjectInstance::SetTextureParameterSelectedOptionT(const FString& TextureParamName, UTexture2D* TextureValue, const int32 RangeIndex)
+{
+	Descriptor.SetTextureParameterSelectedOptionT(TextureParamName, TextureValue, RangeIndex);
+}
+
+
 FLinearColor UCustomizableObjectInstance::GetColorParameterSelectedOption(const FString& ColorParamName) const
 {
 	return Descriptor.GetColorParameterSelectedOption(ColorParamName);
diff --git a/Engine/Plugins/Experimental/Mutable/Source/CustomizableObject/Private/MuCO/CustomizableObjectInstanceDescriptor.cpp b/Engine/Plugins/Experimental/Mutable/Source/CustomizableObject/Private/MuCO/CustomizableObjectInstanceDescriptor.cpp
--- a/Engine/Plugins/Experimental/Mutable/Source/CustomizableObject/Private/MuCO/CustomizableObjectInstanceDescriptor.cpp
+++ b/Engine/Plugins/Experimental/Mutable/Source/CustomizableObject/Private/MuCO/CustomizableObjectInstanceDescriptor.cpp
@@ -10,6 +10,7 @@
 #include "MuCO/CustomizableObject.h"
 #include "MuCO/CustomizableObjectPrivate.h"
 #include "MuCO/CustomizableObjectSystem.h"
+#include "MuCO/DefaultImageProvider.h"
 #include "MuR/Model.h"
 #include "MuR/Parameters.h"
 #include "MuR/RefCounted.h"
@@ -190,15 +191,19 @@ void FCustomizableObjectInstanceDescriptor::SaveDescriptor(FArchive& Ar)
 		case EMutableParameterType::Texture:
 		{
 			uint64 Value = 0;
+			TArray<uint64> RangeValues;
+
 			for (const FCustomizableObjectTextureParameterValue& P : TextureParameters)
 			{
 				if (P.ParameterName == Name)
 				{
 					Value = P.ParameterValue;
+					RangeValues = P.ParameterRangeValues;
 					break;
 				}
 			}
 			Ar << Value;
+			Ar << RangeValues;
 
 			break;
 		}
@@ -384,12 +389,16 @@ void FCustomizableObjectInstanceDescriptor::LoadDescriptor(FArchive& Ar)
 		case EMutableParameterType::Texture:
 		{
 			uint64 Value = 0;
+			TArray<uint64> RangeValues;
 			Ar << Value;
+			Ar << RangeValues;
+
 			for (FCustomizableObjectTextureParameterValue& P : TextureParameters)
 			{
 				if (P.ParameterName == Name)
 				{
 					P.ParameterValue = Value;
+					P.ParameterRangeValues = RangeValues;
 					break;
 				}
 			}
@@ -632,10 +641,21 @@ mu::ParametersPtr FCustomizableObjectInstanceDescriptor::GetParameters() const
 			{
 				if (TextureParameter.ParameterName == Name || (!Uid.IsEmpty() && TextureParameter.Uid == Uid))
 				{
-					MutableParameters->SetImageValue(ParamIndex, TextureParameter.ParameterValue);
+					if (mu::RangeIndexPtr RangeIdxPtr = MutableParameters->NewRangeIndex(ParamIndex))
+					{
+						for (int RangeIndex = 0; RangeIndex < TextureParameter.ParameterRangeValues.Num(); ++RangeIndex)
+						{
+							RangeIdxPtr->SetPosition(0, RangeIndex);
+							MutableParameters->SetImageValue(ParamIndex, TextureParameter.ParameterRangeValues[RangeIndex], RangeIdxPtr);
+						}
+					}
+					else
+					{
+						MutableParameters->SetImageValue(ParamIndex, TextureParameter.ParameterValue);
+					}
+
+					break;
 				}
-	
-				break;
 			}
 
 			break;
@@ -813,8 +833,7 @@ void FCustomizableObjectInstanceDescriptor::ReloadParameters()
 							Param.ParameterRangeValues.AddDefaulted(ValueIndex + 1 - Param.ParameterRangeValues.Num());
 						}
 
-						float Value = MutableParameters->GetFloatValue(ParamIndex, RangeValueIdxPtr);
-						Param.ParameterRangeValues[RangeIndex] = Value;
+						Param.ParameterRangeValues[RangeIndex] = MutableParameters->GetFloatValue(ParamIndex, RangeValueIdxPtr);
 					}
 				}
 				else
@@ -950,14 +969,41 @@ void FCustomizableObjectInstanceDescriptor::ReloadParameters()
 			{
 				return P.ParameterName == Name || (!Uid.IsEmpty() && P.Uid == Uid);
 			};
-
+			
 			if (FCustomizableObjectTextureParameterValue* Result = OldTextureParameters.FindByPredicate(FindByNameAndUid))
-			{
-				Param.ParameterValue = Result->ParameterValue;
-			}
+			{	
+				if (mu::RangeIndexPtr RangeIdxPtr = MutableParameters->NewRangeIndex(ParamIndex))
+				{
+					Param.ParameterRangeValues = Result->ParameterRangeValues;
+				}
+				else
+				{
+					Param.ParameterValue = Result->ParameterValue;
+				}
+			} 
 			else // Not found in Instance Parameters. Use Mutable Parameters.
 			{
-				Param.ParameterValue = MutableParameters->GetImageValue(ParamIndex);
+				if (mu::RangeIndexPtr RangeIdxPtr = MutableParameters->NewRangeIndex(ParamIndex))
+				{
+					const int32 ValueCount = MutableParameters->GetValueCount(ParamIndex);
+
+					for (int ValueIndex = 0; ValueIndex < ValueCount; ++ValueIndex)
+					{	
+						mu::RangeIndexPtr RangeValueIdxPtr = MutableParameters->GetValueIndex(ParamIndex, ValueIndex);
+						int32 RangeIndex = RangeValueIdxPtr->GetPosition(0);
+
+						if (!Param.ParameterRangeValues.IsValidIndex(ValueIndex))
+						{
+							Param.ParameterRangeValues.AddDefaulted(ValueIndex + 1 - Param.ParameterRangeValues.Num());
+						}
+
+						Param.ParameterRangeValues[RangeIndex] = MutableParameters->GetImageValue(ParamIndex, RangeValueIdxPtr);
+					}
+				}
+				else
+				{
+					Param.ParameterValue = MutableParameters->GetImageValue(ParamIndex);
+				}
 			}
 
 			TextureParameters.Add(Param);
@@ -1205,6 +1251,87 @@ void FCustomizableObjectInstanceDescriptor::SetFloatParameterSelectedOption(cons
 }
 
 
+uint64 FCustomizableObjectInstanceDescriptor::GetTextureParameterSelectedOption(const FString& TextureParamName, const int32 RangeIndex) const
+{
+	check(CustomizableObject);
+
+	const int32 ParameterIndexInObject = CustomizableObject->FindParameter(TextureParamName);
+
+	const int32 TextureParamIndex = FindTextureParameterNameIndex(TextureParamName);
+
+	if (ParameterIndexInObject >= 0 && TextureParamIndex >= 0)
+	{
+		if (RangeIndex == -1)
+		{
+			check(!IsParamMultidimensional(ParameterIndexInObject)); // This param is multidimensional, it must have a RangeIndex of 0 or more
+			return TextureParameters[TextureParamIndex].ParameterValue;
+		}
+		else
+		{
+			check(IsParamMultidimensional(ParameterIndexInObject)); // This param is not multidimensional, it must have a RangeIndex of -1
+
+			if (TextureParameters[TextureParamIndex].ParameterRangeValues.IsValidIndex(RangeIndex))
+			{
+				return TextureParameters[TextureParamIndex].ParameterRangeValues[RangeIndex];
+			}
+		}
+	}
+
+	return -1; 
+}
+
+
+UTexture2D* FCustomizableObjectInstanceDescriptor::GetTextureParameterSelectedOptionT(const FString& TextureParamName, const int32 RangeIndex) const
+{
+	UDefaultImageProvider& ImageProvider = UCustomizableObjectSystem::GetInstance()->GetDefaultImageProvider();
+
+	const uint64 TextureValue = GetTextureParameterSelectedOption(TextureParamName, RangeIndex);
+	return ImageProvider.Get(TextureValue);
+}
+
+
+void FCustomizableObjectInstanceDescriptor::SetTextureParameterSelectedOption(const FString& TextureParamName, const uint64 TextureValue, const int32 RangeIndex)
+{
+	check(CustomizableObject);
+
+	const int32 ParameterIndexInObject = CustomizableObject->FindParameter(TextureParamName);
+
+	const int32 TextureParamIndex = FindTextureParameterNameIndex(TextureParamName);
+
+	if (ParameterIndexInObject >= 0 && TextureParamIndex >= 0)
+	{
+		if (RangeIndex == -1)
+		{
+			check(!IsParamMultidimensional(ParameterIndexInObject)); // This param is multidimensional, it must have a RangeIndex of 0 or more
+			TextureParameters[TextureParamIndex].ParameterValue = TextureValue;
+		}
+		else
+		{
+			check(IsParamMultidimensional(ParameterIndexInObject)); // This param is not multidimensional, it must have a RangeIndex of -1
+
+			if (!TextureParameters[TextureParamIndex].ParameterRangeValues.IsValidIndex(RangeIndex))
+			{
+				const int32 InsertionIndex = TextureParameters[TextureParamIndex].ParameterRangeValues.Num();
+				const int32 NumInsertedElements = RangeIndex + 1 - TextureParameters[TextureParamIndex].ParameterRangeValues.Num();
+				TextureParameters[TextureParamIndex].ParameterRangeValues.InsertDefaulted(InsertionIndex, NumInsertedElements);
+			}
+
+			check(TextureParameters[TextureParamIndex].ParameterRangeValues.IsValidIndex(RangeIndex));
+			TextureParameters[TextureParamIndex].ParameterRangeValues[RangeIndex] = TextureValue;
+		}
+	}
+}
+
+
+void FCustomizableObjectInstanceDescriptor::SetTextureParameterSelectedOptionT(const FString& TextureParamName, UTexture2D* TextureValue, const int32 RangeIndex)
+{
+	UDefaultImageProvider& ImageProvider = UCustomizableObjectSystem::GetInstance()->GetDefaultImageProvider();
+
+	const uint64 NewUid = ImageProvider.GetOrAdd(TextureValue);
+	SetTextureParameterSelectedOption(TextureParamName, NewUid, RangeIndex);
+}
+
+
 FLinearColor FCustomizableObjectInstanceDescriptor::GetColorParameterSelectedOption(const FString& ColorParamName) const
 {
 	check(CustomizableObject);
@@ -1837,6 +1964,38 @@ int32 FCustomizableObjectInstanceDescriptor::RemoveValueFromFloatRange(const FSt
 }
 
 
+int32 FCustomizableObjectInstanceDescriptor::RemoveValueFromTextureRange(const FString& ParamName)
+{
+	const int32 TextureParameterIndex = FindTextureParameterNameIndex(ParamName);
+	if (TextureParameterIndex != -1)
+	{
+		FCustomizableObjectTextureParameterValue& TextureParameter = TextureParameters[TextureParameterIndex];
+		if (TextureParameter.ParameterRangeValues.Num() > 0)
+		{
+			TextureParameter.ParameterRangeValues.Pop();
+			return TextureParameter.ParameterRangeValues.Num() - 1;
+		}
+	}
+	return -1;
+}
+
+
+int32 FCustomizableObjectInstanceDescriptor::RemoveValueFromTextureRange(const FString& ParamName, const int32 RangeIndex)
+{
+	const int32 TextureParameterIndex = FindTextureParameterNameIndex(ParamName);
+	if (TextureParameterIndex != -1)
+	{
+		FCustomizableObjectTextureParameterValue& TextureParameter = TextureParameters[TextureParameterIndex];
+		if (TextureParameter.ParameterRangeValues.Num() > 0)
+		{
+			TextureParameter.ParameterRangeValues.RemoveAt(RangeIndex);
+			return TextureParameter.ParameterRangeValues.Num() - 1;
+		}
+	}
+	return -1;
+}
+
+
 int32 FCustomizableObjectInstanceDescriptor::RemoveValueFromProjectorRange(const FString& ParamName)
 {
 	const int32 ProjectorParameterIndex = FindProjectorParameterNameIndex(ParamName);
diff --git a/Engine/Plugins/Experimental/Mutable/Source/CustomizableObject/Private/MuCO/CustomizableObjectSystem.cpp b/Engine/Plugins/Experimental/Mutable/Source/CustomizableObject/Private/MuCO/CustomizableObjectSystem.cpp
--- a/Engine/Plugins/Experimental/Mutable/Source/CustomizableObject/Private/MuCO/CustomizableObjectSystem.cpp
+++ b/Engine/Plugins/Experimental/Mutable/Source/CustomizableObject/Private/MuCO/CustomizableObjectSystem.cpp
@@ -52,6 +52,7 @@
 #include "MuCO/CustomizableObjectParameterTypeDefinitions.h"
 #include "MuCO/CustomizableObjectPrivate.h"
 #include "MuCO/CustomizableObjectSystemPrivate.h"
+#include "MuCO/DefaultImageProvider.h"
 #include "MuCO/CustomizableObjectUIData.h"
 #include "MuCO/CustomizableSkeletalComponent.h"
 #include "MuCO/ICustomizableObjectModule.h"
@@ -375,15 +376,10 @@ void UCustomizableObjectSystem::InitSystem()
 	Private->MutableSystem->SetImageParameterGenerator(Provider);
 
 #if WITH_EDITOR
-	// Register an example provider of texture parameters, used by editor preview.
-	PreviewExternalImageProvider = NewObject<UCustomizableObjectImageProviderArray>();
-	RegisterImageProvider(PreviewExternalImageProvider);
-
 	if (!IsRunningGame())
 	{
 		FEditorDelegates::PreBeginPIE.AddUObject(this, &UCustomizableObjectSystem::OnPreBeginPIE);
 	}
-
 #endif
 
 	if (FParse::Param(FCommandLine::Get(), TEXT("MutablePortableSerialization")))
@@ -817,6 +813,11 @@ void FCustomizableObjectSystemPrivate::InitUpdateSkeletalMesh(UCustomizableObjec
 		return; // The the requested update is equal to the running update.
 	}
 
+	if (const TObjectPtr<UDefaultImageProvider> DefaultImageProvider = UCustomizableObjectSystem::GetInstance()->DefaultImageProvider)
+	{
+		DefaultImageProvider->CacheTextures(Instance);
+	}
+	
 	// These delegates must be called at the end of the begin update.
 	Instance.BeginUpdateDelegate.Broadcast(&Instance);
 	Instance.BeginUpdateNativeDelegate.Broadcast(&Instance);
@@ -2389,12 +2390,17 @@ void UCustomizableObjectSystem::UnregisterImageProvider(UCustomizableSystemImage
 	Private->ImageProvider->ImageProviders.Remove(Provider);
 }
 
-#if WITH_EDITOR
-UCustomizableObjectImageProviderArray* UCustomizableObjectSystem::GetEditorExternalImageProvider()
-{ 
-	return PreviewExternalImageProvider; 
+
+UDefaultImageProvider& UCustomizableObjectSystem::GetDefaultImageProvider()
+{
+	if (!DefaultImageProvider)
+	{
+		DefaultImageProvider = NewObject<UDefaultImageProvider>();
+		RegisterImageProvider(DefaultImageProvider);
+	}
+
+	return *DefaultImageProvider;
 }
-#endif
 
 
 FMutableOperation FMutableOperation::CreateInstanceUpdate(UCustomizableObjectInstance* InCustomizableObjectInstance, bool bInNeverStream, int32 InMipsToSkip)
@@ -2463,57 +2469,6 @@ void FMutableOperation::MutableIsDisabledCase()
 }
 
 
-UCustomizableSystemImageProvider::ValueType UCustomizableObjectImageProviderArray::HasTextureParameterValue(int64 ID)
-{
-	if (ID >= FirstId && ID < FirstId + Textures.Num())
-	{
-		return UCustomizableSystemImageProvider::ValueType::Unreal;
-	}
-
-	return UCustomizableSystemImageProvider::ValueType::None;
-}
-
-
-UTexture2D* UCustomizableObjectImageProviderArray::GetTextureParameterValue(int64 ID)
-{
-	if (ID >= FirstId && ID < FirstId + Textures.Num())
-	{
-		return Textures[ID - FirstId];
-	}
-
-	return nullptr;
-}
-
-
-void UCustomizableObjectImageProviderArray::GetTextureParameterValues(TArray<FCustomizableObjectExternalTexture>& OutValues)
-{
-	for (int i = 0; i < Textures.Num(); ++i)
-	{
-		if (Textures[i])
-		{
-			FCustomizableObjectExternalTexture Data;
-			Data.Name = Textures[i]->GetName();
-			Data.Value = FirstId + i;
-			OutValues.Add(Data);
-		}
-	}
-}
-
-
-void UCustomizableObjectImageProviderArray::InvalidateIds()
-{
-	FirstId += Textures.Num();
-}
-
-
-#if WITH_EDITOR
-void UCustomizableObjectImageProviderArray::PostEditChangeProperty(struct FPropertyChangedEvent& PropertyChangedEvent)
-{
-	TexturesChangeDelegate.Broadcast();
-}
-#endif
-
-
 int32 UCustomizableObjectSystem::GetNumInstances() const
 {
 	return Private->NumInstances;
@@ -2802,40 +2757,32 @@ uint64 UCustomizableObjectSystem::GetMaxChunkSizeForPlatform(const ITargetPlatfo
 	return MaxChunkSize;
 }
 
+#endif // WITH_EDITOR
+
 
 void UCustomizableObjectSystem::CacheImage(uint64 ImageId)
 {
-	if (GetPrivate()->ImageProvider)
-	{
-		GetPrivate()->ImageProvider->CacheImage(ImageId);
-	}
+	check(GetPrivate()->ImageProvider);
+	GetPrivate()->ImageProvider->CacheImage(ImageId);
 }
 
 
 void UCustomizableObjectSystem::UnCacheImage(uint64 ImageId)
 {
-	if (GetPrivate()->ImageProvider)
-	{
-		GetPrivate()->ImageProvider->UnCacheImage(ImageId);
-	}
+	check(GetPrivate()->ImageProvider);
+	GetPrivate()->ImageProvider->UnCacheImage(ImageId);
 }
 
 
 void UCustomizableObjectSystem::CacheAllImagesInAllProviders(bool bClearPreviousCacheImages)
 {
-	if (GetPrivate()->ImageProvider)
-	{
-		GetPrivate()->ImageProvider->CacheAllImagesInAllProviders(bClearPreviousCacheImages);
-	}
+	check(GetPrivate()->ImageProvider);
+	GetPrivate()->ImageProvider->CacheAllImagesInAllProviders(bClearPreviousCacheImages);
 }
 
 
 void UCustomizableObjectSystem::ClearImageCache()
 {
-	if (GetPrivate()->ImageProvider)
-	{
-		GetPrivate()->ImageProvider->ClearCache();
-	}
-}
-
-#endif // WITH_EDITOR
+	check(GetPrivate()->ImageProvider);
+	GetPrivate()->ImageProvider->ClearCache();
+}
\ No newline at end of file
diff --git a/Engine/Plugins/Experimental/Mutable/Source/CustomizableObject/Private/MuCO/DefaultImageProvider.cpp b/Engine/Plugins/Experimental/Mutable/Source/CustomizableObject/Private/MuCO/DefaultImageProvider.cpp
new file mode 100644
--- /dev/null
+++ b/Engine/Plugins/Experimental/Mutable/Source/CustomizableObject/Private/MuCO/DefaultImageProvider.cpp
@@ -0,0 +1,211 @@
+// Copyright Epic Games, Inc. All Rights Reserved.
+
+#include "MuCO/DefaultImageProvider.h" 
+
+#include "MuCO/CustomizableObjectInstance.h"
+#include "MuR/System.h"
+
+
+namespace UDefaultImageProviderCVars
+{
+	void CheckImageMode(IConsoleVariable* Var);
+	
+	int32 ImageMode = 3;
+	FAutoConsoleVariableRef CVarImageMode(
+		TEXT("Mutable.DefaultImageProvider.ImageMode"),
+		ImageMode,
+		TEXT("0 = None: Texture is not provided by this provider.\n"
+		"2 = Unreal: Data will be provided from an unreal texture, loaded in the game thread and kept in memory.\n"
+		"3 = Unreal_Deferred (default): Data will be provided from an unreal texture. Will only be loaded when actually needed in the Mutable thread."),
+		FConsoleVariableDelegate::CreateStatic(CheckImageMode));
+
+	void CheckImageMode(IConsoleVariable* Var)
+	{
+		if ((ImageMode < 0 || ImageMode >= static_cast<int32>(UCustomizableSystemImageProvider::ValueType::Count)) &&
+			ImageMode == static_cast<int32>(UCustomizableSystemImageProvider::ValueType::Raw)) // Raw not supported.
+		{
+			UE_LOG(LogMutable, Error, TEXT("DefaultImageProvider: Incorrect Image Mode. Setting Texture Mode to \"None\"."));
+			ImageMode = static_cast<int32>(UCustomizableSystemImageProvider::ValueType::None);
+		}
+	}
+}
+
+
+UCustomizableSystemImageProvider::ValueType UDefaultImageProvider::HasTextureParameterValue(int64 ID)
+{
+	return Textures[ID] ?
+		static_cast<ValueType>(UDefaultImageProviderCVars::ImageMode) :
+		ValueType::None;
+}
+
+
+UTexture2D* UDefaultImageProvider::GetTextureParameterValue(int64 ID)
+{
+	return Textures[ID];
+}
+
+
+void UDefaultImageProvider::GetTextureParameterValues(TArray<FCustomizableObjectExternalTexture>& OutValues)
+{
+	for (int32 TextureIndex = 0; TextureIndex < Textures.Num(); ++TextureIndex)
+	{
+		if (const UTexture2D* Texture = Textures[TextureIndex])
+		{
+			FCustomizableObjectExternalTexture Data;
+			Data.Name = Texture->GetName();
+			Data.Value = TextureIndex;
+			OutValues.Add(Data);
+		}
+	}
+}
+
+
+UTexture2D* UDefaultImageProvider::Get(const uint64 TextureId) const
+{
+	return Textures[ToIndex(TextureId)];
+}
+
+
+int64 UDefaultImageProvider::Get(UTexture2D* Texture) const
+{
+	for (int32 TextureIndex = 0; TextureIndex < Textures.Num(); ++TextureIndex)
+	{
+		if (Textures[TextureIndex] == Texture)
+		{
+			return ToTextureId(TextureIndex);
+		}
+	}
+
+	return INDEX_NONE;
+}
+
+
+uint64 UDefaultImageProvider::GetOrAdd(UTexture2D* Texture)
+{
+	int32 Hole = INDEX_NONE;
+
+	for (int32 TextureIndex = 0; TextureIndex < Textures.Num(); ++TextureIndex)
+	{
+		if (Textures[TextureIndex] == Texture)
+		{
+			return ToTextureId(TextureIndex);
+		}
+		else if (!Texture && Hole == INDEX_NONE)
+		{
+			Hole = TextureIndex;
+		}
+	}
+	
+	if (Hole == INDEX_NONE)
+	{
+		Hole = Textures.Num();
+
+		const int32 NumElements = Hole + 1;
+		check(NumElements < MAX_IDS); // Max number of TextureIds reached.
+
+		Textures.SetNumUninitialized(NumElements);
+		KeepTextures.SetNumUninitialized(NumElements);
+	} 
+	
+	Textures[Hole] = Texture;
+	KeepTextures[Hole] = false;
+	
+	return ToTextureId(Hole);
+}
+
+
+void UDefaultImageProvider::CacheTextures(UCustomizableObjectInstance& Instance)
+{
+	GarbageCollectTextureIds();
+
+	UCustomizableObjectSystem* System = UCustomizableObjectSystem::GetInstance();
+
+	for (const FCustomizableObjectTextureParameterValue& TextureParameter : Instance.GetDescriptor().GetTextureParameters())
+	{
+		if (const int32 Index = ToIndex(TextureParameter.ParameterValue);
+			Textures.IsValidIndex(Index) && Textures[Index])
+		{
+			System->CacheImage(Index);
+		}
+
+		for (const uint64 RangeValue : TextureParameter.ParameterRangeValues)
+		{
+			if (const int32 Index = ToIndex(RangeValue);
+				Textures.IsValidIndex(Index) && Textures[Index])
+			{
+				System->CacheImage(Index);
+			}
+		}
+	}
+}
+
+
+void UDefaultImageProvider::Keep(const uint64 TextureId, const bool bKeep)
+{
+	KeepTextures[ToIndex(TextureId)] = bKeep;
+}
+
+
+void UDefaultImageProvider::GarbageCollectTextureIds()
+{
+	TBitArray IdUsed;
+	IdUsed.Init(false, Textures.Num());
+
+	for (TObjectIterator<UCustomizableObjectInstance> It; It; ++It)
+	{
+		for (const FCustomizableObjectTextureParameterValue& TextureParameter :  It->GetDescriptor().GetTextureParameters())
+		{			
+			if (const int32 Index = ToIndex(TextureParameter.ParameterValue);
+				IdUsed.IsValidIndex(Index))
+			{
+				IdUsed[Index] = true;
+			}
+
+			for (const uint64 RangeValue : TextureParameter.ParameterRangeValues)
+			{
+				if (const int32 Index = ToIndex(RangeValue);
+					IdUsed.IsValidIndex(Index))
+				{
+					IdUsed[Index] = true;
+				}
+			}
+		}
+	}
+
+	UCustomizableObjectSystem* System = UCustomizableObjectSystem::GetInstance();
+
+	int32 LastToKeep = -1;
+	
+	for (int32 TextureIndex = 0; TextureIndex < Textures.Num(); ++TextureIndex)
+	{
+		if (!IdUsed[TextureIndex] && !KeepTextures[TextureIndex])
+		{
+			Textures[TextureIndex] = nullptr;
+			System->UnCacheImage(TextureIndex);
+		}
+		else
+		{
+			LastToKeep = TextureIndex;
+		}
+	}
+
+	if (const int32 NumToKeep = LastToKeep + 1;
+		NumToKeep < Textures.Num())
+	{
+		Textures.SetNumUninitialized(NumToKeep);
+		KeepTextures.SetNumUninitialized(NumToKeep);
+	}
+}
+
+
+int32 UDefaultImageProvider::ToIndex(const uint64 TextureId) const
+{
+	return TextureId - BASE_ID;
+}
+
+
+uint64 UDefaultImageProvider::ToTextureId(const int32 Index) const
+{
+	return Index + BASE_ID;
+}
+
diff --git a/Engine/Plugins/Experimental/Mutable/Source/CustomizableObject/Public/MuCO/CustomizableObjectInstanceDescriptor.h b/Engine/Plugins/Experimental/Mutable/Source/CustomizableObject/Public/MuCO/CustomizableObjectInstanceDescriptor.h
--- a/Engine/Plugins/Experimental/Mutable/Source/CustomizableObject/Public/MuCO/CustomizableObjectInstanceDescriptor.h
+++ b/Engine/Plugins/Experimental/Mutable/Source/CustomizableObject/Public/MuCO/CustomizableObjectInstanceDescriptor.h
@@ -107,6 +107,18 @@ struct CUSTOMIZABLEOBJECT_API FCustomizableObjectInstanceDescriptor
 	/** Sets the float value "FloatValue" of a float parameter with index "FloatParamIndex". */
 	void SetFloatParameterSelectedOption(const FString& FloatParamName, float FloatValue, int32 RangeIndex = -1);
 
+	/** Gets the value of a texture parameter with name "TextureParamName". */
+	uint64 GetTextureParameterSelectedOption(const FString& TextureParamName, int32 RangeIndex) const;
+
+	/** Gets the texture of a texture parameter with name "TextureParamName". */
+	UTexture2D* GetTextureParameterSelectedOptionT(const FString& TextureParamName, int32 RangeIndex) const;
+
+	/** Sets the texture value "TextureValue" of a texture parameter with index "TextureParamIndex". */
+	void SetTextureParameterSelectedOption(const FString& TextureParamName, uint64 TextureValue, int32 RangeIndex);
+
+	/** Sets the texture "Texture" of a texture parameter with index "TextureParamIndex". */
+	void SetTextureParameterSelectedOptionT(const FString& TextureParamName, UTexture2D* TextureValue, int32 RangeIndex);
+
 	/** Gets the value of a color parameter with name "ColorParamName". */
 	FLinearColor GetColorParameterSelectedOption(const FString& ColorParamName) const;
 
@@ -200,6 +212,10 @@ struct CUSTOMIZABLEOBJECT_API FCustomizableObjectInstanceDescriptor
 	 * The added value is initialized with 0.5f and is the last one of the range. */
 	int32 AddValueToFloatRange(const FString& ParamName);
 
+	/** Increases the range of values of the float with ParamName, returns the index of the new float value, -1 otherwise. 
+	 * The added value is not initialized. */
+	int32 AddValueToTextureRange(const FString& ParamName);
+
 	/** Increases the range of values of the projector with ParamName, returns the index of the new projector value, -1 otherwise.
 	 * The added value is initialized with the default projector as set up in the editor and is the last one of the range. */
 	int32 AddValueToProjectorRange(const FString& ParamName);
@@ -216,6 +232,12 @@ struct CUSTOMIZABLEOBJECT_API FCustomizableObjectInstanceDescriptor
 	/** Remove the RangeIndex element of the float range of values from the parameter ParamName, returns the index of the last valid float, -1 if no values left. */
 	int32 RemoveValueFromFloatRange(const FString& ParamName, int32 RangeIndex);
 
+	/** Remove the last of the texture range of values from the parameter ParamName, returns the index of the last valid float, -1 if no values left. */
+	int32 RemoveValueFromTextureRange(const FString& ParamName);
+
+	/** Remove the RangeIndex element of the texture range of values from the parameter ParamName, returns the index of the last valid float, -1 if no values left. */
+	int32 RemoveValueFromTextureRange(const FString& ParamName, int32 RangeIndex);
+
 	/** Remove the last of the projector range of values from the parameter ParamName, returns the index of the last valid projector, -1 if no values left. */
 	int32 RemoveValueFromProjectorRange(const FString& ParamName);
 
diff --git a/Engine/Plugins/Experimental/Mutable/Source/CustomizableObject/Public/MuCO/CustomizableObjectParameterTypeDefinitions.h b/Engine/Plugins/Experimental/Mutable/Source/CustomizableObject/Public/MuCO/CustomizableObjectParameterTypeDefinitions.h
--- a/Engine/Plugins/Experimental/Mutable/Source/CustomizableObject/Public/MuCO/CustomizableObjectParameterTypeDefinitions.h
+++ b/Engine/Plugins/Experimental/Mutable/Source/CustomizableObject/Public/MuCO/CustomizableObjectParameterTypeDefinitions.h
@@ -159,6 +159,9 @@ struct FCustomizableObjectTextureParameterValue
 
 	UPROPERTY(Category = CustomizableObjectTextureParameterValue, VisibleAnywhere)
 	FString Uid;
+
+	UPROPERTY(Category = CustomizableObjectTextureParameterValue, VisibleAnywhere)
+	TArray<uint64> ParameterRangeValues;
 };
 
 
@@ -167,6 +170,11 @@ inline uint32 GetTypeHash(const FCustomizableObjectTextureParameterValue& Key)
 	uint32 Hash = GetTypeHash(Key.ParameterName);
 	Hash = HashCombine(Hash, GetTypeHash(Key.ParameterValue));
 
+	for (const uint64 Value : Key.ParameterRangeValues)
+    {
+    	Hash = HashCombine(Hash, GetTypeHash(Value));
+    }
+
 	return Hash;
 }
 
diff --git a/Engine/Plugins/Experimental/Mutable/Source/CustomizableObject/Public/MuCO/DefaultImageProvider.h b/Engine/Plugins/Experimental/Mutable/Source/CustomizableObject/Public/MuCO/DefaultImageProvider.h
new file mode 100644
--- /dev/null
+++ b/Engine/Plugins/Experimental/Mutable/Source/CustomizableObject/Public/MuCO/DefaultImageProvider.h
@@ -0,0 +1,55 @@
+// Copyright Epic Games, Inc. All Rights Reserved.
+
+#pragma once
+
+#include "MuCO/CustomizableObjectSystem.h"
+
+#include "DefaultImageProvider.generated.h"
+
+class UTexture2D;
+class UCustomizableObjectInstance;
+
+
+UCLASS()
+class CUSTOMIZABLEOBJECT_API UDefaultImageProvider : public UCustomizableSystemImageProvider
+{
+	GENERATED_BODY()
+
+public:
+	// UCustomizableSystemImageProvider interface
+	virtual ValueType HasTextureParameterValue(int64 ID) override;
+	virtual UTexture2D* GetTextureParameterValue(int64 ID) override;
+	virtual void GetTextureParameterValues(TArray<FCustomizableObjectExternalTexture>& OutValues) override;
+
+	// Own interface
+
+	UTexture2D* Get(uint64 TextureId) const;
+
+	int64 Get(UTexture2D* Texture) const;
+
+	/** If the TextureId does not get assigned to an Instance Texture Parameter it will be GC on the next Instance update. */
+	uint64 GetOrAdd(UTexture2D* Texture); 
+
+	void CacheTextures(UCustomizableObjectInstance& Instance);
+
+	/** Keep the TextureId from being collected by the GC. */
+	void Keep(uint64 TextureId, bool bKeep);
+	
+	void GarbageCollectTextureIds();
+
+private:
+	int32 ToIndex(uint64 TextureId) const;
+
+	uint64 ToTextureId(int32 Index) const;
+	
+	/** Change it to avoid TextureId collisions with other Image Providers. */
+	const int32 BASE_ID = 0;
+	const int32 MAX_IDS = 1000;
+
+	UPROPERTY()
+	TArray<TObjectPtr<UTexture2D>> Textures;
+
+	UPROPERTY()
+	TArray<bool> KeepTextures;
+};
+
diff --git a/Engine/Plugins/Experimental/Mutable/Source/CustomizableObjectEditor/Private/MuCOE/CustomizableInstanceDetails.cpp b/Engine/Plugins/Experimental/Mutable/Source/CustomizableObjectEditor/Private/MuCOE/CustomizableInstanceDetails.cpp
--- a/Engine/Plugins/Experimental/Mutable/Source/CustomizableObjectEditor/Private/MuCOE/CustomizableInstanceDetails.cpp
+++ b/Engine/Plugins/Experimental/Mutable/Source/CustomizableObjectEditor/Private/MuCOE/CustomizableInstanceDetails.cpp
@@ -47,25 +47,6 @@ void FCustomizableInstanceDetails::CustomizeDetails(const TSharedPtr<IDetailLayo
 			.CustomInstance(CustomInstance)
 			.InstanceDetails(SharedThis(this))
 	];
-
-	
-	// If the current instance has texture parameters, show the user interface to set the possible preview
-	// values for those parameters, in case no provider is registered.
-	UCustomizableObjectSystem* System = UCustomizableObjectSystem::GetInstance();
-
-	if (CustomInstance->GetTextureParameters().Num())
-	{
-		if (UCustomizableObjectImageProviderArray* ImageProvider = System->GetEditorExternalImageProvider())
-		{
-			ImageProvider->TexturesChangeDelegate.AddSP(this, &FCustomizableInstanceDetails::UpdateInstance);
-
-			IDetailCategoryBuilder& PreviewCategory = DetailBuilder->EditCategory("Preview Texture Parameter Options");
-
-			TArray<UObject*> Objs;
-			Objs.Add(ImageProvider);
-			PreviewCategory.AddExternalObjectProperty(Objs, FName("Textures"));
-		}
-	}
 }
 
 
@@ -78,13 +59,4 @@ void FCustomizableInstanceDetails::Refresh() const
 }
 
 
-void FCustomizableInstanceDetails::UpdateInstance() const
-{
-	UCustomizableObjectSystem* System = UCustomizableObjectSystem::GetInstance();
-	System->CacheAllImagesInAllProviders(true);
-
-	CustomInstance->UpdateSkeletalMeshAsync(true, true);
-}
-
-
 #undef LOCTEXT_NAMESPACE
diff --git a/Engine/Plugins/Experimental/Mutable/Source/CustomizableObjectEditor/Private/MuCOE/CustomizableInstanceDetails.h b/Engine/Plugins/Experimental/Mutable/Source/CustomizableObjectEditor/Private/MuCOE/CustomizableInstanceDetails.h
--- a/Engine/Plugins/Experimental/Mutable/Source/CustomizableObjectEditor/Private/MuCOE/CustomizableInstanceDetails.h
+++ b/Engine/Plugins/Experimental/Mutable/Source/CustomizableObjectEditor/Private/MuCOE/CustomizableInstanceDetails.h
@@ -25,7 +25,8 @@ public:
 
 	/** Customize details here. */
 	virtual void CustomizeDetails(const TSharedPtr<IDetailLayoutBuilder>& DetailBuilder) override;
-
+	
+	// Own interface	
 	/** Refresh the custom details. */
 	void Refresh() const;
 
@@ -33,9 +34,7 @@ private:
 	TWeakObjectPtr<UCustomizableObjectInstance> CustomInstance;
 	TWeakObjectPtr<UCustomizableSkeletalComponent> CustomizableSkeletalComponent;
 	TSharedPtr<SCustomizableInstanceProperties> InstancePropertiesWidget;
-
-	TWeakPtr<IDetailLayoutBuilder> LayoutBuilder;
 	
-	void UpdateInstance() const;
+	TWeakPtr<IDetailLayoutBuilder> LayoutBuilder;
 };
 
