[mutable] Fix for code debugger using same color for veryExpensive and standard operation types.

diff --git a/Engine/Plugins/Experimental/Mutable/Source/CustomizableObjectEditor/Private/MuCOE/SMutableCodeViewer.h b/Engine/Plugins/Experimental/Mutable/Source/CustomizableObjectEditor/Private/MuCOE/SMutableCodeViewer.h
--- a/Engine/Plugins/Experimental/Mutable/Source/CustomizableObjectEditor/Private/MuCOE/SMutableCodeViewer.h
+++ b/Engine/Plugins/Experimental/Mutable/Source/CustomizableObjectEditor/Private/MuCOE/SMutableCodeViewer.h
@@ -223,7 +223,7 @@ private:
 	{
 		FSlateColor(FLinearColor(1,1,1,1)),		// Standard cost color
 		FSlateColor(FLinearColor(1,0.4,0.2,1)),	// Expensive cost color
-		FSlateColor(FLinearColor(1,1,1,1))			// Very Expensive cost color
+		FSlateColor(FLinearColor(1,0.1,0.1,1))		// Very Expensive cost color
 	};
 	
 	/** Provided an operation type it returns the category representing how much costs to run an operation of this type
