[mutable] Improvements over navigation system
- Swapped navigation type from being based on UI data to UI Backend data
- Improved performance and usability of the navigation system
More information can be found on the task closure comment.

diff --git a/Engine/Plugins/Experimental/Mutable/Source/CustomizableObjectEditor/Private/MuCOE/SMutableCodeViewer.cpp b/Engine/Plugins/Experimental/Mutable/Source/CustomizableObjectEditor/Private/MuCOE/SMutableCodeViewer.cpp
--- a/Engine/Plugins/Experimental/Mutable/Source/CustomizableObjectEditor/Private/MuCOE/SMutableCodeViewer.cpp
+++ b/Engine/Plugins/Experimental/Mutable/Source/CustomizableObjectEditor/Private/MuCOE/SMutableCodeViewer.cpp
@@ -132,6 +132,14 @@ public:
 
 			// Prepare the text shown on the UI side of the operation tree
 			FString MainLabel = FString::Printf(TEXT("%s (%d) : %s"), *RowItem->Caption, int32(RowItem->MutableOperation), OpName);
+
+#if UE_BUILD_DEBUG
+			FString IndexOnTree = FString::FromInt(RowItem->IndexOnTree);
+			IndexOnTree.Append(TEXT("- "));
+			MainLabel.InsertAt(0,IndexOnTree);
+#endif
+
+			
 			if (RowItem->DuplicatedOf)
 			{
 				MainLabel.Append(TEXT(" (duplicated)"));
@@ -348,18 +356,9 @@ void SMutableCodeViewer::Construct(const FArguments& InArgs, const mu::ModelPtr&
 	// create & initialize tab manager
 	//TabManager = FGlobalTabmanager::Get()->NewTabManager(ConstructUnderMajorTab.ToSharedRef());
 
-	const mu::Model::Private* ModelPrivate = MutableModel->GetPrivate();
-	const int32 StateCount = ModelPrivate->m_program.m_states.Num();
-	for ( int32 StateIndex=0; StateIndex<StateCount; ++StateIndex )
-	{
-		const mu::FProgram::FState& State = ModelPrivate->m_program.m_states[StateIndex];
-		FString Caption = FString::Printf( TEXT("state [%s]"), ANSI_TO_TCHAR(State.m_name.c_str()) );
-
-		const FSlateColor LabelColor = ColorPerComputationalCost[StaticCast<uint8>(GetOperationTypeComputationalCost(
-			ModelPrivate->m_program.GetOpType(State.m_root)))];
-		RootNodes.Add(MakeShareable(new FMutableCodeTreeElement(MutableModel, State.m_root, Caption,LabelColor)));
-	}
-
+	// Generate all elements before starting the tree UI so we have a deterministic set of unique and duplicated elements
+	GenerateAllTreeElements();
+	
 	// Setup Navigation system
 	{
 		// Store the addresses of the root nodes for later search operations
@@ -627,7 +626,7 @@ void SMutableCodeViewer::Construct(const FArguments& InArgs, const mu::ModelPtr&
 					[
 						// Generate a new Constants panel to show the data stored on the current mutable program
 						SAssignNew(ConstantsWidget, SMutableConstantsWidget,
-							&(ModelPrivate->m_program),
+							&(MutableModel->GetPrivate()->m_program),
 							SharedThis(this))
 					]
 
@@ -733,22 +732,16 @@ void SMutableCodeViewer::OnNavigationSelectedOperationChanged(
 	// Only do the internal work if the type is one that makes sense searching
 	if (OperationTypeToSearch != mu::OP_TYPE::NONE)
 	{
-		// Reset the navigation element
-		SetCurrentNavigationElement(nullptr);
-		
 		// Locate all operations on the mutable operations tree (not the visual one) that do share the same operation type
 		// as the one selected. This will fill the array with the elements we should be looking for during the navigation operation
-		CacheAddressesOfOperationsOfType(OperationTypeToSearch);
-		
-		// Experimental
-		TreeView->ScrollToTop();
+		CacheAddressesOfOperationsOfType();
 	}
 	// None can be set by the user or be an indication that we are navigating over constant related operations
 	// todo: Separate both operations in some way on the UI to avoid complications in the code and in the UI's UX
 	else
 	{
 		// Clear all the elements on the navigation addresses 
-		NavigationOPAddresses.Empty();
+		NavigationElements.Empty();
 	}
 	
 }
@@ -785,76 +778,87 @@ void SMutableCodeViewer::OnSelectedOperationTypeFromTree()
 	
 	// Set the type operation type to be looking for -> Will invoke OnOptionTypeSelectionChanged
 	TargetedTypeSelector->SetSelectedItem(FoundModelOperationTypeElements[OperationIndex]);
+	
+	// Reset the navigation index so it starts from scratch
+	NavigationIndex = -1;
 }
 
 
-void SMutableCodeViewer::CacheAddressesOfOperationsOfType(mu::OP_TYPE TargetedOperationType)
+
+void SMutableCodeViewer::SortNavigationElements()
+{
+	// Sort the array from lower index to bigger index (0 , 1 , 2 ...)
+	NavigationElements.Sort([](const TSharedPtr<FMutableCodeTreeElement> A , const TSharedPtr<FMutableCodeTreeElement> B)
+	{
+		return A->IndexOnTree < B->IndexOnTree;
+	});
+}
+
+
+void SMutableCodeViewer::CacheAddressesOfOperationsOfType()
 {
 	// Clear previous data
-	NavigationOPAddresses.Empty();
-	
-	// Locate all addresses of operations of the provided operation type
-	TSet<mu::OP::ADDRESS> ProcessedAddresses;
-	TSet<mu::OP::ADDRESS> OperationsWithTargetedType;
+	NavigationElements.Empty();
+	check(RootNodeAddresses.Num());
 
+	// Object containing all data required by the search operation to be able to be called recursively
+	FElementsSearchCache SearchPayload;
+	// Initialize the Search Payload with the root node addresses. This way the search will use them as the root nodes where
+	// to start searching
+	SearchPayload.SetupRootBatch(RootNodeAddresses);
+	
 	// Main update procedure run for the targeted state and the targeted parameter values
 	const mu::FProgram& Program = MutableModel->GetPrivate()->m_program;
-	GetOperationsOfType(TargetedOperationType,RootNodeAddresses,Program,OperationsWithTargetedType,ProcessedAddresses);
+	GetOperationsOfType(OperationTypeToSearch,SearchPayload, Program);
 	
-	if (!OperationsWithTargetedType.IsEmpty())
+	if (!SearchPayload.FoundElements.IsEmpty())
 	{
 		// Cache the navigation addresses so we are able to navigate over them
-		NavigationOPAddresses = OperationsWithTargetedType.Array();
-
-		// Tell the system changes have been made over Navigation Addresses
-		bUpdatedNavigationAddresses = true;
+		NavigationElements = MoveTemp(SearchPayload.FoundElements);
+		SortNavigationElements();
+		
+		// Reset the navigation index
+		NavigationIndex = -1;
 	}
-
 }
 
-void SMutableCodeViewer::GetOperationsOfType ( const mu::OP_TYPE& TargetOperationType,const TArray<mu::OP::ADDRESS>& InParentAddresses,const mu::FProgram& InProgram,
-	TSet<mu::OP::ADDRESS>& OutAddressesOfType,
-	TSet<mu::OP::ADDRESS>& AlreadyProcessedAddresses)
+void SMutableCodeViewer::GetOperationsOfType(const mu::OP_TYPE& TargetOperationType,
+                                             FElementsSearchCache& InSearchPayload,
+                                             const mu::FProgram& InProgram)
 {
-	TArray<mu::OP::ADDRESS> NextBatch;
+	// next batch of addresses to be explored 
+	TArray<FItemCacheKey> NextBatchAddressesData;
 	
-	for (const mu::OP::ADDRESS& ParentAddress : InParentAddresses)
+	// uint32 CurrentParentBeingScanned = 0;
+	for	(int32 ParentIndex = 0 ; ParentIndex < InSearchPayload.BatchData.Num(); ParentIndex++)
 	{
-		// Process parent
-		if (!AlreadyProcessedAddresses.Contains(ParentAddress))
-		{
-			// Cache if same data type and we share the same address (means this op is pointing at the provided resource)
-			if (InProgram.GetOpType(ParentAddress) == TargetOperationType)
-			{
-				OutAddressesOfType.Add(ParentAddress);
-			}
+		// Get one of the previous run "children" and treat as a parent to get it's children and process them
+		const mu::OP::ADDRESS& CurrentAddress = InSearchPayload.BatchData[ParentIndex].Child;
 		
-			// Cache to avoid processing it again later
-			AlreadyProcessedAddresses.Add(ParentAddress);
-
-			// Generic case for unnamed children traversal.
-			mu::ForEachReference(InProgram, ParentAddress, [this,&NextBatch]( mu::OP::ADDRESS ChildAddress)
-			{
-				// If the parent does have a child then process it 
-				if (ChildAddress)
-				{
-					// Add for expansion and search
-					NextBatch.Add(ChildAddress);
-				}
-			});
+		// Cache if same data type and we share the same address (means this op is pointing at the provided resource)
+		// It will cache duplicated entries
+		if ( InProgram.GetOpType(CurrentAddress) == TargetOperationType)
+		{
+			// Since this element is of the type we are looking for then cache it on InSearchPayload.FoundElements
+			InSearchPayload.AddToFoundElements(CurrentAddress,ParentIndex,ItemCache);
 		}
+		
+		// Get all NON PROCESSED the children of this operation to later be able to process them (on next recursive call)
+		InSearchPayload.CacheChildrenOfAddressIfNotProcessed(CurrentAddress, InProgram, NextBatchAddressesData);
 	}
 
-	if (NextBatch.Num())
+	// Explore children if found 
+	if (NextBatchAddressesData.Num())
 	{
+		// Cache next batch data so the next invocations are able to locate the provided addresses on the itemsCache
+		InSearchPayload.BatchData = MoveTemp(NextBatchAddressesData);
+		
 		// Process the children of this object
-		GetOperationsOfType(TargetOperationType,NextBatch,InProgram,OutAddressesOfType,AlreadyProcessedAddresses);
+		GetOperationsOfType(TargetOperationType, InSearchPayload, InProgram);
 	}
-
 }
 
 
-
 void SMutableCodeViewer::CacheOperationTypesPresentOnModel()
 {
 	check(MutableModel)
@@ -901,7 +905,6 @@ void SMutableCodeViewer::CacheOperationTypesPresentOnModel()
 	// ModelOperationTypes is now an array starting with None and all the types found on the operations tree in alphabetical order
 }
 
-
 void SMutableCodeViewer::GetOperationTypesPresentOnModel(
 	const TArray< mu::OP::ADDRESS>& InParentAddresses, const mu::FProgram& InProgram, TSet<mu::OP_TYPE>& OutLocatedOperations,
 	TSet<mu::OP::ADDRESS>& AlreadyProcessedAddresses)
@@ -920,7 +923,7 @@ void SMutableCodeViewer::GetOperationTypesPresentOnModel(
 			AlreadyProcessedAddresses.Add(ParentAddress);
 
 			// Generic case for unnamed children traversal.
-			mu::ForEachReference(InProgram, ParentAddress, [this, &InProgram, &NextBatch, &OutLocatedOperations, &AlreadyProcessedAddresses]( mu::OP::ADDRESS ChildAddress)
+			mu::ForEachReference(InProgram, ParentAddress, [this, &NextBatch]( mu::OP::ADDRESS ChildAddress)
 			{
 				// Avoid processing items more than once
 				if (ChildAddress)
@@ -941,422 +944,184 @@ void SMutableCodeViewer::GetOperationTypesPresentOnModel(
 }
 
 
-
-void SMutableCodeViewer::LocateNavigationElementsOnTree()
-{
-	NavigationFoundElements.Empty();
-	NavigationFoundElements.Reserve(TreeElements.Num());
-	
-	for (TSharedPtr<FMutableCodeTreeElement> TreeElement : TreeElements)
-	{
-		// Consider that you have found elements if their address corresponds to the address on our cached navigation items
-		const bool bIsNavigableElement = NavigationOPAddresses.Contains(TreeElement->MutableOperation);
-		if (bIsNavigableElement && !TreeElement->DuplicatedOf)
-		{
-			// Add the element to our list of elements we want to navigate over
-			NavigationFoundElements.Add(TreeElement);
-		}
-	}
-	
-	NavigationFoundElements.Shrink();
-
-#if UE_BUILD_DEBUG
-	UE_LOG(LogTemp, Log, TEXT("Located %d elements of %d for the current operation type"), NavigationFoundElements.Num(),TreeElements.Num());
-#endif
-
-	if (NavigationFoundElements.IsEmpty())
-	{
-		// No elements have been found for the newly selected type, set the selected element to be null.
-		SetCurrentNavigationElement(nullptr);
-	}
-}
-
 FText SMutableCodeViewer::OnPrintNavigableObjectAddressesCount() const
 {
-	// Depending on the amount of navigable objects (addresses, not actual elements) display the amount there are
-	return FText::AsNumber(NavigationOPAddresses.Num());
-}
-
-void SMutableCodeViewer::SetCurrentNavigationElement(TSharedPtr<FMutableCodeTreeElement> NewelyFoundElement)
-{
-	if (CurrentNavigationElement && NewelyFoundElement)
+	FString OutputString = "";
+	if (const int32 NavigationElementsCount = NavigationElements.Num())
 	{
-		// Avoid operating if the new selected element is the same than before
-		if (CurrentNavigationElement->MutableOperation == NewelyFoundElement->MutableOperation)
+		// Show the index if the index showing adds information
+		if (NavigationIndex >= 0)
 		{
-#if UE_BUILD_DEBUG
-			UE_LOG(LogTemp,Warning,TEXT("Skipping CurrentSelectedElement update since it is the same as previously set."));
-#endif
-			return;
+			OutputString.Append( FString::FromInt( NavigationIndex+1));
+			OutputString.Append(" / ");
 		}
-	}
-
-	// Internally change the selected element for the navigation system to use as starting point for the next movement
-	// operation
-	CurrentNavigationElement = NewelyFoundElement;
-	
-	// Clear the selection and select the new object
-	// TreeView->ClearSelection();
-	if (CurrentNavigationElement)
-	{
-		// Perform the UI selection operation so the viewer gets triggered as if the user did select it manually
-		TreeView->SetSelection(CurrentNavigationElement);
-
-#if UE_BUILD_DEBUG 
-		// Report currently selected one (not actually selected in UI, only internally at the moment)
-		const int32 SelectedOperationIndex = NavigationFoundElements.IndexOfByKey(CurrentNavigationElement);
-		UE_LOG(LogTemp,Log,TEXT("\tSelected element %d of %d"),SelectedOperationIndex +1,NavigationFoundElements.Num());
-#endif
-
-		// TODO: Mind making the system automatically focus on the selected object. Currently is not always on the view
-		// and that could be improved. Scrolling over the tree may require a tree refresh, this means you may need to
-		// wait until that refresh is performed.
 		
+		OutputString.Append( FString::FromInt(NavigationElementsCount));
+
+		// Format -> 1 / 12 or 12
 	}
-#if UE_BUILD_DEBUG 
-	else
-	{
-		UE_LOG(LogTemp,Log,TEXT("\tNo selected element has been found"));
-	}
-#endif
 	
+	// Depending on the amount of navigable objects (addresses, not actual elements) display the amount there are
+	return FText::FromString(OutputString);
 }
 
 
 bool SMutableCodeViewer::CanInteractWithPreviousOperationButton() const
 {
 	// Only navigable if there are more than 0 elements to traverse and we are not scrolling
-	return !bIsScrolling && NavigationOPAddresses.Num() > 0;
+	return NavigationElements.Num() > 0 && NavigationIndex > 0 && (!bWasScrollToTargetRequested && !bWasUniqueExpansionInvokedForNavigation);
 }
 
 bool SMutableCodeViewer::CanInteractWithNextOperationButton() const
 {
 	// Only navigable if there are more than 0 elements to traverse and we are not scrolling
-	return !bIsScrolling && NavigationOPAddresses.Num() > 0;
+	return NavigationElements.Num() > 0 && NavigationIndex < NavigationElements.Num() -1 && (!bWasScrollToTargetRequested && !bWasUniqueExpansionInvokedForNavigation);
 }
 
 
-void SMutableCodeViewer::SortNavigationElementsArray()
+FReply SMutableCodeViewer::OnGoToPreviousOperationButtonPressed()
 {
-	// Do not waste time sorting an empty array of elements
-	// If only one element is present there is no reason for performing a sort operation
-	if (NavigationFoundElements.Num() < 2)
-	{
-		return;
-	}
-
-	// Cache the position on the tree for each of the navigation elements on view
-	TMap<int32,TSharedPtr<FMutableCodeTreeElement>> PositionPerElement;
-	PositionPerElement.Reserve(NavigationFoundElements.Num());
-	for	(int32 i = 0; i < NavigationFoundElements.Num(); i++)
-	{
-		const int32 PositionOnTree = TreeView->WidgetFromItem(NavigationFoundElements[i])->GetIndexInList();
-		PositionPerElement.Add(PositionOnTree, NavigationFoundElements[i]);
-	}
-
-	// By working with a TMap we can ensure the indices provided get automatically sorted on Addition.
+	// Focus on previous target
+	NavigationIndex = NavigationIndex<=0 ? 0 : NavigationIndex - 1;
+	FocusViewOnNavigationTarget();
 	
-	// Grab the values as an array to be later be used as the new array for NavigationElements
-	TArray<TSharedPtr< FMutableCodeTreeElement>> SortedNavigationElements;
-	PositionPerElement.GenerateValueArray(SortedNavigationElements);
-
-	// Make the targeted elements of type the new array with the correct values but sorted out
-	NavigationFoundElements = MoveTemp(SortedNavigationElements);
+	return FReply::Handled();
 }
 
-
-FReply SMutableCodeViewer::OnGoToPreviousOperationButtonPressed()
+FReply SMutableCodeViewer::OnGoToNextOperationButtonPressed()
 {
-	// Filter out not valid type
-	if (NavigationOPAddresses.IsEmpty())
-	{
-		return FReply::Handled();
-	}
-
-	// Sort the array of navigation elements just in case the order of the elements on the view changed
-	// This can happen if the user scrolls the view and new elements of the targeted type come to view
-	SortNavigationElementsArray();
+	// Focus on next target
+	NavigationIndex = NavigationIndex>=NavigationElements.Num() -1 ? NavigationElements.Num() -1 : NavigationIndex + 1;
+	FocusViewOnNavigationTarget();
 	
-	// Make a copy of the array to later be able to compare it to look for new elements of the targeted type
-	NavigationPreviouslyFoundElements = NavigationFoundElements;
+	return FReply::Handled();
+}
 
-	// Fully expand tree if required to uncover all possible elements of the targeted type.
-	// The expansion does not discern between element types, all non duplicated elements will be expanded
-	if (bTreeWasExpanded == false && !bIsScrolling)
+void SMutableCodeViewer::FocusViewOnNavigationTarget()
+{
+	// Stage 1 : Expand all tree so all navigable elements get to be reachable
+	if (!bWasUniqueExpansionInvokedForNavigation && !bWasScrollToTargetRequested)
 	{
-		// Expand all the elements by calling the standard expansion operation (ignoring duplicates of course)
-		TreeExpandElements(RootNodes,false);
-
-		// Notify the navigation system that we just expanded the tree
-		bTreeWasExpanded = true;
-
-		// It will trigger an automatic scroll operation if possible depending on the tree
-		bIsSearchingForPreviousElement = true;
-			
-		return FReply::Handled();
+		TreeExpandUnique();
+		bWasUniqueExpansionInvokedForNavigation = true;
+		// Early exit, this method will get called again later after tree update
+		return;		
 	}
-
-	// At this point we know the tree has been fully expanded and we should be able to traverse it without the
-	// need of expanding it again unless we get to another object of the targeted type with should, if used as
-	// start for another search, perform another cautionary tree expansion.
 	
-	// Get the index of the currently selected element on the array of operations of the current type
-	// If it is < 0 means that the current object is not part of the actual elements of the targeted type
-	int32 SelectedOperationIndex = NavigationFoundElements.IndexOfByKey(CurrentNavigationElement);
-
-	if ( CurrentNavigationElement )
+	// Stage 2 : Try to get to the targeted element. if not visible scroll into view
+	
+	// Locate the element on our map of elements
+	const TSharedPtr<FMutableCodeTreeElement> TargetElement = NavigationElements[NavigationIndex];
+	check (TargetElement.IsValid());
+	
+	// If required scroll to the area where we know the element is going to be in view
+	// a way to ensure this happens is by calling 
+	if (TreeView->IsItemVisible(TargetElement))
 	{
-		// A current element is selected and is not the first one. Select the previous element
-		if ( SelectedOperationIndex > 0)
-		{
-			// Get the previous element on the array (the one on top of the current one on the tree structure)
-			SetCurrentNavigationElement(NavigationFoundElements[SelectedOperationIndex - 1]);
-			return FReply::Handled();
-		}
+		// Stage 3-b : Select the element we have provided since now is sure to be in view
+		
+		// This line selects the element with at the same time updates the UI to show the row representing this element selected
+		TreeView->SetSelection(TargetElement);
+	
+		// Done!
+		// We have the element in view and we have selected it!
 	}
-	// No current element is currently selected
 	else
 	{
-		// We have no current element selected but there are elements there to be used. Select the last one
-		if ( !NavigationFoundElements.IsEmpty() )
-		{
-			// Get the last object and use it as the current object
-			SelectedOperationIndex = NavigationFoundElements.Num() -1;
-			SetCurrentNavigationElement(NavigationFoundElements[SelectedOperationIndex]);
-			return FReply::Handled();
-		}
-	}
-
-	// If no other element of the type we want is found it means we should look for other elements by
-	// scrolling the view up to the top or until we find another element
-	{
-		// It will trigger an automatic scroll operation if possible depending on the tree.
-		bIsSearchingForPreviousElement = true;
-	}
-
-	return FReply::Handled();
-}
-
-void SMutableCodeViewer::GoToPreviousOperationAfterRefresh()
-{
-	// Is mandatory for this system to be stable to operate after the tree refresh
-	check(TreeView->IsPendingRefresh() == false);
-	
-	// Required for all the possible contexts where this method could  partake
-	check(bIsSearchingForPreviousElement);
-	
-	// Check that we did change the amount of elements of our current targeted type.
-	// If the tree was fully expanded on the previous update then check the changes by entering the if block
-	const bool bNewElementFound = HaveNewElementsBeenFound();
-	if ( bNewElementFound || bTreeWasExpanded)
-	{
-		// Reset the flag since we are already being run after the tree refresh
-		bIsSearchingForPreviousElement = false;
+		// Stage 3-a (optional) : Ask for the provided element to be scrolled into view.
 		
-		// Continue the operation now that the tree view has been refreshes and the new rows are there to be used
-		OnGoToPreviousOperationButtonPressed();
+		// Failing this check would mean we have performed a scroll but we are still not able to view the element
+		check (!bWasScrollToTargetRequested);
 		
-		// Reset the just expanded flag if found to be true
-		bTreeWasExpanded = false;
+		// Amount of elements that can be shown by the UI at any given time (search for better way of knowing this)
+		TreeView->RequestScrollIntoView(TargetElement);
 		
-		// Reset the scrolling flag in case it was active
-		bIsScrolling = false;
-		return;
-	}
-
-	// We do know that no element of the targeted type has been found. So, we should start searching
-	
-	// If scrollable try to scroll over the tree to locate new elements
-	if (TreeView->IsScrollbarNeeded() &&
-		TreeView->GetScrollOffset() > 0 )
-	{
-		// Flag the system so it is able to tell if the scrolling is being performed or not
-		bIsScrolling = true;
+		// Read this variable after the update and then select the object (easy at this point)
+		// You may want to just call again this method after refresh since the element will be on view
+		bWasScrollToTargetRequested = true;
 
-		// Add the max possible offset to be fast but secure on the translation process over the graph
-		ComputedMaxViewScrollStep = FMath::Max( TreeElements.Num() - 1,MinTreeScrollStep);
-		TreeView->AddScrollOffset(-ComputedMaxViewScrollStep,true);
-
-		// Check to avoid children drawing crash when having negative offset values
-		if (TreeView->GetScrollOffset() < 0)
-		{
-			TreeView->ScrollToTop();
-		}
-			
+		// Early exit, this method will get called again later after tree update once the scroll has been completed
 		return;
 	}
 
-	// If not scrollable then just ignore this operation since it means no new elements have been added and
-	// no targeted elements can be retrieved from the graph
-	bIsSearchingForPreviousElement = false;
-	bIsScrolling = false;
+	// Reset the control flag so we do not expand all tree again if not required
+	bWasUniqueExpansionInvokedForNavigation = false;
+	bWasScrollToTargetRequested = false;
 }
 
+#pragma endregion 
 
-FReply SMutableCodeViewer::OnGoToNextOperationButtonPressed()
-{
-	// Go forward one index on the array TargetedElementsOfType.
-	// If in last index try to expand other rows
-	
-	// Filter out not valid type
-	if (NavigationOPAddresses.IsEmpty())
-	{
-		return FReply::Handled();
-	}
-
-	// Sort the array of navigation elements just in case the order of the elements on the view changed
-	// This can happen if the user scrolls the view and new elements of the targeted type come to view
-	SortNavigationElementsArray();
+#pragma region Operation Cost Color Hints
 
-	// Make a copy of the array to be able to later compare it to the tree with the element expanded
-	NavigationPreviouslyFoundElements = NavigationFoundElements;
+void SMutableCodeViewer::GenerateAllTreeElements()
+{
+	// By generating all tree elements prior to usage we are able to :
+	//	- Compute the index of each one to aidd on navigation
+	//	- Remove indeterministic assignation of the "Duplicated" state of elements. It was due to user interaction with the tree
+	//	Only unique elements, their children and duplicated elements will be generated. Children of duplicates will
+	// 	be ignored due to how we handle them when expanding and contracting elements (OnExpansionChanged)
 
-	// Fully expand tree if required to uncover all possible elements of the targeted type.
-	// The expansion does not discern between element types, all non duplicated elements will be expanded
-	if (bTreeWasExpanded == false && !bIsScrolling)
+	// Generate all root nodes
+	const mu::Model::Private* ModelPrivate = MutableModel->GetPrivate();
+	const mu::FProgram& Program = ModelPrivate->m_program;
+	const uint32 StateCount = ModelPrivate->m_program.m_states.Num();
+	for ( uint32 StateIndex=0; StateIndex<StateCount; ++StateIndex )
 	{
-		// Expand all elements to avoid leaving any of them behind when searching for the next element
-		TreeExpandElements(RootNodes,false);
+		const mu::FProgram::FState& State = ModelPrivate->m_program.m_states[StateIndex];
+		FString Caption = FString::Printf( TEXT("state [%s]"), ANSI_TO_TCHAR(State.m_name.c_str()) );
 
-		// Tell the system to continue operating after the tree refresh
-		bIsSearchingForNextElement = true;
+		const FSlateColor LabelColor = ColorPerComputationalCost[StaticCast<uint8>(GetOperationTypeComputationalCost(
+			ModelPrivate->m_program.GetOpType(State.m_root)))];
+		
+		RootNodes.Add(MakeShareable(new FMutableCodeTreeElement(ItemCache.Num(),MutableModel, State.m_root, Caption,LabelColor)));
 
-		// Save a flag to avoid redoing this again later (only required once per selected object)
-		bTreeWasExpanded = true;
-			
-		return FReply::Handled();
-	}
+		// Iterate over each root node and generate all the elements in a human readable pattern (Z Pattern)
 
-	// At this point we know the tree has been fully expanded and we should be able to traverse it without the
-	// need of expanding it again unless we get to another object of the targeted type with should, if used as
-	// start for another search, perform another cautionary tree expansion.
-	
-	// Get the index of the currently selected element on the array of operations of the current type
-	// If it is < 0 means that the current object is not part of the actual elements of the targeted type
-	int32 SelectedOperationIndex = NavigationFoundElements.IndexOfByKey(CurrentNavigationElement);
-	
-	// If we have an element already set then use it to compute the values
-	if (CurrentNavigationElement)
-	{
-		// If we are not the last element of the array then move on to the next one normally
-		// Since we have expanded the current element we should be able to access contents inside of it (children of it)
-		if (SelectedOperationIndex < NavigationFoundElements.Num() - 1)
-		{
-			// Select the new element
-			SetCurrentNavigationElement(NavigationFoundElements[SelectedOperationIndex + 1]);
-			return FReply::Handled();
-		}
-	}
-	// We do not have a current element to use as starting point for our search.
-	else
-	{
-		// Since we are moving downwards, grab the first element visible to make it be the current object.
-		if (NavigationFoundElements.Num() > 0)
-		{
-			// Get the first object and store its direction
-			SelectedOperationIndex = 0;
-			SetCurrentNavigationElement(NavigationFoundElements[SelectedOperationIndex]);
-			return FReply::Handled();
-		}
-	}
+		constexpr  mu::OP::ADDRESS CommonParent = 0;
+		const FItemCacheKey Key = { CommonParent, State.m_root, StateIndex };
 
-	// At this point we are not able to reach any other element of the current type. We must then
-	// start scrolling down in order to get to the next element or the end of the tree view.
-	{
-		// Enable this flag to make the system search for new elements on the tree on the next tick
-		bIsSearchingForNextElement = true;
+		// Add the element to the cache so we keep the indices straight.
+		ItemCache.Add(Key, RootNodes.Last());
+		GenerateElementRecursive(State.m_root,Program);
 	}
 	
-	return FReply::Handled();
 }
 
-
-void SMutableCodeViewer::GoToNextOperationAfterRefresh()
+void SMutableCodeViewer::GenerateElementRecursive(mu::OP::ADDRESS InParentAddress,  const mu::FProgram& InProgram)
 {
-	// Is mandatory for this system to be stable to operate after the tree refresh
-	check(TreeView->IsPendingRefresh() == false);
-	
-	// Required for all the possible contexts where this method could  partake
-	check(bIsSearchingForNextElement);
-	
-	// Check if the new CurrentArray does have elements not already set on the previous array
-	// and also check if we expanded the tree on the previous tick. If so, enter the if block
-	const bool bNewElementFound = HaveNewElementsBeenFound();
-	if (bNewElementFound || bTreeWasExpanded)
-	{
-		// Reset the flag since we are already being run after the tree refresh
-		bIsSearchingForNextElement = false;
-		
-		// Recreate the button press to move to the next element after expanding the tree
-		OnGoToNextOperationButtonPressed();
-		
-		// Reset the just expanded flag
-		// Setting this to false makes the system able to fully expand the tree later when we start searching again
-		bTreeWasExpanded = false;
-		
-		// Reset the scrolling flag in case it was active
-		bIsScrolling = false;
-		return;
-	}
-
-	// We do know that no element of the targeted type has been found. So, we should start searching
-	
-	// If scrollable try to scroll over the tree to locate new elements
-	if (TreeView->IsScrollbarNeeded()
-		&& TreeView->GetScrollDistanceRemaining().Y > 0)
+	// Find children of the provided element
+	uint32 ChildIndex = 0;
+	mu::ForEachReference(InProgram, InParentAddress, [this, InParentAddress, &ChildIndex, &InProgram](mu::OP::ADDRESS ChildAddress)
 	{
-		// Flag the system so it is able to tell if the scrolling is being performed or not
-		bIsScrolling = true;
-			
-		// Add a bit of scroll
-		ComputedMaxViewScrollStep = FMath::Max( TreeElements.Num() - 1,MinTreeScrollStep);
-		TreeView->AddScrollOffset(ComputedMaxViewScrollStep,true);
-
-		// Do not trust the tree to be able to handle this kind of artificial scrolling and avoid surpassing the end
-		if (TreeView->GetScrollDistanceRemaining().Y < 0 )
+		if (ChildAddress)
 		{
-			TreeView->ScrollToBottom();
-		}
-			
-		return;
-	}
+			const FItemCacheKey Key = { InParentAddress, ChildAddress, ChildIndex };
+			const TSharedPtr<FMutableCodeTreeElement>* CachedItem = ItemCache.Find(Key);
 
-	// If not scrollable then just ignore this operation since it means no new elements have been added and
-	// no targeted elements can be retrieved from the graph
-	bIsSearchingForNextElement = false;
-	bIsScrolling = false;
-}
+			if (!CachedItem)
+			{
+				// Locate the "original" tree element
+				const TSharedPtr<FMutableCodeTreeElement>* MainItemPtr = MainItemPerOp.Find(ChildAddress);
 
-bool SMutableCodeViewer::HaveNewElementsBeenFound() const
-{
-	// Check if the new CurrentArray does have elements not already set on the previous array
-	bool NewElementAdded = false;
+				const FSlateColor LabelColor = ColorPerComputationalCost[ StaticCast<uint8> (GetOperationTypeComputationalCost(InProgram.GetOpType(ChildAddress))) ];
+				const TSharedPtr<FMutableCodeTreeElement> Item = MakeShareable(new FMutableCodeTreeElement(ItemCache.Num(),MutableModel, ChildAddress, TEXT(""),LabelColor, MainItemPtr));
 
-	// If the current array does have data and the old not then we directly know that new elements have been added.
-	if (NavigationFoundElements.Num() && NavigationPreviouslyFoundElements.IsEmpty())
-	{
-		return true;
-	}
-	
-	// Iterate the current array to 
-	for (int32 CurrentArrayIndex = 0; CurrentArrayIndex < NavigationFoundElements.Num(); CurrentArrayIndex++)
-	{
-		// If the element is not found on the previous array it means it is new, and therefore a valid element to select
-		if (!NavigationPreviouslyFoundElements.Contains(NavigationFoundElements[CurrentArrayIndex]))
-		{
-			NewElementAdded = true;
-			break; 
+				// Cache this element for later access
+				ItemCache.Add(Key, Item);
+				
+				// It is not a duplicated of another one, then we can continue searching
+				if (!MainItemPtr)
+				{
+					MainItemPerOp.Add(ChildAddress, Item);
+					
+					GenerateElementRecursive(ChildAddress,InProgram);
+				}
+			}
 		}
-	}
-
-	return NewElementAdded;
+		++ChildIndex;
+	});
 }
 
-
-#pragma endregion 
-
-#pragma region Operation Cost Color Hints
-
 SMutableCodeViewer::EOperationComputationalCost SMutableCodeViewer::GetOperationTypeComputationalCost(mu::OP_TYPE OperationType) const
 {
 	if (VeryExpensiveOperationTypes.Contains(OperationType))
@@ -1388,16 +1153,6 @@ TSharedRef<ITableRow> SMutableCodeViewer::GenerateRowForNodeTree(TSharedPtr<FMut
 	// Generate a row element
 	TSharedRef<SMutableCodeTreeRow> Row = SNew(SMutableCodeTreeRow, InOwnerTable, InTreeNode);
 	
-	// Only add it to our collection of navigation elements if we know it has an address we want to be able to navigate to
-	if (NavigationOPAddresses.Contains(InTreeNode->MutableOperation))
-	{
-		// Ignore duplicates
-		if (!InTreeNode->DuplicatedOf)
-		{
-			NavigationFoundElements.Add(InTreeNode);
-		}
-	}
-	
 	// Determine if a row should be painted as highlighted based on the selected item
 	if (TreeView->GetNumItemsSelected())
 	{
@@ -1431,7 +1186,7 @@ void SMutableCodeViewer::GetChildrenForInfo(TSharedPtr<FMutableCodeTreeElement>
 
 	// Generic case for unnamed children traversal.
 	uint32 ChildIndex = 0;
-	mu::ForEachReference(Program, InInfo->MutableOperation, [this, ParentAddress, &ChildIndex, &OutChildren, &Program](mu::OP::ADDRESS ChildAddress)
+	mu::ForEachReference(Program, InInfo->MutableOperation, [this, ParentAddress, &ChildIndex, &OutChildren](mu::OP::ADDRESS ChildAddress)
 	{
 		if (ChildAddress)
 		{
@@ -1444,19 +1199,8 @@ void SMutableCodeViewer::GetChildrenForInfo(TSharedPtr<FMutableCodeTreeElement>
 			}
 			else
 			{
-				const TSharedPtr<FMutableCodeTreeElement>* MainItemPtr = MainItemPerOp.Find(ChildAddress);
-
-				const FSlateColor LabelColor = ColorPerComputationalCost[ StaticCast<uint8> (GetOperationTypeComputationalCost(Program.GetOpType(ChildAddress))) ];
-				
-				const TSharedPtr<FMutableCodeTreeElement> Item = MakeShareable(new FMutableCodeTreeElement(MutableModel, ChildAddress, TEXT(""),LabelColor, MainItemPtr));
-				
-				OutChildren.Add(Item);
-				ItemCache.Add(Key, Item);
-
-				if (!MainItemPtr)
-				{
-					MainItemPerOp.Add(ChildAddress, Item);
-				}
+				// if all elements have been already cached this should never happen
+				checkNoEntry();
 			}
 		}
 		++ChildIndex;
@@ -1475,15 +1219,18 @@ void SMutableCodeViewer::OnExpansionChanged(TSharedPtr<FMutableCodeTreeElement>
 			TreeView->SetItemExpansion(*ExpandedElement, false);
 		}
 		
+		// Cache this element as one currently expanded
 		ExpandedElements.Add(MutableOperation, InItem);
 	}
 	else
 	{
+		// Remove this element from the cache of expanded elements
 		ExpandedElements.Remove(InItem->MutableOperation);
 	}
 }
 
 
+
 void SMutableCodeViewer::OnSelectionChanged(TSharedPtr<FMutableCodeTreeElement> InNode, ESelectInfo::Type InSelectInfo)
 {
 	if (bIsElementHighlighted)
@@ -1675,16 +1422,6 @@ void SMutableCodeViewer::OnRowReleased(const TSharedRef<ITableRow>& InTreeRow)
 	SMutableCodeTreeRow* CastedTableRow = static_cast<SMutableCodeTreeRow*>(&InTreeRow.Get());
 	const TSharedPtr<FMutableCodeTreeElement>& RowElement = CastedTableRow->GetItem();
 	TreeElements.Remove(RowElement);
-	
-	// Remove the element from the array of elements of the current type
-	if (NavigationFoundElements.Contains(RowElement))
-	{
-		NavigationFoundElements.Remove(RowElement);
-		if (CurrentNavigationElement == RowElement)
-		{
-			SetCurrentNavigationElement(nullptr);
-		}
-	}
 }
 #pragma endregion
 
@@ -1909,6 +1646,7 @@ void SMutableCodeViewer::TreeExpandInstance()
 	// Expand only the items that match the datatype provided
 	TreeExpandElements(RootNodes,false, mu::DT_INSTANCE);
 }
+
 #pragma endregion
 
 #pragma region Caching of operations related to constant resource
@@ -1935,78 +1673,83 @@ void SMutableCodeViewer::CacheRootNodeAddresses()
 	const int32 IndexOnConstantsArray)
 {
 	check(MutableModel);
+	check(RootNodeAddresses.Num());
+
 	if (IndexOnConstantsArray < 0)
 	{
 		// Not valid index.
 		UE_LOG(LogTemp,Error,TEXT("The provided index [%d] is not valid."),IndexOnConstantsArray );
 		return;
 	}
-	
-	// Iterate over all the operations found on this model. For that use as start the root elements of this graph and 
-	// move deeper in it
-	
-	TSet<mu::OP::ADDRESS> ProcessedAddresses;
-	TSet<mu::OP::ADDRESS> ConstantResourcesAddresses; 
 
+	// Object containing all data required by the search operation to be able to be called recursively
+	FElementsSearchCache SearchPayload;
+	// Initialize the Search Payload with the root node addresses. This way the search will use them as the root nodes where
+	// to start searching
+	SearchPayload.SetupRootBatch(RootNodeAddresses);
+	
 	// Main update procedure run for the targeted state and the targeted parameter values
-	const mu::FProgram& Program = MutableModel->GetPrivate()->m_program;
-	for (const mu::OP::ADDRESS& RootOperationAddress : RootNodeAddresses)
-	{
-		GetOperationsReferencingConstantResource(ConstantDataType, IndexOnConstantsArray,Program, RootOperationAddress,
-											 ConstantResourcesAddresses,ProcessedAddresses);
-	}
+	const mu::FProgram& Program = MutableModel->GetPrivate()->m_program; 
+	GetOperationsReferencingConstantResource(ConstantDataType,IndexOnConstantsArray,SearchPayload, Program);
 	
 	// At this point we did get all the addresses of operations that do involve the usage of our resource
-	if (ConstantResourcesAddresses.Num() > 0)
+	if (SearchPayload.FoundElements.Num() > 0)
 	{
-		// Managed by TargetedTypeSelector->SetSelectedItem callback method
-		// SetCurrentNavigationElement(nullptr);
-		
 		// Reset operation selection object to show no element to navigate to since we are not navigating over op types
 		// Set the type operation type to NONE (we know it is always the first one)
 		TargetedTypeSelector->SetSelectedItem(FoundModelOperationTypeElements[0]);
 
 		// Dump the located resources array onto the navigation array since we have content to navigate over
-		NavigationOPAddresses = ConstantResourcesAddresses.Array();
-
-		// Move the view to the top if possible
-		// this will take some time to finish -> wait until update end
-		TreeView->ScrollToTop();
+		NavigationElements = MoveTemp(SearchPayload.FoundElements);
+		SortNavigationElements();
 		
-		// Tell the system changes have been made over Navigation Addresses
-		bUpdatedNavigationAddresses = true;
+		// Reset the navigation index
+		NavigationIndex = -1;
 	}
 	else
 	{
 		UE_LOG(LogTemp,Error,TEXT("The provided constant index does not seem to be used anywere : Make sure the index is valid and that IsConstantResourceUsedByOperation() switch is up to date"));
 	}
-
 }
 
-void SMutableCodeViewer::GetOperationsReferencingConstantResource(const mu::DATATYPE ConstantDataType, const int32 IndexOnConstantsArray, const mu::FProgram& InProgram,
-	const mu::OP::ADDRESS& InParentAddress,
-	TSet<mu::OP::ADDRESS>& OutAddressesWithPresence,
-	TSet<mu::OP::ADDRESS>& AlreadyProcessedAddresses)
+
+void SMutableCodeViewer::GetOperationsReferencingConstantResource(
+	const mu::DATATYPE ConstantDataType,
+	const int32 IndexOnConstantsArray,
+	FElementsSearchCache& InSearchPayload,
+	const mu::FProgram& InProgram)
 {
-	// Generic case for unnamed children traversal.
-	mu::ForEachReference(InProgram, InParentAddress, [this, &OutAddressesWithPresence, &ConstantDataType, &IndexOnConstantsArray, &AlreadyProcessedAddresses, &InProgram]( mu::OP::ADDRESS ChildAddress)
+	// next batch of addresses to be explored 
+	TArray<FItemCacheKey> NextBatchAddressesData;
+	
+	for	(int32 ParentIndex = 0 ; ParentIndex < InSearchPayload.BatchData.Num(); ParentIndex++)
 	{
-		// If the parent does have a child then process it 
-		if (ChildAddress && !AlreadyProcessedAddresses.Contains(ChildAddress))
+		// Get one of the previous run "children" and treat as a parent to get it's children and process them
+		const mu::OP::ADDRESS& ParentAddress = InSearchPayload.BatchData[ParentIndex].Child;
+		
+		// Cache if same data type and we share the same address (means this op is pointing at the provided resource)
+		// It will cache duplicated entries
+		if (IsConstantResourceUsedByOperation(IndexOnConstantsArray, ConstantDataType, ParentAddress,InProgram))
 		{
-			// Cache if same data type and we share the same address (means this op is pointing at the provided resource)
-			if (IsConstantResourceUsedByOperation(IndexOnConstantsArray,ConstantDataType,ChildAddress,InProgram))
-			{
-				OutAddressesWithPresence.Add(ChildAddress);
-			}
+			// Since this element is related with the provided constant resource cache it on InSearchPayload.FoundElements
+			InSearchPayload.AddToFoundElements(ParentAddress,ParentIndex,ItemCache);
+		}
 		
-			// Cache to avoid processing it again later
-			AlreadyProcessedAddresses.Add(ChildAddress);
+		// Get all NON PROCESSED the children of this operation to later be able to process them (on next recursive call)
+		InSearchPayload.CacheChildrenOfAddressIfNotProcessed(ParentAddress, InProgram, NextBatchAddressesData);
+	}
+
+	// At this point all the addresses to be computed on the next batch have already been set and will be computed on
+	// the next recursive call
+	
+	// Explore children if found 
+	if (NextBatchAddressesData.Num())
+	{
+		// Cache next batch data so the next invocations is able to locate the provided addresses on the itemsCache
+		InSearchPayload.BatchData = MoveTemp(NextBatchAddressesData);
 		
-			// Process the children of this object
-			GetOperationsReferencingConstantResource(ConstantDataType,IndexOnConstantsArray,InProgram,ChildAddress,OutAddressesWithPresence,AlreadyProcessedAddresses);
-		}
-	});
+		GetOperationsReferencingConstantResource(ConstantDataType, IndexOnConstantsArray, InSearchPayload, InProgram);
+	}
 }
 
 bool SMutableCodeViewer::IsConstantResourceUsedByOperation(const int32 IndexOnConstantsArray,
@@ -2215,39 +1958,10 @@ void SMutableCodeViewer::Tick(const FGeometry& AllottedGeometry, const double In
 	// there to be selected or inspected.
 	if (!TreeView->IsPendingRefresh())
 	{
-		// If the addresses set have changed then start looking for the first element that matches the first element on
-		// the addresses array
-		if ( bUpdatedNavigationAddresses)
-		{
-			// Reset the flag
-			bUpdatedNavigationAddresses = false;
-
-			// Locate the elements that can be currently navigated over.
-			LocateNavigationElementsOnTree();
-
-			// Simulate a button press to go to the next operation. At this point we should have no operation set so it
-			// will find and select the first one
-			// Commented out to avoid the jumping around behaviour of the tree view when selecting a constant resource.
-			// It also avoids situations where we select one constant that invokes a viewer (of type StringViewport for example) and close after, due
-			// to the operation that hosts that resource being selected automatically its that operation previewer witch gets open therefore
-			// clouding the initial viewport we were using.
-			// OnGoToNextOperationButtonPressed();
-
-			return;
-		}
-
-		// Are we searching upwards?
-		if (bIsSearchingForPreviousElement)
-		{
-			GoToPreviousOperationAfterRefresh();
-			return;
-		}
-
-		// Are we searching downwards?
-		if (bIsSearchingForNextElement)
+		/** If we have expanded the tree elements in order to reach one of them then continue the operation */
+		if (bWasUniqueExpansionInvokedForNavigation || bWasScrollToTargetRequested)
 		{
-			GoToNextOperationAfterRefresh();
-			return;
+			FocusViewOnNavigationTarget();
 		}
 	}
 	
@@ -2511,4 +2225,4 @@ void SMutableCodeViewer::PreviewMutableShape(const mu::FShape* Shape)
 	UE_LOG(LogMutable,Warning,TEXT("Previewer for Mutable Shapes not yet implemented"))
 }
 
-#undef LOCTEXT_NAMESPACE 
+#undef LOCTEXT_NAMESPACE
diff --git a/Engine/Plugins/Experimental/Mutable/Source/CustomizableObjectEditor/Private/MuCOE/SMutableCodeViewer.h b/Engine/Plugins/Experimental/Mutable/Source/CustomizableObjectEditor/Private/MuCOE/SMutableCodeViewer.h
--- a/Engine/Plugins/Experimental/Mutable/Source/CustomizableObjectEditor/Private/MuCOE/SMutableCodeViewer.h
+++ b/Engine/Plugins/Experimental/Mutable/Source/CustomizableObjectEditor/Private/MuCOE/SMutableCodeViewer.h
@@ -174,6 +174,12 @@ private:
 	* Tree widget objects
 	*/
 
+	/** Map with all the generated elements of the tree. Unique and duplicated elements will be present on this list and
+	 * also the children of the unique elements.
+	 * @note The children of duplicated elements will only be present once, as children of the unique element they duplicate.
+	 * This is to avoid having identical elements on the tree (witch would cause a crash) while also being pointless due to
+	 * how we manage expansion of duplicated elements.
+	 */
 	TMap< FItemCacheKey, TSharedPtr<FMutableCodeTreeElement>> ItemCache;
 
 	/** Main tree item for each op. An op can be represented with multiple tree nodes if it is reachable from different paths. */
@@ -185,6 +191,22 @@ private:
 	/** Array with all the elements that have been manually expanded by the user */
 	TMap< mu::OP::ADDRESS, TSharedPtr<FMutableCodeTreeElement>> ExpandedElements;
 
+	
+	/** Before any UI operation generate all the elements that may be navigable over the tree. No children of duplicated 
+	 * addresses will be generated.
+	 */
+	void GenerateAllTreeElements();
+
+	/** Support GenerateAllTreeElements by generating the elements recursively */
+	void GenerateElementRecursive(mu::OP::ADDRESS InParentAddress, const mu::FProgram& InProgram);
+
+	
+	/** The addresses of the root operations. Cached when this object gets loaded on the Construct method of this slate */
+	TArray<mu::OP::ADDRESS> RootNodeAddresses;
+	
+	/** Utility method that provides you with the addresses of the root nodes of the program */
+	void CacheRootNodeAddresses();
+	
 	/** Control boolean to know if there are any highlighted elements on the tree */
 	bool bIsElementHighlighted = false;
 
@@ -232,6 +254,7 @@ private:
 	 */
 	EOperationComputationalCost GetOperationTypeComputationalCost(mu::OP_TYPE OperationType) const;
 	
+	
 	/*
 	 * Navigation : Operation type navigation Type selection object
 	 */
@@ -260,9 +283,7 @@ private:
 	/** Generate a list of elements that will be used as backend for the navigation type selection dropdown */
 	void GenerateNavigationDropdownElements();
 	
-	/** Update the array with all the elements using the current operation. Required when changing the targeted operation type */
-	void LocateNavigationElementsOnTree();
-
+	
 	/** Fills ModelOperationTypes with all the types present on the current model.
 	 * It also makes sure we have a NONE operation and that the operations are sorted alphabetically.
 	 * @note Method designed to be called once per model.*/
@@ -277,6 +298,7 @@ private:
 	void GetOperationTypesPresentOnModel(
 		const TArray< mu::OP::ADDRESS>& InParentAddresses, const mu::FProgram& InProgram, TSet<mu::OP_TYPE>& OutLocatedOperations,
 		TSet<mu::OP::ADDRESS>& AlreadyProcessedAddresses);
+
 	
 	/*
 	 * Navigation : UI callback methods
@@ -303,76 +325,198 @@ private:
 	 */
 	FText OnPrintNavigableObjectAddressesCount() const;
 	
-	/** Method that tells the UI if the bottom to go to the previous element can be interacted
-	 * @return It returns true if the UI element should be interactable, false if not.
-	 */
+	/** Method that tells the UI if the bottom to go to the previous element can be interacted */
 	bool CanInteractWithPreviousOperationButton() const;
 
-	/** Method that tells the UI if the bottom to go to the next element can be interacted
-	 * @return It returns true if the UI element should be interactable, false if not.
-	 */
+	/** Method that tells the UI if the bottom to go to the next element can be interacted */
 	bool CanInteractWithNextOperationButton() const;
 	
 
+	/** Used by the UI and internally does what is necessary to get to the previous element of the targeted search type.
+	 * It is designed to work alongside GoToPreviousOperationAfterRefresh(...)
+	 * @return A reply to tell if the UI action has been handled or not.
+	 */
+	FReply OnGoToPreviousOperationButtonPressed();
+
+	/** Used by the UI and internally does what is necessary to get to the next element of the targeted search type.
+	 * It is designed to work alongside GoToNextOperationAfterRefresh(...)
+	 * @return A reply to tell if the UI action has been handled or not.
+	 */
+	FReply OnGoToNextOperationButtonPressed();
+	
+	
 	/** Callback method used to allow the user to directly set the type of operation to be scanning directly by selecting a
 	 * operation on the tree and using that operation type as the type to search for
 	 */
 	void OnSelectedOperationTypeFromTree();
-
+	
 	/*
 	 * Navigation : control flags
 	 */
 	
-	/** Flag that tells the system the user does want to select the previous operation of the targeted type. */
-	bool bIsSearchingForPreviousElement = false;
+	/** Flag monitoring if we have requested a scroll operation to reach the targeted element */
+	bool bWasScrollToTargetRequested = false;
+
+	/** Flag designed to tell the system when the expansion of unique elements have been performed as part of the navigation operation */
+	bool bWasUniqueExpansionInvokedForNavigation = false;
 	
-	/** Flag that tells the system the user does want to select the next operation of the targeted type. */
-	bool bIsSearchingForNextElement = false;
+	/*
+	 * Navigation : Shared objects between navigation search types 
+	 */
+
+	/** Array with all the elements of the type we are looking for (shared type of constant resource). */
+	TArray<TSharedPtr<FMutableCodeTreeElement>> NavigationElements;
 	
-	/** Flag used by the navigation system to know when it is required the full expansion of the tree.
-	 * It allows the system to safely jump from one element to another without the risk of leaving elements in the
-	 * middle without being accessed.
+	/** Operation types present on the currently set mutable model. */
+	int64 NavigationIndex = -1;
+	
+	/** Sort the contents of NavigationElements to follow a sequential pattern using the indices of the elements from 0 to +n*/
+	void SortNavigationElements();
+
+	/** Focus the current NavigationElement and places it into view. It also selects it so the previewer for the
+	 * element gets invoked
 	 */
-	bool bTreeWasExpanded = false;
+	void FocusViewOnNavigationTarget();
 	
-	/** Flag set to true when the navigation system is automatically scrolling over the graph */
-	bool bIsScrolling = false;
+	/** Wrapper struct designed to be used as cache for all elements found during the navigation system search for elements
+	 * of X Type or relation with a targeted constant resource. It is designed to be used and then destroyed once the
+	 * search operation has been completed */
+	struct FElementsSearchCache
+	{
+		/** Set of addresses that have already been searched for relevant data */
+		TSet<mu::OP::ADDRESS> ProcessedAddresses;
+
+		/** Collection of elements that have been found during the search. They may be related by OP_Type or used constant resource*/
+		TArray<TSharedPtr<FMutableCodeTreeElement>> FoundElements;
+		
+		/** Array containing all the next addresses to be processed.
+		 * The child is the address itself to be later processed
+		 * The parent is the parent address of the child.
+		 * The ChildIndexInParent is the index (or child position) o the child address on Child as part of the child set of the parent address.
+		 */
+		TArray<FItemCacheKey> BatchData;
+		
+
+		/** Generates the structures to be able to start the search of elements. It uses the root addresses as the
+		 * start of the search operation.
+		 * @param InRootNodeAddresses The addresses of the root operations of the operations tree
+		 */
+		void SetupRootBatch(const TArray<mu::OP::ADDRESS>& InRootNodeAddresses)
+		{
+			check (InRootNodeAddresses.Num());
+			// This method should only be called once when no data is present on this cache
+			check (BatchData.IsEmpty());
+			
+			BatchData.Reserve(InRootNodeAddresses.Num());
+			// The child address of each parent operation on it's parent
+			for (int32 RootIndex = 0; RootIndex < InRootNodeAddresses.Num(); RootIndex++)
+			{
+				FItemCacheKey Key;
+				{
+					Key.Child = InRootNodeAddresses[RootIndex];
+					// Store the parent of this object as 0 to have a "virtual" address witch all root addresses are
+					// children of 
+					Key.Parent = 0;
+					// And also the index of the parent on it's parent "virtual" structure
+					Key.ChildIndexInParent = RootIndex;		
+				}
 
-	/** Flag that tells the system when the navigation addresses have been updated and therefore action should be
-	 * taken after the tree refresh. Currently we look for the first element present on the array of addresses */
-	bool bUpdatedNavigationAddresses = false;
+				// Add this new entry point for the search of addresses 
+				BatchData.Add(Key);
+			}
+		}
+		
+		/** Method to cache the provided mutable address as one of the addresses that are of the type we are looking for or maybe
+		* is related with the constant resource we are looking for operations related with.
+		* @param OpAddress The address to save as one related with a operation type or constant resource
+		* @param IndexAsChildOfInputAddress The index on BatchData that represents the parent of the provided operation address
+		* @param InItemCache Cache with all the elements of the tree. Here is where we search for the element based on
+		* the OpAddress and the parent and ChildIndexInParent provided thanks to IndexAsChildOfInputAddress
+		*/
+		void AddToFoundElements(const mu::OP::ADDRESS OpAddress,const int32 IndexAsChildOfInputAddress,
+		                        const TMap<FItemCacheKey, TSharedPtr<FMutableCodeTreeElement>>& InItemCache)
+		{
+			FItemCacheKey Key;
+			{
+				Key.Child = OpAddress;
+				// Store the parent address of this object
+				Key.Parent = BatchData[IndexAsChildOfInputAddress].Parent;
+				// And also the index of the parent on it's parent structure
+				Key.ChildIndexInParent = BatchData[IndexAsChildOfInputAddress].ChildIndexInParent;		
+			}
+			// Generate a key for this element in order to search in on the map with all the elements
+								
+			// Find that element on the tree, (check error if not found since all elements should be there)
+			TSharedPtr<FMutableCodeTreeElement> FoundElement = *InItemCache.Find(Key);
+			check (FoundElement != nullptr);
+			
+			// Store this element on our temp Map of elements
+			FoundElements.Add(FoundElement);
+		}
+
+		/** Caches the provided parent address to the Search payload so they can be later read and processed on another batch of the
+		 * method tasked with finding related operations to operation type or constant resource.
+		 * @param InParentAddress The parent address to search children of.
+		 * @param InProgram The mutable program that will be used to perform the children search.
+		 * @param OutFoundChildrenData Array with all the ItemCacheKeys that represent all the children found
+		 * @note It will not add the children of any provided ParentAddress for the next batch if the parent address provided
+		 * has already been processed and therefore whose children have already been searched or prepared for searching.
+		 */
+		void CacheChildrenOfAddressIfNotProcessed(mu::OP::ADDRESS InParentAddress,
+		                                          const mu::FProgram& InProgram,
+		                                          TArray<FItemCacheKey>& OutFoundChildrenData)
+		{
+			if (!ProcessedAddresses.Contains(InParentAddress))
+			{
+				// Cache to avoid processing it again later
+				ProcessedAddresses.Add(InParentAddress);
+	
+				// Generic case for unnamed children traversal.
+				uint32 ChildIndex = 0;
+				mu::ForEachReference(InProgram, InParentAddress, [this, &InParentAddress, &ChildIndex,&OutFoundChildrenData]( mu::OP::ADDRESS ChildAddress)
+				{
+					// If the parent does have a child then process it 
+					if (ChildAddress) 
+					{
+						FItemCacheKey Key;
+						{
+							Key.Child = ChildAddress;
+							Key.Parent = InParentAddress;
+							Key.ChildIndexInParent = ChildIndex;		
+						}
+
+						// Save it to the output so can later be placed onto BatchData safely 
+						OutFoundChildrenData.Add(Key);
+					}
+					ChildIndex++;
+				});
+			}
+		}
+	};
+	
 	
 	/*
-	 * Navigation : OP Address based navigation
+	 * Navigation : Operation type based navigation
 	 */
 	
-	/** Collection of Mutable addresses. Cache used during tree navigation.
-	 * They serve as the navigation system backend and determine if an element will or not be stopped at during
-	 * the navigation process.
-	 * @note Not set as a TSet since I want to, in future revisions, be able to trust the order of the elements of the collection.
-	 */
-	TArray<mu::OP::ADDRESS> NavigationOPAddresses;
-
-	/** Provided an Operation type store as navigation addresses all operations of the targeted Op type
+	/** Provided an Operation type store as navigation addresses all operations of the targeted navigation Op type
 	 * @note The operations found get saved on the collection  NavigationOPAddresses.
-	 * @param TargetedOperationType The type of operation we want to cache to be later able to navigate over it.
 	 */
-	void CacheAddressesOfOperationsOfType(mu::OP_TYPE TargetedOperationType);
+	void CacheAddressesOfOperationsOfType();
 
 	/** Fills an array with all the addresses of operations that do have in common the same targeted operation type
 	 * @param TargetOperationType The operation type used to discriminate what operation addresses we want to retrieve.
-	 * @param InParentAddresses The addresses of the parent object. Used to get the children and then process them too.
-	 * @param OutAddressesOfType Output with all the addresses that are of the same type as TargetOperationType.
-	 * @param  AlreadyProcessedAddresses Used during recursive calls. Cache the already processed operations to avoid
 	 * processing them more than once.
+	 * @param InSearchPayload A caching structure designed to hold the data that gets passed from one recursive call to
+	 * the other. It also stores the found elements and other data.
+	 * @param InProgram The mutable program holding the data to be searched over
 	 */
-	void GetOperationsOfType ( const mu::OP_TYPE& TargetOperationType,const TArray<mu::OP::ADDRESS>& InParentAddresses,const mu::FProgram& InProgram,
-		TSet<mu::OP::ADDRESS>& OutAddressesOfType,
-		TSet<mu::OP::ADDRESS>& AlreadyProcessedAddresses);
-
+	void GetOperationsOfType(const mu::OP_TYPE& TargetOperationType,
+	                         FElementsSearchCache& InSearchPayload,
+	                         const mu::FProgram& InProgram);
 
 	/*
-	 * Navigation : Cache operation instances based on constant resource
+	 * Navigation : Navigation based on constant resource relation with addresses
 	 */
 
 public:
@@ -390,29 +534,20 @@ public:
 	void CacheAddressesRelatedWithConstantResource(const mu::DATATYPE ConstantDataType, const int32 IndexOnConstantsArray);
 
 private:
-
-	/** The addresses of the root operations. Cached when this object gets loaded on the Construct method of this slate */
-	TArray<mu::OP::ADDRESS> RootNodeAddresses;
-	
-	/** Utility method that provides you with the addresses of the root nodes of the program */
-	void CacheRootNodeAddresses();
 	
 	/**
 	 * Provided a data type and the index of the constant object using said data type locate all operations that do
 	 * directly make use of said constant resource.
 	 * @param ConstantDataType The type of the resource we are providing an index of. 
 	 * @param IndexOnConstantsArray The index of the constant resource we want to know what operations are using it.
+	 * @param InSearchPayload Cache object that will end up containing all the elements related with the provided constant
+	 * resource
 	 * @param InProgram The mutable program containing all the operations of the graph.
-	 * @param InParentAddress The address of the root object whose child operations we are searching the provided constant on.
-	 * @param OutAddressesWithPresence An array with all the addresses to mutable operations making use of the provided constant resource.
-	 * @param AlreadyProcessedAddresses Cache with the addresses already processed. It is only intended to be used on recursive calls.
 	 */
 	void GetOperationsReferencingConstantResource(const mu::DATATYPE ConstantDataType,
 	                                              const int32 IndexOnConstantsArray,
-	                                              const mu::FProgram& InProgram,
-	                                              const mu::OP::ADDRESS& InParentAddress,
-	                                              TSet<mu::OP::ADDRESS>& OutAddressesWithPresence,
-	                                              TSet<mu::OP::ADDRESS>& AlreadyProcessedAddresses);
+	                                              FElementsSearchCache& InSearchPayload,
+	                                              const mu::FProgram& InProgram);
 
 	/**
 	 * Provided an operation address this method tells us if that address is making use of our constant resource or not.
@@ -431,76 +566,6 @@ private:
 	                                       const mu::FProgram& InProgram) const;
 	
 	
-	/*
-	 * Navigation : operation methods and variables from Tree
-	 */
-	
-	/** The element on FoundElementsOfType that we are inspecting and using as origin for our search.
-	 * @note Do not update it's value manually */
-	TSharedPtr< FMutableCodeTreeElement> CurrentNavigationElement;
-	
-	/** Array with all the elements on the tree of the same type as the one defined on "TargetedOperationType" */
-	TArray<TSharedPtr<FMutableCodeTreeElement>> NavigationFoundElements;
-	
-	/** Copy of the array FoundElementsOfType done before moving on to the next element. Required to check if the
-	 * expansion of the tree view did add or remove elements of the type we are looking for.
-	 */
-	TArray<TSharedPtr<FMutableCodeTreeElement>> NavigationPreviouslyFoundElements;
-	
-	
-	/** Amount of elements the scroll will jump up and down while searching for new elements
-	 * It is required to be able to locate elements whose rows have not been yet generated
-	 * Do not change.
-	 */
-	const int32 MinTreeScrollStep = 1;
-
-	/** Variable holding the scroll step to be used. It will vary depending on the amount of elements the tree can display
-	 * witch may change from resolution to resolution. This is then the value to be used during the automatic scrolling
-	 * operation of the tree navigation system. It will start at the min possible value and only increase when posible.
-	 */
-	int32 ComputedMaxViewScrollStep = MinTreeScrollStep;
-
-	/** Tree view manipulation method designed to set the element to be selected on the UI side of the tree. It is also
-	 *required to take note of what the current found element is to be able to navigate the tree of elements
-	 *@param NewlyFoundElement - The element we want to now be using as the NavigationFoundElement */
-	void SetCurrentNavigationElement(TSharedPtr<FMutableCodeTreeElement> NewlyFoundElement);
-
-	
-	/** Consistency focused method designed to sort the list of navigation elements based on the position of them
-	* on the tree view. It is required since new elements may be added or removed due to user activity.
-	*/
-	void SortNavigationElementsArray();
-	
-	/** Used by the UI and internally does what is necessary to get to the previous element of the targeted search type.
-	 * It is designed to work alongside GoToPreviousOperationAfterRefresh(...)
-	 * @return A reply to tell if the UI action has been handled or not.
-	 */
-	FReply OnGoToPreviousOperationButtonPressed();
-
-	/** Used by the UI and internally does what is necessary to get to the next element of the targeted search type.
-	 * It is designed to work alongside GoToNextOperationAfterRefresh(...)
-	 * @return A reply to tell if the UI action has been handled or not.
-	 */
-	FReply OnGoToNextOperationButtonPressed();
-	
-	/** Makes the currently selected operation be the previous one able to be located on the tree of operations.
-	 * @note It must only be called after the tree has been fully refreshed. If not inconsistent behaviour may arise.
-	 */
-	void GoToPreviousOperationAfterRefresh();
-	
-	/** Makes the currently selected operation of the targeted type to be the next one. Called after the tree has been
-	 * refreshed with new contents (fully refreshed).
-	 * @note It must only be called after the tree has been fully refreshed. If not inconsistent behaviour may arise.
-	 */
-	void GoToNextOperationAfterRefresh();
-
-	/** Method that compares the contents of "PreviousTargetedElementsOfType" against the ones on "CurrentTargetedElementsOfType" to
-	 *check if there is any new element found on the current elements array not found on the Previous one
-	 * @return True if there is any new element defined on the "CurrentTargetedElementsOfType" not present on "PreviousTargetedElementsOfType",
-	 * false if not.
-	 */
-	bool HaveNewElementsBeenFound() const;
-	
 	/*
 	* Main callbacks from the tree widget standard operation. 
 	*/
@@ -522,6 +587,7 @@ private:
 	*/
 	void OnExpansionChanged(TSharedPtr<FMutableCodeTreeElement> InInfo, bool bInExpanded);
 
+
 	/** 
 	* Callback invoked each time the selected element changes used to generate the previews depending of the 
 	* operation selected
@@ -580,7 +646,7 @@ private:
 	
 	/** Grabs the selected element and expands all the elements inside the selected branch. Duplicates are ignored */
 	void TreeExpandSelected();
-
+	
 	/** 
 	* Highlights all tree elements that share the same operation as the element provided
 	* @param InTargetElement - The info object to be used as blueprint to search for similar objects and highlight them
@@ -637,13 +703,14 @@ private:
 	void PrepareProjectorViewer();
 };
 
+/** The data of a row on the operation type dropdown. */
 class FMutableOperationElement : public TSharedFromThis<FMutableOperationElement>
 {
 public:
-	FMutableOperationElement(mu::OP_TYPE InOperationType, FText OperationTypeBeingReprepresented,FSlateColor OperationColor)
+	FMutableOperationElement(mu::OP_TYPE InOperationType, FText OperationTypeBeingRepresented,FSlateColor OperationColor)
 	{
 		OperationType = InOperationType;
-		OperationTypeText = OperationTypeBeingReprepresented;
+		OperationTypeText = OperationTypeBeingRepresented;
 		OperationTextColor = OperationColor;
 	}
 
@@ -659,12 +726,13 @@ public:
 class FMutableCodeTreeElement : public TSharedFromThis<FMutableCodeTreeElement>
 {
 public:
-	FMutableCodeTreeElement(const mu::ModelPtr& InModel, mu::OP::ADDRESS InOperation, const FString& InCaption,const FSlateColor InLabelColor, const TSharedPtr<FMutableCodeTreeElement>* InDuplicatedOf = nullptr)
+	FMutableCodeTreeElement(int32 InIndexOnTree ,const mu::ModelPtr& InModel, mu::OP::ADDRESS InOperation, const FString& InCaption,const FSlateColor InLabelColor, const TSharedPtr<FMutableCodeTreeElement>* InDuplicatedOf = nullptr)
 	{
 		MutableModel = InModel;
 		MutableOperation = InOperation;
 		Caption = InCaption;
 		LabelColor = InLabelColor;
+		IndexOnTree = InIndexOnTree;
 		if (InDuplicatedOf)
 		{
 			DuplicatedOf = *InDuplicatedOf;
@@ -732,4 +800,11 @@ public:
 
 	/** The color to be used by the row representing this object */
 	FSlateColor LabelColor;
+
+	/*
+	 * Navigation metadata
+	 */
+	
+	/** The current position of this element on the tree view. Used for navigation */
+	int32 IndexOnTree;
 };
