[mutable] Fixed access to invalid data on the Performance Reporter util.
- No stop signal was being sent to the test being ran witch made it access invalid data after the host slate was destroyed causing a check fail.

diff --git a/Engine/Plugins/Experimental/Mutable/Source/CustomizableObjectEditor/Private/MuCOE/SCustomizableObjectEditorPerformanceReport.cpp b/Engine/Plugins/Experimental/Mutable/Source/CustomizableObjectEditor/Private/MuCOE/SCustomizableObjectEditorPerformanceReport.cpp
--- a/Engine/Plugins/Experimental/Mutable/Source/CustomizableObjectEditor/Private/MuCOE/SCustomizableObjectEditorPerformanceReport.cpp
+++ b/Engine/Plugins/Experimental/Mutable/Source/CustomizableObjectEditor/Private/MuCOE/SCustomizableObjectEditorPerformanceReport.cpp
@@ -248,6 +248,8 @@ void SCustomizableObjecEditorPerformanceReport::Construct(const FArguments& InAr
 
 SCustomizableObjecEditorPerformanceReport::~SCustomizableObjecEditorPerformanceReport()
 {
+	StopTests();
+	
 	if (CurrentReportWorstCase && CurrentReportWorstCase->WorstCaseInstance)
 	{
 		CurrentReportWorstCase->WorstCaseInstance->UpdatedDelegate.RemoveDynamic(HelperCallback, &UPerformanceReportHelper::DelegatedCallback);
@@ -410,7 +412,7 @@ FReply SCustomizableObjecEditorPerformanceReport::RunPerformanceReport(UCustomiz
 	}
 
 	if (ReportSubject)
-	{
+	{	
 		StopTests();
 
 		if (CurrentReportInstance || CurrentReportWorstCase)
