[Mutable] Reduce the amount of skeletons to merge to avoid bone collisions.
+ Fix CO compilation error when using group projectors.
+ Fix generation issues on low-end platforms when lower LODs are discarded.

#rnx

diff --git a/Engine/Plugins/Experimental/Mutable/Source/CustomizableObject/Classes/MuCO/CustomizableObject.h b/Engine/Plugins/Experimental/Mutable/Source/CustomizableObject/Classes/MuCO/CustomizableObject.h
--- a/Engine/Plugins/Experimental/Mutable/Source/CustomizableObject/Classes/MuCO/CustomizableObject.h
+++ b/Engine/Plugins/Experimental/Mutable/Source/CustomizableObject/Classes/MuCO/CustomizableObject.h
@@ -1178,7 +1178,7 @@ private:
 	// This is a manual version number for the binary blobs in this asset.
 	// Increasing it invalidates all the previously compiled models.
 	// Warning: If while merging code both versions have changed, take the highest+1.
-	static const int32 CurrentSupportedVersion = 351;
+	static const int32 CurrentSupportedVersion = 352;
 
 public:
 
diff --git a/Engine/Plugins/Experimental/Mutable/Source/CustomizableObject/Private/MuCO/CustomizableObjectInstance.cpp b/Engine/Plugins/Experimental/Mutable/Source/CustomizableObject/Private/MuCO/CustomizableObjectInstance.cpp
--- a/Engine/Plugins/Experimental/Mutable/Source/CustomizableObject/Private/MuCO/CustomizableObjectInstance.cpp
+++ b/Engine/Plugins/Experimental/Mutable/Source/CustomizableObject/Private/MuCO/CustomizableObjectInstance.cpp
@@ -2937,14 +2937,21 @@ bool UCustomizableInstancePrivateData::BuildSkeletalMeshSkeletonData(const TShar
 			FName BoneName = MutSkeletonData.BoneNames[BoneIndex];
 			check(BoneName != NAME_None);
 
-			const int32 SourceBoneIndex = ReferenceSkeleton.FindRawBoneIndex(BoneName);
+			int32 SourceBoneIndex = ReferenceSkeleton.FindRawBoneIndex(BoneName);
+			
 			if (SourceBoneIndex == INDEX_NONE)
 			{
-				// Merged skeleton is missing some bones! This happens if one of the skeletons involved in the merge is discarded due to being incompatible with the rest
-				// or if the source mesh is not in sync with the skeleton. 
-				UE_LOG(LogMutable, Warning, TEXT("Building instance: generated mesh has a bone [%s] not present in the reference mesh [%s]. Failing to generate mesh. "),
-					*BoneName.ToString(), *SkeletalMesh->GetName());
-				return false;
+				if (MutSkeletonData.BoneMatricesWithScale.Find(BoneName))
+				{
+					// Merged skeleton is missing some bones! This happens if one of the skeletons involved in the merge is discarded due to being incompatible with the rest
+					// or if the source mesh is not in sync with the skeleton. 
+					UE_LOG(LogMutable, Warning, TEXT("Building instance: generated mesh has a bone [%s] not present in the reference mesh [%s]. Failing to generate mesh. "),
+						*BoneName.ToString(), *SkeletalMesh->GetName());
+					return false;
+				}
+
+				// The bone is missing but it's not required. Fix the SourceBoneIndex and continue.
+				SourceBoneIndex = 0;
 			}
 
 			BoneToFinalBoneIndexMap.Add(BoneName, SourceBoneIndex);
@@ -2967,10 +2974,13 @@ bool UCustomizableInstancePrivateData::BuildSkeletalMeshSkeletonData(const TShar
 				BoneIndex = BoneToFinalBoneIndexMap[MutSkeletonData.BoneNames[BoneIndex]];
 			}
 
-			for (uint16& BoneIndex : Component.ActiveBones)
+			TArray<uint16> UniqueActiveBones;
+			for (const uint16 BoneIndex : Component.ActiveBones)
 			{
-				BoneIndex = BoneToFinalBoneIndexMap[MutSkeletonData.BoneNames[BoneIndex]];
+				UniqueActiveBones.AddUnique(BoneToFinalBoneIndexMap[MutSkeletonData.BoneNames[BoneIndex]]);
 			}
+
+			Exchange(Component.ActiveBones, UniqueActiveBones);
 			Component.ActiveBones.Sort();
 		}
 	}
diff --git a/Engine/Plugins/Experimental/Mutable/Source/CustomizableObject/Private/MuCO/CustomizableObjectSystem.cpp b/Engine/Plugins/Experimental/Mutable/Source/CustomizableObject/Private/MuCO/CustomizableObjectSystem.cpp
--- a/Engine/Plugins/Experimental/Mutable/Source/CustomizableObject/Private/MuCO/CustomizableObjectSystem.cpp
+++ b/Engine/Plugins/Experimental/Mutable/Source/CustomizableObject/Private/MuCO/CustomizableObjectSystem.cpp
@@ -1384,6 +1384,13 @@ namespace impl
 					continue;
 				}
 
+				// Add SkeletonIds 
+				const int32 SkeletonIDsCount = Mesh->GetSkeletonIDsCount();
+				for (int32 SkeletonIndex = 0; SkeletonIndex < SkeletonIDsCount; ++SkeletonIndex)
+				{
+					SkeletonIds.AddUnique(Mesh->GetSkeletonID(SkeletonIndex));
+				}
+
 				const mu::Skeleton* Skeleton = Mesh->GetSkeleton().get();
 				const int32 BoneCount = Skeleton->GetBoneCount();
 
@@ -1518,9 +1525,6 @@ namespace impl
 							{
 								// Add bone even if it doesn't have weights
 								AddBone(BoneName);
-
-								// Add SkeletonId 
-								SkeletonIds.AddUnique(Skeleton->GetBoneId(BoneIndex));
 							}
 
 							// Add root as a placeholder
@@ -1553,9 +1557,6 @@ namespace impl
 									ParentChain[ParentChainCount] = ParentName;
 									++ParentChainCount;
 
-									// Add SkeletonId 
-									SkeletonIds.AddUnique(Skeleton->GetBoneId(NextMutParentIndex));
-
 									NextMutParentIndex = Skeleton->GetBoneParent(NextMutParentIndex);
 								}
 
@@ -1570,9 +1571,6 @@ namespace impl
 							// Add the used bone to the hierarchy
 							const int16 NewBoneIndex = AddBone(BoneName);
 							ComponentBoneMap.Add(NewBoneIndex);
-
-							// Add SkeletonId 
-							SkeletonIds.AddUnique(Skeleton->GetBoneId(BoneIndex));
 						}
 						else
 						{
@@ -2089,6 +2087,12 @@ namespace impl
 		// Prepare streaming for the current customizable object
 		System->GetPrivate()->Streamer->PrepareStreamingForObject(CustomizableObject);
 
+		// Ensure we're not requesting invalid LODs
+		PrivateInstance->LastUpdateMinLOD = FMath::Max(PrivateInstance->LastUpdateMinLOD, (int32)PrivateInstance->FirstLODAvailable);
+		PrivateInstance->LastUpdateMinLOD = FMath::Min(PrivateInstance->LastUpdateMinLOD, PrivateInstance->FirstLODAvailable + PrivateInstance->NumMaxLODsToStream);
+
+		PrivateInstance->LastUpdateMaxLOD = FMath::Max(PrivateInstance->LastUpdateMaxLOD, PrivateInstance->LastUpdateMinLOD);
+
 
 		// Task: Mutable Update and GetMesh
 		//-------------------------------------------------------------
diff --git a/Engine/Plugins/Experimental/Mutable/Source/CustomizableObjectEditor/Private/MuCOE/GenerateMutableSource/GenerateMutableSourceMesh.cpp b/Engine/Plugins/Experimental/Mutable/Source/CustomizableObjectEditor/Private/MuCOE/GenerateMutableSource/GenerateMutableSourceMesh.cpp
--- a/Engine/Plugins/Experimental/Mutable/Source/CustomizableObjectEditor/Private/MuCOE/GenerateMutableSource/GenerateMutableSourceMesh.cpp
+++ b/Engine/Plugins/Experimental/Mutable/Source/CustomizableObjectEditor/Private/MuCOE/GenerateMutableSource/GenerateMutableSourceMesh.cpp
@@ -533,15 +533,16 @@ mu::MeshPtr ConvertSkeletalMeshToMutable(USkeletalMesh* InSkeletalMesh, int LOD,
 		return nullptr;
 	}
 
-	// Use this to set bone IDs in Mutable to specify to which skeleton they belong
-	int32 SkeletonID = 0;
-
 	// This map will be used to find and replace vertex influences.
 	TMap<int32, int32> InfluencesToReplaceMap;
 
 	// Check for a matching skeleton to the reference mesh
 	if (!bIgnoreSkeleton)
 	{
+		// Add the RefSkeleton ID to the mesh.
+		const int32 RefSkeletonID = GenerationContext.ReferencedSkeletons.Find(ComponentRefSkeleton);
+		MutableMesh->AddSkeletonID(RefSkeletonID);
+
 		if (InSkeleton != ComponentRefSkeleton)
 		{
 			bool* SkeletonCompatibility = MutComponentInfo.SkeletonCompatibility.Find(InSkeleton);
@@ -600,14 +601,10 @@ mu::MeshPtr ConvertSkeletalMeshToMutable(USkeletalMesh* InSkeletalMesh, int LOD,
 				// Incompatible skeleton 
 				return nullptr;
 			}
-			
-			// Is it a skeleton that we have already used?
-			SkeletonID = GenerationContext.ReferencedSkeletons.AddUnique(InSkeleton);
-		}
-		else
-		{
-			// Since we could have multiple reference skeletons, we can't assume its index is 0 
-			SkeletonID = GenerationContext.ReferencedSkeletons.Find(InSkeleton);
+
+			// Add the skeleton to the list of referenced skeletons and add its index to the mesh
+			const int32 SkeletonID = GenerationContext.ReferencedSkeletons.AddUnique(InSkeleton);
+			MutableMesh->AddSkeletonID(SkeletonID);
 		}
 
 		{
@@ -1336,12 +1333,12 @@ mu::MeshPtr ConvertSkeletalMeshToMutable(USkeletalMesh* InSkeletalMesh, int LOD,
 	if (!bIgnoreSkeleton)
 	{
 		// Skeleton
-		int32 NumBonesInMesh = ImportedModel->LODModels[LOD].Sections[MaterialIndex].BoneMap.Num();
-
 
 		// Add the bones in the order they are used in the mesh
 		TMap<int32, int32> InverseBoneMap; // Inverse of ImportedModel->LODModels[LOD].Chunks[MaterialIndex].BoneMap for the chunk's bone indices
 		TQueue<int32> SkippedBoneMapIndices; // The order and index of the bones in the BoneMap must remain the same even if some bones are removed
+
+		const int32 NumBonesInMesh = ImportedModel->LODModels[LOD].Sections[MaterialIndex].BoneMap.Num();
 		for (int32 BoneMapIndex = 0; BoneMapIndex < NumBonesInMesh; ++BoneMapIndex)
 		{
 			if (InfluencesToReplaceMap.Find(BoneMapIndex))
@@ -1398,7 +1395,6 @@ mu::MeshPtr ConvertSkeletalMeshToMutable(USkeletalMesh* InSkeletalMesh, int LOD,
 			int32 ParentRefSkeletonIndex = SkeletalMesh->GetRefSkeleton().GetParentIndex(RefSkelIndex);
 			int ParentBoneMapIndex = ParentRefSkeletonIndex >= 0 ? InverseBoneMap[ParentRefSkeletonIndex] : -1;
 			MutableSkeleton->SetBoneParent(BoneMapIndex, ParentBoneMapIndex);
-			MutableSkeleton->SetBoneId(BoneMapIndex,SkeletonID);
 
 			FMatrix44f BaseInvMatrix = SkeletalMesh->GetRefBasesInvMatrix()[RefSkelIndex];
 			FTransform3f BaseInvTransform;
diff --git a/Engine/Plugins/Experimental/Mutable/Source/CustomizableObjectEditor/Private/MuCOE/GenerateMutableSource/GenerateMutableSourceSurface.cpp b/Engine/Plugins/Experimental/Mutable/Source/CustomizableObjectEditor/Private/MuCOE/GenerateMutableSource/GenerateMutableSourceSurface.cpp
--- a/Engine/Plugins/Experimental/Mutable/Source/CustomizableObjectEditor/Private/MuCOE/GenerateMutableSource/GenerateMutableSourceSurface.cpp
+++ b/Engine/Plugins/Experimental/Mutable/Source/CustomizableObjectEditor/Private/MuCOE/GenerateMutableSource/GenerateMutableSourceSurface.cpp
@@ -759,7 +759,7 @@ mu::NodeSurfacePtr GenerateMutableSourceSurface(const UEdGraphPin * Pin, FMutabl
 					if (bShareProjectionTexturesBetweenLODs && bIsGroupProjectorImage)
 					{
 						// Add to the GroupProjectorLODCache to potentially reuse this projection texture in higher LODs
-						ensure(LOD == 0);
+						ensure(LOD == GenerationContext.FirstLODAvailable);
 						float* AlternateProjectionResFactor = TextureNameToProjectionResFactor.Find(ImageName);
 						GenerationContext.GroupProjectorLODCache.Add(MaterialImageId,
 							FGroupProjectorImageInfo(ImageNodePtr, ImageName, ImageName, TypedNodeMat,
@@ -769,7 +769,7 @@ mu::NodeSurfacePtr GenerateMutableSourceSurface(const UEdGraphPin * Pin, FMutabl
 			}
 			else
 			{
-				ensure(LOD > 0);
+				ensure(LOD > GenerationContext.FirstLODAvailable);
 				check(ProjectorInfo->SurfNode->GetImage(ImageIndex) == ProjectorInfo->ImageNode);
 				SurfNode->SetImage(ImageIndex, ProjectorInfo->ImageNode);
 				SurfNode->SetImageName(ImageIndex, TCHAR_TO_ANSI(*ProjectorInfo->TextureName));
@@ -933,14 +933,14 @@ mu::NodeSurfacePtr GenerateMutableSourceSurface(const UEdGraphPin * Pin, FMutabl
 
 						if (ProjectorInfo)
 						{
-							ensure(LOD == 0);
+							ensure(LOD == GenerationContext.FirstLODAvailable);
 							ProjectorInfo->ImageResizeNode = NodeImageResize;
 							ProjectorInfo->bIsAlternateResolutionResized = true;
 						}
 					}
 					else
 					{
-						ensure(LOD > 0);
+						ensure(LOD > GenerationContext.FirstLODAvailable);
 						check(ProjectorInfo->bIsAlternateResolutionResized);
 						SurfNode2->SetImage(ImageIndex, ProjectorInfo->ImageResizeNode);
 					}
@@ -1052,7 +1052,7 @@ mu::NodeSurfacePtr GenerateMutableSourceSurface(const UEdGraphPin * Pin, FMutabl
 
 					if (ProjectorInfo)
 					{
-						ensure(LOD > 0);
+						ensure(LOD > GenerationContext.FirstLODAvailable);
 						check(ProjectorInfo->SurfNode->GetImage(ImageIndex) == ProjectorInfo->ImageNode);
 						ImageNode = ProjectorInfo->ImageNode;
 
diff --git a/Engine/Plugins/Experimental/Mutable/Source/MutableRuntime/Private/MuR/Mesh.cpp b/Engine/Plugins/Experimental/Mutable/Source/MutableRuntime/Private/MuR/Mesh.cpp
--- a/Engine/Plugins/Experimental/Mutable/Source/MutableRuntime/Private/MuR/Mesh.cpp
+++ b/Engine/Plugins/Experimental/Mutable/Source/MutableRuntime/Private/MuR/Mesh.cpp
@@ -73,6 +73,9 @@ MeshPtr Mesh::Clone() const
 	// Clone bone poses
 	pResult->BonePoses = BonePoses;
 
+	// Clone SkeletonIDs
+	pResult->SkeletonIDs = SkeletonIDs;
+
     return pResult;
 }
 
@@ -149,6 +152,12 @@ MeshPtr Mesh::Clone(EMeshCloneFlags Flags) const
 		pResult->BonePoses = BonePoses;
 	}
 
+	// Clone SkeletonIDs
+	if (EnumHasAnyFlags(Flags, EMeshCloneFlags::WithSkeletonIDs))
+	{
+		pResult->SkeletonIDs = SkeletonIDs;
+	}
+
     return pResult;
 }
 
@@ -514,6 +523,28 @@ void mu::Mesh::GetBoneTransform(int32 BoneIndex, FTransform3f& Transform) const
 }
 
 
+//---------------------------------------------------------------------------------------------
+int32 Mesh::GetSkeletonIDsCount() const
+{
+    return SkeletonIDs.Num();
+}
+
+
+//---------------------------------------------------------------------------------------------
+int32 Mesh::GetSkeletonID(int32 SkeletonIndex) const
+{
+	return SkeletonIDs.IsValidIndex(SkeletonIndex) ? SkeletonIDs[SkeletonIndex] : INDEX_NONE;
+}
+
+
+//---------------------------------------------------------------------------------------------
+void Mesh::AddSkeletonID(int32 SkeletonID)
+{
+	check(SkeletonID != INDEX_NONE);
+	SkeletonIDs.AddUnique(SkeletonID);
+}
+
+
 //---------------------------------------------------------------------------------------------
 //---------------------------------------------------------------------------------------------
 //---------------------------------------------------------------------------------------------
diff --git a/Engine/Plugins/Experimental/Mutable/Source/MutableRuntime/Private/MuR/OpMeshFormat.cpp b/Engine/Plugins/Experimental/Mutable/Source/MutableRuntime/Private/MuR/OpMeshFormat.cpp
--- a/Engine/Plugins/Experimental/Mutable/Source/MutableRuntime/Private/MuR/OpMeshFormat.cpp
+++ b/Engine/Plugins/Experimental/Mutable/Source/MutableRuntime/Private/MuR/OpMeshFormat.cpp
@@ -678,6 +678,8 @@ namespace mu
 
 		pResult->BonePoses = pSource->BonePoses;
 
+		pResult->SkeletonIDs = pSource->SkeletonIDs;
+
 		return pResult;
 
 	}
diff --git a/Engine/Plugins/Experimental/Mutable/Source/MutableRuntime/Private/MuR/OpMeshMerge.h b/Engine/Plugins/Experimental/Mutable/Source/MutableRuntime/Private/MuR/OpMeshMerge.h
--- a/Engine/Plugins/Experimental/Mutable/Source/MutableRuntime/Private/MuR/OpMeshMerge.h
+++ b/Engine/Plugins/Experimental/Mutable/Source/MutableRuntime/Private/MuR/OpMeshMerge.h
@@ -437,6 +437,14 @@ namespace mu
 		// Skeleton
 		//---------------------------
 
+		// Add SkeletonIDs
+		pResult->SkeletonIDs = pFirst->SkeletonIDs;
+
+		for (const int32 SkeletonID : pSecond->SkeletonIDs)
+		{
+			pResult->SkeletonIDs.AddUnique(SkeletonID);
+		}
+
         // Do they have the same skeleton?
         bool remapNeeded = pFirst->GetSkeleton() != pSecond->GetSkeleton();
 
@@ -485,7 +493,6 @@ namespace mu
                 {
                     tb = pResultSkeleton->m_bones.Num();
                     pResultSkeleton->m_bones.Add( pSecondSkeleton->m_bones[ob] );
-					pResultSkeleton->m_boneIds.Add(pSecondSkeleton->m_boneIds[ob]);
 					
 					// Add an incorrect index, to be fixed below in case the parent index is later in the bone array.
 					pResultSkeleton->m_boneParents.Add(pSecondSkeleton->m_boneParents[ob]);
@@ -1018,9 +1025,6 @@ namespace mu
 
                 // Will be remapped below
                 pBase->m_boneParents.Add(pOther->m_boneParents[b] );
-
-                // Copy the external bone id
-				pBase->m_boneIds.Add(pOther->m_boneIds[b]);
             }
             else
             {
diff --git a/Engine/Plugins/Experimental/Mutable/Source/MutableRuntime/Private/MuR/Skeleton.cpp b/Engine/Plugins/Experimental/Mutable/Source/MutableRuntime/Private/MuR/Skeleton.cpp
--- a/Engine/Plugins/Experimental/Mutable/Source/MutableRuntime/Private/MuR/Skeleton.cpp
+++ b/Engine/Plugins/Experimental/Mutable/Source/MutableRuntime/Private/MuR/Skeleton.cpp
@@ -34,7 +34,6 @@ namespace mu
 
 		pResult->m_bones = m_bones;
 		pResult->m_boneParents = m_boneParents;
-		pResult->m_boneIds = m_boneIds;
 
 		return pResult;
 	}
@@ -52,7 +51,6 @@ namespace mu
 	{
 		m_bones.SetNum(b);
 		m_boneParents.Init(INDEX_NONE, b);
-		m_boneIds.Init(INDEX_NONE, b);
 	}
 
 
@@ -104,34 +102,4 @@ namespace mu
 	}
 
 
-	//-------------------------------------------------------------------------------------------------
-	int32 Skeleton::GetBoneId(int32 boneIndex) const
-	{
-		if (boneIndex >= 0 && boneIndex < m_boneIds.Num())
-		{
-			return int(m_boneIds[boneIndex]);
-		}
-
-		return 0;
-	}
-
-
-	//-------------------------------------------------------------------------------------------------
-	void Skeleton::SetBoneId(int32 boneIndex, int32 id)
-	{
-		int boneCount = GetBoneCount();
-		check(boneIndex >= 0 && boneIndex < boneCount);
-		if (boneIndex >= 0 && boneIndex < boneCount)
-		{
-			// Enlarge the id buffer if necessary, but enlarge it to the total number of bones to avoid
-			// repetitive allocation.
-			if (boneIndex >= m_boneIds.Num())
-			{
-				m_boneIds.SetNumZeroed(boneCount);
-			}
-
-			m_boneIds[boneIndex] = id;
-		}
-	}
-
 }
\ No newline at end of file
diff --git a/Engine/Plugins/Experimental/Mutable/Source/MutableRuntime/Public/MuR/Mesh.h b/Engine/Plugins/Experimental/Mutable/Source/MutableRuntime/Public/MuR/Mesh.h
--- a/Engine/Plugins/Experimental/Mutable/Source/MutableRuntime/Public/MuR/Mesh.h
+++ b/Engine/Plugins/Experimental/Mutable/Source/MutableRuntime/Public/MuR/Mesh.h
@@ -125,7 +125,8 @@ namespace mu
 		WithFaceBuffers       = 1 << 9,
 		WithAdditionalBuffers = 1 << 10,
 		WithLayouts           = 1 << 11,
-		WithPoses			  = 1 << 12
+		WithPoses			  = 1 << 12,
+		WithSkeletonIDs		  = 1 << 13
 	};
 	
 	ENUM_CLASS_FLAGS(EMeshCloneFlags);
@@ -289,6 +290,15 @@ namespace mu
 		//! Return a matrix stored per bone. It is a set of 16-float values.
 		void GetBoneTransform(int32 BoneIndex, FTransform3f& Transform) const;
 
+		//!
+		int32 GetSkeletonIDsCount() const;
+
+		//!
+		int32 GetSkeletonID(int32 SkeletonIndex) const;
+
+		//!
+		void AddSkeletonID(int32 SkeletonID);
+
         //! \}
 
 
@@ -328,6 +338,9 @@ namespace mu
 
 		TArray<MESH_SURFACE> m_surfaces;
 
+		// Externally provided SkeletonIDs of the skeletons required by this mesh.
+		TArray<uint32> SkeletonIDs;
+
 		//! This skeleton and physics body are not owned and may be used by other meshes, so it cannot be modified
 		//! once the mesh has been fully created.
 		Ptr<const Skeleton> m_pSkeleton;
@@ -427,7 +440,7 @@ namespace mu
 		//!
 		inline void Serialise(OutputArchive& arch) const
 		{
-			uint32 ver = 13;
+			uint32 ver = 14;
 			arch << ver;
 
 			arch << m_IndexBuffers;
@@ -436,6 +449,8 @@ namespace mu
 			arch << m_AdditionalBuffers;
 			arch << m_layouts;
 
+			arch << SkeletonIDs;
+
 			arch << m_pSkeleton;
 			arch << m_pPhysicsBody;
 
@@ -453,7 +468,7 @@ namespace mu
 		{
 			uint32 ver;
 			arch >> ver;
-			check(ver <= 13);
+			check(ver <= 14);
 
 			arch >> m_IndexBuffers;
 			arch >> m_VertexBuffers;
@@ -461,6 +476,11 @@ namespace mu
 			arch >> m_AdditionalBuffers;
 			arch >> m_layouts;
 
+			if (ver >= 14)
+			{
+				arch >> SkeletonIDs;
+			}
+
 			arch >> m_pSkeleton;
 			if (ver >= 12)
 			{ 
@@ -519,6 +539,7 @@ namespace mu
 			if (equal) equal = (m_surfaces == o.m_surfaces);
 			if (equal) equal = (m_faceGroups == o.m_faceGroups);
 			if (equal) equal = (m_tags == o.m_tags);
+			if (equal) equal = (SkeletonIDs == o.SkeletonIDs);
 
 			for (int32 i = 0; equal && i < m_layouts.Num(); ++i)
 			{
diff --git a/Engine/Plugins/Experimental/Mutable/Source/MutableRuntime/Public/MuR/Skeleton.h b/Engine/Plugins/Experimental/Mutable/Source/MutableRuntime/Public/MuR/Skeleton.h
--- a/Engine/Plugins/Experimental/Mutable/Source/MutableRuntime/Public/MuR/Skeleton.h
+++ b/Engine/Plugins/Experimental/Mutable/Source/MutableRuntime/Public/MuR/Skeleton.h
@@ -61,12 +61,6 @@ namespace mu
         int32 GetBoneParent(int32 boneIndex) const;
         void SetBoneParent(int32 boneIndex, int32 parentBoneIndex);
 
-        //! Get and set an optional external Id for each bone that can be retrieved from generated
-        //! skeletons. This data is opaque to mutable, so it is copied around and merged as
-        //! necessary, but never changed. By default all indices are 0.
-        int32 GetBoneId(int32 boneIndex) const;
-        void SetBoneId(int32 boneIndex, int32 boneId);
-
 	protected:
 
 		//! Forbidden. Manage with the Ptr<> template.
@@ -82,20 +76,15 @@ namespace mu
 		//! This array must have the same size than the m_bones array.
 		TArray<int16> m_boneParents;
 
-		//! For each bone, an externally provided optional id. This vector can be empty or smaller
-		//! Than the number of bones, and missing bones are considered to have id 0.
-		TArray<int32> m_boneIds;
-
 
 		//!
 		inline void Serialise(OutputArchive& arch) const
 		{
-			uint32 ver = 4;
+			uint32 ver = 5;
 			arch << ver;
 
 			arch << m_bones;
 			arch << m_boneParents;
-			arch << m_boneIds;
 		}
 
 		//!
@@ -113,13 +102,15 @@ namespace mu
 			}
 
 			arch >> m_boneParents;
-			arch >> m_boneIds;
 
-			if (ver == 3)
+			if (ver <= 4)
 			{
-				// ensure m_boneIds' size is equal to the number of bones
-				m_boneIds.SetNum(m_bones.Num());
+				TArray<int32> boneIds;
+				arch >> boneIds;
+			}
 
+			if (ver == 3)
+			{
 				bool bBoneTransformModified;
 				arch >> bBoneTransformModified;
 			}
@@ -129,9 +120,8 @@ namespace mu
 		//!
 		inline bool operator==(const Skeleton& o) const
 		{
-			return m_bones == o.m_bones 
-				&& m_boneParents == o.m_boneParents
-				&& m_boneIds == o.m_boneIds;
+			return m_bones == o.m_bones
+				&& m_boneParents == o.m_boneParents;
 		}
 
 
