[Mutable] Disable CO Editor using ingame bounds by default.

#rnx

diff --git a/Engine/Plugins/Experimental/Mutable/Source/CustomizableObjectEditor/Private/MuCOE/CustomizableObjectEditor.cpp b/Engine/Plugins/Experimental/Mutable/Source/CustomizableObjectEditor/Private/MuCOE/CustomizableObjectEditor.cpp
--- a/Engine/Plugins/Experimental/Mutable/Source/CustomizableObjectEditor/Private/MuCOE/CustomizableObjectEditor.cpp
+++ b/Engine/Plugins/Experimental/Mutable/Source/CustomizableObjectEditor/Private/MuCOE/CustomizableObjectEditor.cpp
@@ -473,7 +473,6 @@ bool FCustomizableObjectEditor::CreatePreviewComponent(int32 ComponentIndex)
 
 		if (PreviewSkeletalMeshComponents[ComponentIndex])
 		{
-			PreviewSkeletalMeshComponents[ComponentIndex]->UseInGameBounds(true);
 			PreviewCustomizableSkeletalComponents[ComponentIndex]->AttachToComponent(PreviewSkeletalMeshComponents[ComponentIndex], FAttachmentTransformRules::KeepRelativeTransform);
 
 			return true;
