[mutable] Disable skeleton and physics reshape bones not found warnings.

- The expected usage is to input all possible bones for a CO and ignore the ones not present in an specific skeletal mesh.

- Warnig messages have been compacted. Now only one message is emited by node indicating a sublist and the total number of elements not found or valid.

#rnx

diff --git a/Engine/Plugins/Experimental/Mutable/Source/CustomizableObjectEditor/Private/MuCOE/GenerateMutableSource/GenerateMutableSourceMesh.cpp b/Engine/Plugins/Experimental/Mutable/Source/CustomizableObjectEditor/Private/MuCOE/GenerateMutableSource/GenerateMutableSourceMesh.cpp
--- a/Engine/Plugins/Experimental/Mutable/Source/CustomizableObjectEditor/Private/MuCOE/GenerateMutableSource/GenerateMutableSourceMesh.cpp
+++ b/Engine/Plugins/Experimental/Mutable/Source/CustomizableObjectEditor/Private/MuCOE/GenerateMutableSource/GenerateMutableSourceMesh.cpp
@@ -2068,6 +2068,9 @@ bool GetAndValidateReshapeBonesToDeform(
 {
 	bool bSetRefreshWarning = false;
 
+	TArray<uint8> MissingBones;
+	MissingBones.Init(true, InBonesToDeform.Num());
+
 	if(SelectionMethod == EBoneDeformSelectionMethod::ONLY_SELECTED)
 	{
 		int32 NumBonesToDeform = InBonesToDeform.Num();
@@ -2087,16 +2090,51 @@ bool GetAndValidateReshapeBonesToDeform(
 						OutBonesToDeform.AddUnique(BoneName.ToString());
 					}
 
-					bMissingBone = false;
+					MissingBones[InBoneIndex] &= false;
 					break;
 				}
 			}
+		}
 
-			if (bMissingBone)
+		constexpr bool bEmitWarnings = false;
+		// Don't emit wanings for now, the expected usage of the list is to include all possible bones for all meshes and
+		// ignore the ones that are not present in the specific mesh.
+		if (bEmitWarnings)
+		{
+			const auto MakeCompactMissingBoneListMessage = [&MissingBones, &InBonesToDeform]() -> FString
 			{
-				FString msg = FString::Printf(TEXT("Could not find the selected Bone to Deform [%s] in Skeleton"), *(BoneName.ToString()));
+				FString Msg = "";
 
-				GenerationContext.Compiler->CompilerLog(FText::FromString(msg), Node);
+				constexpr int32 MaxNumDisplayElems = 3;
+				int32 NumDisplayedElems = 0;
+
+				const int32 NumBones = InBonesToDeform.Num();
+				for (int32 IndexToDeform = 0; IndexToDeform < NumBones && NumDisplayedElems < MaxNumDisplayElems; ++IndexToDeform)
+				{
+					if (MissingBones[IndexToDeform])
+					{
+						Msg += (NumDisplayedElems == 0 ? " " : ", ") + InBonesToDeform[IndexToDeform].BoneName.ToString();
+						++NumDisplayedElems;
+					}
+				}
+
+				if (NumDisplayedElems >= MaxNumDisplayElems)
+				{
+					const int32 NumMissingBones = Algo::CountIf(MissingBones, [](const uint8& B) { return B; });
+					Msg += FString::Printf(TEXT(", ... and %d more"), NumMissingBones - NumDisplayedElems);
+				}
+
+				return Msg;
+			};
+
+			if (Algo::AnyOf(MissingBones, [](const uint8& B) { return B; }))
+			{
+				GenerationContext.Compiler->CompilerLog(
+					FText::FromString(
+						"Could not find the selected bones to deform " +
+						MakeCompactMissingBoneListMessage() +
+						" in the Skeleton."),
+					Node, EMessageSeverity::Warning);
 
 				bSetRefreshWarning = true;
 			}
@@ -2174,17 +2212,17 @@ bool GetAndValidateReshapeBonesToDeform(
 
 
 bool GetAndValidateReshapePhysicsToDeform(
-		TArray<FString>& OutPhysiscsToDeform,
-		const TArray<FMeshReshapeBoneReference>& InPhysicsToDeform,
-		const TArray<USkeletalMesh*>& SkeletalMeshes,
-		EBoneDeformSelectionMethod SelectionMethod,
-		const UCustomizableObjectNode* Node,
-		FMutableGraphGenerationContext& GenerationContext)
+	TArray<FString>& OutPhysiscsToDeform,
+	const TArray<FMeshReshapeBoneReference>& InPhysicsToDeform,
+	const TArray<USkeletalMesh*>& SkeletalMeshes,
+	EBoneDeformSelectionMethod SelectionMethod,
+	const UCustomizableObjectNode* Node,
+	FMutableGraphGenerationContext& GenerationContext)
 {
 
 	const bool bIsReferenceSkeletalMeshMethod =
-			SelectionMethod == EBoneDeformSelectionMethod::DEFORM_REF_SKELETON ||
-			SelectionMethod == EBoneDeformSelectionMethod::DEFORM_NONE_REF_SKELETON;
+		SelectionMethod == EBoneDeformSelectionMethod::DEFORM_REF_SKELETON ||
+		SelectionMethod == EBoneDeformSelectionMethod::DEFORM_NONE_REF_SKELETON;
 
 	// Get the participant bone names.
 	TArray<FName> BoneNamesInUserSelection;
@@ -2208,14 +2246,14 @@ bool GetAndValidateReshapePhysicsToDeform(
 
 	int32 NumUserSelectedBones = BoneNamesInUserSelection.Num();
 
-	struct FMissingBoneState
+	struct FMissingBoneStatus
 	{
 		uint8 bMissingBone : 1;
 		uint8 bMissingBody : 1;
 	};
 
-	TArray<FMissingBoneState> MissingBones;
-	MissingBones.Init(FMissingBoneState{ false, true }, NumUserSelectedBones);
+	TArray<FMissingBoneStatus> MissingBones;
+	MissingBones.Init(FMissingBoneStatus{ false, true }, NumUserSelectedBones);
 
 	for (const USkeletalMesh* SkeletalMesh : SkeletalMeshes)
 	{
@@ -2229,10 +2267,10 @@ bool GetAndValidateReshapePhysicsToDeform(
 			continue;
 		}
 
-		const FReferenceSkeleton& RefSkeleton = bIsReferenceSkeletalMeshMethod 
-				? GenerationContext.ComponentInfos[GenerationContext.CurrentMeshComponent].RefSkeletalMesh->GetRefSkeleton()
-				: SkeletalMesh->GetRefSkeleton();
-	
+		const FReferenceSkeleton& RefSkeleton = bIsReferenceSkeletalMeshMethod
+			? GenerationContext.ComponentInfos[GenerationContext.CurrentMeshComponent].RefSkeletalMesh->GetRefSkeleton()
+			: SkeletalMesh->GetRefSkeleton();
+
 		TArray<uint8> BoneInclusionSet;
 		BoneInclusionSet.Init(0, PhysicsAsset->SkeletalBodySetups.Num());
 
@@ -2240,10 +2278,14 @@ bool GetAndValidateReshapePhysicsToDeform(
 		for (int32 IndexToDeform = 0; IndexToDeform < NumUserSelectedBones; ++IndexToDeform)
 		{
 			const FName& BodyBoneName = BoneNamesInUserSelection[IndexToDeform];
+			const bool bBoneFound = RefSkeleton.FindBoneIndex(BodyBoneName) == INDEX_NONE;
+
 			MissingBones[IndexToDeform].bMissingBone = RefSkeleton.FindBoneIndex(BodyBoneName) == INDEX_NONE;
 
-			if (!MissingBones[IndexToDeform].bMissingBone)
+			if (!bBoneFound)
 			{
+				MissingBones[IndexToDeform].bMissingBone |= false;
+
 				const int32 FoundIndex = PhysicsAsset->SkeletalBodySetups.IndexOfByPredicate(
 					[&BodyBoneName](const TObjectPtr<USkeletalBodySetup>& Setup) {  return Setup->BoneName == BodyBoneName; });
 
@@ -2256,13 +2298,13 @@ bool GetAndValidateReshapePhysicsToDeform(
 		}
 
 		const bool bFlipSelection =
-				SelectionMethod == EBoneDeformSelectionMethod::ALL_BUT_SELECTED ||
-				SelectionMethod == EBoneDeformSelectionMethod::DEFORM_NONE_REF_SKELETON;
+			SelectionMethod == EBoneDeformSelectionMethod::ALL_BUT_SELECTED ||
+			SelectionMethod == EBoneDeformSelectionMethod::DEFORM_NONE_REF_SKELETON;
 		if (bFlipSelection)
 		{
-			for (uint8& Elem : BoneInclusionSet) 
-			{ 
-				Elem = 1 - Elem; 
+			for (uint8& Elem : BoneInclusionSet)
+			{
+				Elem = 1 - Elem;
 			}
 		}
 
@@ -2284,31 +2326,71 @@ bool GetAndValidateReshapePhysicsToDeform(
 		return false;
 	}
 
-	// Emit warnings if some explicitly selected bone is not present or has no phyiscs attached.
+	// Emit info message if some explicitly selected bone is not present or has no phyiscs attached.
+	// Usually the list of bones will contain bones referenced thruout the CO (the same list for all deforms.)
+
+	constexpr bool bEmitWarnings = false;
+
 	bool bSetRefreshWarning = false;
-	for (int32 IndexToDeform = 0; IndexToDeform < NumUserSelectedBones; ++IndexToDeform)
+	// Don't emit wanings for now, the expected usage of the list is to include all possible bones for all meshes and
+	// ignore the ones that are not present in the specific mesh.
+	if (bEmitWarnings)
 	{
-		if (MissingBones[IndexToDeform].bMissingBone)
+		const auto MakeCompactMissingBoneListMessage = [&MissingBones, &BoneNamesInUserSelection]
+		(auto&& MissingBonesStatusProjection) -> FString
 		{
-			FString BoneName = BoneNamesInUserSelection[IndexToDeform].ToString();
-			FString msg = FString::Printf(TEXT("Could not find the selected Physics Body bone to Deform [%s] in Skeleton"), *BoneName);
+			FString Msg = "";
 
-			GenerationContext.Compiler->CompilerLog(FText::FromString(msg), Node);
+			constexpr int32 MaxNumDisplayElems = 3;
+			int32 NumDisplayedElems = 0;
 
-			bSetRefreshWarning = true;
-		}
+			const int32 NumBones = BoneNamesInUserSelection.Num();
+			for (int32 IndexToDeform = 0; IndexToDeform < NumBones && NumDisplayedElems < MaxNumDisplayElems; ++IndexToDeform)
+			{
+				if (MissingBonesStatusProjection(MissingBones[IndexToDeform]))
+				{
+					Msg += (NumDisplayedElems == 0 ? " " : ", ") + BoneNamesInUserSelection[IndexToDeform].ToString();
+					++NumDisplayedElems;
+				}
+			}
+
+			if (NumDisplayedElems >= MaxNumDisplayElems)
+			{
+				const int32 NumMissingBones = Algo::CountIf(MissingBones, MissingBonesStatusProjection);
+				Msg += FString::Printf(TEXT(", ... and %d more"), NumMissingBones - NumDisplayedElems);
+			}
+
+			return Msg;
+		};
+
+		auto IsMissingBone = [](const FMissingBoneStatus& S) -> bool { return S.bMissingBone; };
+		auto IsMissingBody = [](const FMissingBoneStatus& S) -> bool { return S.bMissingBody; };
 
-		if (MissingBones[IndexToDeform].bMissingBody)
+		if (Algo::AnyOf(MissingBones, IsMissingBone))
 		{
-			FString BoneName = BoneNamesInUserSelection[IndexToDeform].ToString();
-			FString msg = FString::Printf(TEXT("Selected Bone to Deform [%s] does not have any physics body attached."), *BoneName);
+			GenerationContext.Compiler->CompilerLog(
+				FText::FromString(
+					"Could not find the selected physics bodies bones to deform " +
+					MakeCompactMissingBoneListMessage(IsMissingBone) +
+					" in the Skeleton."),
+				Node, EMessageSeverity::Warning);
 
-			GenerationContext.Compiler->CompilerLog(FText::FromString(msg), Node);
+			bSetRefreshWarning = true;
+		}
 
+		if (Algo::AnyOf(MissingBones, IsMissingBody))
+		{
+			GenerationContext.Compiler->CompilerLog(
+				FText::FromString(
+					"Selected Bones to deform " +
+					MakeCompactMissingBoneListMessage(IsMissingBody) +
+					" do not have any physics body attached."),
+				Node, EMessageSeverity::Warning);
+			
 			bSetRefreshWarning = true;
 		}
-	}
 
+	}
 	return bSetRefreshWarning;
 }
 
