[mutable] Re-structure virtual machine's layout operations so that they can be optimized against mesh operations slowing down progressive image generation.
Renamed some structs to Unreal's naming convention.
Fix optimization bug that could cause wrong operation reuse when sinking operations.
#rnx

diff --git a/Engine/Plugins/Experimental/Mutable/Source/CustomizableObject/Classes/MuCO/CustomizableObject.h b/Engine/Plugins/Experimental/Mutable/Source/CustomizableObject/Classes/MuCO/CustomizableObject.h
--- a/Engine/Plugins/Experimental/Mutable/Source/CustomizableObject/Classes/MuCO/CustomizableObject.h
+++ b/Engine/Plugins/Experimental/Mutable/Source/CustomizableObject/Classes/MuCO/CustomizableObject.h
@@ -1178,7 +1178,7 @@ private:
 	// This is a manual version number for the binary blobs in this asset.
 	// Increasing it invalidates all the previously compiled models.
 	// Warning: If while merging code both versions have changed, take the highest+1.
-	static const int32 CurrentSupportedVersion = 352;
+	static const int32 CurrentSupportedVersion = 353;
 
 public:
 
diff --git a/Engine/Plugins/Experimental/Mutable/Source/CustomizableObjectEditor/Private/MuCOE/SMutableCodeViewer.cpp b/Engine/Plugins/Experimental/Mutable/Source/CustomizableObjectEditor/Private/MuCOE/SMutableCodeViewer.cpp
--- a/Engine/Plugins/Experimental/Mutable/Source/CustomizableObjectEditor/Private/MuCOE/SMutableCodeViewer.cpp
+++ b/Engine/Plugins/Experimental/Mutable/Source/CustomizableObjectEditor/Private/MuCOE/SMutableCodeViewer.cpp
@@ -344,21 +344,25 @@ FString SMutableCodeViewer::GetReferencerName() const
 }
 
 
-void SMutableCodeViewer::Construct(const FArguments& InArgs, const mu::ModelPtr& InMutableModel/*, const TSharedPtr<SDockTab>& ConstructUnderMajorTab*/)
+void SMutableCodeViewer::SetCurrentModel(const mu::ModelPtr& InMutableModel)
 {
 	MutableModel = InMutableModel;
 	PreviewParameters = MutableModel->NewParameters();
-	
-	// Min width allowed for the column. Needed to avoid having issues with the constants space being to small
-	// and then getting too tall on the y axis crashing the drawer.
-	constexpr float MinParametersCollWidth = 400;
-	
-	// create & initialize tab manager
-	//TabManager = FGlobalTabmanager::Get()->NewTabManager(ConstructUnderMajorTab.ToSharedRef());
+
+	RootNodes.Empty();
+	RootNodeAddresses.Empty();
+	ItemCache.Empty();
+	MainItemPerOp.Empty();
+	TreeElements.Empty();
+	ExpandedElements.Empty();
+	FoundModelOperationTypeElements.Empty();
+	ModelOperationTypes.Empty();
+	ModelOperationTypeStrings.Empty();
+	NavigationElements.Empty();
 
 	// Generate all elements before starting the tree UI so we have a deterministic set of unique and duplicated elements
 	GenerateAllTreeElements();
-	
+
 	// Setup Navigation system
 	{
 		// Store the addresses of the root nodes for later search operations
@@ -369,15 +373,24 @@ void SMutableCodeViewer::Construct(const FArguments& InArgs, const mu::ModelPtr&
 
 		// Get an array of mutable types as an array of FStrings for the UI
 		GenerateNavigationOpTypeStrings();
-		
+
 		// Generate list elements for the found operation types so we are able to search over them on our type dropdown
 		GenerateNavigationDropdownElements();
 
 		// Check we did find types (witch should always happen in a normal run) and select the first option as default (NONE)
-		check (FoundModelOperationTypeElements.Num())
+		check(FoundModelOperationTypeElements.Num())
 		CurrentlySelectedOperationTypeElement = FoundModelOperationTypeElements[0];
 	}
+}
+
 
+void SMutableCodeViewer::Construct(const FArguments& InArgs, const mu::ModelPtr& InMutableModel/*, const TSharedPtr<SDockTab>& ConstructUnderMajorTab*/)
+{
+	// Min width allowed for the column. Needed to avoid having issues with the constants space being to small
+	// and then getting too tall on the y axis crashing the drawer.
+	constexpr float MinParametersCollWidth = 400;
+
+	SetCurrentModel(InMutableModel);
 	
 	FToolBarBuilder ToolbarBuilder(TSharedPtr<const FUICommandList>(), FMultiBoxCustomization::None, TSharedPtr<FExtender>(), true);
 	ToolbarBuilder.SetLabelVisibility(EVisibility::Visible);
@@ -2225,4 +2238,84 @@ void SMutableCodeViewer::PreviewMutableShape(const mu::FShape* Shape)
 	UE_LOG(LogMutable,Warning,TEXT("Previewer for Mutable Shapes not yet implemented"))
 }
 
+
+
+FReply SMutableCodeViewer::OnDragOver(const FGeometry& MyGeometry, const FDragDropEvent& DragDropEvent)
+{
+	if (TSharedPtr<FExternalDragOperation> DragDropOp = DragDropEvent.GetOperationAs<FExternalDragOperation>())
+	{
+		if (DragDropOp->HasFiles())
+		{
+			// For now, only allow a single file.
+			const TArray<FString>& Files = DragDropOp->GetFiles();
+			if (Files.Num() == 1)
+			{
+				const FString DraggedFileExtension = FPaths::GetExtension(Files[0], true);
+				if (DraggedFileExtension == TEXT(".mutable_compiled"))
+				{
+					// Dump source model to a file.
+					mu::InputFileStream stream(TCHAR_TO_ANSI(*Files[0]));
+
+					char MutableSourceTag[4] = {};
+					stream.Read(MutableSourceTag, 4);
+
+					if (!FMemory::Memcmp(MutableSourceTag, MUTABLE_COMPILED_MODEL_FILETAG, 4))
+					{
+						return FReply::Handled();
+					}
+
+					return FReply::Unhandled();
+				}
+			}
+		}
+	}
+
+	return FReply::Unhandled();
+}
+
+
+FReply SMutableCodeViewer::OnDrop(const FGeometry& MyGeometry, const FDragDropEvent& DragDropEvent)
+{
+	SCompoundWidget::OnDrop(MyGeometry, DragDropEvent);
+
+	if (TSharedPtr<FExternalDragOperation> DragDropOp = DragDropEvent.GetOperationAs<FExternalDragOperation>())
+	{
+		if (DragDropOp->HasFiles())
+		{
+			// For now, only allow a single file.
+			const TArray<FString>& Files = DragDropOp->GetFiles();
+			if (Files.Num() == 1)
+			{
+				const FString DraggedFileExtension = FPaths::GetExtension(Files[0], true);
+				if (DraggedFileExtension == TEXT(".mutable_compiled"))
+				{
+					// Dump source model to a file.
+					mu::InputFileStream stream(TCHAR_TO_ANSI(*Files[0]));
+
+					char MutableSourceTag[4] = {};
+					stream.Read(MutableSourceTag, 4);
+
+					if (!FMemory::Memcmp(MutableSourceTag, MUTABLE_COMPILED_MODEL_FILETAG, 4))
+					{
+						mu::InputArchive arch(&stream);
+						mu::ModelPtr Model = mu::Model::StaticUnserialise(arch);
+						SetCurrentModel(Model);
+
+						TreeView->RequestTreeRefresh();
+
+						return FReply::Handled();
+					}
+
+					return FReply::Unhandled();
+				}
+			}
+		}
+
+		return FReply::Unhandled();
+	}
+
+	return FReply::Unhandled();
+}
+
+
 #undef LOCTEXT_NAMESPACE
diff --git a/Engine/Plugins/Experimental/Mutable/Source/CustomizableObjectEditor/Private/MuCOE/SMutableCodeViewer.h b/Engine/Plugins/Experimental/Mutable/Source/CustomizableObjectEditor/Private/MuCOE/SMutableCodeViewer.h
--- a/Engine/Plugins/Experimental/Mutable/Source/CustomizableObjectEditor/Private/MuCOE/SMutableCodeViewer.h
+++ b/Engine/Plugins/Experimental/Mutable/Source/CustomizableObjectEditor/Private/MuCOE/SMutableCodeViewer.h
@@ -79,6 +79,8 @@ public:
 
 	// SWidget interface
 	virtual void Tick(const FGeometry& AllottedGeometry, const double InCurrentTime, const float InDeltaTime) override;
+	FReply OnDragOver(const FGeometry& MyGeometry, const FDragDropEvent& DragDropEvent) override;
+	FReply OnDrop(const FGeometry& MyGeometry, const FDragDropEvent& DragDropEvent) override;
 
 	// FGCObject interface
 	virtual void AddReferencedObjects(FReferenceCollector& Collector) override;
@@ -191,7 +193,9 @@ private:
 	/** Array with all the elements that have been manually expanded by the user */
 	TMap< mu::OP::ADDRESS, TSharedPtr<FMutableCodeTreeElement>> ExpandedElements;
 
-	
+	/** Prepare the widget for the given model. */
+	void SetCurrentModel(const mu::ModelPtr&);
+
 	/** Before any UI operation generate all the elements that may be navigable over the tree. No children of duplicated 
 	 * addresses will be generated.
 	 */
@@ -222,12 +226,16 @@ private:
 	{
 		mu::OP_TYPE::ME_BINDSHAPE,
 		mu::OP_TYPE::ME_MASKCLIPMESH,
+		mu::OP_TYPE::ME_FORMAT,
+		mu::OP_TYPE::ME_DIFFERENCE,
+		mu::OP_TYPE::IM_MAKEGROWMAP,
 	};
 
 	/** Collection with all expensive to run operation types */
 	const TArray<mu::OP_TYPE> ExpensiveOperationTypes
 	{
-		// none
+		mu::OP_TYPE::IM_PIXELFORMAT,
+		mu::OP_TYPE::ME_PROJECT,
 	};
 	
 	/** Enum designed to be able to notify the row generation of the type of operation being generated */
diff --git a/Engine/Plugins/Experimental/Mutable/Source/MutableRuntime/Private/MuR/CodeRunner.cpp b/Engine/Plugins/Experimental/Mutable/Source/MutableRuntime/Private/MuR/CodeRunner.cpp
--- a/Engine/Plugins/Experimental/Mutable/Source/MutableRuntime/Private/MuR/CodeRunner.cpp
+++ b/Engine/Plugins/Experimental/Mutable/Source/MutableRuntime/Private/MuR/CodeRunner.cpp
@@ -4883,7 +4883,7 @@ namespace mu
 
         case OP_TYPE::LA_CONSTANT:
         {
-            auto args = pModel->GetPrivate()->m_program.GetOpArgs<OP::ResourceConstantArgs>(item.at);
+			OP::ResourceConstantArgs args = pModel->GetPrivate()->m_program.GetOpArgs<OP::ResourceConstantArgs>(item.at);
             check( args.value < (uint32)pModel->GetPrivate()->m_program.m_constantLayouts.Num() );
 
             LayoutPtrConst pResult = pModel->GetPrivate()->m_program.m_constantLayouts
@@ -4894,19 +4894,19 @@ namespace mu
 
         case OP_TYPE::LA_MERGE:
         {
-            auto args = pModel->GetPrivate()->m_program.GetOpArgs<OP::LayoutMergeArgs>(item.at);
+			OP::LayoutMergeArgs args = pModel->GetPrivate()->m_program.GetOpArgs<OP::LayoutMergeArgs>(item.at);
             switch (item.stage)
             {
             case 0:
                     AddOp( SCHEDULED_OP( item.at, item, 1),
-                           SCHEDULED_OP( args.base, item),
-                           SCHEDULED_OP( args.added, item) );
+                           SCHEDULED_OP( args.Base, item),
+                           SCHEDULED_OP( args.Added, item) );
                 break;
 
             case 1:
             {
-                Ptr<const Layout> pA = GetMemory().GetLayout( CACHE_ADDRESS(args.base,item) );
-                Ptr<const Layout> pB = GetMemory().GetLayout( CACHE_ADDRESS(args.added,item) );
+                Ptr<const Layout> pA = GetMemory().GetLayout( CACHE_ADDRESS(args.Base,item) );
+                Ptr<const Layout> pB = GetMemory().GetLayout( CACHE_ADDRESS(args.Added,item) );
 
                 LayoutPtrConst pResult;
 
@@ -4936,17 +4936,17 @@ namespace mu
 
         case OP_TYPE::LA_PACK:
         {
-            auto args = pModel->GetPrivate()->m_program.GetOpArgs<OP::LayoutPackArgs>(item.at);
+			OP::LayoutPackArgs args = pModel->GetPrivate()->m_program.GetOpArgs<OP::LayoutPackArgs>(item.at);
             switch (item.stage)
             {
             case 0:
                     AddOp( SCHEDULED_OP( item.at, item, 1),
-                           SCHEDULED_OP( args.layout, item) );
+                           SCHEDULED_OP( args.Source, item) );
                 break;
 
             case 1:
             {
-                Ptr<const Layout> pSource = GetMemory().GetLayout( CACHE_ADDRESS(args.layout,item) );
+                Ptr<const Layout> pSource = GetMemory().GetLayout( CACHE_ADDRESS(args.Source,item) );
 
 				LayoutPtr pResult;
 
@@ -4976,39 +4976,75 @@ namespace mu
             break;
         }
 
-        case OP_TYPE::LA_REMOVEBLOCKS:
-        {
-            auto args = pModel->GetPrivate()->m_program.GetOpArgs<OP::LayoutRemoveBlocksArgs>(item.at);
-            switch (item.stage)
-            {
-            case 0:
-                    AddOp( SCHEDULED_OP( item.at, item, 1),
-                           SCHEDULED_OP( args.source, item),
-                           SCHEDULED_OP( args.mesh, item) );
-                break;
+		case OP_TYPE::LA_FROMMESH:
+		{
+			OP::LayoutFromMeshArgs args = pModel->GetPrivate()->m_program.GetOpArgs<OP::LayoutFromMeshArgs>(item.at);
+			switch (item.stage)
+			{
+			case 0:
+				AddOp(SCHEDULED_OP(item.at, item, 1),
+					SCHEDULED_OP(args.Mesh, item));
+				break;
 
-            case 1:
-            {
-                Ptr<const Layout> pSource = GetMemory().GetLayout( CACHE_ADDRESS(args.source,item) );
-                Ptr<const Mesh> pMesh = GetMemory().GetMesh( CACHE_ADDRESS(args.mesh,item) );
+			case 1:
+			{
+				Ptr<const Mesh> Mesh = GetMemory().GetMesh(CACHE_ADDRESS(args.Mesh, item));
 
-				LayoutPtr pResult;
+				Ptr<const Layout> Result;
 
-				if (pSource && pMesh)
+				if (Mesh && Mesh->GetLayoutCount()>int32(args.LayoutIndex))
 				{
-					pResult = LayoutRemoveBlocks(pSource.get(), pMesh.get(), args.meshLayoutIndex);
+					Result = Mesh->GetLayout( args.LayoutIndex );
 				}
 
-                GetMemory().SetLayout( item, pResult );
-                break;
-            }
+				GetMemory().SetLayout(item, Result);
+				break;
+			}
 
-            default:
-                check(false);
-            }
+			default:
+				check(false);
+			}
 
-            break;
-        }
+			break;
+		}
+
+		case OP_TYPE::LA_REMOVEBLOCKS:
+		{
+			OP::LayoutRemoveBlocksArgs args = pModel->GetPrivate()->m_program.GetOpArgs<OP::LayoutRemoveBlocksArgs>(item.at);
+			switch (item.stage)
+			{
+			case 0:
+				AddOp(SCHEDULED_OP(item.at, item, 1),
+					SCHEDULED_OP(args.Source, item),
+					SCHEDULED_OP(args.ReferenceLayout, item));
+				break;
+
+			case 1:
+			{
+				Ptr<const Layout> Source = GetMemory().GetLayout(CACHE_ADDRESS(args.Source, item));
+				Ptr<const Layout> ReferenceLayout = GetMemory().GetLayout(CACHE_ADDRESS(args.ReferenceLayout, item));
+
+				Ptr<const Layout> pResult;
+
+				if (Source && ReferenceLayout)
+				{
+					pResult = LayoutRemoveBlocks(Source.get(), ReferenceLayout.get());
+				}
+				else if (Source)
+				{
+					pResult = Source;
+				}
+
+				GetMemory().SetLayout(item, pResult);
+				break;
+			}
+
+			default:
+				check(false);
+			}
+
+			break;
+		}
 
         default:
             // Operation not implemented
diff --git a/Engine/Plugins/Experimental/Mutable/Source/MutableRuntime/Private/MuR/OpLayoutRemoveBlocks.h b/Engine/Plugins/Experimental/Mutable/Source/MutableRuntime/Private/MuR/OpLayoutRemoveBlocks.h
--- a/Engine/Plugins/Experimental/Mutable/Source/MutableRuntime/Private/MuR/OpLayoutRemoveBlocks.h
+++ b/Engine/Plugins/Experimental/Mutable/Source/MutableRuntime/Private/MuR/OpLayoutRemoveBlocks.h
@@ -11,62 +11,100 @@ namespace mu
 {
 
 	//---------------------------------------------------------------------------------------------
-    inline LayoutPtr LayoutRemoveBlocks( const Layout* pSource, const Mesh* pMesh, int layoutIndex )
-    {
-        // Create the list of blocks in the mesh
+	inline Ptr<Layout> LayoutRemoveBlocks_Deprecated(const Layout* pSource, const Mesh* pMesh, int layoutIndex)
+	{
+		// Create the list of blocks in the mesh
 		TArray<bool> blocksFound;
 		blocksFound.SetNumZeroed(1024);
 
-        UntypedMeshBufferIteratorConst itBlocks( pMesh->GetVertexBuffers(), MBS_LAYOUTBLOCK,
-                                                 layoutIndex );
-        if ( itBlocks.GetFormat()==MBF_UINT16 )
-        {
-            const uint16* pBlocks = reinterpret_cast<const uint16*>( itBlocks.ptr() );
-            for ( int i=0; i<pMesh->GetVertexCount(); ++i )
-            {
-                if (pBlocks[i]>=blocksFound.Num())
-                {
-                    blocksFound.SetNumZeroed(pBlocks[i]+1024);
-                }
-
-                blocksFound[ pBlocks[i] ] = true;
-            }
-        }
-        else if ( itBlocks.GetFormat()==MBF_NONE )
-        {
-            // This seems to happen.
-            // May this happen when entire meshes are removed?
-            return pSource->Clone();
-        }
-        else
-        {
-            // Format not supported yet
-            check( false );
-        }
-
-        // Remove blocks that are not in the mesh
-		LayoutPtr pResult = pSource->Clone();
+		UntypedMeshBufferIteratorConst itBlocks(pMesh->GetVertexBuffers(), MBS_LAYOUTBLOCK, layoutIndex);
+		if (itBlocks.GetFormat() == MBF_UINT16)
+		{
+			const uint16* pBlocks = reinterpret_cast<const uint16*>(itBlocks.ptr());
+			for (int i = 0; i < pMesh->GetVertexCount(); ++i)
+			{
+				if (pBlocks[i] >= blocksFound.Num())
+				{
+					blocksFound.SetNumZeroed(pBlocks[i] + 1024);
+				}
+
+				blocksFound[pBlocks[i]] = true;
+			}
+		}
+		else if (itBlocks.GetFormat() == MBF_NONE)
+		{
+			// This seems to happen.
+			// May this happen when entire meshes are removed?
+			return pSource->Clone();
+		}
+		else
+		{
+			// Format not supported yet
+			check(false);
+		}
+
+		// Remove blocks that are not in the mesh
+		Ptr<Layout> pResult = pSource->Clone();
+		int dest = 0;
+		for (int32 b = 0; b < pResult->m_blocks.Num(); ++b)
+		{
+			int blockIndex = pResult->m_blocks[b].m_id;
+			if (blockIndex < (int)blocksFound.Num() && blocksFound[blockIndex])
+			{
+				// keep
+				pResult->m_blocks[dest] = pResult->m_blocks[b];
+				++dest;
+			}
+		}
+		pResult->SetBlockCount(dest);
+
+		return pResult;
+	}
+
+
+	//---------------------------------------------------------------------------------------------
+	inline Ptr<Layout> LayoutRemoveBlocks(const Layout* Source, const Layout* ReferenceLayout)
+	{
+		// Create the list of blocks in the mesh
+		TArray<bool,TInlineAllocator<1024>> BlocksFound;
+		BlocksFound.SetNumZeroed(1024);
+
+		if (ReferenceLayout)
+		{
+			for (const Layout::FBlock& Block: ReferenceLayout->m_blocks)
+			{
+				if (Block.m_id >= BlocksFound.Num())
+				{
+					BlocksFound.SetNumZeroed(Block.m_id + 1024);
+				}
+
+				BlocksFound[Block.m_id] = true;
+			}
+		}
+
+		// Remove blocks that are not in the mesh
+		Ptr<Layout> pResult = Source->Clone();
 		int dest = 0;
-		for ( int32 b=0; b<pResult->m_blocks.Num(); ++b )
+		for (int32 b = 0; b < pResult->m_blocks.Num(); ++b)
 		{
-            int blockIndex = pResult->m_blocks[b].m_id;
-            if ( blockIndex<(int)blocksFound.Num() && blocksFound[blockIndex] )
+			int blockIndex = pResult->m_blocks[b].m_id;
+			if (blockIndex < (int)BlocksFound.Num() && BlocksFound[blockIndex])
 			{
 				// keep
 				pResult->m_blocks[dest] = pResult->m_blocks[b];
 				++dest;
 			}
 		}
-		pResult->SetBlockCount( dest );
+		pResult->SetBlockCount(dest);
 
-        return pResult;
+		return pResult;
 	}
 
 
 	//---------------------------------------------------------------------------------------------
-	inline LayoutPtr LayoutMerge(const Layout* pA, const Layout* pB )
+	inline Ptr<Layout> LayoutMerge(const Layout* pA, const Layout* pB )
 	{
-		LayoutPtr pResult = pA->Clone();
+		Ptr<Layout> pResult = pA->Clone();
 
 		// This is faster but fails in the rare case of a block being in both layouts, which may 
 		// happen if we merge a mesh with itself.
diff --git a/Engine/Plugins/Experimental/Mutable/Source/MutableRuntime/Private/MuR/Operations.cpp b/Engine/Plugins/Experimental/Mutable/Source/MutableRuntime/Private/MuR/Operations.cpp
--- a/Engine/Plugins/Experimental/Mutable/Source/MutableRuntime/Private/MuR/Operations.cpp
+++ b/Engine/Plugins/Experimental/Mutable/Source/MutableRuntime/Private/MuR/Operations.cpp
@@ -138,8 +138,9 @@ namespace mu
 
         { DT_LAYOUT,	true,   false,		{ 0, 0, 0, 0, 0, 0, 0, 0, 0, 0, 0, 0, 0, 0, 0, 0, 0, 0, 0, 0, 0, 0 }		},	// LA_PACK
         { DT_LAYOUT,	true,   false,		{ 0, 0, 0, 0, 0, 0, 0, 0, 0, 0, 0, 0, 0, 0, 0, 0, 0, 0, 0, 0, 0, 0 }		},	// LA_MERGE
-        { DT_LAYOUT,	true,   false,		{ 0, 0, 0, 0, 0, 0, 0, 0, 0, 0, 0, 0, 0, 0, 0, 0, 0, 0, 0, 0, 0, 0 }		},	// LA_REMOVEBLOCKS
-    };
+		{ DT_LAYOUT,	true,   false,		{ 0, 0, 0, 0, 0, 0, 0, 0, 0, 0, 0, 0, 0, 0, 0, 0, 0, 0, 0, 0, 0, 0 }		},	// LA_REMOVEBLOCKS
+		{ DT_LAYOUT,	true,   false,		{ 0, 0, 0, 0, 0, 0, 0, 0, 0, 0, 0, 0, 0, 0, 0, 0, 0, 0, 0, 0, 0, 0 }		},	// LA_FROMMESH
+};
 
     // clang-format on
 
@@ -422,21 +423,6 @@ namespace mu
 			f(&op.args.MeshProject.projector);
 			break;
 
-        //-------------------------------------------------------------------------------------
-        case OP_TYPE::LA_PACK:
-            f(&op.args.LayoutPack.layout );
-            break;
-
-        case OP_TYPE::LA_MERGE:
-            f(&op.args.LayoutMerge.base );
-            f(&op.args.LayoutMerge.added );
-            break;
-
-        case OP_TYPE::LA_REMOVEBLOCKS:
-            f(&op.args.LayoutRemoveBlocks.source );
-            f(&op.args.LayoutRemoveBlocks.mesh );
-            break;
-
         default:
 			check( false );
             break;
@@ -1162,25 +1148,32 @@ namespace mu
         case OP_TYPE::LA_PACK:
         {
 			OP::LayoutPackArgs args = program.GetOpArgs<OP::LayoutPackArgs>(at);
-            f(args.layout );
+            f(args.Source );
             break;
         }
 
         case OP_TYPE::LA_MERGE:
         {
 			OP::LayoutMergeArgs args = program.GetOpArgs<OP::LayoutMergeArgs>(at);
-            f(args.base );
-            f(args.added );
+            f(args.Base );
+            f(args.Added );
             break;
         }
 
-        case OP_TYPE::LA_REMOVEBLOCKS:
-        {
+		case OP_TYPE::LA_REMOVEBLOCKS:
+		{
 			OP::LayoutRemoveBlocksArgs args = program.GetOpArgs<OP::LayoutRemoveBlocksArgs>(at);
-            f(args.source );
-            f(args.mesh );
-            break;
-        }
+			f(args.Source);
+			f(args.ReferenceLayout);
+			break;
+		}
+
+		case OP_TYPE::LA_FROMMESH:
+		{
+			OP::LayoutFromMeshArgs args = program.GetOpArgs<OP::LayoutFromMeshArgs>(at);
+			f(args.Mesh);
+			break;
+		}
 
         default:
 			check( false );
diff --git a/Engine/Plugins/Experimental/Mutable/Source/MutableRuntime/Private/MuR/Operations.h b/Engine/Plugins/Experimental/Mutable/Source/MutableRuntime/Private/MuR/Operations.h
--- a/Engine/Plugins/Experimental/Mutable/Source/MutableRuntime/Private/MuR/Operations.h
+++ b/Engine/Plugins/Experimental/Mutable/Source/MutableRuntime/Private/MuR/Operations.h
@@ -73,7 +73,8 @@ namespace mu
         //-----------------------------------------------------------------------------------------
 
         //! Compare two scalars, return true if the first is less than the second.
-        BO_LESS,
+		//! \TODO: Deprecated?
+		BO_LESS,
 
         //! Compare an integerexpression with an integer constant
         BO_EQUAL_INT_CONST,
@@ -96,7 +97,8 @@ namespace mu
         //-----------------------------------------------------------------------------------------
 
         //! Multiply a scalar value by another onw and add a third one to the result
-        SC_MULTIPLYADD,
+		//! \TODO: Deprecated?
+		SC_MULTIPLYADD,
 
         //! Apply an arithmetic operation to two scalars
         SC_ARITHMETIC,
@@ -115,10 +117,12 @@ namespace mu
         CO_SWIZZLE,
 
         //! Make a "color" from the size of an image. Colours are used as a generic vector here.
-        CO_IMAGESIZE,
+		//! \TODO: Deprecated?
+		CO_IMAGESIZE,
 
         //! Encode a layout block transformation in a vector: position = (x,y) size = ( z, w )
-        CO_LAYOUTBLOCKTRANSFORM,
+		//! \TODO: Deprecated?
+		CO_LAYOUTBLOCKTRANSFORM,
 
         //! Compose a vector from 4 scalars
         CO_FROMSCALARS,
@@ -194,7 +198,8 @@ namespace mu
         IM_PLAINCOLOUR,
 
         //! Image resulting from a GPU program
-        IM_GPU,
+		//! \TODO: Deprecated?
+		IM_GPU,
 
         //! Cut a rect from an image
         IM_CROP,
@@ -252,7 +257,8 @@ namespace mu
         ME_MASKDIFF,
 
         //! Remove from a source mesh all the faces in common with another mesh
-        ME_SUBTRACT,
+		//! \TODO: Deprecated?
+		ME_SUBTRACT,
 
         //! Remove all the geometry selected by a mask.
         ME_REMOVEMASK,
@@ -264,7 +270,8 @@ namespace mu
         ME_EXTRACTLAYOUTBLOCK,
 
         //! Extract a fragment of a mesh containing a specific face group.
-        ME_EXTRACTFACEGROUP,
+		//! \TODO: Deprecated?
+		ME_EXTRACTFACEGROUP,
 
         //! Apply a transform in a 4x4 matrix to the geometry channels of the mesh
         ME_TRANSFORM,
@@ -291,6 +298,7 @@ namespace mu
 		ME_REMAPINDICES,
 
 		//! Apply a geometry core operation to a mesh.
+		//! \TODO: Deprecated?
 		ME_GEOMETRYOPERATION,
 
 		//! Calculate the binding of a mesh on a shape
@@ -347,6 +355,9 @@ namespace mu
         //! This operation is for the new way of managing layout blocks.
         LA_REMOVEBLOCKS,
 
+		//! Extract a layout from a mesh
+		LA_FROMMESH,
+
         //-----------------------------------------------------------------------------------------
         // Utility values
         //-----------------------------------------------------------------------------------------
@@ -966,24 +977,30 @@ namespace mu
         //-------------------------------------------------------------------------------------
         struct LayoutPackArgs
         {
-            ADDRESS layout;
+            ADDRESS Source;
         };
 
         struct LayoutMergeArgs
         {
-            ADDRESS base;
-            ADDRESS added;
+            ADDRESS Base;
+            ADDRESS Added;
         };
 
-        struct LayoutRemoveBlocksArgs
-        {
-            // Layout to be processed
-            ADDRESS source;
+		struct LayoutRemoveBlocksArgs
+		{
+			/** Layout to be processedand modified. */
+			ADDRESS Source;
 
-            // Source mesh to scan for active blocks
-            ADDRESS mesh;            
-            uint8 meshLayoutIndex;
-        };
+			/** Source layout to scan for active blocks. */
+			ADDRESS ReferenceLayout;
+		};
+
+		struct LayoutFromMeshArgs
+		{
+			/** Source mesh to retrieve the layout from. */
+			ADDRESS Mesh;
+			uint8 LayoutIndex;
+		};
 
 
         //-------------------------------------------------------------------------------------
@@ -1000,7 +1017,7 @@ namespace mu
         {
             if (o.type!=type) return false;
             //if (o.flags!=flags) return false;
-            if (FMemory::Memcmp(&o.args, &args, sizeof(ARGS))) return false;
+            if (FMemory::Memcmp(&o.args, &args, sizeof(FArgs))) return false;
             return true;
         }
 
@@ -1010,7 +1027,7 @@ namespace mu
             if (o.type>type) return false;
 //            if (o.flags<flags) return true;
 //            if (o.flags>flags) return false;
-            if (FMemory::Memcmp(&o.args, &args, sizeof(ARGS))<0) return true;
+            if (FMemory::Memcmp(&o.args, &args, sizeof(FArgs))<0) return true;
             return false;
         }
 
@@ -1093,14 +1110,9 @@ namespace mu
             MeshSetSkeletonArgs MeshSetSkeleton;
 			MeshProjectArgs MeshProject;
 
-            //-------------------------------------------------------------------------------------
-            LayoutPackArgs LayoutPack;
-            LayoutMergeArgs LayoutMerge;
-            LayoutRemoveBlocksArgs LayoutRemoveBlocks;
-
-        } ARGS;
+        } FArgs;
 
-        ARGS args;
+        FArgs args;
     };
 
     typedef t_OP<uint32> OP;
@@ -1110,7 +1122,7 @@ namespace mu
     MUTABLE_DEFINE_POD_VECTOR_SERIALISABLE( OP );
     
     // Check that we didn't go out of control with the operation size
-    static_assert( sizeof(OP::ARGS)==28, "Argument union has an unexpected size." );
+    static_assert( sizeof(OP::FArgs)==28, "Argument union has an unexpected size." );
     static_assert( sizeof(OP)==32, "Operation struct has an unexpected size." );
 
 
diff --git a/Engine/Plugins/Experimental/Mutable/Source/MutableRuntime/Public/MuR/Model.h b/Engine/Plugins/Experimental/Mutable/Source/MutableRuntime/Public/MuR/Model.h
--- a/Engine/Plugins/Experimental/Mutable/Source/MutableRuntime/Public/MuR/Model.h
+++ b/Engine/Plugins/Experimental/Mutable/Source/MutableRuntime/Public/MuR/Model.h
@@ -7,21 +7,20 @@
 #include "MuR/RefCounted.h"
 #include "Templates/Function.h"
 
-namespace mu { class InputArchive; }
-namespace mu { class ModelStreamer; }
-namespace mu { class OutputArchive; }
-
 
 //! This version number changes whenever there is a compatibility-breaking change in the Model
 //! data structures. Compiled models are not necessarily compatible when the runtime is updated,
 //! so this version number can be used externally to verify this. It is not used internally, and
 //! serializing models from different versions than this runtime will probably result in a crash.
-#define MUTABLE_COMPILED_MODEL_CODE_VERSION		uint32( 48 )
+#define MUTABLE_COMPILED_MODEL_CODE_VERSION		uint32( 49 )
 #define MUTABLE_PARAMETERS_VERSION              uint32( 1 )
 
 
 namespace mu
 {
+	class InputArchive;
+	class ModelStreamer;
+	class OutputArchive;
 
 	typedef Ptr<Parameters> ParametersPtr;
 	typedef Ptr<const Parameters> ParametersPtrConst;
diff --git a/Engine/Plugins/Experimental/Mutable/Source/MutableTools/Private/MuT/AST.cpp b/Engine/Plugins/Experimental/Mutable/Source/MutableTools/Private/MuT/AST.cpp
--- a/Engine/Plugins/Experimental/Mutable/Source/MutableTools/Private/MuT/AST.cpp
+++ b/Engine/Plugins/Experimental/Mutable/Source/MutableTools/Private/MuT/AST.cpp
@@ -916,16 +916,16 @@ bool ASTOp::IsColourConstant(FVector4f&) const
 }
 
 
-void ASTOp::GetBlockLayoutSize( int, int*, int*, BLOCK_LAYOUT_SIZE_CACHE* )
+void ASTOp::GetBlockLayoutSize( int, int*, int*, FBlockLayoutSizeCache* )
 {
     check(false);
 }
 
 void ASTOp::GetBlockLayoutSizeCached( int blockIndex, int* pBlockX, int* pBlockY,
-                               BLOCK_LAYOUT_SIZE_CACHE* cache )
+                               FBlockLayoutSizeCache* cache )
 {
-	BLOCK_LAYOUT_SIZE_CACHE::KeyType key(this,blockIndex);
-	BLOCK_LAYOUT_SIZE_CACHE::ValueType* ValuePtr = cache->Find( key );
+	FBlockLayoutSizeCache::KeyType key(this,blockIndex);
+	FBlockLayoutSizeCache::ValueType* ValuePtr = cache->Find( key );
     if (ValuePtr)
     {
         *pBlockX = ValuePtr->Key;
@@ -935,7 +935,7 @@ void ASTOp::GetBlockLayoutSizeCached( int blockIndex, int* pBlockX, int* pBlockY
 
     GetBlockLayoutSize( blockIndex, pBlockX, pBlockY, cache );
 
-	cache->Add(key, BLOCK_LAYOUT_SIZE_CACHE::ValueType( *pBlockX, *pBlockY) );
+	cache->Add(key, FBlockLayoutSizeCache::ValueType( *pBlockX, *pBlockY) );
 }
 
 
@@ -1290,45 +1290,6 @@ FImageDesc ASTOpFixed::GetImageDesc( bool returnBestOption, FGetImageDescContext
 }
 
 
-void ASTOpFixed::GetBlockLayoutSize( int blockIndex, int* pBlockX, int* pBlockY,
-                                     BLOCK_LAYOUT_SIZE_CACHE* cache )
-{
-    switch ( op.type )
-    {
-    case OP_TYPE::LA_PACK:
-    {
-        GetBlockLayoutSize( blockIndex, op.args.LayoutPack.layout,
-                            pBlockX, pBlockY, cache );
-        break;
-    }
-
-    case OP_TYPE::LA_MERGE:
-    {
-        GetBlockLayoutSize( blockIndex, op.args.LayoutMerge.base,
-                            pBlockX, pBlockY, cache );
-
-        if (!*pBlockX)
-        {
-            GetBlockLayoutSize( blockIndex, op.args.LayoutMerge.added,
-                                pBlockX, pBlockY, cache );
-        }
-
-        break;
-    }
-
-    case OP_TYPE::LA_REMOVEBLOCKS:
-    {
-        GetBlockLayoutSize( blockIndex, op.args.LayoutRemoveBlocks.source,
-                            pBlockX, pBlockY, cache );
-        break;
-    }
-
-    default:
-        checkf( false, TEXT("Instruction not supported") );
-    }
-}
-
-
 void ASTOpFixed::GetLayoutBlockSize( int* pBlockX, int* pBlockY )
 {
     switch ( op.type )
diff --git a/Engine/Plugins/Experimental/Mutable/Source/MutableTools/Private/MuT/AST.h b/Engine/Plugins/Experimental/Mutable/Source/MutableTools/Private/MuT/AST.h
--- a/Engine/Plugins/Experimental/Mutable/Source/MutableTools/Private/MuT/AST.h
+++ b/Engine/Plugins/Experimental/Mutable/Source/MutableTools/Private/MuT/AST.h
@@ -179,7 +179,7 @@ namespace mu
 
 
     //! Detailed optimization flags
-    struct MODEL_OPTIMIZATION_OPTIONS
+    struct FModelOptimizationOptions
     {
         bool m_enabled = true;
         bool m_optimiseOverlappedMasks = false;
@@ -212,6 +212,24 @@ namespace mu
 	};
 
 
+	//! For each operation we sink, the map from old instructions to new instructions.
+	struct FSinkerOldToNewKey
+	{
+		Ptr<const ASTOp> Op;
+		Ptr<const ASTOp> SinkingOp;
+
+		friend inline uint32 GetTypeHash(const FSinkerOldToNewKey& Key)
+		{
+			return HashCombine(::GetTypeHash(Key.Op.get()), ::GetTypeHash(Key.SinkingOp.get()));
+		}
+
+		friend inline bool operator==(const FSinkerOldToNewKey& A, const FSinkerOldToNewKey& B)
+		{
+			return A.Op==B.Op && A.SinkingOp==B.SinkingOp;
+		}
+	};
+
+
 	class Sink_ImageCropAST
 	{
 	public:
@@ -223,13 +241,11 @@ namespace mu
 		const ASTOp* m_root = nullptr;
 		Ptr<ASTOp> m_initialSource;
 		//! For each operation we sink, the map from old instructions to new instructions.
-		TMap<Ptr<const ASTOpFixed>, std::unordered_map<Ptr<ASTOp>, Ptr<ASTOp>>> m_oldToNew;
-		TArray<Ptr<ASTOp>> m_newOps;
+		TMap<FSinkerOldToNewKey, Ptr<ASTOp>> OldToNew;
 
 		Ptr<ASTOp> Visit(Ptr<ASTOp> at, const ASTOpFixed* currentCropOp);
 	};
 
-
 	//---------------------------------------------------------------------------------------------
 	//!
 	//---------------------------------------------------------------------------------------------
@@ -243,9 +259,7 @@ namespace mu
 
 		const class ASTOpImagePixelFormat* m_root = nullptr;
 		Ptr<ASTOp> m_initialSource;
-		//! For each operation we sink, the map from old instructions to new instructions.
-		TMap<Ptr<const ASTOp>, std::unordered_map<Ptr<ASTOp>, Ptr<ASTOp>>> m_oldToNew;
-		TArray<Ptr<ASTOp>> m_newOps;
+		TMap<FSinkerOldToNewKey, Ptr<ASTOp>> OldToNew;
 
 		Ptr<ASTOp> Visit(Ptr<ASTOp> at, const class ASTOpImagePixelFormat* currentFormatOp);
 	};
@@ -259,16 +273,15 @@ namespace mu
 	public:
 
 		// \TODO This is recursive and may cause stack overflows in big models.
-		mu::Ptr<ASTOp> Apply(const ASTOp* root);
+		Ptr<ASTOp> Apply(const ASTOp* root);
 
 	protected:
 
 		const class ASTOpMeshFormat* m_root = nullptr;
-		mu::Ptr<ASTOp> m_initialSource;
-		TMap<mu::Ptr<ASTOp>, mu::Ptr<ASTOp>> m_oldToNew;
-		TArray<mu::Ptr<ASTOp>> m_newOps;
+		Ptr<ASTOp> m_initialSource;
+		TMap<FSinkerOldToNewKey, Ptr<ASTOp>> OldToNew;
 
-		mu::Ptr<ASTOp> Visit(const mu::Ptr<ASTOp>& at, const class ASTOpMeshFormat* currentFormatOp);
+		Ptr<ASTOp> Visit(const Ptr<ASTOp>& at, const class ASTOpMeshFormat* currentFormatOp);
 	};
 
 
@@ -287,14 +300,13 @@ namespace mu
 		Ptr<ASTOp> m_initialSource;
 
 		//! For each operation we sink, the map from old instructions to new instructions.
-		TMap<Ptr<ASTOp>, Ptr<ASTOp>> m_oldToNew;
-		TArray<Ptr<ASTOp>> m_newOps;
+		TMap<FSinkerOldToNewKey, Ptr<ASTOp>> OldToNew;
 
 		Ptr<ASTOp> Visit(Ptr<ASTOp> at, const class ASTOpImageMipmap* currentMipmapOp);
 
 	};
 
-	struct OPTIMIZE_SINK_CONTEXT
+	struct FOptimizeSinkContext
 	{
 		Sink_ImageCropAST ImageCropSinker;
 		Sink_ImagePixelFormatAST ImagePixelFormatSinker;
@@ -435,8 +447,8 @@ namespace mu
 
         //!
         virtual Ptr<ASTOp> OptimiseSize() const { return nullptr; }
-        virtual Ptr<ASTOp> OptimiseSemantic(const MODEL_OPTIMIZATION_OPTIONS&) const { return nullptr; }
-        virtual Ptr<ASTOp> OptimiseSink(const MODEL_OPTIMIZATION_OPTIONS&, OPTIMIZE_SINK_CONTEXT& ) const { return nullptr; }
+        virtual Ptr<ASTOp> OptimiseSemantic(const FModelOptimizationOptions&) const { return nullptr; }
+        virtual Ptr<ASTOp> OptimiseSink(const FModelOptimizationOptions&, FOptimizeSinkContext& ) const { return nullptr; }
         virtual Ptr<ImageSizeExpression> GetImageSizeExpression() const
         {
             check( false );
@@ -671,13 +683,13 @@ namespace mu
         virtual FImageDesc GetImageDesc( bool returnBestOption=false, FGetImageDescContext* context=nullptr ) const;
 
         //! Optional cache struct to use int he method below.
-        using BLOCK_LAYOUT_SIZE_CACHE=TMap< const TPair<ASTOp*,int>, TPair<int,int>>;
+        using FBlockLayoutSizeCache=TMap< const TPair<ASTOp*,int>, TPair<int,int>>;
 
         //! Return the size in layout blocks of a particular block given by absolute index
         virtual void GetBlockLayoutSize( int blockIndex, int* pBlockX, int* pBlockY,
-                                         BLOCK_LAYOUT_SIZE_CACHE* cache );
+                                         FBlockLayoutSizeCache* cache );
         void GetBlockLayoutSizeCached( int blockIndex, int* pBlockX, int* pBlockY,
-                                       BLOCK_LAYOUT_SIZE_CACHE* cache );
+                                       FBlockLayoutSizeCache* cache );
 
 
         //! Return the size in pixels of the layout grid block for the image operation
@@ -970,12 +982,10 @@ namespace mu
         bool IsEqual(const ASTOp& otherUntyped) const override;
 		uint64 Hash() const override;
 		FImageDesc GetImageDesc( bool returnBestOption, FGetImageDescContext* context ) const override;
-        void GetBlockLayoutSize( int blockIndex, int* pBlockX, int* pBlockY,
-                                 BLOCK_LAYOUT_SIZE_CACHE* cache ) override;
         void GetLayoutBlockSize( int* pBlockX, int* pBlockY ) override;
         Ptr<ASTOp> OptimiseSize() const override;
-        Ptr<ASTOp> OptimiseSemantic(const MODEL_OPTIMIZATION_OPTIONS&) const override;
-        Ptr<ASTOp> OptimiseSink(const MODEL_OPTIMIZATION_OPTIONS&, OPTIMIZE_SINK_CONTEXT&) const override;
+        Ptr<ASTOp> OptimiseSemantic(const FModelOptimizationOptions&) const override;
+        Ptr<ASTOp> OptimiseSink(const FModelOptimizationOptions&, FOptimizeSinkContext&) const override;
         FBoolEvalResult EvaluateBool( ASTOpList& facts, FEvaluateBoolCache* cache ) const override;
         int EvaluateInt( ASTOpList& facts, bool &unknown ) const override;
         bool IsImagePlainConstant(FVector4f& colour ) const override;
@@ -1031,15 +1041,6 @@ namespace mu
             return res;
         }
 
-        inline void GetBlockLayoutSize( int blockIndex, OP::ADDRESS at, int* pBlockX, int* pBlockY,
-                                        BLOCK_LAYOUT_SIZE_CACHE* cache )
-        {
-            if (children[at])
-            {
-                children[at]->GetBlockLayoutSizeCached(blockIndex,pBlockX,pBlockY,cache);
-            }
-        }
-
         inline void GetLayoutBlockSize( OP::ADDRESS at, int* pBlockX, int* pBlockY )
         {
             if (children[at])
@@ -1049,7 +1050,7 @@ namespace mu
         }
 
         //! \todo move to ast swizzle op class when created
-        Ptr<ASTOp> OptimiseSwizzle( const MODEL_OPTIMIZATION_OPTIONS& ) const;
+        Ptr<ASTOp> OptimiseSwizzle( const FModelOptimizationOptions& ) const;
     };
 
 
diff --git a/Engine/Plugins/Experimental/Mutable/Source/MutableTools/Private/MuT/ASTOpConditional.cpp b/Engine/Plugins/Experimental/Mutable/Source/MutableTools/Private/MuT/ASTOpConditional.cpp
--- a/Engine/Plugins/Experimental/Mutable/Source/MutableTools/Private/MuT/ASTOpConditional.cpp
+++ b/Engine/Plugins/Experimental/Mutable/Source/MutableTools/Private/MuT/ASTOpConditional.cpp
@@ -214,7 +214,7 @@ namespace mu
 
 	//-------------------------------------------------------------------------------------------------
 	void ASTOpConditional::GetBlockLayoutSize(int blockIndex, int* pBlockX, int* pBlockY,
-		BLOCK_LAYOUT_SIZE_CACHE* cache)
+		FBlockLayoutSizeCache* cache)
 	{
 		if (type == OP_TYPE::LA_CONDITIONAL)
 		{
@@ -282,7 +282,7 @@ namespace mu
 
 
 	//-------------------------------------------------------------------------------------------------
-	mu::Ptr<ASTOp> ASTOpConditional::OptimiseSemantic(const MODEL_OPTIMIZATION_OPTIONS&) const
+	mu::Ptr<ASTOp> ASTOpConditional::OptimiseSemantic(const FModelOptimizationOptions&) const
 	{
 		if (!condition)
 		{
diff --git a/Engine/Plugins/Experimental/Mutable/Source/MutableTools/Private/MuT/ASTOpConditional.h b/Engine/Plugins/Experimental/Mutable/Source/MutableTools/Private/MuT/ASTOpConditional.h
--- a/Engine/Plugins/Experimental/Mutable/Source/MutableTools/Private/MuT/ASTOpConditional.h
+++ b/Engine/Plugins/Experimental/Mutable/Source/MutableTools/Private/MuT/ASTOpConditional.h
@@ -46,9 +46,8 @@ struct FProgram;
 		void Link(FProgram& program, const FLinkerOptions* Options) override;
 		FImageDesc GetImageDesc(bool returnBestOption, class FGetImageDescContext* context) const override;
 		void GetLayoutBlockSize(int* pBlockX, int* pBlockY) override;
-		void GetBlockLayoutSize(int blockIndex, int* pBlockX, int* pBlockY,
-			BLOCK_LAYOUT_SIZE_CACHE* cache) override;
-		Ptr<ASTOp> OptimiseSemantic(const MODEL_OPTIMIZATION_OPTIONS&) const override;
+		void GetBlockLayoutSize(int blockIndex, int* pBlockX, int* pBlockY, FBlockLayoutSizeCache* cache) override;
+		Ptr<ASTOp> OptimiseSemantic(const FModelOptimizationOptions&) const override;
 		bool GetNonBlackRect(FImageRect& maskUsage) const override;
 		Ptr<ImageSizeExpression> GetImageSizeExpression() const override;
 	};
diff --git a/Engine/Plugins/Experimental/Mutable/Source/MutableTools/Private/MuT/ASTOpConstantResource.cpp b/Engine/Plugins/Experimental/Mutable/Source/MutableTools/Private/MuT/ASTOpConstantResource.cpp
--- a/Engine/Plugins/Experimental/Mutable/Source/MutableTools/Private/MuT/ASTOpConstantResource.cpp
+++ b/Engine/Plugins/Experimental/Mutable/Source/MutableTools/Private/MuT/ASTOpConstantResource.cpp
@@ -200,7 +200,7 @@ namespace mu
 
 	//-------------------------------------------------------------------------------------------------
 	void ASTOpConstantResource::GetBlockLayoutSize(int blockIndex, int* pBlockX, int* pBlockY,
-		BLOCK_LAYOUT_SIZE_CACHE*)
+		FBlockLayoutSizeCache*)
 	{
 		switch (type)
 		{
diff --git a/Engine/Plugins/Experimental/Mutable/Source/MutableTools/Private/MuT/ASTOpConstantResource.h b/Engine/Plugins/Experimental/Mutable/Source/MutableTools/Private/MuT/ASTOpConstantResource.h
--- a/Engine/Plugins/Experimental/Mutable/Source/MutableTools/Private/MuT/ASTOpConstantResource.h
+++ b/Engine/Plugins/Experimental/Mutable/Source/MutableTools/Private/MuT/ASTOpConstantResource.h
@@ -67,7 +67,7 @@ template <class SCALAR> class vec4;
 		void Link(FProgram& program, const FLinkerOptions*) override;
 		FImageDesc GetImageDesc(bool, class FGetImageDescContext*) const override;
 		void GetBlockLayoutSize(int blockIndex, int* pBlockX, int* pBlockY,
-			BLOCK_LAYOUT_SIZE_CACHE* cache) override;
+			FBlockLayoutSizeCache* cache) override;
 		void GetLayoutBlockSize(int* pBlockX, int* pBlockY) override;
 		bool GetNonBlackRect(FImageRect& maskUsage) const override;
 		bool IsImagePlainConstant(FVector4f& colour) const override;
diff --git a/Engine/Plugins/Experimental/Mutable/Source/MutableTools/Private/MuT/ASTOpImageCompose.cpp b/Engine/Plugins/Experimental/Mutable/Source/MutableTools/Private/MuT/ASTOpImageCompose.cpp
--- a/Engine/Plugins/Experimental/Mutable/Source/MutableTools/Private/MuT/ASTOpImageCompose.cpp
+++ b/Engine/Plugins/Experimental/Mutable/Source/MutableTools/Private/MuT/ASTOpImageCompose.cpp
@@ -196,7 +196,7 @@ void ASTOpImageCompose::GetLayoutBlockSize(int* pBlockX, int* pBlockY)
 	if (Layout.child())
 	{
 		MUTABLE_CPUPROFILER_SCOPE(GetLayoutBlockSize_GetBlockLayoutSize);
-		BLOCK_LAYOUT_SIZE_CACHE cache;
+		FBlockLayoutSizeCache cache;
 		Layout->GetBlockLayoutSizeCached(BlockIndex,
 			&layoutBlocksX, &layoutBlocksY,
 			&cache);
@@ -217,7 +217,7 @@ void ASTOpImageCompose::GetLayoutBlockSize(int* pBlockX, int* pBlockY)
 }
 
 
-mu::Ptr<ASTOp> ASTOpImageCompose::OptimiseSemantic(const MODEL_OPTIMIZATION_OPTIONS& options) const
+mu::Ptr<ASTOp> ASTOpImageCompose::OptimiseSemantic(const FModelOptimizationOptions& options) const
 {
 	mu::Ptr<ASTOp> at;
 
diff --git a/Engine/Plugins/Experimental/Mutable/Source/MutableTools/Private/MuT/ASTOpImageCompose.h b/Engine/Plugins/Experimental/Mutable/Source/MutableTools/Private/MuT/ASTOpImageCompose.h
--- a/Engine/Plugins/Experimental/Mutable/Source/MutableTools/Private/MuT/ASTOpImageCompose.h
+++ b/Engine/Plugins/Experimental/Mutable/Source/MutableTools/Private/MuT/ASTOpImageCompose.h
@@ -36,8 +36,8 @@ template <class SCALAR> class vec4;
 		Ptr<ASTOp> Clone(MapChildFuncRef mapChild) const override;
 		void ForEachChild(const TFunctionRef<void(ASTChild&)>) override;
 		void Link(FProgram& program, const FLinkerOptions* Options) override;
-		//TODO: Ptr<ASTOp> OptimiseSink(const MODEL_OPTIMIZATION_OPTIONS& options, OPTIMIZE_SINK_CONTEXT& context) const override;
-		Ptr<ASTOp> OptimiseSemantic(const MODEL_OPTIMIZATION_OPTIONS& options) const;
+		//TODO: Ptr<ASTOp> OptimiseSink(const FModelOptimizationOptions& options, FOptimizeSinkContext& context) const override;
+		Ptr<ASTOp> OptimiseSemantic(const FModelOptimizationOptions& options) const;
 		FImageDesc GetImageDesc(bool returnBestOption, FGetImageDescContext* context) const override;
 		void GetLayoutBlockSize(int* pBlockX, int* pBlockY) override;
 		bool IsImagePlainConstant(FVector4f& colour) const override;
diff --git a/Engine/Plugins/Experimental/Mutable/Source/MutableTools/Private/MuT/ASTOpImageMipmap.cpp b/Engine/Plugins/Experimental/Mutable/Source/MutableTools/Private/MuT/ASTOpImageMipmap.cpp
--- a/Engine/Plugins/Experimental/Mutable/Source/MutableTools/Private/MuT/ASTOpImageMipmap.cpp
+++ b/Engine/Plugins/Experimental/Mutable/Source/MutableTools/Private/MuT/ASTOpImageMipmap.cpp
@@ -113,7 +113,7 @@ void ASTOpImageMipmap::Link( FProgram& program, const FLinkerOptions* )
 }
 
 
-mu::Ptr<ASTOp> ASTOpImageMipmap::OptimiseSink(const MODEL_OPTIMIZATION_OPTIONS& options, OPTIMIZE_SINK_CONTEXT& context) const
+mu::Ptr<ASTOp> ASTOpImageMipmap::OptimiseSink(const FModelOptimizationOptions& options, FOptimizeSinkContext& context) const
 {
 	mu::Ptr<ASTOp> at;
 
@@ -257,8 +257,7 @@ mu::Ptr<ASTOp> Sink_ImageMipmapAST::Apply(const ASTOp* root)
 	check(root->GetOpType() == OP_TYPE::IM_MIPMAP);
 
 	m_root = dynamic_cast<const ASTOpImageMipmap*>(root);
-	m_oldToNew.Empty();
-	m_newOps.Empty();
+	OldToNew.Empty();
 
 	if (m_root->bOnlyTail)
 	{
@@ -317,17 +316,11 @@ mu::Ptr<ASTOp> Sink_ImageMipmapAST::Visit(mu::Ptr<ASTOp> at, const ASTOpImageMip
 {
 	if (!at) return nullptr;
 
-	// Newly created?
-	if (m_newOps.Contains(at))
-	{
-		return at;
-	}
-
 	// Already visited?
-	auto cacheIt = m_oldToNew.Find(at);
-	if (cacheIt)
+	const Ptr<ASTOp>* Cached = OldToNew.Find({ at,currentMipmapOp });
+	if (Cached)
 	{
-		return *cacheIt;
+		return *Cached;
 	}
 
 	mu::Ptr<ASTOp> newAt = at;
@@ -450,7 +443,7 @@ mu::Ptr<ASTOp> Sink_ImageMipmapAST::Visit(mu::Ptr<ASTOp> at, const ASTOpImageMip
 		newAt = newOp;
 	}
 
-	m_oldToNew.Add(at, newAt);
+	OldToNew.Add({ at,currentMipmapOp }, newAt);
 
 	return newAt;
 }
diff --git a/Engine/Plugins/Experimental/Mutable/Source/MutableTools/Private/MuT/ASTOpImageMipmap.h b/Engine/Plugins/Experimental/Mutable/Source/MutableTools/Private/MuT/ASTOpImageMipmap.h
--- a/Engine/Plugins/Experimental/Mutable/Source/MutableTools/Private/MuT/ASTOpImageMipmap.h
+++ b/Engine/Plugins/Experimental/Mutable/Source/MutableTools/Private/MuT/ASTOpImageMipmap.h
@@ -47,7 +47,7 @@ template <class SCALAR> class vec4;
 		Ptr<ASTOp> Clone(MapChildFuncRef mapChild) const override;
 		void ForEachChild(const TFunctionRef<void(ASTChild&)>) override;
 		void Link(FProgram& program, const FLinkerOptions* Options) override;
-		Ptr<ASTOp> OptimiseSink(const MODEL_OPTIMIZATION_OPTIONS& options, OPTIMIZE_SINK_CONTEXT& context) const override;
+		Ptr<ASTOp> OptimiseSink(const FModelOptimizationOptions& options, FOptimizeSinkContext& context) const override;
 		FImageDesc GetImageDesc(bool returnBestOption, FGetImageDescContext* context) const override;
 		void GetLayoutBlockSize(int* pBlockX, int* pBlockY) override;
 		bool IsImagePlainConstant(FVector4f& colour) const override;
diff --git a/Engine/Plugins/Experimental/Mutable/Source/MutableTools/Private/MuT/ASTOpImagePixelFormat.cpp b/Engine/Plugins/Experimental/Mutable/Source/MutableTools/Private/MuT/ASTOpImagePixelFormat.cpp
--- a/Engine/Plugins/Experimental/Mutable/Source/MutableTools/Private/MuT/ASTOpImagePixelFormat.cpp
+++ b/Engine/Plugins/Experimental/Mutable/Source/MutableTools/Private/MuT/ASTOpImagePixelFormat.cpp
@@ -98,7 +98,7 @@ void ASTOpImagePixelFormat::Link( FProgram& program, const FLinkerOptions* )
 }
 
 
-mu::Ptr<ASTOp> ASTOpImagePixelFormat::OptimiseSink(const MODEL_OPTIMIZATION_OPTIONS& options, OPTIMIZE_SINK_CONTEXT& context) const
+mu::Ptr<ASTOp> ASTOpImagePixelFormat::OptimiseSink(const FModelOptimizationOptions& options, FOptimizeSinkContext& context) const
 {
 	mu::Ptr<ASTOp> at;
 
@@ -278,8 +278,7 @@ mu::Ptr<ImageSizeExpression> ASTOpImagePixelFormat::GetImageSizeExpression() con
 mu::Ptr<ASTOp> Sink_ImagePixelFormatAST::Apply(const ASTOp* root)
 {
 	m_root = dynamic_cast<const ASTOpImagePixelFormat*>(root);
-	m_oldToNew.Empty();
-	m_newOps.Empty();
+	OldToNew.Empty();
 
 	check(root->GetOpType() == OP_TYPE::IM_PIXELFORMAT);
 
@@ -303,24 +302,15 @@ mu::Ptr<ASTOp> Sink_ImagePixelFormatAST::Visit(mu::Ptr<ASTOp> at, const ASTOpIma
 {
 	if (!at) return nullptr;
 
-	// Newly created?
-	if (m_newOps.Find(at)!=INDEX_NONE)
-	{
-		return at;
-	}
-
 	EImageFormat format = currentFormatOp->Format;
 	bool isCompressedFormat = IsCompressedFormat(format);
 	bool isBlockFormat = GetImageFormatData(format).m_pixelsPerBlockX != 0;
 
 	// Already visited?
+	const Ptr<ASTOp>* Cached = OldToNew.Find({ at,currentFormatOp });
+	if (Cached)
 	{
-		auto& CurrentSet = m_oldToNew.FindOrAdd(currentFormatOp);
-		auto cacheIt = CurrentSet.find(at);
-		if (cacheIt != CurrentSet.end())
-		{
-			return cacheIt->second;
-		}
+		return *Cached;
 	}
 
 	mu::Ptr<ASTOp> newAt = at;
@@ -524,7 +514,7 @@ mu::Ptr<ASTOp> Sink_ImagePixelFormatAST::Visit(mu::Ptr<ASTOp> at, const ASTOpIma
 		newAt = newOp;
 	}
 
-	m_oldToNew[currentFormatOp][at] = newAt;
+	OldToNew.Add( { at, currentFormatOp }, newAt);
 
 	return newAt;
 }
diff --git a/Engine/Plugins/Experimental/Mutable/Source/MutableTools/Private/MuT/ASTOpImagePixelFormat.h b/Engine/Plugins/Experimental/Mutable/Source/MutableTools/Private/MuT/ASTOpImagePixelFormat.h
--- a/Engine/Plugins/Experimental/Mutable/Source/MutableTools/Private/MuT/ASTOpImagePixelFormat.h
+++ b/Engine/Plugins/Experimental/Mutable/Source/MutableTools/Private/MuT/ASTOpImagePixelFormat.h
@@ -34,7 +34,7 @@ template <class SCALAR> class vec4;
 		Ptr<ASTOp> Clone(MapChildFuncRef mapChild) const override;
 		void ForEachChild(const TFunctionRef<void(ASTChild&)>) override;
 		void Link(FProgram& program, const FLinkerOptions* Options) override;
-		Ptr<ASTOp> OptimiseSink(const MODEL_OPTIMIZATION_OPTIONS& options, OPTIMIZE_SINK_CONTEXT& context) const override;
+		Ptr<ASTOp> OptimiseSink(const FModelOptimizationOptions& options, FOptimizeSinkContext& context) const override;
 		FImageDesc GetImageDesc(bool returnBestOption, FGetImageDescContext* context) const override;
 		void GetLayoutBlockSize(int* pBlockX, int* pBlockY) override;
 		bool IsImagePlainConstant(FVector4f& colour) const override;
diff --git a/Engine/Plugins/Experimental/Mutable/Source/MutableTools/Private/MuT/ASTOpLayoutFromMesh.cpp b/Engine/Plugins/Experimental/Mutable/Source/MutableTools/Private/MuT/ASTOpLayoutFromMesh.cpp
new file mode 100644
--- /dev/null
+++ b/Engine/Plugins/Experimental/Mutable/Source/MutableTools/Private/MuT/ASTOpLayoutFromMesh.cpp
@@ -0,0 +1,271 @@
+// Copyright Epic Games, Inc. All Rights Reserved.
+
+#include "MuT/ASTOpLayoutFromMesh.h"
+
+#include "HAL/UnrealMemory.h"
+#include "Misc/AssertionMacros.h"
+#include "MuR/MemoryPrivate.h"
+#include "MuR/Mesh.h"
+#include "MuR/MeshBufferSet.h"
+#include "MuR/ModelPrivate.h"
+#include "MuR/MutableTrace.h"
+#include "MuR/RefCounted.h"
+#include "MuR/System.h"
+#include "MuR/Types.h"
+#include "MuT/ASTOpConditional.h"
+#include "MuT/ASTOpConstantResource.h"
+#include "MuT/ASTOpMeshClipMorphPlane.h"
+#include "MuT/ASTOpMeshRemoveMask.h"
+#include "MuT/ASTOpMeshMorph.h"
+#include "MuT/ASTOpMeshFormat.h"
+#include "MuT/ASTOpMeshApplyShape.h"
+#include "MuT/ASTOpMeshBindShape.h"
+#include "MuT/ASTOpSwitch.h"
+#include "MuT/ASTOpLayoutMerge.h"
+#include "MuT/StreamsPrivate.h"
+
+#include <functional>
+
+
+namespace mu
+{
+
+
+	//-------------------------------------------------------------------------------------------------
+	ASTOpLayoutFromMesh::ASTOpLayoutFromMesh()
+		: Mesh(this)
+	{
+	}
+
+
+	ASTOpLayoutFromMesh::~ASTOpLayoutFromMesh()
+	{
+		// Explicit call needed to avoid recursive destruction
+		ASTOp::RemoveChildren();
+	}
+
+
+	bool ASTOpLayoutFromMesh::IsEqual(const ASTOp& otherUntyped) const
+	{
+		if (const ASTOpLayoutFromMesh* other = dynamic_cast<const ASTOpLayoutFromMesh*>(&otherUntyped))
+		{
+			return Mesh == other->Mesh && LayoutIndex == other->LayoutIndex;
+		}
+		return false;
+	}
+
+
+	uint64 ASTOpLayoutFromMesh::Hash() const
+	{
+		uint64 res = std::hash<void*>()(Mesh.child().get());
+		hash_combine(res, LayoutIndex);
+		return res;
+	}
+
+
+	mu::Ptr<ASTOp> ASTOpLayoutFromMesh::Clone(MapChildFuncRef mapChild) const
+	{
+		mu::Ptr<ASTOpLayoutFromMesh> n = new ASTOpLayoutFromMesh();
+		n->Mesh = mapChild(Mesh.child());
+		n->LayoutIndex = LayoutIndex;
+		return n;
+	}
+
+
+	void ASTOpLayoutFromMesh::ForEachChild(const TFunctionRef<void(ASTChild&)> f)
+	{
+		f(Mesh);
+	}
+
+
+	void ASTOpLayoutFromMesh::Link(FProgram& program, const FLinkerOptions*)
+	{
+		// Already linked?
+		if (!linkedAddress)
+		{
+			OP::LayoutFromMeshArgs args;
+			FMemory::Memzero(&args, sizeof(args));
+
+			args.LayoutIndex = LayoutIndex;
+			if (Mesh) args.Mesh = Mesh->linkedAddress;
+
+			linkedAddress = (OP::ADDRESS)program.m_opAddress.Num();
+			//program.m_code.push_back(op);
+			program.m_opAddress.Add((uint32_t)program.m_byteCode.Num());
+			AppendCode(program.m_byteCode, GetOpType());
+			AppendCode(program.m_byteCode, args);
+		}
+
+	}
+
+
+	void ASTOpLayoutFromMesh::GetBlockLayoutSize(int blockIndex, int* pBlockX, int* pBlockY, FBlockLayoutSizeCache* cache)
+	{
+		// This shouldn't happen for this operation because it is always in a branch of layout operations that is not the main one.
+		check(false);
+	}
+
+
+	namespace
+	{
+
+		/** Handle the optimization of a ASTOpLayoutFromMesh operation by moving it down its subtree. */
+		class FSinkLayoutFromMesh
+		{
+		public:
+
+			// \TODO This is recursive and may cause stack overflows in big models.
+			mu::Ptr<ASTOp> Apply(const ASTOpLayoutFromMesh* InRoot)
+			{
+				InitialRoot = InRoot;
+				check(InitialRoot);
+
+				OldToNew.Empty();
+
+				InitialSource = InitialRoot->Mesh.child();
+				mu::Ptr<ASTOp> NewSource = Visit(InitialSource, InitialRoot);
+
+				// If there is any change, it is the new root.
+				if (NewSource != InitialSource)
+				{
+					return NewSource;
+				}
+
+				return nullptr;
+			}
+
+
+		protected:
+
+			const class ASTOpLayoutFromMesh* InitialRoot = nullptr;
+			Ptr<ASTOp> InitialSource;
+			TMap<FSinkerOldToNewKey, Ptr<ASTOp>> OldToNew;
+
+			mu::Ptr<ASTOp> Visit(const mu::Ptr<ASTOp>& at, const ASTOpLayoutFromMesh* CurrentSinkingOp)
+			{
+				if (!at) return nullptr;
+
+				// Already visited?
+				const Ptr<ASTOp>* Cached = OldToNew.Find({ at,CurrentSinkingOp });
+				if (Cached)
+				{
+					return *Cached;
+				}
+
+				mu::Ptr<ASTOp> newAt = at;
+				switch (at->GetOpType())
+				{
+
+				case OP_TYPE::ME_MORPH2:
+				{
+					// Sink, ignoring the op
+					const ASTOpMeshMorph* Typed = dynamic_cast<const ASTOpMeshMorph*>(at.get());
+					check(Typed);
+					newAt = Visit(Typed->Base.child(), CurrentSinkingOp);
+					break;
+				}
+
+				case OP_TYPE::ME_FORMAT:
+				{
+					// Sink, ignoring the op
+					const ASTOpMeshFormat* Typed = dynamic_cast<const ASTOpMeshFormat*>(at.get());
+					check(Typed);
+					newAt = Visit(Typed->Source.child(), CurrentSinkingOp);
+					break;
+				}
+
+				case OP_TYPE::ME_APPLYSHAPE:
+				{
+					// Sink, ignoring the op
+					const ASTOpMeshApplyShape* Typed = dynamic_cast<const ASTOpMeshApplyShape*>(at.get());
+					check(Typed);
+					newAt = Visit(Typed->Mesh.child(), CurrentSinkingOp);
+					break;
+				}
+
+				case OP_TYPE::ME_BINDSHAPE:
+				{
+					// Sink, ignoring the op
+					const ASTOpMeshBindShape* Typed = dynamic_cast<const ASTOpMeshBindShape*>(at.get());
+					check(Typed);
+					newAt = Visit(Typed->Mesh.child(), CurrentSinkingOp);
+					break;
+				}
+
+				case OP_TYPE::ME_INTERPOLATE:
+				{
+					// Sink, ignoring the op
+					const ASTOpFixed* Typed = dynamic_cast<const ASTOpFixed*>(at.get());
+					check(Typed);
+					newAt = Visit(Typed->children[Typed->op.args.MeshInterpolate.base].child(), CurrentSinkingOp);
+					break;
+				}
+
+				case OP_TYPE::ME_CONDITIONAL:
+				{
+					Ptr<ASTOpConditional> NewConditional = mu::Clone<ASTOpConditional>(at);
+					NewConditional->type = OP_TYPE::LA_CONDITIONAL;
+					NewConditional->yes = Visit(NewConditional->yes.child(), CurrentSinkingOp);
+					NewConditional->no = Visit(NewConditional->no.child(), CurrentSinkingOp);
+					newAt = NewConditional;
+					break;
+				}
+
+				case OP_TYPE::ME_SWITCH:
+				{
+					Ptr<ASTOpSwitch> NewOp = mu::Clone<ASTOpSwitch>(at);
+					NewOp->type = OP_TYPE::LA_SWITCH;
+					NewOp->def = Visit(NewOp->def.child(), CurrentSinkingOp);
+					for (ASTOpSwitch::FCase& c : NewOp->cases)
+					{
+						c.branch = Visit(c.branch.child(), CurrentSinkingOp);
+					}
+					newAt = NewOp;
+					break;
+				}
+
+				case OP_TYPE::ME_MERGE:
+				{
+					const ASTOpFixed* Typed = dynamic_cast<const ASTOpFixed*>(at.get());
+					check(Typed);
+
+					Ptr<ASTOpLayoutMerge> NewMerge = new ASTOpLayoutMerge;
+					NewMerge->Base = Visit(Typed->children[Typed->op.args.MeshMerge.base].child(), CurrentSinkingOp);
+					NewMerge->Added = Visit(Typed->children[Typed->op.args.MeshMerge.added].child(), CurrentSinkingOp);
+					newAt = NewMerge;
+					break;
+				}
+
+				default:
+					if (at != InitialSource)
+					{
+						mu::Ptr<ASTOpLayoutFromMesh> NewOp = mu::Clone<ASTOpLayoutFromMesh>(CurrentSinkingOp);
+						NewOp->Mesh = at;
+						newAt = NewOp;
+					}
+					break;
+
+				}
+
+				OldToNew.Add({ at,CurrentSinkingOp }, newAt);
+
+				return newAt;
+			}
+
+		};
+
+	}
+
+
+	mu::Ptr<ASTOp> ASTOpLayoutFromMesh::OptimiseSink(const FModelOptimizationOptions& options, FOptimizeSinkContext& context) const
+	{
+		MUTABLE_CPUPROFILER_SCOPE(ASTOpLayoutFromMesh_Sink);
+
+		FSinkLayoutFromMesh Sinker;
+		mu::Ptr<ASTOp> at = Sinker.Apply(this);
+
+		return at;
+	}
+
+}
+
diff --git a/Engine/Plugins/Experimental/Mutable/Source/MutableTools/Private/MuT/ASTOpLayoutFromMesh.h b/Engine/Plugins/Experimental/Mutable/Source/MutableTools/Private/MuT/ASTOpLayoutFromMesh.h
new file mode 100644
--- /dev/null
+++ b/Engine/Plugins/Experimental/Mutable/Source/MutableTools/Private/MuT/ASTOpLayoutFromMesh.h
@@ -0,0 +1,40 @@
+// Copyright Epic Games, Inc. All Rights Reserved.
+
+#pragma once
+
+#include "HAL/PlatformMath.h"
+#include "MuR/Operations.h"
+#include "MuR/Ptr.h"
+#include "MuT/AST.h"
+
+
+namespace mu
+{
+	struct FProgram;
+	
+	class ASTOpLayoutFromMesh : public ASTOp
+	{
+	public:
+
+		ASTChild Mesh;
+		uint8_t LayoutIndex = 0;
+
+	public:
+
+		ASTOpLayoutFromMesh();
+		ASTOpLayoutFromMesh(const ASTOpLayoutFromMesh&) = delete;
+		~ASTOpLayoutFromMesh();
+
+		OP_TYPE GetOpType() const override { return OP_TYPE::LA_FROMMESH; }
+		uint64 Hash() const override;
+		bool IsEqual(const ASTOp& otherUntyped) const override;
+		Ptr<ASTOp> Clone(MapChildFuncRef mapChild) const override;
+		void ForEachChild(const TFunctionRef<void(ASTChild&)>) override;
+		void Link(FProgram& program, const FLinkerOptions* Options) override;
+		Ptr<ASTOp> OptimiseSink(const FModelOptimizationOptions& options, FOptimizeSinkContext& context) const override;
+		void GetBlockLayoutSize(int blockIndex, int* pBlockX, int* pBlockY, FBlockLayoutSizeCache* cache) override;
+
+	};
+
+}
+
diff --git a/Engine/Plugins/Experimental/Mutable/Source/MutableTools/Private/MuT/ASTOpLayoutMerge.cpp b/Engine/Plugins/Experimental/Mutable/Source/MutableTools/Private/MuT/ASTOpLayoutMerge.cpp
new file mode 100644
--- /dev/null
+++ b/Engine/Plugins/Experimental/Mutable/Source/MutableTools/Private/MuT/ASTOpLayoutMerge.cpp
@@ -0,0 +1,112 @@
+// Copyright Epic Games, Inc. All Rights Reserved.
+
+#include "MuT/ASTOpLayoutMerge.h"
+
+#include "HAL/UnrealMemory.h"
+#include "Misc/AssertionMacros.h"
+#include "MuR/MemoryPrivate.h"
+#include "MuR/Mesh.h"
+#include "MuR/MeshBufferSet.h"
+#include "MuR/ModelPrivate.h"
+#include "MuR/MutableTrace.h"
+#include "MuR/RefCounted.h"
+#include "MuR/System.h"
+#include "MuR/Types.h"
+#include "MuT/ASTOpConditional.h"
+#include "MuT/ASTOpConstantResource.h"
+#include "MuT/ASTOpMeshClipMorphPlane.h"
+#include "MuT/ASTOpMeshRemoveMask.h"
+#include "MuT/ASTOpMeshMorph.h"
+#include "MuT/ASTOpSwitch.h"
+#include "MuT/StreamsPrivate.h"
+
+
+
+namespace mu
+{
+
+
+	//-------------------------------------------------------------------------------------------------
+	ASTOpLayoutMerge::ASTOpLayoutMerge()
+		: Base(this)
+		, Added(this)
+	{
+	}
+
+
+	ASTOpLayoutMerge::~ASTOpLayoutMerge()
+	{
+		// Explicit call needed to avoid recursive destruction
+		ASTOp::RemoveChildren();
+	}
+
+
+	bool ASTOpLayoutMerge::IsEqual(const ASTOp& otherUntyped) const
+	{
+		if (const ASTOpLayoutMerge* other = dynamic_cast<const ASTOpLayoutMerge*>(&otherUntyped))
+		{
+			return Base == other->Base && Added == other->Added;
+		}
+		return false;
+	}
+
+
+	uint64 ASTOpLayoutMerge::Hash() const
+	{
+		uint64 res = std::hash<void*>()(Base.child().get());
+		hash_combine(res, Added.child().get());
+		return res;
+	}
+
+
+	mu::Ptr<ASTOp> ASTOpLayoutMerge::Clone(MapChildFuncRef mapChild) const
+	{
+		mu::Ptr<ASTOpLayoutMerge> n = new ASTOpLayoutMerge();
+		n->Base = mapChild(Base.child());
+		n->Added = mapChild(Added.child());
+		return n;
+	}
+
+
+	void ASTOpLayoutMerge::ForEachChild(const TFunctionRef<void(ASTChild&)> f)
+	{
+		f(Base);
+		f(Added);
+	}
+
+
+	void ASTOpLayoutMerge::Link(FProgram& program, const FLinkerOptions*)
+	{
+		// Already linked?
+		if (!linkedAddress)
+		{
+			OP::LayoutMergeArgs args;
+			FMemory::Memzero(&args, sizeof(args));
+
+			if (Base) args.Base = Base->linkedAddress;
+			if (Added) args.Added = Added->linkedAddress;
+
+			linkedAddress = (OP::ADDRESS)program.m_opAddress.Num();
+			program.m_opAddress.Add((uint32_t)program.m_byteCode.Num());
+			AppendCode(program.m_byteCode, GetOpType());
+			AppendCode(program.m_byteCode, args);
+		}
+
+	}
+
+
+	void ASTOpLayoutMerge::GetBlockLayoutSize(int blockIndex, int* pBlockX, int* pBlockY, FBlockLayoutSizeCache* cache)
+	{
+		if (Base)
+		{
+			Base->GetBlockLayoutSize(blockIndex, pBlockX, pBlockY, cache);
+		}
+
+		if (!*pBlockX && Added)
+		{
+			Added->GetBlockLayoutSize(blockIndex, pBlockX, pBlockY, cache);
+		}
+	}
+
+
+}
diff --git a/Engine/Plugins/Experimental/Mutable/Source/MutableTools/Private/MuT/ASTOpLayoutMerge.h b/Engine/Plugins/Experimental/Mutable/Source/MutableTools/Private/MuT/ASTOpLayoutMerge.h
new file mode 100644
--- /dev/null
+++ b/Engine/Plugins/Experimental/Mutable/Source/MutableTools/Private/MuT/ASTOpLayoutMerge.h
@@ -0,0 +1,39 @@
+// Copyright Epic Games, Inc. All Rights Reserved.
+
+#pragma once
+
+#include "HAL/PlatformMath.h"
+#include "MuR/Operations.h"
+#include "MuR/Ptr.h"
+#include "MuT/AST.h"
+
+
+namespace mu
+{
+	struct FProgram;
+	
+	class ASTOpLayoutMerge : public ASTOp
+	{
+	public:
+
+		ASTChild Base;
+		ASTChild Added;
+
+	public:
+
+		ASTOpLayoutMerge();
+		ASTOpLayoutMerge(const ASTOpLayoutMerge&) = delete;
+		~ASTOpLayoutMerge();
+
+		OP_TYPE GetOpType() const override { return OP_TYPE::LA_MERGE; }
+		uint64 Hash() const override;
+		bool IsEqual(const ASTOp& otherUntyped) const override;
+		Ptr<ASTOp> Clone(MapChildFuncRef mapChild) const override;
+		void ForEachChild(const TFunctionRef<void(ASTChild&)>) override;
+		void Link(FProgram& program, const FLinkerOptions* Options) override;
+		void GetBlockLayoutSize(int blockIndex, int* pBlockX, int* pBlockY, FBlockLayoutSizeCache* cache) override;
+
+	};
+
+}
+
diff --git a/Engine/Plugins/Experimental/Mutable/Source/MutableTools/Private/MuT/ASTOpLayoutPack.cpp b/Engine/Plugins/Experimental/Mutable/Source/MutableTools/Private/MuT/ASTOpLayoutPack.cpp
new file mode 100644
--- /dev/null
+++ b/Engine/Plugins/Experimental/Mutable/Source/MutableTools/Private/MuT/ASTOpLayoutPack.cpp
@@ -0,0 +1,102 @@
+// Copyright Epic Games, Inc. All Rights Reserved.
+
+#include "MuT/ASTOpLayoutPack.h"
+
+#include "HAL/UnrealMemory.h"
+#include "Misc/AssertionMacros.h"
+#include "MuR/MemoryPrivate.h"
+#include "MuR/Mesh.h"
+#include "MuR/MeshBufferSet.h"
+#include "MuR/ModelPrivate.h"
+#include "MuR/MutableTrace.h"
+#include "MuR/RefCounted.h"
+#include "MuR/System.h"
+#include "MuR/Types.h"
+#include "MuT/ASTOpConditional.h"
+#include "MuT/ASTOpConstantResource.h"
+#include "MuT/ASTOpMeshClipMorphPlane.h"
+#include "MuT/ASTOpMeshRemoveMask.h"
+#include "MuT/ASTOpMeshMorph.h"
+#include "MuT/ASTOpSwitch.h"
+#include "MuT/StreamsPrivate.h"
+
+
+
+namespace mu
+{
+
+
+	//-------------------------------------------------------------------------------------------------
+	ASTOpLayoutPack::ASTOpLayoutPack()
+		: Source(this)
+	{
+	}
+
+
+	ASTOpLayoutPack::~ASTOpLayoutPack()
+	{
+		// Explicit call needed to avoid recursive destruction
+		ASTOp::RemoveChildren();
+	}
+
+
+	bool ASTOpLayoutPack::IsEqual(const ASTOp& otherUntyped) const
+	{
+		if (const ASTOpLayoutPack* other = dynamic_cast<const ASTOpLayoutPack*>(&otherUntyped))
+		{
+			return Source == other->Source;
+		}
+		return false;
+	}
+
+
+	uint64 ASTOpLayoutPack::Hash() const
+	{
+		uint64 res = std::hash<void*>()(Source.child().get());
+		return res;
+	}
+
+
+	mu::Ptr<ASTOp> ASTOpLayoutPack::Clone(MapChildFuncRef mapChild) const
+	{
+		mu::Ptr<ASTOpLayoutPack> n = new ASTOpLayoutPack();
+		n->Source = mapChild(Source.child());
+		return n;
+	}
+
+
+	void ASTOpLayoutPack::ForEachChild(const TFunctionRef<void(ASTChild&)> f)
+	{
+		f(Source);
+	}
+
+
+	void ASTOpLayoutPack::Link(FProgram& program, const FLinkerOptions*)
+	{
+		// Already linked?
+		if (!linkedAddress)
+		{
+			OP::LayoutRemoveBlocksArgs args;
+			FMemory::Memzero(&args, sizeof(args));
+
+			if (Source) args.Source = Source->linkedAddress;
+
+			linkedAddress = (OP::ADDRESS)program.m_opAddress.Num();
+			program.m_opAddress.Add((uint32_t)program.m_byteCode.Num());
+			AppendCode(program.m_byteCode, GetOpType());
+			AppendCode(program.m_byteCode, args);
+		}
+
+	}
+
+
+	void ASTOpLayoutPack::GetBlockLayoutSize(int blockIndex, int* pBlockX, int* pBlockY, FBlockLayoutSizeCache* cache)
+	{
+		if (Source)
+		{
+			Source->GetBlockLayoutSize(blockIndex, pBlockX, pBlockY, cache);
+		}
+	}
+
+
+}
diff --git a/Engine/Plugins/Experimental/Mutable/Source/MutableTools/Private/MuT/ASTOpLayoutPack.h b/Engine/Plugins/Experimental/Mutable/Source/MutableTools/Private/MuT/ASTOpLayoutPack.h
new file mode 100644
--- /dev/null
+++ b/Engine/Plugins/Experimental/Mutable/Source/MutableTools/Private/MuT/ASTOpLayoutPack.h
@@ -0,0 +1,38 @@
+// Copyright Epic Games, Inc. All Rights Reserved.
+
+#pragma once
+
+#include "HAL/PlatformMath.h"
+#include "MuR/Operations.h"
+#include "MuR/Ptr.h"
+#include "MuT/AST.h"
+
+
+namespace mu
+{
+	struct FProgram;
+	
+	class ASTOpLayoutPack : public ASTOp
+	{
+	public:
+
+		ASTChild Source;
+
+	public:
+
+		ASTOpLayoutPack();
+		ASTOpLayoutPack(const ASTOpLayoutPack&) = delete;
+		~ASTOpLayoutPack();
+
+		OP_TYPE GetOpType() const override { return OP_TYPE::LA_PACK; }
+		uint64 Hash() const override;
+		bool IsEqual(const ASTOp& otherUntyped) const override;
+		Ptr<ASTOp> Clone(MapChildFuncRef mapChild) const override;
+		void ForEachChild(const TFunctionRef<void(ASTChild&)>) override;
+		void Link(FProgram& program, const FLinkerOptions* Options) override;
+		void GetBlockLayoutSize(int blockIndex, int* pBlockX, int* pBlockY, FBlockLayoutSizeCache* cache) override;
+
+	};
+
+}
+
diff --git a/Engine/Plugins/Experimental/Mutable/Source/MutableTools/Private/MuT/ASTOpLayoutRemoveBlocks.cpp b/Engine/Plugins/Experimental/Mutable/Source/MutableTools/Private/MuT/ASTOpLayoutRemoveBlocks.cpp
new file mode 100644
--- /dev/null
+++ b/Engine/Plugins/Experimental/Mutable/Source/MutableTools/Private/MuT/ASTOpLayoutRemoveBlocks.cpp
@@ -0,0 +1,107 @@
+// Copyright Epic Games, Inc. All Rights Reserved.
+
+#include "MuT/ASTOpLayoutRemoveBlocks.h"
+
+#include "HAL/UnrealMemory.h"
+#include "Misc/AssertionMacros.h"
+#include "MuR/MemoryPrivate.h"
+#include "MuR/Mesh.h"
+#include "MuR/MeshBufferSet.h"
+#include "MuR/ModelPrivate.h"
+#include "MuR/MutableTrace.h"
+#include "MuR/RefCounted.h"
+#include "MuR/System.h"
+#include "MuR/Types.h"
+#include "MuT/ASTOpConditional.h"
+#include "MuT/ASTOpConstantResource.h"
+#include "MuT/ASTOpMeshClipMorphPlane.h"
+#include "MuT/ASTOpMeshRemoveMask.h"
+#include "MuT/ASTOpMeshMorph.h"
+#include "MuT/ASTOpSwitch.h"
+#include "MuT/StreamsPrivate.h"
+
+
+
+namespace mu
+{
+
+
+	//-------------------------------------------------------------------------------------------------
+	ASTOpLayoutRemoveBlocks::ASTOpLayoutRemoveBlocks()
+		: Source(this)
+		, ReferenceLayout(this)
+	{
+	}
+
+
+	ASTOpLayoutRemoveBlocks::~ASTOpLayoutRemoveBlocks()
+	{
+		// Explicit call needed to avoid recursive destruction
+		ASTOp::RemoveChildren();
+	}
+
+
+	bool ASTOpLayoutRemoveBlocks::IsEqual(const ASTOp& otherUntyped) const
+	{
+		if (const ASTOpLayoutRemoveBlocks* other = dynamic_cast<const ASTOpLayoutRemoveBlocks*>(&otherUntyped))
+		{
+			return Source == other->Source && ReferenceLayout == other->ReferenceLayout;
+		}
+		return false;
+	}
+
+
+	uint64 ASTOpLayoutRemoveBlocks::Hash() const
+	{
+		uint64 res = std::hash<void*>()(Source.child().get());
+		hash_combine(res, ReferenceLayout.child().get());
+		return res;
+	}
+
+
+	mu::Ptr<ASTOp> ASTOpLayoutRemoveBlocks::Clone(MapChildFuncRef mapChild) const
+	{
+		mu::Ptr<ASTOpLayoutRemoveBlocks> n = new ASTOpLayoutRemoveBlocks();
+		n->Source = mapChild(Source.child());
+		n->ReferenceLayout = mapChild(ReferenceLayout.child());
+		return n;
+	}
+
+
+	void ASTOpLayoutRemoveBlocks::ForEachChild(const TFunctionRef<void(ASTChild&)> f)
+	{
+		f(Source);
+		f(ReferenceLayout);
+	}
+
+
+	void ASTOpLayoutRemoveBlocks::Link(FProgram& program, const FLinkerOptions*)
+	{
+		// Already linked?
+		if (!linkedAddress)
+		{
+			OP::LayoutRemoveBlocksArgs args;
+			FMemory::Memzero(&args, sizeof(args));
+
+			if (Source) args.Source = Source->linkedAddress;
+			if (ReferenceLayout) args.ReferenceLayout = ReferenceLayout->linkedAddress;
+
+			linkedAddress = (OP::ADDRESS)program.m_opAddress.Num();
+			program.m_opAddress.Add((uint32_t)program.m_byteCode.Num());
+			AppendCode(program.m_byteCode, OP_TYPE::LA_REMOVEBLOCKS);
+			AppendCode(program.m_byteCode, args);
+		}
+
+	}
+
+
+	void ASTOpLayoutRemoveBlocks::GetBlockLayoutSize(int blockIndex, int* pBlockX, int* pBlockY, FBlockLayoutSizeCache* cache)
+	{
+		if (Source)
+		{
+			Source->GetBlockLayoutSize(blockIndex, pBlockX, pBlockY, cache);
+		}
+	}
+
+
+}
diff --git a/Engine/Plugins/Experimental/Mutable/Source/MutableTools/Private/MuT/ASTOpLayoutRemoveBlocks.h b/Engine/Plugins/Experimental/Mutable/Source/MutableTools/Private/MuT/ASTOpLayoutRemoveBlocks.h
new file mode 100644
--- /dev/null
+++ b/Engine/Plugins/Experimental/Mutable/Source/MutableTools/Private/MuT/ASTOpLayoutRemoveBlocks.h
@@ -0,0 +1,39 @@
+// Copyright Epic Games, Inc. All Rights Reserved.
+
+#pragma once
+
+#include "HAL/PlatformMath.h"
+#include "MuR/Operations.h"
+#include "MuR/Ptr.h"
+#include "MuT/AST.h"
+
+
+namespace mu
+{
+	struct FProgram;
+	
+	class ASTOpLayoutRemoveBlocks : public ASTOp
+	{
+	public:
+
+		ASTChild Source;
+		ASTChild ReferenceLayout;
+
+	public:
+
+		ASTOpLayoutRemoveBlocks();
+		ASTOpLayoutRemoveBlocks(const ASTOpLayoutRemoveBlocks&) = delete;
+		~ASTOpLayoutRemoveBlocks();
+
+		OP_TYPE GetOpType() const override { return OP_TYPE::LA_REMOVEBLOCKS; }
+		uint64 Hash() const override;
+		bool IsEqual(const ASTOp& otherUntyped) const override;
+		Ptr<ASTOp> Clone(MapChildFuncRef mapChild) const override;
+		void ForEachChild(const TFunctionRef<void(ASTChild&)>) override;
+		void Link(FProgram& program, const FLinkerOptions* Options) override;
+		void GetBlockLayoutSize(int blockIndex, int* pBlockX, int* pBlockY, FBlockLayoutSizeCache* cache) override;
+
+	};
+
+}
+
diff --git a/Engine/Plugins/Experimental/Mutable/Source/MutableTools/Private/MuT/ASTOpMeshBindShape.cpp b/Engine/Plugins/Experimental/Mutable/Source/MutableTools/Private/MuT/ASTOpMeshBindShape.cpp
--- a/Engine/Plugins/Experimental/Mutable/Source/MutableTools/Private/MuT/ASTOpMeshBindShape.cpp
+++ b/Engine/Plugins/Experimental/Mutable/Source/MutableTools/Private/MuT/ASTOpMeshBindShape.cpp
@@ -159,7 +159,7 @@ namespace mu
 
 
 	//-------------------------------------------------------------------------------------------------
-	mu::Ptr<ASTOp> ASTOpMeshBindShape::OptimiseSink(const MODEL_OPTIMIZATION_OPTIONS&, OPTIMIZE_SINK_CONTEXT&) const
+	mu::Ptr<ASTOp> ASTOpMeshBindShape::OptimiseSink(const FModelOptimizationOptions&, FOptimizeSinkContext&) const
 	{
 		Ptr<ASTOp> NewOp;
 
diff --git a/Engine/Plugins/Experimental/Mutable/Source/MutableTools/Private/MuT/ASTOpMeshBindShape.h b/Engine/Plugins/Experimental/Mutable/Source/MutableTools/Private/MuT/ASTOpMeshBindShape.h
--- a/Engine/Plugins/Experimental/Mutable/Source/MutableTools/Private/MuT/ASTOpMeshBindShape.h
+++ b/Engine/Plugins/Experimental/Mutable/Source/MutableTools/Private/MuT/ASTOpMeshBindShape.h
@@ -47,7 +47,7 @@ struct FProgram;
 		Ptr<ASTOp> Clone(MapChildFuncRef mapChild) const override;
 		void ForEachChild(const TFunctionRef<void(ASTChild&)>) override;
 		void Link(FProgram& program, const FLinkerOptions* Options) override;
-		Ptr<ASTOp> OptimiseSink(const MODEL_OPTIMIZATION_OPTIONS&, OPTIMIZE_SINK_CONTEXT&) const override;
+		Ptr<ASTOp> OptimiseSink(const FModelOptimizationOptions&, FOptimizeSinkContext&) const override;
 
 	};
 
diff --git a/Engine/Plugins/Experimental/Mutable/Source/MutableTools/Private/MuT/ASTOpMeshDifference.cpp b/Engine/Plugins/Experimental/Mutable/Source/MutableTools/Private/MuT/ASTOpMeshDifference.cpp
--- a/Engine/Plugins/Experimental/Mutable/Source/MutableTools/Private/MuT/ASTOpMeshDifference.cpp
+++ b/Engine/Plugins/Experimental/Mutable/Source/MutableTools/Private/MuT/ASTOpMeshDifference.cpp
@@ -104,7 +104,7 @@ namespace mu
 
 
 	//-------------------------------------------------------------------------------------------------
-	mu::Ptr<ASTOp> ASTOpMeshDifference::OptimiseSink(const MODEL_OPTIMIZATION_OPTIONS&, OPTIMIZE_SINK_CONTEXT&) const
+	mu::Ptr<ASTOp> ASTOpMeshDifference::OptimiseSink(const FModelOptimizationOptions&, FOptimizeSinkContext&) const
 	{
 		Ptr<ASTOp> NewOp;
 
diff --git a/Engine/Plugins/Experimental/Mutable/Source/MutableTools/Private/MuT/ASTOpMeshDifference.h b/Engine/Plugins/Experimental/Mutable/Source/MutableTools/Private/MuT/ASTOpMeshDifference.h
--- a/Engine/Plugins/Experimental/Mutable/Source/MutableTools/Private/MuT/ASTOpMeshDifference.h
+++ b/Engine/Plugins/Experimental/Mutable/Source/MutableTools/Private/MuT/ASTOpMeshDifference.h
@@ -53,7 +53,7 @@ namespace mu
 		bool IsEqual(const ASTOp& otherUntyped) const override;
 		Ptr<ASTOp> Clone(MapChildFuncRef mapChild) const override;
 		void Link(FProgram& program, const FLinkerOptions*) override;
-		Ptr<ASTOp> OptimiseSink(const MODEL_OPTIMIZATION_OPTIONS&, OPTIMIZE_SINK_CONTEXT&) const override;
+		Ptr<ASTOp> OptimiseSink(const FModelOptimizationOptions&, FOptimizeSinkContext&) const override;
 
 	};
 
diff --git a/Engine/Plugins/Experimental/Mutable/Source/MutableTools/Private/MuT/ASTOpMeshFormat.cpp b/Engine/Plugins/Experimental/Mutable/Source/MutableTools/Private/MuT/ASTOpMeshFormat.cpp
--- a/Engine/Plugins/Experimental/Mutable/Source/MutableTools/Private/MuT/ASTOpMeshFormat.cpp
+++ b/Engine/Plugins/Experimental/Mutable/Source/MutableTools/Private/MuT/ASTOpMeshFormat.cpp
@@ -101,7 +101,7 @@ void ASTOpMeshFormat::Link( FProgram& program, const FLinkerOptions* )
 }
 
 
-mu::Ptr<ASTOp> ASTOpMeshFormat::OptimiseSink(const MODEL_OPTIMIZATION_OPTIONS& options, OPTIMIZE_SINK_CONTEXT& context) const
+mu::Ptr<ASTOp> ASTOpMeshFormat::OptimiseSink(const FModelOptimizationOptions& options, FOptimizeSinkContext& context) const
 {
 	mu::Ptr<ASTOp> at = context.MeshFormatSinker.Apply(this);
 	return at;
@@ -116,7 +116,7 @@ mu::Ptr<ASTOp> Sink_MeshFormatAST::Apply(const ASTOp* root)
 	m_root = dynamic_cast<const ASTOpMeshFormat*>(root);
 	check(m_root);
 
-	m_oldToNew.Empty();
+	OldToNew.Empty();
 
 	m_initialSource = m_root->Source.child();
 	mu::Ptr<ASTOp> newSource = Visit(m_initialSource, m_root);
@@ -230,17 +230,11 @@ mu::Ptr<ASTOp> Sink_MeshFormatAST::Visit(const mu::Ptr<ASTOp>& at, const ASTOpMe
 {
 	if (!at) return nullptr;
 
-	// Newly created?
-	if (m_newOps.Find(at) != INDEX_NONE)
-	{
-		return at;
-	}
-
 	// Already visited?
-	auto cacheIt = m_oldToNew.Find(at);
-	if (cacheIt)
+	const Ptr<ASTOp>* Cached = OldToNew.Find({ at,currentFormatOp });
+	if (Cached)
 	{
-		return *cacheIt;
+		return *Cached;
 	}
 
 	mu::Ptr<ASTOp> newAt = at;
@@ -407,7 +401,7 @@ mu::Ptr<ASTOp> Sink_MeshFormatAST::Visit(const mu::Ptr<ASTOp>& at, const ASTOpMe
 
 	}
 
-	m_oldToNew.Add(at, newAt);
+	OldToNew.Add({ at,currentFormatOp }, newAt);
 
 	return newAt;
 }
diff --git a/Engine/Plugins/Experimental/Mutable/Source/MutableTools/Private/MuT/ASTOpMeshFormat.h b/Engine/Plugins/Experimental/Mutable/Source/MutableTools/Private/MuT/ASTOpMeshFormat.h
--- a/Engine/Plugins/Experimental/Mutable/Source/MutableTools/Private/MuT/ASTOpMeshFormat.h
+++ b/Engine/Plugins/Experimental/Mutable/Source/MutableTools/Private/MuT/ASTOpMeshFormat.h
@@ -10,7 +10,7 @@
 
 namespace mu
 {
-struct FProgram;
+	struct FProgram;
 	
 	class ASTOpMeshFormat : public ASTOp
 	{
@@ -32,7 +32,7 @@ struct FProgram;
 		Ptr<ASTOp> Clone(MapChildFuncRef mapChild) const override;
 		void ForEachChild(const TFunctionRef<void(ASTChild&)>) override;
 		void Link(FProgram& program, const FLinkerOptions* Options) override;
-		Ptr<ASTOp> OptimiseSink(const MODEL_OPTIMIZATION_OPTIONS& options, OPTIMIZE_SINK_CONTEXT& context) const override;
+		Ptr<ASTOp> OptimiseSink(const FModelOptimizationOptions& options, FOptimizeSinkContext& context) const override;
 
 	};
 
diff --git a/Engine/Plugins/Experimental/Mutable/Source/MutableTools/Private/MuT/ASTOpMeshMaskClipMesh.cpp b/Engine/Plugins/Experimental/Mutable/Source/MutableTools/Private/MuT/ASTOpMeshMaskClipMesh.cpp
--- a/Engine/Plugins/Experimental/Mutable/Source/MutableTools/Private/MuT/ASTOpMeshMaskClipMesh.cpp
+++ b/Engine/Plugins/Experimental/Mutable/Source/MutableTools/Private/MuT/ASTOpMeshMaskClipMesh.cpp
@@ -188,7 +188,7 @@ namespace mu
 					// We move the mask creation down all the paths
 					Ptr<ASTOpSwitch> newOp = mu::Clone<ASTOpSwitch>(at);
 					newOp->def = Visit(newOp->def.child());
-					for (ASTOpSwitch::CASE& c : newOp->cases)
+					for (ASTOpSwitch::FCase& c : newOp->cases)
 					{
 						c.branch = Visit(c.branch.child());
 					}
@@ -263,7 +263,7 @@ namespace mu
 				if (!at) return nullptr;
 
 				// Newly created?
-				if (m_newOps.Find(at) != INDEX_NONE)
+				if (m_newOps.Contains(at))
 				{
 					return at;
 				}
@@ -297,7 +297,7 @@ namespace mu
 					// We move the mask creation down all the paths
 					Ptr<ASTOpSwitch> newOp = mu::Clone<ASTOpSwitch>(at);
 					newOp->def = Visit(newOp->def.child());
-					for (ASTOpSwitch::CASE& c : newOp->cases)
+					for (ASTOpSwitch::FCase& c : newOp->cases)
 					{
 						c.branch = Visit(c.branch.child());
 					}
@@ -330,7 +330,7 @@ namespace mu
 
 
 	//-------------------------------------------------------------------------------------------------
-	mu::Ptr<ASTOp> ASTOpMeshMaskClipMesh::OptimiseSink(const MODEL_OPTIMIZATION_OPTIONS&, OPTIMIZE_SINK_CONTEXT&) const
+	mu::Ptr<ASTOp> ASTOpMeshMaskClipMesh::OptimiseSink(const FModelOptimizationOptions&, FOptimizeSinkContext&) const
 	{
 		// \TODO: Add logic state to the sinkers to avoid explosion with switches in both branches and similar cases.
 
diff --git a/Engine/Plugins/Experimental/Mutable/Source/MutableTools/Private/MuT/ASTOpMeshMaskClipMesh.h b/Engine/Plugins/Experimental/Mutable/Source/MutableTools/Private/MuT/ASTOpMeshMaskClipMesh.h
--- a/Engine/Plugins/Experimental/Mutable/Source/MutableTools/Private/MuT/ASTOpMeshMaskClipMesh.h
+++ b/Engine/Plugins/Experimental/Mutable/Source/MutableTools/Private/MuT/ASTOpMeshMaskClipMesh.h
@@ -35,7 +35,7 @@ struct FProgram;
 		Ptr<ASTOp> Clone(MapChildFuncRef mapChild) const override;
 		void ForEachChild(const TFunctionRef<void(ASTChild&)>) override;
 		void Link(FProgram& program, const FLinkerOptions* Options) override;
-		Ptr<ASTOp> OptimiseSink(const MODEL_OPTIMIZATION_OPTIONS&, OPTIMIZE_SINK_CONTEXT&) const override;
+		Ptr<ASTOp> OptimiseSink(const FModelOptimizationOptions&, FOptimizeSinkContext&) const override;
 
 	};
 
diff --git a/Engine/Plugins/Experimental/Mutable/Source/MutableTools/Private/MuT/ASTOpMeshMorph.cpp b/Engine/Plugins/Experimental/Mutable/Source/MutableTools/Private/MuT/ASTOpMeshMorph.cpp
--- a/Engine/Plugins/Experimental/Mutable/Source/MutableTools/Private/MuT/ASTOpMeshMorph.cpp
+++ b/Engine/Plugins/Experimental/Mutable/Source/MutableTools/Private/MuT/ASTOpMeshMorph.cpp
@@ -113,7 +113,7 @@ namespace mu
 
 
 	//-------------------------------------------------------------------------------------------------
-	mu::Ptr<ASTOp> ASTOpMeshMorph::OptimiseSink(const MODEL_OPTIMIZATION_OPTIONS&, OPTIMIZE_SINK_CONTEXT&) const
+	mu::Ptr<ASTOp> ASTOpMeshMorph::OptimiseSink(const FModelOptimizationOptions&, FOptimizeSinkContext&) const
 	{
 		return nullptr;
 	}
diff --git a/Engine/Plugins/Experimental/Mutable/Source/MutableTools/Private/MuT/ASTOpMeshMorph.h b/Engine/Plugins/Experimental/Mutable/Source/MutableTools/Private/MuT/ASTOpMeshMorph.h
--- a/Engine/Plugins/Experimental/Mutable/Source/MutableTools/Private/MuT/ASTOpMeshMorph.h
+++ b/Engine/Plugins/Experimental/Mutable/Source/MutableTools/Private/MuT/ASTOpMeshMorph.h
@@ -23,7 +23,7 @@ namespace mu
 		//! Factor selecting what morphs to apply and with what weight.
 		ASTChild Factor;
 
-		//! Base mesh to remove from.
+		//! Base mesh to morph.
 		ASTChild Base;
 
 		//! Targets to apply on the base depending on the factor
@@ -42,7 +42,7 @@ namespace mu
 		bool IsEqual(const ASTOp& otherUntyped) const override;
 		Ptr<ASTOp> Clone(MapChildFuncRef mapChild) const override;
 		void Link(FProgram& program, const FLinkerOptions*) override;
-		Ptr<ASTOp> OptimiseSink(const MODEL_OPTIMIZATION_OPTIONS&, OPTIMIZE_SINK_CONTEXT&) const override;
+		Ptr<ASTOp> OptimiseSink(const FModelOptimizationOptions&, FOptimizeSinkContext&) const override;
 
 		// Own interface
 		void AddTarget(const Ptr<ASTOp>&);
diff --git a/Engine/Plugins/Experimental/Mutable/Source/MutableTools/Private/MuT/ASTOpMeshRemoveMask.cpp b/Engine/Plugins/Experimental/Mutable/Source/MutableTools/Private/MuT/ASTOpMeshRemoveMask.cpp
--- a/Engine/Plugins/Experimental/Mutable/Source/MutableTools/Private/MuT/ASTOpMeshRemoveMask.cpp
+++ b/Engine/Plugins/Experimental/Mutable/Source/MutableTools/Private/MuT/ASTOpMeshRemoveMask.cpp
@@ -221,7 +221,7 @@ namespace mu
 
 
 	//-------------------------------------------------------------------------------------------------
-	mu::Ptr<ASTOp> ASTOpMeshRemoveMask::OptimiseSink(const MODEL_OPTIMIZATION_OPTIONS&, OPTIMIZE_SINK_CONTEXT&) const
+	mu::Ptr<ASTOp> ASTOpMeshRemoveMask::OptimiseSink(const FModelOptimizationOptions&, FOptimizeSinkContext&) const
 	{
 		Sink_MeshRemoveMaskAST sinker;
 		mu::Ptr<ASTOp> at = sinker.Apply(this);
diff --git a/Engine/Plugins/Experimental/Mutable/Source/MutableTools/Private/MuT/ASTOpMeshRemoveMask.h b/Engine/Plugins/Experimental/Mutable/Source/MutableTools/Private/MuT/ASTOpMeshRemoveMask.h
--- a/Engine/Plugins/Experimental/Mutable/Source/MutableTools/Private/MuT/ASTOpMeshRemoveMask.h
+++ b/Engine/Plugins/Experimental/Mutable/Source/MutableTools/Private/MuT/ASTOpMeshRemoveMask.h
@@ -39,7 +39,7 @@ namespace mu
 		bool IsEqual(const ASTOp& otherUntyped) const override;
 		Ptr<ASTOp> Clone(MapChildFuncRef mapChild) const override;
 		void Link(FProgram& program, const FLinkerOptions*) override;
-		Ptr<ASTOp> OptimiseSink(const MODEL_OPTIMIZATION_OPTIONS&, OPTIMIZE_SINK_CONTEXT&) const override;
+		Ptr<ASTOp> OptimiseSink(const FModelOptimizationOptions&, FOptimizeSinkContext&) const override;
 
 		// Own interface
 		void AddRemove(const Ptr<ASTOp>& condition, const Ptr<ASTOp>& mask);
diff --git a/Engine/Plugins/Experimental/Mutable/Source/MutableTools/Private/MuT/ASTOpSwitch.cpp b/Engine/Plugins/Experimental/Mutable/Source/MutableTools/Private/MuT/ASTOpSwitch.cpp
--- a/Engine/Plugins/Experimental/Mutable/Source/MutableTools/Private/MuT/ASTOpSwitch.cpp
+++ b/Engine/Plugins/Experimental/Mutable/Source/MutableTools/Private/MuT/ASTOpSwitch.cpp
@@ -183,7 +183,7 @@ namespace mu
 			AppendCode(program.m_byteCode, DefAddress);
 			AppendCode(program.m_byteCode, (uint32_t)cases.Num());
 
-			for (const CASE& Case : cases)
+			for (const FCase& Case : cases)
 			{
 				OP::ADDRESS CaseBranchAddress = Case.branch ? Case.branch->linkedAddress : 0;
 				AppendCode(program.m_byteCode, Case.condition);
@@ -287,7 +287,7 @@ namespace mu
 
 	//-------------------------------------------------------------------------------------------------
 	void ASTOpSwitch::GetBlockLayoutSize(int blockIndex, int* pBlockX, int* pBlockY,
-		BLOCK_LAYOUT_SIZE_CACHE* cache)
+		FBlockLayoutSizeCache* cache)
 	{
 		switch (type)
 		{
@@ -410,7 +410,7 @@ namespace mu
 
 
 	//-------------------------------------------------------------------------------------------------
-	mu::Ptr<ASTOp> ASTOpSwitch::OptimiseSemantic(const MODEL_OPTIMIZATION_OPTIONS&) const
+	mu::Ptr<ASTOp> ASTOpSwitch::OptimiseSemantic(const FModelOptimizationOptions&) const
 	{
 		// Constant condition?
 		if (variable->GetOpType() == OP_TYPE::NU_CONSTANT)
diff --git a/Engine/Plugins/Experimental/Mutable/Source/MutableTools/Private/MuT/ASTOpSwitch.h b/Engine/Plugins/Experimental/Mutable/Source/MutableTools/Private/MuT/ASTOpSwitch.h
--- a/Engine/Plugins/Experimental/Mutable/Source/MutableTools/Private/MuT/ASTOpSwitch.h
+++ b/Engine/Plugins/Experimental/Mutable/Source/MutableTools/Private/MuT/ASTOpSwitch.h
@@ -32,9 +32,9 @@ template <class SCALAR> class vec4;
 		//! Default branch in case none matches the value in the variable.
 		ASTChild def;
 
-		struct CASE
+		struct FCase
 		{
-			CASE(int32_t cond, Ptr<ASTOp> parent, Ptr<ASTOp> b)
+			FCase(int32_t cond, Ptr<ASTOp> parent, Ptr<ASTOp> b)
 				: condition(cond)
 				, branch(parent.get(), b)
 			{
@@ -44,14 +44,14 @@ template <class SCALAR> class vec4;
 			ASTChild branch;
 
 			//!
-			bool operator==(const CASE& other) const
+			bool operator==(const FCase& other) const
 			{
 				return condition == other.condition &&
 					branch == other.branch;
 			}
 		};
 
-		TArray<CASE> cases;
+		TArray<FCase> cases;
 
 	public:
 
@@ -68,9 +68,9 @@ template <class SCALAR> class vec4;
 		void Assert() override;
 		void Link(FProgram& program, const FLinkerOptions* Options) override;
 		FImageDesc GetImageDesc(bool returnBestOption, class FGetImageDescContext* context) const  override;
-		void GetBlockLayoutSize(int blockIndex, int* pBlockX, int* pBlockY, BLOCK_LAYOUT_SIZE_CACHE* cache) override;
+		void GetBlockLayoutSize(int blockIndex, int* pBlockX, int* pBlockY, FBlockLayoutSizeCache* cache) override;
 		void GetLayoutBlockSize(int* pBlockX, int* pBlockY) override;
-		Ptr<ASTOp> OptimiseSemantic(const MODEL_OPTIMIZATION_OPTIONS&) const override;
+		Ptr<ASTOp> OptimiseSemantic(const FModelOptimizationOptions&) const override;
 		bool GetNonBlackRect(FImageRect& maskUsage) const override;
 		bool IsImagePlainConstant(FVector4f& colour) const override;
 		Ptr<ImageSizeExpression> GetImageSizeExpression() const override;
diff --git a/Engine/Plugins/Experimental/Mutable/Source/MutableTools/Private/MuT/CodeGenerator.cpp b/Engine/Plugins/Experimental/Mutable/Source/MutableTools/Private/MuT/CodeGenerator.cpp
--- a/Engine/Plugins/Experimental/Mutable/Source/MutableTools/Private/MuT/CodeGenerator.cpp
+++ b/Engine/Plugins/Experimental/Mutable/Source/MutableTools/Private/MuT/CodeGenerator.cpp
@@ -29,6 +29,10 @@
 #include "MuT/ASTOpMeshDifference.h"
 #include "MuT/ASTOpMeshMorph.h"
 #include "MuT/ASTOpParameter.h"
+#include "MuT/ASTOpLayoutRemoveBlocks.h"
+#include "MuT/ASTOpLayoutFromMesh.h"
+#include "MuT/ASTOpLayoutMerge.h"
+#include "MuT/ASTOpLayoutPack.h"
 #include "MuT/CodeGenerator_SecondPass.h"
 #include "MuT/CodeOptimiser.h"
 #include "MuT/CompilerPrivate.h"
@@ -1202,13 +1206,9 @@ namespace mu
                             }
 
                             // Merge operation
-                            Ptr<ASTOpFixed> mergeAd = new ASTOpFixed();
-                            {
-                                mergeAd->op.type = OP_TYPE::LA_MERGE;
-                                // A normal mesh that we merge
-                                mergeAd->SetChild(mergeAd->op.args.LayoutMerge.base, layoutOp );
-                                mergeAd->SetChild(mergeAd->op.args.LayoutMerge.added, layoutFragmentAd );
-                            }
+                            Ptr<ASTOpLayoutMerge> mergeAd = new ASTOpLayoutMerge();
+                            mergeAd->Base = layoutOp;
+	                        mergeAd->Added = layoutFragmentAd;
 
                             // Condition to apply
                             if (data.condition)
@@ -1235,17 +1235,19 @@ namespace mu
                         // Add layout packing instructions
                         {
                             // Make sure we removed unnecessary blocks
-                            Ptr<ASTOpFixed> op = new ASTOpFixed();
-                            op->op.type = OP_TYPE::LA_REMOVEBLOCKS;
-                            op->SetChild(op->op.args.LayoutRemoveBlocks.source, layoutOp );
-                            op->SetChild(op->op.args.LayoutRemoveBlocks.mesh, lastMeshOp );
-                            op->op.args.LayoutRemoveBlocks.meshLayoutIndex = uint8_t( l );
-                            layoutOp = op;
+                            Ptr<ASTOpLayoutFromMesh> ExtractOp = new ASTOpLayoutFromMesh();
+							ExtractOp->Mesh = lastMeshOp;
+							check(l<256);
+							ExtractOp->LayoutIndex = uint8( l );
+
+							Ptr<ASTOpLayoutRemoveBlocks> RemoveOp = new ASTOpLayoutRemoveBlocks();
+							RemoveOp->Source = layoutOp;
+							RemoveOp->ReferenceLayout = ExtractOp;
+							layoutOp = RemoveOp;
 
                             // Pack uv blocks
-                            op = new ASTOpFixed();
-                            op->op.type = OP_TYPE::LA_PACK;
-                            op->SetChild(op->op.args.LayoutPack.layout, layoutOp );
+                            Ptr<ASTOpLayoutPack> op = new ASTOpLayoutPack();
+                            op->Source = layoutOp;
                             layoutOp = op;
                         }
 
diff --git a/Engine/Plugins/Experimental/Mutable/Source/MutableTools/Private/MuT/CodeGenerator.h b/Engine/Plugins/Experimental/Mutable/Source/MutableTools/Private/MuT/CodeGenerator.h
--- a/Engine/Plugins/Experimental/Mutable/Source/MutableTools/Private/MuT/CodeGenerator.h
+++ b/Engine/Plugins/Experimental/Mutable/Source/MutableTools/Private/MuT/CodeGenerator.h
@@ -749,7 +749,7 @@ namespace mu
             check( pTable->GetPrivate()->m_rows[i].m_id <= 0xFFFF);
             auto condition = (uint16)pTable->GetPrivate()->m_rows[i].m_id;
             Ptr<ASTOp> Branch = GenerateOption( node, colIndex, (int)i, m_pErrorLog.get() );
-			SwitchOp->cases.Add(ASTOpSwitch::CASE(condition, SwitchOp, Branch ));
+			SwitchOp->cases.Add(ASTOpSwitch::FCase(condition, SwitchOp, Branch ));
         }
 
         return SwitchOp;
diff --git a/Engine/Plugins/Experimental/Mutable/Source/MutableTools/Private/MuT/CodeOptimiser.cpp b/Engine/Plugins/Experimental/Mutable/Source/MutableTools/Private/MuT/CodeOptimiser.cpp
--- a/Engine/Plugins/Experimental/Mutable/Source/MutableTools/Private/MuT/CodeOptimiser.cpp
+++ b/Engine/Plugins/Experimental/Mutable/Source/MutableTools/Private/MuT/CodeOptimiser.cpp
@@ -1052,7 +1052,7 @@ namespace mu
 
 
 	//---------------------------------------------------------------------------------------------
-	void SkeletonCleanerAST( TArray<mu::Ptr<ASTOp>>& roots, const MODEL_OPTIMIZATION_OPTIONS& options )
+	void SkeletonCleanerAST( TArray<mu::Ptr<ASTOp>>& roots, const FModelOptimizationOptions& options )
 	{
 		// This collects all the meshes that require a skeleton because they are used in operations
 		// that require it.
diff --git a/Engine/Plugins/Experimental/Mutable/Source/MutableTools/Private/MuT/CodeOptimiser.h b/Engine/Plugins/Experimental/Mutable/Source/MutableTools/Private/MuT/CodeOptimiser.h
--- a/Engine/Plugins/Experimental/Mutable/Source/MutableTools/Private/MuT/CodeOptimiser.h
+++ b/Engine/Plugins/Experimental/Mutable/Source/MutableTools/Private/MuT/CodeOptimiser.h
@@ -122,13 +122,13 @@ namespace mu
     //! All kinds of optimisations that depend on the meaning of each operation
     //---------------------------------------------------------------------------------------------
     extern bool SemanticOptimiserAST( ASTOpList& roots,
-                                      const MODEL_OPTIMIZATION_OPTIONS& optimisationOptions );
+                                      const FModelOptimizationOptions& optimisationOptions );
 
     //---------------------------------------------------------------------------------------------
     //! Semantic operator that reorders instructions moving expensive ones down to the
     //! leaves of the expressions trying to turn them into constants.
     //---------------------------------------------------------------------------------------------
-    extern bool SinkOptimiserAST( ASTOpList& roots, const MODEL_OPTIMIZATION_OPTIONS& );
+    extern bool SinkOptimiserAST( ASTOpList& roots, const FModelOptimizationOptions& );
 
     //---------------------------------------------------------------------------------------------
     //!
@@ -245,7 +245,7 @@ namespace mu
     public:
 
         ParameterOptimiserAST( STATE_COMPILATION_DATA &s,
-                               const MODEL_OPTIMIZATION_OPTIONS& optimisationOptions );
+                               const FModelOptimizationOptions& optimisationOptions );
 
         bool Apply();
 
@@ -256,7 +256,7 @@ namespace mu
 
         bool m_modified;
 
-        MODEL_OPTIMIZATION_OPTIONS m_optimisationOptions;
+        FModelOptimizationOptions m_optimisationOptions;
 
         RuntimeParameterVisitorAST m_hasRuntimeParamVisitor;
 
diff --git a/Engine/Plugins/Experimental/Mutable/Source/MutableTools/Private/MuT/CodeOptimiser_Parameters.cpp b/Engine/Plugins/Experimental/Mutable/Source/MutableTools/Private/MuT/CodeOptimiser_Parameters.cpp
--- a/Engine/Plugins/Experimental/Mutable/Source/MutableTools/Private/MuT/CodeOptimiser_Parameters.cpp
+++ b/Engine/Plugins/Experimental/Mutable/Source/MutableTools/Private/MuT/CodeOptimiser_Parameters.cpp
@@ -22,6 +22,7 @@
 #include "MuT/ASTOpImagePixelFormat.h"
 #include "MuT/ASTOpMeshMorph.h"
 #include "MuT/ASTOpInstanceAdd.h"
+#include "MuT/ASTOpLayoutFromMesh.h"
 #include "MuT/ASTOpParameter.h"
 #include "MuT/CodeOptimiser.h"
 #include "MuT/Compiler.h"
@@ -268,7 +269,7 @@ class TaskManager;
     //---------------------------------------------------------------------------------------------
     ParameterOptimiserAST::ParameterOptimiserAST(
             STATE_COMPILATION_DATA& s,
-            const MODEL_OPTIMIZATION_OPTIONS& optimisationOptions
+            const FModelOptimizationOptions& optimisationOptions
             )
             : m_stateProps(s)
             , m_modified(false)
@@ -1673,24 +1674,19 @@ class TaskManager;
                     break;
                 }
 
-                case OP_TYPE::LA_REMOVEBLOCKS:
-                {
-                    // Manually choose how to recurse this op
-                    auto pTyped = dynamic_cast<const ASTOpFixed*>( at.get() );
-
-                    if ( auto& source = pTyped->children[pTyped->op.args.LayoutRemoveBlocks.source] )
-                    {
-						pending.Add({ source.m_child, state });
-                    }
+				case OP_TYPE::LA_FROMMESH:
+				{
+					// Manually choose how to recurse this op
+					const ASTOpLayoutFromMesh* pTyped = dynamic_cast<const ASTOpLayoutFromMesh*>(at.get());
 
-                    // For that mesh we only want to know about the layouts
-                    if ( auto& mesh = pTyped->children[pTyped->op.args.LayoutRemoveBlocks.mesh] )
-                    {
-						pending.Add({ mesh.m_child, true });
-                    }
+					// For that mesh we only want to know about the layouts
+					if (const ASTChild& Mesh = pTyped->Mesh)
+					{
+						pending.Add({ Mesh.m_child, true });
+					}
 
-                    return false;
-                }
+					return false;
+				}
 
                 case OP_TYPE::ME_MORPH2:
                 {
diff --git a/Engine/Plugins/Experimental/Mutable/Source/MutableTools/Private/MuT/CodeOptimiser_Semantic.cpp b/Engine/Plugins/Experimental/Mutable/Source/MutableTools/Private/MuT/CodeOptimiser_Semantic.cpp
--- a/Engine/Plugins/Experimental/Mutable/Source/MutableTools/Private/MuT/CodeOptimiser_Semantic.cpp
+++ b/Engine/Plugins/Experimental/Mutable/Source/MutableTools/Private/MuT/CodeOptimiser_Semantic.cpp
@@ -80,7 +80,7 @@ namespace mu
     //---------------------------------------------------------------------------------------------
     bool SemanticOptimiserAST(
             ASTOpList& roots,
-            const MODEL_OPTIMIZATION_OPTIONS& optimisationOptions
+            const FModelOptimizationOptions& optimisationOptions
             )
     {
         MUTABLE_CPUPROFILER_SCOPE(SemanticOptimiserAST);
@@ -105,7 +105,7 @@ namespace mu
 
 
     //---------------------------------------------------------------------------------------------
-    Ptr<ASTOp> ASTOpFixed::OptimiseSemantic( const MODEL_OPTIMIZATION_OPTIONS& options ) const
+    Ptr<ASTOp> ASTOpFixed::OptimiseSemantic( const FModelOptimizationOptions& options ) const
     {
         Ptr<ASTOp> at;
 
@@ -478,7 +478,7 @@ namespace mu
 
 
     //---------------------------------------------------------------------------------------------
-    Ptr<ASTOp> ASTOpFixed::OptimiseSwizzle( const MODEL_OPTIMIZATION_OPTIONS& options ) const
+    Ptr<ASTOp> ASTOpFixed::OptimiseSwizzle( const FModelOptimizationOptions& options ) const
     {
         MUTABLE_CPUPROFILER_SCOPE(OptimiseSwizzleAST);
 
@@ -1002,14 +1002,14 @@ namespace mu
     bool SinkOptimiserAST
     (
             ASTOpList& roots,
-            const MODEL_OPTIMIZATION_OPTIONS& optimisationOptions
+            const FModelOptimizationOptions& optimisationOptions
     )
     {
         MUTABLE_CPUPROFILER_SCOPE(SinkOptimiserAST);
 
         bool modified = false;
 
-		OPTIMIZE_SINK_CONTEXT context;
+		FOptimizeSinkContext context;
 
         ASTOp::Traverse_TopDown_Unique_Imprecise( roots, [&](Ptr<ASTOp>& n)
         {
@@ -1033,8 +1033,7 @@ namespace mu
     Ptr<ASTOp> Sink_ImageCropAST::Apply( const ASTOp* root )
     {
         m_root = root;
-        m_oldToNew.Empty();
-		m_newOps.Empty();
+        OldToNew.Empty();
 
         check(root->GetOpType()==OP_TYPE::IM_CROP);
 
@@ -1057,20 +1056,11 @@ namespace mu
     {
         if (!at) return nullptr;
 
-        // Newly created?
-        if (m_newOps.Find( at )!=INDEX_NONE)
-        {
-            return at;
-        }
-
         // Already visited?
+		const Ptr<ASTOp>* Cached = OldToNew.Find({ at, currentCropOp });
+		if (Cached)
 		{
-			auto& CurrentSet = m_oldToNew.FindOrAdd(currentCropOp);
-			auto cacheIt = CurrentSet.find(at);
-			if (cacheIt != CurrentSet.end())
-			{
-				return cacheIt->second;
-			}
+			return *Cached;
 		}
 
         bool skipSinking=false;
@@ -1300,7 +1290,7 @@ namespace mu
             newAt = newOp;
         }
 
-        m_oldToNew[currentCropOp][at] = newAt;
+		OldToNew.Add({ at, currentCropOp }, newAt);
 
         return newAt;
     }
@@ -1309,7 +1299,7 @@ namespace mu
 	//---------------------------------------------------------------------------------------------
 	//---------------------------------------------------------------------------------------------
 	//---------------------------------------------------------------------------------------------
-	Ptr<ASTOp> ASTOpFixed::OptimiseSink( const MODEL_OPTIMIZATION_OPTIONS& options, OPTIMIZE_SINK_CONTEXT& context ) const
+	Ptr<ASTOp> ASTOpFixed::OptimiseSink( const FModelOptimizationOptions& options, FOptimizeSinkContext& context ) const
     {
         Ptr<ASTOp> at;
 
@@ -1392,63 +1382,7 @@ namespace mu
         //-----------------------------------------------------------------------------------------
         case OP_TYPE::LA_REMOVEBLOCKS:
         {
-            auto meshAt = children[op.args.LayoutRemoveBlocks.mesh].child();
-
-            OP_TYPE meshType = meshAt->GetOpType();
-            switch ( meshType )
-            {
-
-            case OP_TYPE::ME_SWITCH:
-            {
-                // Move the remove-blocks down all the paths
-                Ptr<ASTOpSwitch> nop = mu::Clone<ASTOpSwitch>(meshAt);
-                nop->type = OP_TYPE::LA_SWITCH;
-
-                if (nop->def)
-                {
-					Ptr<ASTOpFixed> defOp = mu::Clone<ASTOpFixed>(this);
-                    defOp->SetChild( defOp->op.args.LayoutRemoveBlocks.mesh, nop->def );
-                    nop->def = defOp;
-                }
-
-                for ( auto& o: nop->cases )
-                {
-                    if ( o.branch )
-                    {
-						Ptr<ASTOpFixed> bOp = mu::Clone<ASTOpFixed>(this);
-                        bOp->SetChild( bOp->op.args.LayoutRemoveBlocks.mesh, o.branch );
-                        o.branch = bOp;
-                    }
-                }
-
-                at = nop;
-                break;
-            }
-
-            case OP_TYPE::ME_MORPH2:
-            {
-                // Move the remove-blocks down the base
-				const ASTOpMeshMorph* typedMeshAt = dynamic_cast<const ASTOpMeshMorph*>(meshAt.get());
-                Ptr<ASTOpFixed> nop = mu::Clone<ASTOpFixed>(this);
-                nop->SetChild( nop->op.args.LayoutRemoveBlocks.mesh, typedMeshAt->Base );
-                at = nop;
-                break;
-            }
-
-            case OP_TYPE::ME_INTERPOLATE:
-            {
-                // Assume that the UVs will not be interpolated across blocks!
-                auto typedMeshAt = dynamic_cast<const ASTOpFixed*>(meshAt.get());
-				Ptr<ASTOpFixed> nop = mu::Clone<ASTOpFixed>(this);
-                nop->SetChild( nop->op.args.LayoutRemoveBlocks.mesh, typedMeshAt->children[ typedMeshAt->op.args.MeshInterpolate.base] );
-                at = nop;
-                break;
-            }
-
-            default:
-                break;
-            }
-
+			check(false);
             break;
         }
 
@@ -1986,7 +1920,7 @@ namespace mu
         Ptr<ASTOp> Apply( const ASTOp* root )
         {
             m_root = root;
-            m_oldToNew.clear();
+            OldToNew.Empty();
 
             check(root->GetOpType()==OP_TYPE::IM_RESIZEREL);
 
@@ -2008,259 +1942,252 @@ namespace mu
         const ASTOp* m_root;
         Ptr<ASTOp> m_initialSource;
         //! For each operation we sink, the map from old instructions to new instructions.
-        std::unordered_map<Ptr<const ASTOpFixed>,std::unordered_map<Ptr<ASTOp>,Ptr<ASTOp>>> m_oldToNew;
-        vector<Ptr<ASTOp>> m_newOps;
-
-        Ptr<ASTOp> Visit( Ptr<ASTOp> at, const ASTOpFixed* currentSinkingOp )
-        {
-            if (!at) return nullptr;
+		TMap<FSinkerOldToNewKey, Ptr<ASTOp>> OldToNew;
 
-            // Newly created?
-            if (std::find(m_newOps.begin(), m_newOps.end(), at )!=m_newOps.end())
-            {
-                return at;
-            }
+		Ptr<ASTOp> Visit(Ptr<ASTOp> at, const ASTOpFixed* currentSinkingOp)
+		{
+			if (!at) return nullptr;
 
-            // Already visited?
-            auto cacheIt = m_oldToNew[currentSinkingOp].find(at);
-            if (cacheIt!=m_oldToNew[currentSinkingOp].end())
-            {
-                return cacheIt->second;
-            }
+			// Already visited?
+			const Ptr<ASTOp>* Cached = OldToNew.Find({ at,currentSinkingOp });
+			if (Cached)
+			{
+				return *Cached;
+			}
 
-            float scaleX = currentSinkingOp->op.args.ImageResizeRel.factor[0];
-            float scaleY = currentSinkingOp->op.args.ImageResizeRel.factor[1];
+			float scaleX = currentSinkingOp->op.args.ImageResizeRel.factor[0];
+			float scaleY = currentSinkingOp->op.args.ImageResizeRel.factor[1];
 
-            Ptr<ASTOp> newAt = at;
-            switch ( at->GetOpType() )
-            {
+			Ptr<ASTOp> newAt = at;
+			switch (at->GetOpType())
+			{
 
-            case OP_TYPE::IM_CONDITIONAL:
-            {
-                // We move the mask creation down the two paths
-                auto newOp = mu::Clone<ASTOpConditional>(at);
-                newOp->yes = Visit(newOp->yes.child(), currentSinkingOp);
-                newOp->no = Visit(newOp->no.child(), currentSinkingOp);
-                newAt = newOp;
-                break;
-            }
+			case OP_TYPE::IM_CONDITIONAL:
+			{
+				// We move the mask creation down the two paths
+				auto newOp = mu::Clone<ASTOpConditional>(at);
+				newOp->yes = Visit(newOp->yes.child(), currentSinkingOp);
+				newOp->no = Visit(newOp->no.child(), currentSinkingOp);
+				newAt = newOp;
+				break;
+			}
 
-            case OP_TYPE::IM_PIXELFORMAT:
-            {
-                // We move the mask creation down format.
-                // \todo: only if shrinking?
-                Ptr<ASTOpImagePixelFormat> newOp = mu::Clone<ASTOpImagePixelFormat>(at);
-                newOp->Source = Visit( newOp->Source.child(), currentSinkingOp);
-                newAt = newOp;
-                break;
-            }
+			case OP_TYPE::IM_PIXELFORMAT:
+			{
+				// We move the mask creation down format.
+				// \todo: only if shrinking?
+				Ptr<ASTOpImagePixelFormat> newOp = mu::Clone<ASTOpImagePixelFormat>(at);
+				newOp->Source = Visit(newOp->Source.child(), currentSinkingOp);
+				newAt = newOp;
+				break;
+			}
 
-            case OP_TYPE::IM_SWITCH:
-            {
-                // We move the mask creation down all the paths
-                auto newOp = mu::Clone<ASTOpSwitch>(at);
-                newOp->def = Visit(newOp->def.child(), currentSinkingOp);
-                for( auto& c:newOp->cases )
-                {
-                    c.branch = Visit(c.branch.child(), currentSinkingOp);
-                }
-                newAt = newOp;
-                break;
-            }
+			case OP_TYPE::IM_SWITCH:
+			{
+				// We move the mask creation down all the paths
+				auto newOp = mu::Clone<ASTOpSwitch>(at);
+				newOp->def = Visit(newOp->def.child(), currentSinkingOp);
+				for (auto& c : newOp->cases)
+				{
+					c.branch = Visit(c.branch.child(), currentSinkingOp);
+				}
+				newAt = newOp;
+				break;
+			}
 
-            case OP_TYPE::IM_SWIZZLE:
-            {
-                auto newOp = mu::Clone<ASTOpFixed>(at);
-                for( int s=0; s<MUTABLE_OP_MAX_SWIZZLE_CHANNELS; ++s )
-                {
-                    auto channelOp = newOp->children[newOp->op.args.ImageSwizzle.sources[s]].child();
-                    if (channelOp)
-                    {
-                        newOp->SetChild( newOp->op.args.ImageSwizzle.sources[s], Visit( channelOp, currentSinkingOp) );
-                    }
-                }
-                newAt = newOp;
-                break;
-            }
+			case OP_TYPE::IM_SWIZZLE:
+			{
+				auto newOp = mu::Clone<ASTOpFixed>(at);
+				for (int s = 0; s < MUTABLE_OP_MAX_SWIZZLE_CHANNELS; ++s)
+				{
+					auto channelOp = newOp->children[newOp->op.args.ImageSwizzle.sources[s]].child();
+					if (channelOp)
+					{
+						newOp->SetChild(newOp->op.args.ImageSwizzle.sources[s], Visit(channelOp, currentSinkingOp));
+					}
+				}
+				newAt = newOp;
+				break;
+			}
 
-            case OP_TYPE::IM_COMPOSE:
-            {
-                // We can only optimise if the layout grid blocks size in pixels is
-                // still an integer after relative scale
-                bool acceptable = false;
-                {
-                    auto typedAt = dynamic_cast<const ASTOpImageCompose*>(at.get());
-                    auto originalBaseOp = typedAt->Base.child();
-
-                    // \todo: recursion-proof cache?
-                    int layoutBlockPixelsX = 0;
-                    int layoutBlockPixelsY = 0;
-                    originalBaseOp->GetLayoutBlockSize( &layoutBlockPixelsX, &layoutBlockPixelsY );
-
-                    int scaledLayoutBlockPixelsX = int(layoutBlockPixelsX*scaleX);
-                    int scaledLayoutBlockPixelsY = int(layoutBlockPixelsY*scaleY);
-                    int unscaledLayoutBlockPixelsX = int(scaledLayoutBlockPixelsX/scaleX);
-                    int unscaledLayoutBlockPixelsY = int(scaledLayoutBlockPixelsY/scaleY);
-                    acceptable =
-                            ( layoutBlockPixelsX!=0 && layoutBlockPixelsY!=0 )
-                            &&
-                            ( layoutBlockPixelsX == unscaledLayoutBlockPixelsX )
-                            &&
-                            ( layoutBlockPixelsY == unscaledLayoutBlockPixelsY );
-                }
+			case OP_TYPE::IM_COMPOSE:
+			{
+				// We can only optimise if the layout grid blocks size in pixels is
+				// still an integer after relative scale
+				bool acceptable = false;
+				{
+					auto typedAt = dynamic_cast<const ASTOpImageCompose*>(at.get());
+					auto originalBaseOp = typedAt->Base.child();
+
+					// \todo: recursion-proof cache?
+					int layoutBlockPixelsX = 0;
+					int layoutBlockPixelsY = 0;
+					originalBaseOp->GetLayoutBlockSize(&layoutBlockPixelsX, &layoutBlockPixelsY);
+
+					int scaledLayoutBlockPixelsX = int(layoutBlockPixelsX * scaleX);
+					int scaledLayoutBlockPixelsY = int(layoutBlockPixelsY * scaleY);
+					int unscaledLayoutBlockPixelsX = int(scaledLayoutBlockPixelsX / scaleX);
+					int unscaledLayoutBlockPixelsY = int(scaledLayoutBlockPixelsY / scaleY);
+					acceptable =
+						(layoutBlockPixelsX != 0 && layoutBlockPixelsY != 0)
+						&&
+						(layoutBlockPixelsX == unscaledLayoutBlockPixelsX)
+						&&
+						(layoutBlockPixelsY == unscaledLayoutBlockPixelsY);
+				}
 
-                if (acceptable)
-                {
-                    auto newOp = mu::Clone<ASTOpImageCompose>(at);
+				if (acceptable)
+				{
+					auto newOp = mu::Clone<ASTOpImageCompose>(at);
 
-                    auto baseOp = newOp->Base.child();
-                    newOp->Base = Visit( baseOp, currentSinkingOp);
+					auto baseOp = newOp->Base.child();
+					newOp->Base = Visit(baseOp, currentSinkingOp);
 
-                    auto blockOp = newOp->BlockImage.child();
-                    newOp->BlockImage = Visit( blockOp, currentSinkingOp);
+					auto blockOp = newOp->BlockImage.child();
+					newOp->BlockImage = Visit(blockOp, currentSinkingOp);
 
-                    newAt = newOp;
-                }
+					newAt = newOp;
+				}
 
-                break;
-            }
+				break;
+			}
 
-            case OP_TYPE::IM_PATCH:
-            {
-                auto newOp = mu::Clone<ASTOpImagePatch>(at);
+			case OP_TYPE::IM_PATCH:
+			{
+				auto newOp = mu::Clone<ASTOpImagePatch>(at);
 
-                newOp->base = Visit( newOp->base.child(), currentSinkingOp );
-                newOp->patch = Visit( newOp->patch.child(), currentSinkingOp);
+				newOp->base = Visit(newOp->base.child(), currentSinkingOp);
+				newOp->patch = Visit(newOp->patch.child(), currentSinkingOp);
 
-                // todo: review if this is always correct, or we need some "divisible" check
-                newOp->location[0] = uint16( newOp->location[0] *  scaleX );
-                newOp->location[1] = uint16( newOp->location[1] *  scaleY );
+				// todo: review if this is always correct, or we need some "divisible" check
+				newOp->location[0] = uint16(newOp->location[0] * scaleX);
+				newOp->location[1] = uint16(newOp->location[1] * scaleY);
 
-                newAt = newOp;
+				newAt = newOp;
 
-                break;
-            }
+				break;
+			}
 
-            case OP_TYPE::IM_MIPMAP:
-            {
-                auto newOp = mu::Clone<ASTOpImageMipmap>(at.get());
-                auto baseOp = newOp->Source.child();
-                newOp->Source = Visit( baseOp, currentSinkingOp );
-                newAt = newOp;
-                break;
-            }
+			case OP_TYPE::IM_MIPMAP:
+			{
+				auto newOp = mu::Clone<ASTOpImageMipmap>(at.get());
+				auto baseOp = newOp->Source.child();
+				newOp->Source = Visit(baseOp, currentSinkingOp);
+				newAt = newOp;
+				break;
+			}
 
-            case OP_TYPE::IM_INTERPOLATE:
-            {
-                auto newOp = mu::Clone<ASTOpFixed>(at);
+			case OP_TYPE::IM_INTERPOLATE:
+			{
+				auto newOp = mu::Clone<ASTOpFixed>(at);
 
-                for ( int v=0; v<MUTABLE_OP_MAX_INTERPOLATE_COUNT; ++v )
-                {
-                    auto child = newOp->children[ newOp->op.args.ImageInterpolate.targets[v] ].child();
-                    auto bOp = Visit( child, currentSinkingOp );
-                    newOp->SetChild( newOp->op.args.ImageInterpolate.targets[v], bOp );
-                }
+				for (int v = 0; v < MUTABLE_OP_MAX_INTERPOLATE_COUNT; ++v)
+				{
+					auto child = newOp->children[newOp->op.args.ImageInterpolate.targets[v]].child();
+					auto bOp = Visit(child, currentSinkingOp);
+					newOp->SetChild(newOp->op.args.ImageInterpolate.targets[v], bOp);
+				}
 
-                newAt = newOp;
-                break;
-            }
+				newAt = newOp;
+				break;
+			}
 
 
-            case OP_TYPE::IM_INTERPOLATE3:
-            {
-                auto newOp = mu::Clone<ASTOpFixed>(at);
+			case OP_TYPE::IM_INTERPOLATE3:
+			{
+				auto newOp = mu::Clone<ASTOpFixed>(at);
 
-                auto top0 = newOp->children[ newOp->op.args.ImageInterpolate3.target0 ].child();
-                newOp->SetChild( newOp->op.args.ImageInterpolate3.target0, Visit(top0, currentSinkingOp) );
+				auto top0 = newOp->children[newOp->op.args.ImageInterpolate3.target0].child();
+				newOp->SetChild(newOp->op.args.ImageInterpolate3.target0, Visit(top0, currentSinkingOp));
 
-                auto top1 = newOp->children[ newOp->op.args.ImageInterpolate3.target1 ].child();
-                newOp->SetChild( newOp->op.args.ImageInterpolate3.target1, Visit(top1, currentSinkingOp) );
+				auto top1 = newOp->children[newOp->op.args.ImageInterpolate3.target1].child();
+				newOp->SetChild(newOp->op.args.ImageInterpolate3.target1, Visit(top1, currentSinkingOp));
 
-                auto top2 = newOp->children[ newOp->op.args.ImageInterpolate3.target2 ].child();
-                newOp->SetChild( newOp->op.args.ImageInterpolate3.target2, Visit(top2, currentSinkingOp) );
+				auto top2 = newOp->children[newOp->op.args.ImageInterpolate3.target2].child();
+				newOp->SetChild(newOp->op.args.ImageInterpolate3.target2, Visit(top2, currentSinkingOp));
 
-                newAt = newOp;
-                break;
-            }
+				newAt = newOp;
+				break;
+			}
 
-            case OP_TYPE::IM_MULTILAYER:
-            {
-                auto nop = mu::Clone<ASTOpImageMultiLayer>(at);
-                nop->base = Visit( nop->base.child(), currentSinkingOp );
-                nop->mask = Visit( nop->mask.child(), currentSinkingOp );
-                nop->blend = Visit( nop->blend.child(), currentSinkingOp );
-                newAt = nop;
-                break;
-            }
+			case OP_TYPE::IM_MULTILAYER:
+			{
+				auto nop = mu::Clone<ASTOpImageMultiLayer>(at);
+				nop->base = Visit(nop->base.child(), currentSinkingOp);
+				nop->mask = Visit(nop->mask.child(), currentSinkingOp);
+				nop->blend = Visit(nop->blend.child(), currentSinkingOp);
+				newAt = nop;
+				break;
+			}
 
-            case OP_TYPE::IM_LAYER:
-            {
-                auto nop = mu::Clone<ASTOpFixed>(at);
+			case OP_TYPE::IM_LAYER:
+			{
+				auto nop = mu::Clone<ASTOpFixed>(at);
 
-                auto aOp = nop->children[ nop->op.args.ImageLayer.base ].child();
-                nop->SetChild( nop->op.args.ImageLayer.base, Visit( aOp, currentSinkingOp ) );
+				auto aOp = nop->children[nop->op.args.ImageLayer.base].child();
+				nop->SetChild(nop->op.args.ImageLayer.base, Visit(aOp, currentSinkingOp));
 
-                auto bOp = nop->children[ nop->op.args.ImageLayer.blended ].child();
-                nop->SetChild( nop->op.args.ImageLayer.blended, Visit( bOp, currentSinkingOp ) );
+				auto bOp = nop->children[nop->op.args.ImageLayer.blended].child();
+				nop->SetChild(nop->op.args.ImageLayer.blended, Visit(bOp, currentSinkingOp));
 
-                auto mOp = nop->children[ nop->op.args.ImageLayer.mask ].child();
-                nop->SetChild( nop->op.args.ImageLayer.mask, Visit( mOp, currentSinkingOp ) );
+				auto mOp = nop->children[nop->op.args.ImageLayer.mask].child();
+				nop->SetChild(nop->op.args.ImageLayer.mask, Visit(mOp, currentSinkingOp));
 
-                newAt = nop;
-                break;
-            }
+				newAt = nop;
+				break;
+			}
 
-            case OP_TYPE::IM_LAYERCOLOUR:
-            {
-                auto nop = mu::Clone<ASTOpFixed>(at);
+			case OP_TYPE::IM_LAYERCOLOUR:
+			{
+				auto nop = mu::Clone<ASTOpFixed>(at);
 
-                auto aOp = nop->children[ nop->op.args.ImageLayerColour.base ].child();
-                nop->SetChild( nop->op.args.ImageLayerColour.base, Visit( aOp, currentSinkingOp ) );
+				auto aOp = nop->children[nop->op.args.ImageLayerColour.base].child();
+				nop->SetChild(nop->op.args.ImageLayerColour.base, Visit(aOp, currentSinkingOp));
 
-                auto mOp = nop->children[ nop->op.args.ImageLayerColour.mask ].child();
-                nop->SetChild( nop->op.args.ImageLayer.mask, Visit( mOp, currentSinkingOp ) );
+				auto mOp = nop->children[nop->op.args.ImageLayerColour.mask].child();
+				nop->SetChild(nop->op.args.ImageLayer.mask, Visit(mOp, currentSinkingOp));
 
-                newAt = nop;
-                break;
-            }
+				newAt = nop;
+				break;
+			}
 
-            case OP_TYPE::IM_DISPLACE:
-            {
-                auto nop = mu::Clone<ASTOpFixed>(at);
+			case OP_TYPE::IM_DISPLACE:
+			{
+				auto nop = mu::Clone<ASTOpFixed>(at);
 
-                auto sourceOp = nop->children[ nop->op.args.ImageDisplace.source ].child();
-                nop->SetChild( nop->op.args.ImageDisplace.source, Visit(sourceOp, currentSinkingOp ) );
+				auto sourceOp = nop->children[nop->op.args.ImageDisplace.source].child();
+				nop->SetChild(nop->op.args.ImageDisplace.source, Visit(sourceOp, currentSinkingOp));
 
-                auto mapOp = nop->children[ nop->op.args.ImageDisplace.displacementMap ].child();
-                nop->SetChild( nop->op.args.ImageDisplace.displacementMap, Visit(mapOp, currentSinkingOp ) );
+				auto mapOp = nop->children[nop->op.args.ImageDisplace.displacementMap].child();
+				nop->SetChild(nop->op.args.ImageDisplace.displacementMap, Visit(mapOp, currentSinkingOp));
 
-                // Make sure we don't scale a constant displacement map, which is very wrong.
-                if ( mapOp->GetOpType()==OP_TYPE::IM_MAKEGROWMAP )
-                {
-                    newAt = nop;
-                }
-                else
-                {
-                    // We cannot resize an already calculated displacement map.
-                }
+				// Make sure we don't scale a constant displacement map, which is very wrong.
+				if (mapOp->GetOpType() == OP_TYPE::IM_MAKEGROWMAP)
+				{
+					newAt = nop;
+				}
+				else
+				{
+					// We cannot resize an already calculated displacement map.
+				}
 
-                break;
-            }
+				break;
+			}
 
-            case OP_TYPE::IM_MAKEGROWMAP:
-            {
-                auto nop = mu::Clone<ASTOpFixed>(at);
-                auto maskOp = nop->children[ nop->op.args.ImageMakeGrowMap.mask ].child();
-                nop->SetChild( nop->op.args.ImageMakeGrowMap.mask, Visit( maskOp, currentSinkingOp ) );
-                newAt = nop;
-                break;
-            }
+			case OP_TYPE::IM_MAKEGROWMAP:
+			{
+				auto nop = mu::Clone<ASTOpFixed>(at);
+				auto maskOp = nop->children[nop->op.args.ImageMakeGrowMap.mask].child();
+				nop->SetChild(nop->op.args.ImageMakeGrowMap.mask, Visit(maskOp, currentSinkingOp));
+				newAt = nop;
+				break;
+			}
 
-            case OP_TYPE::IM_RASTERMESH:
-            {
-                auto nop = mu::Clone<ASTOpFixed>(at);
-                auto maskOp = nop->children[ nop->op.args.ImageRasterMesh.mask ].child();
-                nop->SetChild( nop->op.args.ImageRasterMesh.mask, Visit( maskOp, currentSinkingOp ) );
+			case OP_TYPE::IM_RASTERMESH:
+			{
+				auto nop = mu::Clone<ASTOpFixed>(at);
+				auto maskOp = nop->children[nop->op.args.ImageRasterMesh.mask].child();
+				nop->SetChild(nop->op.args.ImageRasterMesh.mask, Visit(maskOp, currentSinkingOp));
 
 				// Resize the image to project as well, assuming that since the target has a different resolution
 				// it make sense for the source image to have a similar resize.
@@ -2268,28 +2195,28 @@ namespace mu
 				//auto imageOp = nop->children[nop->op.args.ImageRasterMesh.image].child();
 				//nop->SetChild(nop->op.args.ImageRasterMesh.image, Visit(imageOp, currentSinkingOp));
 
-                nop->op.args.ImageRasterMesh.sizeX = uint16( nop->op.args.ImageRasterMesh.sizeX * scaleX + 0.5f );
-                nop->op.args.ImageRasterMesh.sizeY = uint16( nop->op.args.ImageRasterMesh.sizeY * scaleY + 0.5f );
-                newAt = nop;
-                break;
-            }
+				nop->op.args.ImageRasterMesh.sizeX = uint16(nop->op.args.ImageRasterMesh.sizeX * scaleX + 0.5f);
+				nop->op.args.ImageRasterMesh.sizeY = uint16(nop->op.args.ImageRasterMesh.sizeY * scaleY + 0.5f);
+				newAt = nop;
+				break;
+			}
 
-            default:
-                break;
-            }
+			default:
+				break;
+			}
 
-            // end on line, replace with sinking op
-            if (at==newAt && at!=m_initialSource)
-            {
-                Ptr<ASTOpFixed> newOp = mu::Clone<ASTOpFixed>(currentSinkingOp);
-                check(newOp->GetOpType()==OP_TYPE::IM_RESIZEREL);
+			// end on line, replace with sinking op
+			if (at == newAt && at != m_initialSource)
+			{
+				Ptr<ASTOpFixed> newOp = mu::Clone<ASTOpFixed>(currentSinkingOp);
+				check(newOp->GetOpType() == OP_TYPE::IM_RESIZEREL);
 
-                newOp->SetChild( newOp->op.args.ImageResizeRel.source,at);
+				newOp->SetChild(newOp->op.args.ImageResizeRel.source, at);
 
-                newAt = newOp;
-            }
+				newAt = newOp;
+			}
 
-            m_oldToNew[currentSinkingOp][at] = newAt;
+			OldToNew.Add({ at, currentSinkingOp }, newAt);
 
             return newAt;
         }
diff --git a/Engine/Plugins/Experimental/Mutable/Source/MutableTools/Private/MuT/CompilerPrivate.h b/Engine/Plugins/Experimental/Mutable/Source/MutableTools/Private/MuT/CompilerPrivate.h
--- a/Engine/Plugins/Experimental/Mutable/Source/MutableTools/Private/MuT/CompilerPrivate.h
+++ b/Engine/Plugins/Experimental/Mutable/Source/MutableTools/Private/MuT/CompilerPrivate.h
@@ -67,7 +67,7 @@ namespace mu
     public:
 
         //! Detailed optimization options
-        MODEL_OPTIMIZATION_OPTIONS m_optimisationOptions;
+        FModelOptimizationOptions m_optimisationOptions;
 
         bool m_enablePartialOptimise = false;
 		bool m_enableConcurrency = false;
diff --git a/Engine/Plugins/Experimental/Mutable/Source/MutableTools/Private/MuT/DataPacker.cpp b/Engine/Plugins/Experimental/Mutable/Source/MutableTools/Private/MuT/DataPacker.cpp
--- a/Engine/Plugins/Experimental/Mutable/Source/MutableTools/Private/MuT/DataPacker.cpp
+++ b/Engine/Plugins/Experimental/Mutable/Source/MutableTools/Private/MuT/DataPacker.cpp
@@ -24,6 +24,7 @@
 #include "MuT/ASTOpMeshRemoveMask.h"
 #include "MuT/ASTOpMeshDifference.h"
 #include "MuT/ASTOpMeshMorph.h"
+#include "MuT/ASTOpLayoutFromMesh.h"
 
 #include <cstdint>
 #include <memory>
@@ -625,21 +626,34 @@ namespace mu
                 break;
             }
 
-            case OP_TYPE::ME_EXTRACTFACEGROUP:
-            {
-                recurse = false;
+			case OP_TYPE::ME_EXTRACTFACEGROUP:
+			{
+				recurse = false;
 
 				const ASTOpFixed* op = dynamic_cast<const ASTOpFixed*>(node.get());
 
-                // todo: check if we really need all of them
-                uint64_t newState = currentSemantics;
-                newState |= (UINT64_C(1)<<MBS_LAYOUTBLOCK);
-                newState |= (UINT64_C(1)<<MBS_CHART);
-                newState |= (UINT64_C(1)<<MBS_VERTEXINDEX);
+				// todo: check if we really need all of them
+				uint64_t newState = currentSemantics;
+				newState |= (UINT64_C(1) << MBS_LAYOUTBLOCK);
+				newState |= (UINT64_C(1) << MBS_CHART);
+				newState |= (UINT64_C(1) << MBS_VERTEXINDEX);
 
-                RecurseWithState( op->children[op->op.args.MeshExtractFaceGroup.source].child(), newState );
-                break;
-            }
+				RecurseWithState(op->children[op->op.args.MeshExtractFaceGroup.source].child(), newState);
+				break;
+			}
+
+			case OP_TYPE::LA_FROMMESH:
+			{
+				recurse = false;
+
+				const ASTOpLayoutFromMesh* op = dynamic_cast<const ASTOpLayoutFromMesh*>(node.get());
+
+				uint64_t newState = currentSemantics;
+				newState |= (UINT64_C(1) << MBS_LAYOUTBLOCK);
+
+				RecurseWithState(op->Mesh.child(), newState);
+				break;
+			}
 
             case OP_TYPE::IN_ADDMESH:
             {
@@ -720,7 +734,7 @@ namespace mu
 
     //---------------------------------------------------------------------------------------------
     void DataOptimiseAST( int imageCompressionQuality, ASTOpList& roots,
-                          const MODEL_OPTIMIZATION_OPTIONS& options )
+                          const FModelOptimizationOptions& options )
     {
         // Images
         AccumulateImageFormatsAST accFormat;
diff --git a/Engine/Plugins/Experimental/Mutable/Source/MutableTools/Private/MuT/DataPacker.h b/Engine/Plugins/Experimental/Mutable/Source/MutableTools/Private/MuT/DataPacker.h
--- a/Engine/Plugins/Experimental/Mutable/Source/MutableTools/Private/MuT/DataPacker.h
+++ b/Engine/Plugins/Experimental/Mutable/Source/MutableTools/Private/MuT/DataPacker.h
@@ -21,7 +21,7 @@ namespace mu
     //! Convert constant data to different formats, based on their usage
     //---------------------------------------------------------------------------------------------
     extern void DataOptimiseAST( int imageCompressionQuality, ASTOpList& roots,
-                                 const MODEL_OPTIMIZATION_OPTIONS& );
+                                 const FModelOptimizationOptions& );
 
     //---------------------------------------------------------------------------------------------
     //! Find the given constant in a subtree
diff --git a/Engine/Plugins/Experimental/Mutable/Source/MutableTools/Private/MuT/ErrorLog.cpp b/Engine/Plugins/Experimental/Mutable/Source/MutableTools/Private/MuT/ErrorLog.cpp
--- a/Engine/Plugins/Experimental/Mutable/Source/MutableTools/Private/MuT/ErrorLog.cpp
+++ b/Engine/Plugins/Experimental/Mutable/Source/MutableTools/Private/MuT/ErrorLog.cpp
@@ -314,7 +314,8 @@ namespace mu
 
 		TEXT("LA_PACK          "),
 		TEXT("LA_MERGE         "),
-        TEXT("LA_REMOVEBLOCKS  ")
+		TEXT("LA_REMOVEBLOCKS  "),
+		TEXT("LA_FROMMESH	   "),
 	};
 
     // clang-format on
