[Mutable] Group Projector Node add Image Pin.

- Group Projector node now have an Image Pin wich is used as the projected image.
- Renamed Instance/Descriptor CurrentParamRange to GetProjectorValueRange.
- Group Projector Node moved backwards compatible code to the BackwardsCompatibleFixup function.
- Fixed Group Project Node not working with Extend Material.
- Fixed Group Projector Node crash with Extend Material.

#rnx

diff --git a/Engine/Plugins/Experimental/Mutable/Source/CustomizableObject/Classes/MuCO/CustomizableObjectInstance.h b/Engine/Plugins/Experimental/Mutable/Source/CustomizableObject/Classes/MuCO/CustomizableObjectInstance.h
--- a/Engine/Plugins/Experimental/Mutable/Source/CustomizableObject/Classes/MuCO/CustomizableObjectInstance.h
+++ b/Engine/Plugins/Experimental/Mutable/Source/CustomizableObject/Classes/MuCO/CustomizableObjectInstance.h
@@ -312,7 +312,7 @@ public:
 	bool IsParamMultidimensional(const FString& ParamName) const;
 
 	UFUNCTION(BlueprintCallable, Category = CustomizableObjectInstance)
-	int32 CurrentParamRange(const FString& ParamName) const;
+	int32 GetProjectorValueRange(const FString& ParamName) const;
 	
 	bool IsParamMultidimensional(int32 ParamIndex) const;
 	FCustomizableObjectProjector GetProjectorDefaultValue(int32 ParamIndex) const;
diff --git a/Engine/Plugins/Experimental/Mutable/Source/CustomizableObject/Private/MuCO/CustomizableObjectInstance.cpp b/Engine/Plugins/Experimental/Mutable/Source/CustomizableObject/Private/MuCO/CustomizableObjectInstance.cpp
--- a/Engine/Plugins/Experimental/Mutable/Source/CustomizableObject/Private/MuCO/CustomizableObjectInstance.cpp
+++ b/Engine/Plugins/Experimental/Mutable/Source/CustomizableObject/Private/MuCO/CustomizableObjectInstance.cpp
@@ -565,9 +565,9 @@ bool UCustomizableObjectInstance::IsParamMultidimensional(const FString& ParamNa
 }
 
 
-int32 UCustomizableObjectInstance::CurrentParamRange(const FString& ParamName) const
+int32 UCustomizableObjectInstance::GetProjectorValueRange(const FString& ParamName) const
 {
-	return Descriptor.CurrentParamRange(ParamName);
+	return Descriptor.GetProjectorValueRange(ParamName);
 }
 
 
diff --git a/Engine/Plugins/Experimental/Mutable/Source/CustomizableObject/Private/MuCO/CustomizableObjectInstanceDescriptor.cpp b/Engine/Plugins/Experimental/Mutable/Source/CustomizableObject/Private/MuCO/CustomizableObjectInstanceDescriptor.cpp
--- a/Engine/Plugins/Experimental/Mutable/Source/CustomizableObject/Private/MuCO/CustomizableObjectInstanceDescriptor.cpp
+++ b/Engine/Plugins/Experimental/Mutable/Source/CustomizableObject/Private/MuCO/CustomizableObjectInstanceDescriptor.cpp
@@ -1835,22 +1835,17 @@ bool FCustomizableObjectInstanceDescriptor::IsParamMultidimensional(const int32
 }
 
 
-int32 FCustomizableObjectInstanceDescriptor::CurrentParamRange(const FString& ParamName) const
+int32 FCustomizableObjectInstanceDescriptor::GetProjectorValueRange(const FString& ParamName) const
 {
 	check(CustomizableObject);
 
-	const int32 ParameterIndexInObject = CustomizableObject->FindParameter(ParamName);
-
 	const int32 ProjectorParamIndex = FindProjectorParameterNameIndex(ParamName);
-
-	if (ParameterIndexInObject >= 0 && ProjectorParamIndex >= 0)
+	if (ProjectorParamIndex < 0)
 	{
-		return ProjectorParameters[ProjectorParamIndex].RangeValues.Num();
-	}
-	else
-	{
-		return 0;
+		return -1;
 	}
+
+	return ProjectorParameters[ProjectorParamIndex].RangeValues.Num();
 }
 
 
diff --git a/Engine/Plugins/Experimental/Mutable/Source/CustomizableObject/Private/MuCO/MultilayerProjector.cpp b/Engine/Plugins/Experimental/Mutable/Source/CustomizableObject/Private/MuCO/MultilayerProjector.cpp
--- a/Engine/Plugins/Experimental/Mutable/Source/CustomizableObject/Private/MuCO/MultilayerProjector.cpp
+++ b/Engine/Plugins/Experimental/Mutable/Source/CustomizableObject/Private/MuCO/MultilayerProjector.cpp
@@ -27,9 +27,11 @@ void FMultilayerProjectorLayer::Read(const FCustomizableObjectInstanceDescriptor
 		Scale = static_cast<FVector3d>(Projector.Scale);
 		Angle = Projector.Angle;
 	}
-	
+
+	// Image Parameter can not exist due to the Group Projector Node providing the image through its own Image Pin.
+	if (const int32 ImageParamIndex = Descriptor.FindIntParameterNameIndex(ParamName + FMultilayerProjector::IMAGE_PARAMETER_POSTFIX);
+		ImageParamIndex >= 0)
 	{
-		const int32 ImageParamIndex = Descriptor.FindIntParameterNameIndex(ParamName + FMultilayerProjector::IMAGE_PARAMETER_POSTFIX);
 		Image = Descriptor.GetIntParameters()[ImageParamIndex].ParameterRangeValueNames[Index];
 	}
 
@@ -56,9 +58,11 @@ void FMultilayerProjectorLayer::Write(FCustomizableObjectInstanceDescriptor& Des
 		Projector.Scale = static_cast<FVector3f>(Scale);
 		Projector.Angle = Angle;
 	}
-	
+
+	// Image Parameter can not exist due to the Group Projector Node providing the image through its own Image Pin.
+	if (const int32 ImageParamIndex = Descriptor.FindIntParameterNameIndex(ParamName + FMultilayerProjector::IMAGE_PARAMETER_POSTFIX);
+		ImageParamIndex >= 0)
 	{
-		const int32 ImageParamIndex = Descriptor.FindIntParameterNameIndex(ParamName + FMultilayerProjector::IMAGE_PARAMETER_POSTFIX);
 		Descriptor.GetIntParameters()[ImageParamIndex].ParameterRangeValueNames[Index] = Image;
 	}
 
@@ -136,10 +140,10 @@ void FMultilayerProjector::CreateLayer(FCustomizableObjectInstanceDescriptor& De
 		ProjectorParameter.RangeValues.Insert(Projector, Index);
 	}
 	
-	// Selected Image Range.
+	// Selected Image Range. Image Parameter can not exist due to the Group Projector Node providing the image through its own Image Pin.
+	if (const int32 IntParameterIndex = Descriptor.FindIntParameterNameIndex(ParamName.ToString() + FMultilayerProjector::IMAGE_PARAMETER_POSTFIX);
+		IntParameterIndex >= 0)
 	{
-		const int32 IntParameterIndex = Descriptor.FindIntParameterNameIndex(ParamName.ToString() + FMultilayerProjector::IMAGE_PARAMETER_POSTFIX);
-
 		FCustomizableObjectIntParameterValue& IntParameter = Descriptor.GetIntParameters()[IntParameterIndex];
 		const int32 ParamIndexInObject = Object->FindParameter(IntParameter.ParameterName);
 
@@ -178,10 +182,10 @@ void FMultilayerProjector::RemoveLayerAt(FCustomizableObjectInstanceDescriptor&
 		ProjectorParameter.RangeValues.RemoveAt(Index);
 	}
 	
-	// Selected Image Range.
-	{
-		const int32 IntParameterIndex = Descriptor.FindIntParameterNameIndex(ParamName.ToString() + IMAGE_PARAMETER_POSTFIX);
-		
+	// Selected Image Range. Image Parameter can not exist due to the Group Projector Node providing the image through its own Image Pin.
+	if (const int32 IntParameterIndex = Descriptor.FindIntParameterNameIndex(ParamName.ToString() + FMultilayerProjector::IMAGE_PARAMETER_POSTFIX);
+		IntParameterIndex >= 0)
+	{		
 		FCustomizableObjectIntParameterValue& IntParameter = Descriptor.GetIntParameters()[IntParameterIndex];
 		IntParameter.ParameterRangeValueNames.RemoveAt(Index);
 	}
@@ -455,11 +459,7 @@ void FMultilayerProjector::CheckDescriptorParameters(const FCustomizableObjectIn
 		check(ProjectorParameterIndex >= 0) // Descriptor Parameter does not exist.
 	}
 	
-	// Selected Image.
-	{
-		const int32 IntParameterIndex = Descriptor.FindIntParameterNameIndex(ParamNameString + IMAGE_PARAMETER_POSTFIX);
-		check(IntParameterIndex >= 0) // Descriptor Parameter does not exist.
-	}
+	// Selected Image. No longer checked since the image can be provided by its own Node Group Projector pin.
 	
 	// Opacity.
 	{
@@ -492,16 +492,9 @@ bool FMultilayerProjector::AreDescriptorParametersValid(const FCustomizableObjec
 		}
 	}
 	
-	// Selected Image.
-	{
-		const int32 IntParameterIndex = Descriptor.FindIntParameterNameIndex(ParamNameString + IMAGE_PARAMETER_POSTFIX);
-		if (IntParameterIndex < 0)
-		{
-			return false;
-		}
-	}
+	// Selected Image. No longer checked since the image can be provided by its own Node Group Projector pin.
 	
-	// Opacity.
+	// Opacity.	
 	{
 		const int32 FloatParameterIndex = Descriptor.FindFloatParameterNameIndex(ParamNameString + OPACITY_PARAMETER_POSTFIX);
 		if (FloatParameterIndex < 0)
diff --git a/Engine/Plugins/Experimental/Mutable/Source/CustomizableObject/Public/MuCO/CustomizableObjectCustomVersion.h b/Engine/Plugins/Experimental/Mutable/Source/CustomizableObject/Public/MuCO/CustomizableObjectCustomVersion.h
--- a/Engine/Plugins/Experimental/Mutable/Source/CustomizableObject/Public/MuCO/CustomizableObjectCustomVersion.h
+++ b/Engine/Plugins/Experimental/Mutable/Source/CustomizableObject/Public/MuCO/CustomizableObjectCustomVersion.h
@@ -61,6 +61,8 @@ struct CUSTOMIZABLEOBJECT_API FCustomizableObjectCustomVersion
 		DescriptorMultilayerProjectors,
 
 		DeformSkeletonOptionsAdded,
+
+		GroupProjectorImagePin,
 		
 		// -----<new versions can be added above this line>-------------------------------------------------
 		VersionPlusOne,
diff --git a/Engine/Plugins/Experimental/Mutable/Source/CustomizableObject/Public/MuCO/CustomizableObjectInstanceDescriptor.h b/Engine/Plugins/Experimental/Mutable/Source/CustomizableObject/Public/MuCO/CustomizableObjectInstanceDescriptor.h
--- a/Engine/Plugins/Experimental/Mutable/Source/CustomizableObject/Public/MuCO/CustomizableObjectInstanceDescriptor.h
+++ b/Engine/Plugins/Experimental/Mutable/Source/CustomizableObject/Public/MuCO/CustomizableObjectInstanceDescriptor.h
@@ -202,7 +202,8 @@ struct CUSTOMIZABLEOBJECT_API FCustomizableObjectInstanceDescriptor
 	/** Returns true if the parameter is multidimensional (has multiple ranges). */
 	bool IsParamMultidimensional(int32 ParamIndex) const;
 
-	int32 CurrentParamRange(const FString& ParamName) const;
+	/** Gets the range of values of the projector with ParamName, returns -1 if the parameter exists, -1 otherwise. */
+	int32 GetProjectorValueRange(const FString& ParamName) const;
 
 	/** Increases the range of values of the integer with ParamName, returns the index of the new integer value, -1 otherwise.
 	 * The added value is initialized with the first integer option and is the last one of the range. */
diff --git a/Engine/Plugins/Experimental/Mutable/Source/CustomizableObject/Public/MuCO/MultilayerProjector.h b/Engine/Plugins/Experimental/Mutable/Source/CustomizableObject/Public/MuCO/MultilayerProjector.h
--- a/Engine/Plugins/Experimental/Mutable/Source/CustomizableObject/Public/MuCO/MultilayerProjector.h
+++ b/Engine/Plugins/Experimental/Mutable/Source/CustomizableObject/Public/MuCO/MultilayerProjector.h
@@ -53,7 +53,7 @@ struct FMultilayerProjectorLayer
 	/** Layer scale. */
 	FVector Scale;
 
-	/** Layer angle*/
+	/** Layer angle. */
 	float Angle;
 
 	/** Layer selected image. */
diff --git a/Engine/Plugins/Experimental/Mutable/Source/CustomizableObjectEditor/Classes/MuCOE/Nodes/CustomizableObjectNode.h b/Engine/Plugins/Experimental/Mutable/Source/CustomizableObjectEditor/Classes/MuCOE/Nodes/CustomizableObjectNode.h
--- a/Engine/Plugins/Experimental/Mutable/Source/CustomizableObjectEditor/Classes/MuCOE/Nodes/CustomizableObjectNode.h
+++ b/Engine/Plugins/Experimental/Mutable/Source/CustomizableObjectEditor/Classes/MuCOE/Nodes/CustomizableObjectNode.h
@@ -59,6 +59,8 @@ DECLARE_MULTICAST_DELEGATE(FPostReconstructNodeDelegate);
 
 DECLARE_MULTICAST_DELEGATE(FNodeConnectionListChangedDelegate);
 
+DECLARE_MULTICAST_DELEGATE_OneParam(FPinConnectionListChangedDelegate, UEdGraphPin*);
+
 DECLARE_MULTICAST_DELEGATE(FDestroyNodeDelegate);
 
 using FRemapPinsDelegateParameter = TMap<UEdGraphPin*, UEdGraphPin*>; // Required for the delegate macro.
@@ -81,20 +83,21 @@ public:
 	GENERATED_BODY()
 
 	// UObject interface
-	void Serialize(FArchive& Ar) override;
+	virtual void Serialize(FArchive& Ar) override;
 	virtual void PostLoad() override; // Do work at PostBackwardsCompatibleFixup.
 
 	// UEdGraphNode interface
-	void AutowireNewNode(UEdGraphPin* FromPin) override;
-	void NodeConnectionListChanged() override;
-	void DestroyNode() override;
+	virtual void AutowireNewNode(UEdGraphPin* FromPin) override;
+	virtual void NodeConnectionListChanged() override;
+	virtual void PinConnectionListChanged(UEdGraphPin* Pin) override;
+	virtual void DestroyNode() override;
 
 	/** Allocates the default pins using the empty remap pins action. Usually called from CreateNode. */
-	void AllocateDefaultPins() override final; // Final. Override AllocateDefaultPins(UCustomizableObjectNodeRemapPins* RemapPins) instead
+	virtual void AllocateDefaultPins() override final; // Final. Override AllocateDefaultPins(UCustomizableObjectNodeRemapPins* RemapPins) instead
 	
 	/** Reconstructs the node using its default remap pins action. */
-	void ReconstructNode() override final; // Final. Override ReconstructNode(UCustomizableObjectNodeRemapPins* RemapPins) instead
-	void PostInitProperties() override;
+	virtual void ReconstructNode() override final; // Final. Override ReconstructNode(UCustomizableObjectNodeRemapPins* RemapPins) instead
+	virtual void PostInitProperties() override;
 
 	// Own interface
 	/** Called at the beginning of the node lifecycle.*/
@@ -247,6 +250,8 @@ public:
 
 	FNodeConnectionListChangedDelegate NodeConnectionListChangedDelegate;
 
+	FPinConnectionListChangedDelegate PinConnectionListChangedDelegate;
+
 	FDestroyNodeDelegate DestroyNodeDelegate;
 
 	FRemapPinsDelegate RemapPinsDelegate;
diff --git a/Engine/Plugins/Experimental/Mutable/Source/CustomizableObjectEditor/Classes/MuCOE/Nodes/CustomizableObjectNodeGroupProjectorParameter.h b/Engine/Plugins/Experimental/Mutable/Source/CustomizableObjectEditor/Classes/MuCOE/Nodes/CustomizableObjectNodeGroupProjectorParameter.h
--- a/Engine/Plugins/Experimental/Mutable/Source/CustomizableObjectEditor/Classes/MuCOE/Nodes/CustomizableObjectNodeGroupProjectorParameter.h
+++ b/Engine/Plugins/Experimental/Mutable/Source/CustomizableObjectEditor/Classes/MuCOE/Nodes/CustomizableObjectNodeGroupProjectorParameter.h
@@ -54,8 +54,6 @@ class CUSTOMIZABLEOBJECTEDITOR_API UCustomizableObjectNodeGroupProjectorParamete
 public:
 	GENERATED_BODY()
 
-	UCustomizableObjectNodeGroupProjectorParameter();
-
 	/** Return array with the sticker name and UTexture2D for projection */
 	TArray<FGroupProjectorParameterImage> GetOptionImagesFromTable() const;
 
@@ -84,7 +82,7 @@ public:
 	UPROPERTY(EditAnywhere, Category = ProjectorGroup, Meta = (ToolTip = "If true, projection textures will be shared between LODs of the same object, and will save memory. Only use if all the LODs share the same UV layout."))
 	bool bShareProjectionTexturesBetweenLODs = false;
 
-	UPROPERTY(EditAnywhere, Category = ProjectorGroup)
+	UPROPERTY(EditAnywhere, Category = ProjectorGroup,  meta = (EditConditionHides))
 	TArray<FGroupProjectorParameterImage> OptionImages;
 
 	UPROPERTY(EditAnywhere, Category = ProjectorGroup)
@@ -97,7 +95,7 @@ public:
 	/** Table where additional option images besides Option Images are read. The elements in this table have priority
 	* over elements from Option Images in case of duplicity. Use the "Data Table Texture Column Name" property to specify
 	* the name of the column where textures are read in the table. */
-	UPROPERTY(EditAnywhere, Category = ProjectorGroup)
+	UPROPERTY(EditAnywhere, Category = ProjectorGroup, meta = (EditConditionHides))
 	TObjectPtr<UDataTable> OptionImagesDataTable = nullptr;
 
 	UPROPERTY(EditAnywhere, Category = ProjectorGroup)
@@ -109,15 +107,20 @@ public:
 	UPROPERTY(EditAnywhere, Category = ProjectorGroup, Meta = (ClampMin = "0"))
 	int32 UVLayout = 0;
 
-	// UObject interface
-	void Serialize(FArchive& Ar) override;
-
-	// Begin EdGraphNode interface
+	// EdGraphNode interface.
 	virtual FText GetNodeTitle(ENodeTitleType::Type TitleType) const override;
-	FLinearColor GetNodeTitleColor() const override;
-	FText GetTooltipText() const override;
+	virtual FLinearColor GetNodeTitleColor() const override;
+	virtual FText GetTooltipText() const override;
 
 	// UCustomizableObjectNode interface.
-	void AllocateDefaultPins(UCustomizableObjectNodeRemapPins* RemapPins);
+	virtual void BackwardsCompatibleFixup() override;
+	virtual void AllocateDefaultPins(UCustomizableObjectNodeRemapPins* RemapPins);
+
+	// Own interface.
+	UEdGraphPin& GetImagePin() const;
+
+private:
+	UPROPERTY()
+	FEdGraphPinReference ImagePin;
 };
 
diff --git a/Engine/Plugins/Experimental/Mutable/Source/CustomizableObjectEditor/Private/MuCOE/CustomizableObjectEditorModule.cpp b/Engine/Plugins/Experimental/Mutable/Source/CustomizableObjectEditor/Private/MuCOE/CustomizableObjectEditorModule.cpp
--- a/Engine/Plugins/Experimental/Mutable/Source/CustomizableObjectEditor/Private/MuCOE/CustomizableObjectEditorModule.cpp
+++ b/Engine/Plugins/Experimental/Mutable/Source/CustomizableObjectEditor/Private/MuCOE/CustomizableObjectEditorModule.cpp
@@ -77,6 +77,7 @@
 #include "MuCOE/Nodes/CustomizableObjectNodeProjectorConstant.h"
 #include "MuCOE/Nodes/CustomizableObjectNodeProjectorParameter.h"
 #include "MuCOE/Nodes/CustomizableObjectNodeProjectorParameterDetails.h"
+#include "MuCOE/Nodes/CustomizableObjectNodeGroupProjectorParameterDetails.h"
 #include "MuCOE/Nodes/CustomizableObjectNodeSkeletalMesh.h"
 #include "MuCOE/Nodes/CustomizableObjectNodeSkeletalMeshDetails.h"
 #include "MuCOE/Nodes/CustomizableObjectNodeStaticMesh.h"
@@ -84,6 +85,7 @@
 #include "MuCOE/Nodes/CustomizableObjectNodeTableDetails.h"
 #include "PropertyEditorDelegates.h"
 #include "PropertyEditorModule.h"
+#include "MuCOE/Nodes/CustomizableObjectNodeGroupProjectorParameter.h"
 #include "Templates/SubclassOf.h"
 #include "Toolkits/AssetEditorToolkit.h"
 #include "UObject/Class.h"
@@ -202,6 +204,7 @@ void FCustomizableObjectEditorModule::StartupModule()
 	RegisterCustomDetails(PropertyModule, UCustomizableObjectNodeObject::StaticClass(), FOnGetDetailCustomizationInstance::CreateStatic(&FCustomizableObjectNodeObjectDetails::MakeInstance));
 	RegisterCustomDetails(PropertyModule, UCustomizableObjectNodeObjectGroup::StaticClass(), FOnGetDetailCustomizationInstance::CreateStatic(&FCustomizableObjectNodeObjectGroupDetails::MakeInstance));
 	RegisterCustomDetails(PropertyModule, UCustomizableObjectNodeProjectorParameter::StaticClass(), FOnGetDetailCustomizationInstance::CreateStatic(&FCustomizableObjectNodeProjectorParameterDetails::MakeInstance));
+	RegisterCustomDetails(PropertyModule, UCustomizableObjectNodeGroupProjectorParameter::StaticClass(), FOnGetDetailCustomizationInstance::CreateStatic(&FCustomizableObjectNodeGroupProjectorParameterDetails::MakeInstance));
 	RegisterCustomDetails(PropertyModule, UCustomizableObjectNodeProjectorConstant::StaticClass(), FOnGetDetailCustomizationInstance::CreateStatic(&FCustomizableObjectNodeProjectorParameterDetails::MakeInstance));
 	RegisterCustomDetails(PropertyModule, UCustomizableObjectNodeMeshMorph::StaticClass(), FOnGetDetailCustomizationInstance::CreateStatic(&FCustomizableObjectNodeMeshMorphDetails::MakeInstance));
 	RegisterCustomDetails(PropertyModule, UCustomizableObjectNodeMeshClipMorph::StaticClass(), FOnGetDetailCustomizationInstance::CreateStatic(&FCustomizableObjectNodeMeshClipMorphDetails::MakeInstance));
diff --git a/Engine/Plugins/Experimental/Mutable/Source/CustomizableObjectEditor/Private/MuCOE/GenerateMutableSource/GenerateMutableSourceGroupProjector.cpp b/Engine/Plugins/Experimental/Mutable/Source/CustomizableObjectEditor/Private/MuCOE/GenerateMutableSource/GenerateMutableSourceGroupProjector.cpp
--- a/Engine/Plugins/Experimental/Mutable/Source/CustomizableObjectEditor/Private/MuCOE/GenerateMutableSource/GenerateMutableSourceGroupProjector.cpp
+++ b/Engine/Plugins/Experimental/Mutable/Source/CustomizableObjectEditor/Private/MuCOE/GenerateMutableSource/GenerateMutableSourceGroupProjector.cpp
@@ -2,6 +2,7 @@
 
 #include "MuCOE/GenerateMutableSource/GenerateMutableSourceGroupProjector.h"
 
+#include "GenerateMutableSourceImage.h"
 #include "Containers/Array.h"
 #include "Containers/BitArray.h"
 #include "Containers/ContainersFwd.h"
@@ -29,6 +30,7 @@
 #include "MuCO/CustomizableObjectUIData.h"
 #include "MuCO/MultilayerProjector.h"
 #include "MuCOE/CustomizableObjectCompiler.h"
+#include "MuCOE/GraphTraversal.h"
 #include "MuCOE/GenerateMutableSource/GenerateMutableSource.h"
 #include "MuCOE/GenerateMutableSource/GenerateMutableSourceProjector.h"
 #include "MuCOE/Nodes/CustomizableObjectNode.h"
@@ -46,6 +48,7 @@
 #include "MuT/NodeImageConstant.h"
 #include "MuT/NodeImageLayerColour.h"
 #include "MuT/NodeImageMultiLayer.h"
+#include "MuT/NodeImageParameter.h"
 #include "MuT/NodeImagePlainColour.h"
 #include "MuT/NodeImageProject.h"
 #include "MuT/NodeImageResize.h"
@@ -69,6 +72,7 @@ class UPoseAsset;
 #define LOCTEXT_NAMESPACE "CustomizableObjectEditor"
 
 
+PRAGMA_DISABLE_OPTIMIZATION
 mu::NodeImagePtr GenerateMutableGroupProjection(const int32 NodeLOD, const int32 ImageIndex, mu::NodeMeshPtr MeshNode, FMutableGraphGenerationContext& GenerationContext,
 	const UCustomizableObjectNodeMaterialBase* NodeMaterialBase, bool& bShareProjectionTexturesBetweenLODs, bool& bIsGroupProjectorImage,
 	UTexture2D*& GroupProjectionReferenceTexture, TMap<FString, float>& TextureNameToProjectionResFactor, FString& AlternateResStateName,
@@ -111,17 +115,9 @@ mu::NodeImagePtr GenerateMutableGroupProjection(const int32 NodeLOD, const int32
 
 		const bool bProjectToImage = [&]
 		{
-			if (TypedNodeMat)
-			{
-				const FString ParameterName = TypedNodeMat->GetParameterName(EMaterialParameterType::Texture, ImageIndex).ToString();
-				return ParameterName == ProjectorTempData.CustomizableObjectNodeGroupProjectorParameter->MaterialChannelNameToConnect;
-			}
-			else
-			{
-				const FGuid ImageId = ParentMaterial->GetParameterId(EMaterialParameterType::Texture, ImageIndex);
-				const FString ImagePinName = TypedNodeExt->GetUsedImagePin(ImageId)->PinName.ToString();
-				return ImagePinName.StartsWith(ProjectorTempData.CustomizableObjectNodeGroupProjectorParameter->MaterialChannelNameToConnect);
-			}
+			const UCustomizableObjectNodeMaterial* NodeMaterial = TypedNodeMat ? TypedNodeMat : ParentMaterial;
+			const FString ParameterName = NodeMaterial->GetParameterName(EMaterialParameterType::Texture, ImageIndex).ToString();
+			return ParameterName == ProjectorTempData.CustomizableObjectNodeGroupProjectorParameter->MaterialChannelNameToConnect;
 		}();
 		
 		if (bProjectToImage)
@@ -143,15 +139,8 @@ mu::NodeImagePtr GenerateMutableGroupProjection(const int32 NodeLOD, const int32
 			{
 				const FString ImageName = [&]
 				{
-					if (TypedNodeMat)
-					{
-						return TypedNodeMat->GetParameterName(EMaterialParameterType::Texture, ImageIndex).ToString();
-					}
-					else
-					{
-						const FGuid ImageId = ParentMaterial->GetParameterId(EMaterialParameterType::Texture, ImageIndex);
-						return TypedNodeExt->GetUsedImagePin(ImageId)->PinName.ToString();
-					}
+					const UCustomizableObjectNodeMaterial* NodeMaterial = TypedNodeMat ? TypedNodeMat : ParentMaterial;
+					return NodeMaterial->GetParameterName(EMaterialParameterType::Texture, ImageIndex).ToString();
 				}();
 				
 				FString msg = FString::Printf(TEXT("Material image [%s] is connected to an image but will be replaced by a Group Projector."), *ImageName);
@@ -229,10 +218,10 @@ mu::NodeImagePtr GenerateMutableGroupProjection(const int32 NodeLOD, const int32
 			mu::NodeImagePtr ImageToProject = ProjectorTempData.NodeImagePtr;
 			{
 				// Apply LODBias to the images to project
-
-				int32 LODBias = ComputeLODBias(GenerationContext, GroupProjectionReferenceTexture, GroupProjectionReferenceTexture ? GroupProjectionReferenceTexture->MaxTextureSize : 0, TypedNodeMat, ImageIndex);
-
-				if (LODBias > 0)
+				const UCustomizableObjectNodeMaterial* NodeMaterial = TypedNodeMat ? TypedNodeMat : ParentMaterial;
+				const int32 MaxTextureSize = GroupProjectionReferenceTexture ? GroupProjectionReferenceTexture->MaxTextureSize : 0;
+				if (const int32 LODBias =  ComputeLODBias(GenerationContext, GroupProjectionReferenceTexture, MaxTextureSize, NodeMaterial, ImageIndex);
+					LODBias > 0)
 				{
 					mu::NodeImageResizePtr ResizeImage = new mu::NodeImageResize();
 					ResizeImage->SetBase(ImageToProject.get());
@@ -241,7 +230,6 @@ mu::NodeImagePtr GenerateMutableGroupProjection(const int32 NodeLOD, const int32
 					ResizeImage->SetSize(factor, factor);
 					ImageToProject = ResizeImage;
 				}
-
 			}
 
 			ImageNode->SetImage(ImageToProject);
@@ -440,7 +428,7 @@ mu::NodeImagePtr GenerateMutableGroupProjection(const int32 NodeLOD, const int32
 
 	return FinalSwizzleNode;
 }
-
+PRAGMA_ENABLE_OPTIMIZATION
 
 bool GenerateMutableSourceGroupProjector(const UEdGraphPin* Pin, FMutableGraphGenerationContext& GenerationContext, const UCustomizableObjectNodeObjectGroup* originalGroup)
 {
@@ -502,21 +490,7 @@ bool GenerateMutableSourceGroupProjector(const UEdGraphPin* Pin, FMutableGraphGe
 			GroupProjectorTempData.NodeProjectorParameterPtr->SetRangeCount(1);
 			GroupProjectorTempData.NodeProjectorParameterPtr->SetRange(0, NodeRangeFromScalar);
 
-			mu::NodeScalarEnumParameterPtr EnumParameterNode = new mu::NodeScalarEnumParameter;
-			FString NodeEnumParamName = ProjParamNode->ParameterName + FMultilayerProjector::IMAGE_PARAMETER_POSTFIX;
-			EnumParameterNode->SetName(TCHAR_TO_ANSI(*(NodeEnumParamName)));
-			EnumParameterNode->SetUid(TCHAR_TO_ANSI(*(ProjectorParamUid + FString("_SI"))));
-			GenerationContext.AddParameterNameUnique(originalGroup, NodeEnumParamName);
-			EnumParameterNode->SetValueCount(ArrayOptionImage.Num());
-			EnumParameterNode->SetDefaultValueIndex(0);
-			EnumParameterNode->SetRangeCount(1);
-			EnumParameterNode->SetRange(0, NodeRangeFromScalar);
-
-
-			FParameterUIData ParameterUIData(NodeEnumParamName, ProjParamNode->ParamUIMetadata, EMutableParameterType::Int);
-			ParameterUIData.IntegerParameterGroupType = ECustomizableObjectGroupType::COGT_ONE;
-			ParameterUIData.ParamUIMetadata.ExtraInformation.Add(FString("UseThumbnails"));
-
+			
 			mu::NodeScalarEnumParameterPtr PoseEnumParameterNode = new mu::NodeScalarEnumParameter;
 			FString PoseNodeEnumParamName = ProjParamNode->ParameterName + FMultilayerProjector::POSE_PARAMETER_POSTFIX;
 			PoseEnumParameterNode->SetName(TCHAR_TO_ANSI(*(PoseNodeEnumParamName)));
@@ -555,19 +529,6 @@ bool GenerateMutableSourceGroupProjector(const UEdGraphPin* Pin, FMutableGraphGe
 				return false;
 			}
 
-			for (int ImageIndex = 0; ImageIndex < ArrayOptionImage.Num(); ++ImageIndex)
-			{
-				EnumParameterNode->SetValue(ImageIndex, (float)ImageIndex, TCHAR_TO_ANSI(*ArrayOptionImage[ImageIndex].OptionName));
-
-				FMutableParamUIMetadata optionMetadata = ParameterUIData.ParamUIMetadata;
-				optionMetadata.UIThumbnail = ArrayOptionImage[ImageIndex].OptionImage;
-				ParameterUIData.ArrayIntegerParameterOption.Add(FIntegerParameterUIData(
-					ArrayOptionImage[ImageIndex].OptionName,
-					optionMetadata));
-			}
-
-			GenerationContext.ParameterUIDataMap.Add(NodeEnumParamName, ParameterUIData);
-
 			PoseEnumParameterNode->SetValue(0, 0.f, "Default pose");
 
 			for (int PoseIndex = 0; PoseIndex < ProjParamNode->OptionPoses.Num(); ++PoseIndex)
@@ -590,35 +551,68 @@ bool GenerateMutableSourceGroupProjector(const UEdGraphPin* Pin, FMutableGraphGe
 					GroupProjectorTempData.PoseBoneDataArray[PoseIndex].ArrayBoneName, GroupProjectorTempData.PoseBoneDataArray[PoseIndex].ArrayTransform);
 			}
 
-			mu::NodeImageSwitchPtr SwitchNode = new mu::NodeImageSwitch;
-			SwitchNode->SetParameter(EnumParameterNode);
-			SwitchNode->SetOptionCount(ArrayOptionImage.Num());
+			if (UEdGraphPin* ConnectedPin = FollowInputPin(ProjParamNode->GetImagePin()))
+			{
+				GroupProjectorTempData.NodeImagePtr = GenerateMutableSourceImage(ConnectedPin, GenerationContext, 0, NodeRangeFromScalar);
+			}
+			else
+			{
+				mu::NodeScalarEnumParameterPtr EnumParameterNode = new mu::NodeScalarEnumParameter;
+				FString NodeEnumParamName = ProjParamNode->ParameterName + FMultilayerProjector::IMAGE_PARAMETER_POSTFIX;
+				EnumParameterNode->SetName(TCHAR_TO_ANSI(*(NodeEnumParamName)));
+				EnumParameterNode->SetUid(TCHAR_TO_ANSI(*(ProjectorParamUid + FString("_SI"))));
+				GenerationContext.AddParameterNameUnique(originalGroup, NodeEnumParamName);
+				EnumParameterNode->SetValueCount(ArrayOptionImage.Num());
+				EnumParameterNode->SetDefaultValueIndex(0);
+				EnumParameterNode->SetRangeCount(1);
+				EnumParameterNode->SetRange(0, NodeRangeFromScalar);
+
+				FParameterUIData ParameterUIData(NodeEnumParamName, ProjParamNode->ParamUIMetadata, EMutableParameterType::Int);
+				ParameterUIData.IntegerParameterGroupType = ECustomizableObjectGroupType::COGT_ONE;
+				ParameterUIData.ParamUIMetadata.ExtraInformation.Add(FString("UseThumbnails"));
+
+				GenerationContext.ParameterUIDataMap.Add(NodeEnumParamName, ParameterUIData);
+
+				for (int ImageIndex = 0; ImageIndex < ArrayOptionImage.Num(); ++ImageIndex)
+				{
+					EnumParameterNode->SetValue(ImageIndex, (float)ImageIndex, TCHAR_TO_ANSI(*ArrayOptionImage[ImageIndex].OptionName));
 
-			bool bFoundUnlinkedPin = false;
+					FMutableParamUIMetadata optionMetadata = ParameterUIData.ParamUIMetadata;
+					optionMetadata.UIThumbnail = ArrayOptionImage[ImageIndex].OptionImage;
+					ParameterUIData.ArrayIntegerParameterOption.Add(FIntegerParameterUIData(
+						ArrayOptionImage[ImageIndex].OptionName,
+						optionMetadata));
+				}
 
-			for (int SelectorIndex = 0; SelectorIndex < ArrayOptionImage.Num(); ++SelectorIndex)
-			{
-				if (ArrayOptionImage[SelectorIndex].OptionImage)
+				mu::NodeImageSwitchPtr SwitchNode = new mu::NodeImageSwitch;
+				SwitchNode->SetParameter(EnumParameterNode);
+				SwitchNode->SetOptionCount(ArrayOptionImage.Num());
+
+				bool bFoundUnlinkedPin = false; 
+
+				for (int SelectorIndex = 0; SelectorIndex < ArrayOptionImage.Num(); ++SelectorIndex)
 				{
-					mu::NodeImageConstantPtr ImageNode = new mu::NodeImageConstant();
-					GenerationContext.ArrayTextureUnrealToMutableTask.Add(FTextureUnrealToMutableTask(ImageNode, ArrayOptionImage[SelectorIndex].OptionImage, ProjParamNode));
-					SwitchNode->SetOption(SelectorIndex, ImageNode);
+					if (ArrayOptionImage[SelectorIndex].OptionImage)
+					{
+						mu::NodeImageConstantPtr ImageNode = new mu::NodeImageConstant();
+						GenerationContext.ArrayTextureUnrealToMutableTask.Add(FTextureUnrealToMutableTask(ImageNode, ArrayOptionImage[SelectorIndex].OptionImage, ProjParamNode));
+						SwitchNode->SetOption(SelectorIndex, ImageNode);
+					}
+					else
+					{
+						bFoundUnlinkedPin = true;
+					}
 				}
-				else
+
+				if (bFoundUnlinkedPin)
 				{
-					//SwitchNode->SetOption(SelectorIndex, nullptr);
-					bFoundUnlinkedPin = true;
+					FString msg = FString::Printf(TEXT("The group projection node must have a texture for all the options. Please set a texture for all the options."));
+					GenerationContext.Compiler->CompilerLog(FText::FromString(msg), ProjParamNode);
 				}
-			}
 
-			if (bFoundUnlinkedPin)
-			{
-				FString msg = FString::Printf(TEXT("The group projection node must have all the option images connected to a texture. Please set a texture for all the options."));
-				GenerationContext.Compiler->CompilerLog(FText::FromString(msg), ProjParamNode);
+				GroupProjectorTempData.NodeImagePtr = SwitchNode;
 			}
 
-			GroupProjectorTempData.NodeImagePtr = SwitchNode;
-
 			GenerationContext.ProjectorGroupMap.Add(originalGroup, GroupProjectorTempData);
 		}
 		else
diff --git a/Engine/Plugins/Experimental/Mutable/Source/CustomizableObjectEditor/Private/MuCOE/GenerateMutableSource/GenerateMutableSourceImage.cpp b/Engine/Plugins/Experimental/Mutable/Source/CustomizableObjectEditor/Private/MuCOE/GenerateMutableSource/GenerateMutableSourceImage.cpp
--- a/Engine/Plugins/Experimental/Mutable/Source/CustomizableObjectEditor/Private/MuCOE/GenerateMutableSource/GenerateMutableSourceImage.cpp
+++ b/Engine/Plugins/Experimental/Mutable/Source/CustomizableObjectEditor/Private/MuCOE/GenerateMutableSource/GenerateMutableSourceImage.cpp
@@ -157,7 +157,7 @@ mu::NodeImagePtr ResizeToMaxTextureSize(float MaxTextureSize, const UTexture2D*
 }
 
 
-mu::NodeImagePtr GenerateMutableSourceImage(const UEdGraphPin* Pin, FMutableGraphGenerationContext& GenerationContext, float MaxTextureSize)
+mu::NodeImagePtr GenerateMutableSourceImage(const UEdGraphPin* Pin, FMutableGraphGenerationContext& GenerationContext, float MaxTextureSize, mu::NodeRangePtr NodeRange)
 {
 	check(Pin)
 	RETURN_ON_CYCLE(*Pin, GenerationContext)
@@ -213,6 +213,11 @@ mu::NodeImagePtr GenerateMutableSourceImage(const UEdGraphPin* Pin, FMutableGrap
 
 		TextureNode->SetName(TCHAR_TO_ANSI(*TypedNodeParam->ParameterName));
 		TextureNode->SetUid(TCHAR_TO_ANSI(*GenerationContext.GetNodeIdUnique(Node).ToString()));
+		if (NodeRange)
+		{
+			TextureNode->SetRangeCount(1);
+			TextureNode->SetRange(0, NodeRange);
+		}
 
 		// TODO: Set a default value for the texture parameter?
 		//ColorNode->SetDefaultValue(TypedNodeParam->DefaultValue.R, TypedNodeParam->DefaultValue.G, TypedNodeParam->DefaultValue.B);
diff --git a/Engine/Plugins/Experimental/Mutable/Source/CustomizableObjectEditor/Private/MuCOE/GenerateMutableSource/GenerateMutableSourceImage.h b/Engine/Plugins/Experimental/Mutable/Source/CustomizableObjectEditor/Private/MuCOE/GenerateMutableSource/GenerateMutableSourceImage.h
--- a/Engine/Plugins/Experimental/Mutable/Source/CustomizableObjectEditor/Private/MuCOE/GenerateMutableSource/GenerateMutableSourceImage.h
+++ b/Engine/Plugins/Experimental/Mutable/Source/CustomizableObjectEditor/Private/MuCOE/GenerateMutableSource/GenerateMutableSourceImage.h
@@ -4,8 +4,7 @@
 
 
 #include "MuT/NodeImageConstant.h"
-#include "MuT/NodeSurfaceNew.h"
-#include "MuT/Table.h"
+#include "MuT/NodeScalarEnumParameter.h"
 
 class FCustomizableObjectCompiler;
 class UCustomizableObjectNode;
@@ -20,4 +19,4 @@ mu::NodeImagePtr ResizeToMaxTextureSize(float MaxTextureSize, const UTexture2D*
 
 
 /** Convert a CustomizableObject Source Graph into a mutable source graph. */
-mu::NodeImagePtr GenerateMutableSourceImage(const UEdGraphPin* Pin, FMutableGraphGenerationContext& GenerationContext, float MaxTextureSize);
\ No newline at end of file
+mu::NodeImagePtr GenerateMutableSourceImage(const UEdGraphPin* Pin, FMutableGraphGenerationContext& GenerationContext, float MaxTextureSize, mu::NodeRangePtr NodeRange = nullptr);
\ No newline at end of file
diff --git a/Engine/Plugins/Experimental/Mutable/Source/CustomizableObjectEditor/Private/MuCOE/Nodes/CustomizableObjectNode.cpp b/Engine/Plugins/Experimental/Mutable/Source/CustomizableObjectEditor/Private/MuCOE/Nodes/CustomizableObjectNode.cpp
--- a/Engine/Plugins/Experimental/Mutable/Source/CustomizableObjectEditor/Private/MuCOE/Nodes/CustomizableObjectNode.cpp
+++ b/Engine/Plugins/Experimental/Mutable/Source/CustomizableObjectEditor/Private/MuCOE/Nodes/CustomizableObjectNode.cpp
@@ -344,6 +344,12 @@ void UCustomizableObjectNode::NodeConnectionListChanged()
 }
 
 
+void UCustomizableObjectNode::PinConnectionListChanged(UEdGraphPin* Pin)
+{
+	Super::PinConnectionListChanged(Pin);
+	PinConnectionListChangedDelegate.Broadcast(Pin);
+}
+
 
 void UCustomizableObjectNode::PostInitProperties()
 {
diff --git a/Engine/Plugins/Experimental/Mutable/Source/CustomizableObjectEditor/Private/MuCOE/Nodes/CustomizableObjectNodeGroupProjectorParameter.cpp b/Engine/Plugins/Experimental/Mutable/Source/CustomizableObjectEditor/Private/MuCOE/Nodes/CustomizableObjectNodeGroupProjectorParameter.cpp
--- a/Engine/Plugins/Experimental/Mutable/Source/CustomizableObjectEditor/Private/MuCOE/Nodes/CustomizableObjectNodeGroupProjectorParameter.cpp
+++ b/Engine/Plugins/Experimental/Mutable/Source/CustomizableObjectEditor/Private/MuCOE/Nodes/CustomizableObjectNodeGroupProjectorParameter.cpp
@@ -24,13 +24,6 @@ class UCustomizableObjectNodeRemapPins;
 #define LOCTEXT_NAMESPACE "CustomizableObjectEditor"
 
 
-UCustomizableObjectNodeGroupProjectorParameter::UCustomizableObjectNodeGroupProjectorParameter()
-	: Super()
-{
-
-}
-
-
 TArray<FGroupProjectorParameterImage> UCustomizableObjectNodeGroupProjectorParameter::GetOptionImagesFromTable() const
 {
 	TArray<FGroupProjectorParameterImage> ArrayResult;
@@ -104,22 +97,6 @@ TArray<FGroupProjectorParameterImage> UCustomizableObjectNodeGroupProjectorParam
 }
 
 
-void UCustomizableObjectNodeGroupProjectorParameter::Serialize(FArchive& Ar)
-{
-	Super::Serialize(Ar);
-
-	Ar.UsingCustomVersion(FCustomizableObjectCustomVersion::GUID);
-
-	UEdGraphPin* projectorPin = ProjectorPin();
-	if (Ar.CustomVer(FCustomizableObjectCustomVersion::GUID) < FCustomizableObjectCustomVersion::GroupProjectorPinTypeAdded
-		&& projectorPin
-		&& projectorPin->PinType.PinCategory == UEdGraphSchema_CustomizableObject::PC_Projector)
-	{
-		projectorPin->PinType.PinCategory = UEdGraphSchema_CustomizableObject::PC_GroupProjector;
-	}
-}
-
-
 FText UCustomizableObjectNodeGroupProjectorParameter::GetNodeTitle(ENodeTitleType::Type TitleType) const
 {
 	return LOCTEXT("Group_Projector_Parameter", "Group Projector Parameter");
@@ -139,13 +116,39 @@ FText UCustomizableObjectNodeGroupProjectorParameter::GetTooltipText() const
 }
 
 
+void UCustomizableObjectNodeGroupProjectorParameter::BackwardsCompatibleFixup()
+{
+	Super::BackwardsCompatibleFixup();
+
+	const int32 CustomizableObjectCustomVersion = GetLinkerCustomVersion(FCustomizableObjectCustomVersion::GUID);
+	
+	if (CustomizableObjectCustomVersion < FCustomizableObjectCustomVersion::GroupProjectorPinTypeAdded)
+	{
+		if (UEdGraphPin* Pin = ProjectorPin())
+		{
+			Pin->PinType.PinCategory = UEdGraphSchema_CustomizableObject::PC_GroupProjector;
+		}
+	}
+
+	if (CustomizableObjectCustomVersion < FCustomizableObjectCustomVersion::GroupProjectorImagePin)
+	{
+		ReconstructNode();		
+	}
+}
+
+
 void UCustomizableObjectNodeGroupProjectorParameter::AllocateDefaultPins(UCustomizableObjectNodeRemapPins* RemapPins)
 {
 	const UEdGraphSchema_CustomizableObject* Schema = GetDefault<UEdGraphSchema_CustomizableObject>();
 
-	FString PinName = TEXT("Value");
-	UEdGraphPin* ValuePin = CustomCreatePin(EGPD_Output, Schema->PC_GroupProjector, FName(*PinName));
-	ValuePin->bDefaultValueIsIgnored = true;
+	CustomCreatePin(EGPD_Output, Schema->PC_GroupProjector, TEXT("Value"));	
+	ImagePin = CustomCreatePin(EGPD_Input, Schema->PC_Image, TEXT("Texture"));
+}
+
+
+UEdGraphPin& UCustomizableObjectNodeGroupProjectorParameter::GetImagePin() const
+{
+	return *ImagePin.Get();
 }
 
 
diff --git a/Engine/Plugins/Experimental/Mutable/Source/CustomizableObjectEditor/Private/MuCOE/Nodes/CustomizableObjectNodeGroupProjectorParameterDetails.cpp b/Engine/Plugins/Experimental/Mutable/Source/CustomizableObjectEditor/Private/MuCOE/Nodes/CustomizableObjectNodeGroupProjectorParameterDetails.cpp
new file mode 100644
--- /dev/null
+++ b/Engine/Plugins/Experimental/Mutable/Source/CustomizableObjectEditor/Private/MuCOE/Nodes/CustomizableObjectNodeGroupProjectorParameterDetails.cpp
@@ -0,0 +1,41 @@
+// Copyright Epic Games, Inc. All Rights Reserved.
+
+#include "MuCOE/Nodes/CustomizableObjectNodeGroupProjectorParameterDetails.h"
+
+#include "DetailLayoutBuilder.h"
+#include "EdGraph/EdGraphPin.h"
+#include "MuCOE/GraphTraversal.h"
+#include "MuCOE/Nodes/CustomizableObjectNodeGroupProjectorParameter.h"
+
+
+TSharedRef<IDetailCustomization> FCustomizableObjectNodeGroupProjectorParameterDetails::MakeInstance()
+{
+	return MakeShareable(new FCustomizableObjectNodeGroupProjectorParameterDetails);
+}
+
+
+void FCustomizableObjectNodeGroupProjectorParameterDetails::CustomizeDetails(const TSharedPtr<IDetailLayoutBuilder>& DetailBuilder)
+{
+	Node = Cast<UCustomizableObjectNodeGroupProjectorParameter>(DetailBuilder->GetSelectedObjects()[0].Get());
+	Node->PinConnectionListChangedDelegate.AddSP(this, &FCustomizableObjectNodeGroupProjectorParameterDetails::OnPinConnectionListChanged);
+
+	LayoutBuilder = DetailBuilder;
+
+	if (FollowInputPin(Node->GetImagePin()))
+	{
+		DetailBuilder->HideProperty(DetailBuilder->GetProperty(GET_MEMBER_NAME_CHECKED(UCustomizableObjectNodeGroupProjectorParameter, OptionImages)));
+		DetailBuilder->HideProperty(DetailBuilder->GetProperty(GET_MEMBER_NAME_CHECKED(UCustomizableObjectNodeGroupProjectorParameter, OptionImagesDataTable)));
+	}
+}
+
+
+void FCustomizableObjectNodeGroupProjectorParameterDetails::OnPinConnectionListChanged(UEdGraphPin* Pin)
+{
+	if (IDetailLayoutBuilder* Layout = LayoutBuilder.Pin().Get()) // Raw because we don't want to keep alive the details builder when calling the force refresh details
+	{
+		if (Pin == &Node->GetImagePin())
+		{
+			Layout->ForceRefreshDetails();
+		}
+	}
+}
\ No newline at end of file
diff --git a/Engine/Plugins/Experimental/Mutable/Source/CustomizableObjectEditor/Private/MuCOE/Nodes/CustomizableObjectNodeGroupProjectorParameterDetails.h b/Engine/Plugins/Experimental/Mutable/Source/CustomizableObjectEditor/Private/MuCOE/Nodes/CustomizableObjectNodeGroupProjectorParameterDetails.h
new file mode 100644
--- /dev/null
+++ b/Engine/Plugins/Experimental/Mutable/Source/CustomizableObjectEditor/Private/MuCOE/Nodes/CustomizableObjectNodeGroupProjectorParameterDetails.h
@@ -0,0 +1,29 @@
+// Copyright Epic Games, Inc. All Rights Reserved.
+
+#pragma once
+
+#include "IDetailCustomization.h"
+#include "UObject/ObjectPtr.h"
+
+class UCustomizableObjectNodeGroupProjectorParameter;
+class UEdGraphPin;
+
+class FCustomizableObjectNodeGroupProjectorParameterDetails : public IDetailCustomization
+{
+public:
+	/** Makes a new instance of this detail layout class for a specific detail view requesting it */
+	static TSharedRef<IDetailCustomization> MakeInstance();
+
+	// ILayoutDetails interface
+	// Do not use. Add details customization in the other CustomizeDetails signature.
+	virtual void CustomizeDetails(IDetailLayoutBuilder& DetailBuilder) override {}
+
+	virtual void CustomizeDetails(const TSharedPtr<IDetailLayoutBuilder>& DetailBuilder) override;
+
+private:
+	TObjectPtr<UCustomizableObjectNodeGroupProjectorParameter> Node;
+
+	TWeakPtr<IDetailLayoutBuilder> LayoutBuilder;
+
+	void OnPinConnectionListChanged(UEdGraphPin* Pin);
+};
\ No newline at end of file
diff --git a/Engine/Plugins/Experimental/Mutable/Source/CustomizableObjectEditor/Private/MuCOE/SCustomizableInstanceProperties.cpp b/Engine/Plugins/Experimental/Mutable/Source/CustomizableObjectEditor/Private/MuCOE/SCustomizableInstanceProperties.cpp
--- a/Engine/Plugins/Experimental/Mutable/Source/CustomizableObjectEditor/Private/MuCOE/SCustomizableInstanceProperties.cpp
+++ b/Engine/Plugins/Experimental/Mutable/Source/CustomizableObjectEditor/Private/MuCOE/SCustomizableInstanceProperties.cpp
@@ -1073,40 +1073,29 @@ void SCustomizableInstanceProperties::AddParameter(int32 ParamIndexInObject)
 					]
 				];
 
+				FString TextureSwitchEnumParamName = ParamName + FMultilayerProjector::IMAGE_PARAMETER_POSTFIX;
+
+				// Add the necessary ranges. Ranges could be lower if in the previous compilation the Group Projector used the Texture pin.
+				if (int32 ParameterIndex = CustomInstance->FindIntParameterNameIndex(TextureSwitchEnumParamName);
+					ParameterIndex != INDEX_NONE)
+				{
+					const int32 ProjectorRange = ProjectorParameters[ProjectorParamIndex].RangeValues.Num();
+					const int32 ImagesRange = CustomInstance->GetIntParameters()[ParameterIndex].ParameterRangeValueNames.Num();
+					for (int32 Range = 0; Range < ProjectorRange - ImagesRange; ++Range)
+					{
+						CustomInstance->AddValueToIntRange(TextureSwitchEnumParamName);
+					}
+				}
+
 				for (int32 i = 0; i < ProjectorParameters[ProjectorParamIndex].RangeValues.Num(); ++i)
 				{
 					FString ParamNameWithIndex = ParamName + FString::Printf(TEXT("__%d"), i);
-
-					FString TextureSwitchEnumParamName = ParamName + FMultilayerProjector::IMAGE_PARAMETER_POSTFIX;
+					
 					FString TextureSwitchEnumParamNameWithRange = TextureSwitchEnumParamName + FString::Printf(TEXT("__%d"), i);
-					int32 TextureSwitchEnumParamIndexInObject = CustomizableObject->FindParameter(TextureSwitchEnumParamName);
-					check(TextureSwitchEnumParamIndexInObject >= 0);
-					int32 numValues = CustomizableObject->GetIntParameterNumOptions(TextureSwitchEnumParamIndexInObject);
 
 					FString OpacitySliderParamName = ParamName + FMultilayerProjector::OPACITY_PARAMETER_POSTFIX;
 					FString OpacitySliderParamNameWithRange = OpacitySliderParamName + FString::Printf(TEXT("__%d"), i);
 
-					//if (numValues)
-					//{
-					//TSharedPtr< TArray< TSharedPtr<FString> > > OptionNames = MakeShareable(new TArray< TSharedPtr<FString> >());
-					//IntOptionNames.Add(OptionNames);
-
-					TArray<FString> OptionNamesAttribute;
-					FString Value = GetIntParameterValue(TextureSwitchEnumParamName, i);
-					int32 ValueIndex = 0;
-
-					for (int32 CandidateIndex = 0; CandidateIndex < numValues; ++CandidateIndex)
-					{
-						FString PossibleValue = CustomizableObject->GetIntParameterAvailableOption(TextureSwitchEnumParamIndexInObject, CandidateIndex);
-						if (PossibleValue == Value)
-						{
-							ValueIndex = CandidateIndex;
-						}
-						//OptionNames->Add(MakeShareable(new FString(CustomizableObject->GetIntParameterAvailableOption(TextureSwitchEnumParamIndexInObject, CandidateIndex))));
-						OptionNamesAttribute.Add(CustomizableObject->GetIntParameterAvailableOption(TextureSwitchEnumParamIndexInObject, CandidateIndex));
-					}
-					//}
-
 					TSharedPtr<SHorizontalBox> SliderBox;
 					TSharedPtr<SSpinBox<float>> Slider;
 					TSharedPtr<SHorizontalBox> PropertiesHB;
@@ -1149,21 +1138,6 @@ void SCustomizableInstanceProperties::AddParameter(int32 ParamIndexInObject)
 									]
 								]
 							]
-
-							+ SHorizontalBox::Slot()
-							.Padding(1, 0)
-							[
-								SNew(SBox)
-								.MinDesiredWidth(120.f)
-								[
-									SNew(SMutableTextSearchBox)
-									.PossibleSuggestions(OptionNamesAttribute)
-									.InitialText(FText::FromString(OptionNamesAttribute[ValueIndex]))
-									.MustMatchPossibleSuggestions(TAttribute<bool>(true))
-									.SuggestionListPlacement(EMenuPlacement::MenuPlacement_ComboBox)
-									.OnTextCommitted(this, &SCustomizableInstanceProperties::OnProjectorTextureParameterComboBoxChanged, TextureSwitchEnumParamNameWithRange)
-								]
-							]
 						]
 					];
 						
@@ -1299,6 +1273,44 @@ void SCustomizableInstanceProperties::AddParameter(int32 ParamIndexInObject)
 					}
 
 					FloatSliders.Add(FSliderData(Slider, OpacitySliderParamName, i, GetFloatParameterValue(OpacitySliderParamName, i)));
+
+					if (CustomizableObject->FindParameter(TextureSwitchEnumParamName) != INDEX_NONE)
+					{
+						int32 TextureSwitchEnumParamIndexInObject = CustomizableObject->FindParameter(TextureSwitchEnumParamName);
+						check(TextureSwitchEnumParamIndexInObject >= 0);
+						int32 NumValues = CustomizableObject->GetIntParameterNumOptions(TextureSwitchEnumParamIndexInObject);
+
+						TArray<FString> OptionNamesAttribute;
+						int32 ValueIndex = 0;
+
+						FString Value = GetIntParameterValue(TextureSwitchEnumParamName, i);
+
+						for (int32 CandidateIndex = 0; CandidateIndex < NumValues; ++CandidateIndex)
+						{
+							FString PossibleValue = CustomizableObject->GetIntParameterAvailableOption(TextureSwitchEnumParamIndexInObject, CandidateIndex);
+							if (PossibleValue == Value)
+							{
+								ValueIndex = CandidateIndex;
+							}
+
+							OptionNamesAttribute.Add(CustomizableObject->GetIntParameterAvailableOption(TextureSwitchEnumParamIndexInObject, CandidateIndex));
+						}
+
+						PropertiesHB->AddSlot()
+							.Padding(1, 0)
+							[
+								SNew(SBox)
+								.MinDesiredWidth(120.f)
+								[
+									SNew(SMutableTextSearchBox)
+									.PossibleSuggestions(OptionNamesAttribute)
+									.InitialText(FText::FromString(OptionNamesAttribute[ValueIndex]))
+									.MustMatchPossibleSuggestions(TAttribute<bool>(true))
+									.SuggestionListPlacement(EMenuPlacement::MenuPlacement_ComboBox)
+									.OnTextCommitted(this, &SCustomizableInstanceProperties::OnProjectorTextureParameterComboBoxChanged, TextureSwitchEnumParamNameWithRange)
+								]
+							];
+					}
 				}
 			}
 		}
@@ -1310,6 +1322,7 @@ void SCustomizableInstanceProperties::AddParameter(int32 ParamIndexInObject)
 	}
 }
 
+
 void SCustomizableInstanceProperties::Tick(const FGeometry& AllottedGeometry, const double InCurrentTime, const float InDeltaTime)
 {
 	SCompoundWidget::Tick(AllottedGeometry, InCurrentTime, InDeltaTime);
@@ -1818,7 +1831,6 @@ ECheckBoxState SCustomizableInstanceProperties::IsProjectorScale(FString ParamNa
 	return ECheckBoxState::Unchecked;
 }
 
-
 FReply SCustomizableInstanceProperties::OnProjectorLayerAdded(FString ParamName) const
 {
 	const int32 NumLayers = CustomInstance->AddValueToProjectorRange(ParamName) + 1;
@@ -1828,9 +1840,12 @@ FReply SCustomizableInstanceProperties::OnProjectorLayerAdded(FString ParamName)
 	}
 
 	const FString TextureSwitchEnumParamName = ParamName + FMultilayerProjector::IMAGE_PARAMETER_POSTFIX;
-	CustomInstance->AddValueToIntRange(TextureSwitchEnumParamName);
-	check(CustomInstance->FindIntParameterNameIndex(TextureSwitchEnumParamName) != INDEX_NONE);
-	check(NumLayers == CustomInstance->GetIntParameters()[CustomInstance->FindIntParameterNameIndex(TextureSwitchEnumParamName)].ParameterRangeValueNames.Num());
+	if (CustomInstance->GetCustomizableObject()->FindParameter(TextureSwitchEnumParamName) != INDEX_NONE)
+	{
+		CustomInstance->AddValueToIntRange(TextureSwitchEnumParamName);
+		check(CustomInstance->FindIntParameterNameIndex(TextureSwitchEnumParamName) != INDEX_NONE);
+		check(NumLayers == CustomInstance->GetIntParameters()[CustomInstance->FindIntParameterNameIndex(TextureSwitchEnumParamName)].ParameterRangeValueNames.Num());
+	}
 
 	TArray<FCustomizableObjectFloatParameterValue>& FloatParameters = CustomInstance->GetFloatParameters();
 	
@@ -1874,9 +1889,12 @@ FReply SCustomizableInstanceProperties::OnProjectorLayerRemoved(FString ParamNam
 	CustomInstance->RemovedProjectorParameterNameWithIndex = ParamNameWithIndex; // Since the gizmo updates the CustomObjectInstance when unselected, a check is required to avoid updating the removed projector parameter range index
 
 	const FString TextureSwitchEnumParamName = ParamName + FMultilayerProjector::IMAGE_PARAMETER_POSTFIX;
-	CustomInstance->RemoveValueFromIntRange(TextureSwitchEnumParamName, RangeIndex);
-	check(CustomInstance->FindIntParameterNameIndex(TextureSwitchEnumParamName) != INDEX_NONE);
-	check(NumLayers == CustomInstance->GetIntParameters()[CustomInstance->FindIntParameterNameIndex(TextureSwitchEnumParamName)].ParameterRangeValueNames.Num());
+	if (CustomInstance->GetCustomizableObject()->FindParameter(TextureSwitchEnumParamName) != INDEX_NONE)
+	{
+		CustomInstance->RemoveValueFromIntRange(TextureSwitchEnumParamName, RangeIndex);
+		check(CustomInstance->FindIntParameterNameIndex(TextureSwitchEnumParamName) != INDEX_NONE);
+		check(NumLayers == CustomInstance->GetIntParameters()[CustomInstance->FindIntParameterNameIndex(TextureSwitchEnumParamName)].ParameterRangeValueNames.Num());
+	}
 
 	TArray<FCustomizableObjectFloatParameterValue>& FloatParameters = CustomInstance->GetFloatParameters();
 
