[Mutable] Minor code improvements.

- Removed pragma disable optimization.
- Edited function documentation.

#rnx

diff --git a/Engine/Plugins/Experimental/Mutable/Source/CustomizableObject/Public/MuCO/CustomizableObjectInstanceDescriptor.h b/Engine/Plugins/Experimental/Mutable/Source/CustomizableObject/Public/MuCO/CustomizableObjectInstanceDescriptor.h
--- a/Engine/Plugins/Experimental/Mutable/Source/CustomizableObject/Public/MuCO/CustomizableObjectInstanceDescriptor.h
+++ b/Engine/Plugins/Experimental/Mutable/Source/CustomizableObject/Public/MuCO/CustomizableObjectInstanceDescriptor.h
@@ -202,7 +202,7 @@ struct CUSTOMIZABLEOBJECT_API FCustomizableObjectInstanceDescriptor
 	/** Returns true if the parameter is multidimensional (has multiple ranges). */
 	bool IsParamMultidimensional(int32 ParamIndex) const;
 
-	/** Gets the range of values of the projector with ParamName, returns -1 if the parameter exists, -1 otherwise. */
+	/** Gets the range of values of the projector with ParamName, returns -1 if the parameter does not exist. */
 	int32 GetProjectorValueRange(const FString& ParamName) const;
 
 	/** Increases the range of values of the integer with ParamName, returns the index of the new integer value, -1 otherwise.
diff --git a/Engine/Plugins/Experimental/Mutable/Source/CustomizableObjectEditor/Private/MuCOE/GenerateMutableSource/GenerateMutableSourceGroupProjector.cpp b/Engine/Plugins/Experimental/Mutable/Source/CustomizableObjectEditor/Private/MuCOE/GenerateMutableSource/GenerateMutableSourceGroupProjector.cpp
--- a/Engine/Plugins/Experimental/Mutable/Source/CustomizableObjectEditor/Private/MuCOE/GenerateMutableSource/GenerateMutableSourceGroupProjector.cpp
+++ b/Engine/Plugins/Experimental/Mutable/Source/CustomizableObjectEditor/Private/MuCOE/GenerateMutableSource/GenerateMutableSourceGroupProjector.cpp
@@ -72,7 +72,6 @@ class UPoseAsset;
 #define LOCTEXT_NAMESPACE "CustomizableObjectEditor"
 
 
-PRAGMA_DISABLE_OPTIMIZATION
 mu::NodeImagePtr GenerateMutableGroupProjection(const int32 NodeLOD, const int32 ImageIndex, mu::NodeMeshPtr MeshNode, FMutableGraphGenerationContext& GenerationContext,
 	const UCustomizableObjectNodeMaterialBase* NodeMaterialBase, bool& bShareProjectionTexturesBetweenLODs, bool& bIsGroupProjectorImage,
 	UTexture2D*& GroupProjectionReferenceTexture, TMap<FString, float>& TextureNameToProjectionResFactor, FString& AlternateResStateName,
@@ -428,7 +427,7 @@ mu::NodeImagePtr GenerateMutableGroupProjection(const int32 NodeLOD, const int32
 
 	return FinalSwizzleNode;
 }
-PRAGMA_ENABLE_OPTIMIZATION
+
 
 bool GenerateMutableSourceGroupProjector(const UEdGraphPin* Pin, FMutableGraphGenerationContext& GenerationContext, const UCustomizableObjectNodeObjectGroup* originalGroup)
 {
