[mutable] Fixed naming issue found on Cyborg customization UI where "Opacity" was expected instead of "Transparency"
- This change will also apply to other projects using a projected texture (projector node) after CO recompilation.

diff --git a/Engine/Plugins/Experimental/Mutable/Source/CustomizableObjectEditor/Private/MuCOE/GenerateMutableSource/GenerateMutableSourceGroupProjector.cpp b/Engine/Plugins/Experimental/Mutable/Source/CustomizableObjectEditor/Private/MuCOE/GenerateMutableSource/GenerateMutableSourceGroupProjector.cpp
--- a/Engine/Plugins/Experimental/Mutable/Source/CustomizableObjectEditor/Private/MuCOE/GenerateMutableSource/GenerateMutableSourceGroupProjector.cpp
+++ b/Engine/Plugins/Experimental/Mutable/Source/CustomizableObjectEditor/Private/MuCOE/GenerateMutableSource/GenerateMutableSourceGroupProjector.cpp
@@ -540,7 +540,7 @@ bool GenerateMutableSourceGroupProjector(const UEdGraphPin* Pin, FMutableGraphGe
 			OpacityParameterNode->SetRangeCount(1);
 			OpacityParameterNode->SetRange(0, NodeRangeFromScalar);
 			FMutableParamUIMetadata OpacityMetadata = ProjParamNode->ParamUIMetadata;
-			OpacityMetadata.ObjectFriendlyName = FString("Transparency");
+			OpacityMetadata.ObjectFriendlyName = FString("Opacity");
 			GroupProjectorTempData.NodeOpacityParameter = OpacityParameterNode;
 
 			GenerationContext.ParameterUIDataMap.Add(OpacityParameterNodeName, FParameterUIData(
