[mutable] Add flag in CO compilation options to disable physics assets collection.

#rnx

diff --git a/Engine/Plugins/Experimental/Mutable/Source/CustomizableObject/Classes/MuCO/CustomizableObject.h b/Engine/Plugins/Experimental/Mutable/Source/CustomizableObject/Classes/MuCO/CustomizableObject.h
--- a/Engine/Plugins/Experimental/Mutable/Source/CustomizableObject/Classes/MuCO/CustomizableObject.h
+++ b/Engine/Plugins/Experimental/Mutable/Source/CustomizableObject/Classes/MuCO/CustomizableObject.h
@@ -203,6 +203,9 @@ struct FCompilationOptions
 	// Used to enable skin weight profiles.
 	bool bSkinWeightProfilesEnabled = true;
 
+	// Used to enable physics asset merge.
+	bool bPhysicsAssetMergeEnebled = false;
+
 	// Used to reduce the number of notifications when compiling objects
 	bool bSilentCompilation = true;
 
@@ -1069,6 +1072,10 @@ public:
 	UPROPERTY(EditAnywhere, Category = CompileOptions)
 	bool bEnableAltSkinWeightProfiles = false;
 
+	//
+	UPROPERTY(EditAnywhere, Category = CompileOptions)
+	bool bEnablePhysicsAssetMerge = false;
+
 	// Options when compiling this customizable object (see EMutableCompileMeshType declaration for info)
 	UPROPERTY(EditAnywhere, Category = CompileOptions)
 	EMutableCompileMeshType MeshCompileType = EMutableCompileMeshType::LocalAndChildren;
@@ -1178,7 +1185,7 @@ private:
 	// This is a manual version number for the binary blobs in this asset.
 	// Increasing it invalidates all the previously compiled models.
 	// Warning: If while merging code both versions have changed, take the highest+1.
-	static const int32 CurrentSupportedVersion = 353;
+	static const int32 CurrentSupportedVersion = 354;
 
 public:
 
diff --git a/Engine/Plugins/Experimental/Mutable/Source/CustomizableObjectEditor/Private/MuCOE/CustomizableObjectCompiler.cpp b/Engine/Plugins/Experimental/Mutable/Source/CustomizableObjectEditor/Private/MuCOE/CustomizableObjectCompiler.cpp
--- a/Engine/Plugins/Experimental/Mutable/Source/CustomizableObjectEditor/Private/MuCOE/CustomizableObjectCompiler.cpp
+++ b/Engine/Plugins/Experimental/Mutable/Source/CustomizableObjectEditor/Private/MuCOE/CustomizableObjectCompiler.cpp
@@ -951,6 +951,7 @@ void FCustomizableObjectCompiler::CompileInternal(UCustomizableObject* Object, c
 	Options.bClothingEnabled = Object->bEnableClothing;
 	Options.b16BitBoneWeightsEnabled = Object->bEnable16BitBoneWeights;
 	Options.bSkinWeightProfilesEnabled = Object->bEnableAltSkinWeightProfiles;
+	Options.bPhysicsAssetMergeEnebled = Object->bEnablePhysicsAssetMerge;
 
 	if (Object->IsLocked() || !UCustomizableObjectSystem::GetInstance()->LockObject(Object))
 	{
diff --git a/Engine/Plugins/Experimental/Mutable/Source/CustomizableObjectEditor/Private/MuCOE/GenerateMutableSource/GenerateMutableSourceMesh.cpp b/Engine/Plugins/Experimental/Mutable/Source/CustomizableObjectEditor/Private/MuCOE/GenerateMutableSource/GenerateMutableSourceMesh.cpp
--- a/Engine/Plugins/Experimental/Mutable/Source/CustomizableObjectEditor/Private/MuCOE/GenerateMutableSource/GenerateMutableSourceMesh.cpp
+++ b/Engine/Plugins/Experimental/Mutable/Source/CustomizableObjectEditor/Private/MuCOE/GenerateMutableSource/GenerateMutableSourceMesh.cpp
@@ -469,7 +469,8 @@ mu::MeshPtr ConvertSkeletalMeshToMutable(USkeletalMesh* InSkeletalMesh, int LOD,
 	// Get the mesh generation flags to use
 	uint32 CurrentFlags = GenerationContext.MeshGenerationFlags.Last();
 	bool bIgnoreSkeleton = CurrentFlags & uint32(EMutableMeshConversionFlags::IgnoreSkinning);
-	bool bIgnorePhysics = CurrentFlags & uint32(EMutableMeshConversionFlags::IgnorePhysics);
+	bool bIgnorePhysics = (CurrentFlags & uint32(EMutableMeshConversionFlags::IgnorePhysics)) || 
+						  !GenerationContext.Options.bPhysicsAssetMergeEnebled;
 
 	mu::MeshPtr MutableMesh = new mu::Mesh();
 		
