[mutable] Optimize projector's border displacement map when using progressive images, or downscaling in general.
Optimize switching operations when all possible cases end in the same operation.
Extend debugger to show more information for Switching operations in the code tree.
#rnx

diff --git a/Engine/Plugins/Experimental/Mutable/Source/CustomizableObjectEditor/Private/MuCOE/CustomizableObjectCompileRunnable.cpp b/Engine/Plugins/Experimental/Mutable/Source/CustomizableObjectEditor/Private/MuCOE/CustomizableObjectCompileRunnable.cpp
--- a/Engine/Plugins/Experimental/Mutable/Source/CustomizableObjectEditor/Private/MuCOE/CustomizableObjectCompileRunnable.cpp
+++ b/Engine/Plugins/Experimental/Mutable/Source/CustomizableObjectEditor/Private/MuCOE/CustomizableObjectCompileRunnable.cpp
@@ -96,6 +96,9 @@ uint32 FCustomizableObjectCompileRunnable::Run()
 	const int MinRomSize = 128;
 	CompilerOptions->SetDataPackingStrategy(MinRomSize, MinResidentMips);
 
+	/** At object compilation time we don't know if we will want progressive images or not. Assume we will. */
+	CompilerOptions->SetEnableProgressiveImages(true);
+
 	CompilerOptions->SetTextureLayoutStrategy(bDisableTextureLayout
 		? mu::CompilerOptions::TextureLayoutStrategy::None
 		: mu::CompilerOptions::TextureLayoutStrategy::Pack);
diff --git a/Engine/Plugins/Experimental/Mutable/Source/CustomizableObjectEditor/Private/MuCOE/SMutableCodeViewer.cpp b/Engine/Plugins/Experimental/Mutable/Source/CustomizableObjectEditor/Private/MuCOE/SMutableCodeViewer.cpp
--- a/Engine/Plugins/Experimental/Mutable/Source/CustomizableObjectEditor/Private/MuCOE/SMutableCodeViewer.cpp
+++ b/Engine/Plugins/Experimental/Mutable/Source/CustomizableObjectEditor/Private/MuCOE/SMutableCodeViewer.cpp
@@ -869,7 +869,7 @@ void SMutableCodeViewer::GetOperationsOfType(const mu::OP_TYPE& TargetOperationT
 		// Process the children of this object
 		GetOperationsOfType(TargetOperationType, InSearchPayload, InProgram);
 	}
-}
+ }
 
 
 void SMutableCodeViewer::CacheOperationTypesPresentOnModel()
@@ -1086,7 +1086,7 @@ void SMutableCodeViewer::GenerateAllTreeElements()
 		const FSlateColor LabelColor = ColorPerComputationalCost[StaticCast<uint8>(GetOperationTypeComputationalCost(
 			ModelPrivate->m_program.GetOpType(State.m_root)))];
 		
-		RootNodes.Add(MakeShareable(new FMutableCodeTreeElement(ItemCache.Num(),MutableModel, State.m_root, Caption,LabelColor)));
+		RootNodes.Add(MakeShareable(new FMutableCodeTreeElement(ItemCache.Num(), MutableModel, State.m_root, Caption,LabelColor)));
 
 		// Iterate over each root node and generate all the elements in a human readable pattern (Z Pattern)
 
@@ -1102,37 +1102,102 @@ void SMutableCodeViewer::GenerateAllTreeElements()
 
 void SMutableCodeViewer::GenerateElementRecursive(mu::OP::ADDRESS InParentAddress,  const mu::FProgram& InProgram)
 {
-	// Find children of the provided element
+	// This will be used to add operations
 	uint32 ChildIndex = 0;
-	mu::ForEachReference(InProgram, InParentAddress, [this, InParentAddress, &ChildIndex, &InProgram](mu::OP::ADDRESS ChildAddress)
-	{
-		if (ChildAddress)
+	auto AddOpFunc = [this, InParentAddress, &InProgram, &ChildIndex](mu::OP::ADDRESS ChildAddress, const FString& Caption)
 		{
-			const FItemCacheKey Key = { InParentAddress, ChildAddress, ChildIndex };
-			const TSharedPtr<FMutableCodeTreeElement>* CachedItem = ItemCache.Find(Key);
-
-			if (!CachedItem)
+			if (ChildAddress)
 			{
-				// Locate the "original" tree element
-				const TSharedPtr<FMutableCodeTreeElement>* MainItemPtr = MainItemPerOp.Find(ChildAddress);
+				const FItemCacheKey Key = { InParentAddress, ChildAddress, ChildIndex };
+				const TSharedPtr<FMutableCodeTreeElement>* CachedItem = ItemCache.Find(Key);
 
-				const FSlateColor LabelColor = ColorPerComputationalCost[ StaticCast<uint8> (GetOperationTypeComputationalCost(InProgram.GetOpType(ChildAddress))) ];
-				const TSharedPtr<FMutableCodeTreeElement> Item = MakeShareable(new FMutableCodeTreeElement(ItemCache.Num(),MutableModel, ChildAddress, TEXT(""),LabelColor, MainItemPtr));
+				// Will this ever really hit?
+				check(!CachedItem);
 
-				// Cache this element for later access
-				ItemCache.Add(Key, Item);
-				
-				// It is not a duplicated of another one, then we can continue searching
-				if (!MainItemPtr)
+				if (!CachedItem)
 				{
-					MainItemPerOp.Add(ChildAddress, Item);
-					
-					GenerateElementRecursive(ChildAddress,InProgram);
+					// Locate the "original" tree element
+					const TSharedPtr<FMutableCodeTreeElement>* MainItemPtr = MainItemPerOp.Find(ChildAddress);
+
+					const FSlateColor LabelColor = ColorPerComputationalCost[StaticCast<uint8>(GetOperationTypeComputationalCost(InProgram.GetOpType(ChildAddress)))];
+
+					// No caption for the generic tree
+					const TSharedPtr<FMutableCodeTreeElement> Item = MakeShareable(new FMutableCodeTreeElement(ItemCache.Num(), MutableModel, ChildAddress, Caption, LabelColor, MainItemPtr));
+
+					// Cache this element for later access
+					ItemCache.Add(Key, Item);
+
+					// It is not a duplicated of another one, then we can continue searching
+					if (!MainItemPtr)
+					{
+						MainItemPerOp.Add(ChildAddress, Item);
+
+						GenerateElementRecursive(ChildAddress, InProgram);
+					}
 				}
 			}
+			++ChildIndex;
+		};
+
+	
+	// For some specific parent operation types we create more detailed subtrees.
+	const mu::OP_TYPE ParentOperationType = InProgram.GetOpType(InParentAddress);
+	switch (ParentOperationType)
+	{
+	case mu::OP_TYPE::IM_SWITCH:
+	case mu::OP_TYPE::LA_SWITCH:
+	case mu::OP_TYPE::ME_SWITCH:
+	case mu::OP_TYPE::CO_SWITCH:
+	case mu::OP_TYPE::SC_SWITCH:
+	case mu::OP_TYPE::NU_SWITCH:
+	case mu::OP_TYPE::IN_SWITCH:
+	{
+		const uint8* OpData = InProgram.GetOpArgsPointer(InParentAddress);
+
+		mu::OP::ADDRESS VarAddress;
+		FMemory::Memcpy(&VarAddress, OpData, sizeof(mu::OP::ADDRESS));
+		OpData += sizeof(mu::OP::ADDRESS);
+		AddOpFunc(VarAddress, TEXT("var "));
+
+		mu::OP::ADDRESS DefAddress;
+		FMemory::Memcpy(&DefAddress, OpData, sizeof(mu::OP::ADDRESS));
+		OpData += sizeof(mu::OP::ADDRESS);
+		AddOpFunc(DefAddress, TEXT("def "));
+
+		uint32 CaseCount;
+		FMemory::Memcpy(&CaseCount, OpData, sizeof(uint32));
+		OpData += sizeof(uint32);
+
+		for (uint32 C = 0; C < CaseCount; ++C)
+		{
+			int32 Condition;
+			FMemory::Memcpy(&Condition, OpData, sizeof(int32));
+			OpData += sizeof(int32);
+
+			mu::OP::ADDRESS At;
+			FMemory::Memcpy(&At, OpData, sizeof(mu::OP::ADDRESS));
+			OpData += sizeof(mu::OP::ADDRESS);
+
+			FString Caption = FString::Printf(TEXT("case %d "),Condition);
+			AddOpFunc(At, Caption);
 		}
-		++ChildIndex;
-	});
+
+		break;
+	}
+
+	default:
+	{
+		// Generic list of child operations
+
+		// Find children of the provided element
+		mu::ForEachReference(InProgram, InParentAddress, [this, &InProgram, AddOpFunc](mu::OP::ADDRESS ChildAddress)
+			{
+				AddOpFunc(ChildAddress,TEXT(""));
+			});
+		break;
+	}
+
+	}
 }
 
 SMutableCodeViewer::EOperationComputationalCost SMutableCodeViewer::GetOperationTypeComputationalCost(mu::OP_TYPE OperationType) const
@@ -2289,7 +2354,7 @@ FReply SMutableCodeViewer::OnDrop(const FGeometry& MyGeometry, const FDragDropEv
 				const FString DraggedFileExtension = FPaths::GetExtension(Files[0], true);
 				if (DraggedFileExtension == TEXT(".mutable_compiled"))
 				{
-					// Dump source model to a file.
+					// Read a mutable compiled model file.
 					mu::InputFileStream stream(TCHAR_TO_ANSI(*Files[0]));
 
 					char MutableSourceTag[4] = {};
diff --git a/Engine/Plugins/Experimental/Mutable/Source/CustomizableObjectEditor/Private/MuCOE/SMutableCodeViewer.h b/Engine/Plugins/Experimental/Mutable/Source/CustomizableObjectEditor/Private/MuCOE/SMutableCodeViewer.h
--- a/Engine/Plugins/Experimental/Mutable/Source/CustomizableObjectEditor/Private/MuCOE/SMutableCodeViewer.h
+++ b/Engine/Plugins/Experimental/Mutable/Source/CustomizableObjectEditor/Private/MuCOE/SMutableCodeViewer.h
@@ -455,8 +455,10 @@ private:
 			// Generate a key for this element in order to search in on the map with all the elements
 								
 			// Find that element on the tree, (check error if not found since all elements should be there)
-			TSharedPtr<FMutableCodeTreeElement> FoundElement = *InItemCache.Find(Key);
-			check (FoundElement != nullptr);
+			const TSharedPtr<FMutableCodeTreeElement>* PtrFoundElement = InItemCache.Find(Key);
+			check(PtrFoundElement);
+			TSharedPtr<FMutableCodeTreeElement> FoundElement = *PtrFoundElement;
+			check(FoundElement != nullptr);
 			
 			// Store this element on our temp Map of elements
 			FoundElements.Add(FoundElement);
diff --git a/Engine/Plugins/Experimental/Mutable/Source/MutableRuntime/Private/MuR/CodeRunner.cpp b/Engine/Plugins/Experimental/Mutable/Source/MutableRuntime/Private/MuR/CodeRunner.cpp
--- a/Engine/Plugins/Experimental/Mutable/Source/MutableRuntime/Private/MuR/CodeRunner.cpp
+++ b/Engine/Plugins/Experimental/Mutable/Source/MutableRuntime/Private/MuR/CodeRunner.cpp
@@ -3882,7 +3882,7 @@ namespace mu
 
                 Ptr<const Image> pMask = GetMemory().GetImage( CACHE_ADDRESS(args.mask,item) );
 
-                ImagePtr pNew = new Image( pMask->GetSizeX(), pMask->GetSizeY(), 1, EImageFormat::IF_L_UBYTE );
+                ImagePtr pNew = new Image( pMask->GetSizeX(), pMask->GetSizeY(), pMask->GetLODCount(), EImageFormat::IF_L_UBYTE);
 
                 ImageMakeGrowMap( pNew.get(), pMask.get(), args.border );
 				pNew->m_flags |= Image::IF_CANNOT_BE_SCALED;
diff --git a/Engine/Plugins/Experimental/Mutable/Source/MutableRuntime/Private/MuR/Image.cpp b/Engine/Plugins/Experimental/Mutable/Source/MutableRuntime/Private/MuR/Image.cpp
--- a/Engine/Plugins/Experimental/Mutable/Source/MutableRuntime/Private/MuR/Image.cpp
+++ b/Engine/Plugins/Experimental/Mutable/Source/MutableRuntime/Private/MuR/Image.cpp
@@ -109,7 +109,7 @@ namespace mu
 			return Clone();
 		}
 
-		vec2<int> MipSize = CalculateMipSize(Mip);
+		FIntVector2 MipSize = CalculateMipSize(Mip);
 
 		int32 Quality = 4;
 
@@ -295,7 +295,7 @@ namespace mu
             // An RLE image.
             for ( int l=0; l<FMath::Max(1,(int)m_lods); ++l )
             {
-                auto mipSize = CalculateMipSize(l);
+				FIntVector2 mipSize = CalculateMipSize(l);
                 res += mipSize[0] * mipSize[1];
             }
         }
@@ -331,7 +331,7 @@ namespace mu
         else
         {
             // An RLE image.
-            auto mipSize = CalculateMipSize(mip);
+			FIntVector2 mipSize = CalculateMipSize(mip);
             res += mipSize[0] * mipSize[1];
         }
 
@@ -340,11 +340,11 @@ namespace mu
 
 
     //---------------------------------------------------------------------------------------------
-    vec2<int> Image::CalculateMipSize( int mip ) const
+	FIntVector2 Image::CalculateMipSize( int mip ) const
     {
-        vec2<int> res;
+		FIntVector2 res(0,0);
 
-        FImageSize s = m_size;
+		FIntVector2 s = FIntVector2(m_size.x(), m_size.y());
 
         for ( int l=0; l<mip+1; ++l )
         {
@@ -355,8 +355,8 @@ namespace mu
                 return res;
             }
 
-            s[0] = FMath::DivideAndRoundUp( s[0], (uint16)2 );
-            s[1] = FMath::DivideAndRoundUp( s[1], (uint16)2 );
+            s[0] = FMath::DivideAndRoundUp( s[0], 2 );
+            s[1] = FMath::DivideAndRoundUp( s[1], 2 );
         }
 
         return res;
@@ -922,7 +922,7 @@ namespace mu
 				DataToSkip += MipDataSize;
 			}
 
-			vec2<int> FinalSize = CalculateMipSize(LODSToSkip);
+			FIntVector2 FinalSize = CalculateMipSize(LODSToSkip);
 			m_size[0] = uint16(FinalSize[0]);
 			m_size[1] = uint16(FinalSize[1]);
 			m_lods = NewLODCount;
@@ -949,7 +949,7 @@ namespace mu
 				DataToSkip += MipDataSize;
 			}
 
-			vec2<int> FinalSize = CalculateMipSize(LODSToSkip);
+			FIntVector2 FinalSize = CalculateMipSize(LODSToSkip);
 			m_size[0] = uint16(FinalSize[0]);
 			m_size[1] = uint16(FinalSize[1]);
 			m_lods -= LODSToSkip;
diff --git a/Engine/Plugins/Experimental/Mutable/Source/MutableRuntime/Private/MuR/ModelPrivate.h b/Engine/Plugins/Experimental/Mutable/Source/MutableRuntime/Private/MuR/ModelPrivate.h
--- a/Engine/Plugins/Experimental/Mutable/Source/MutableRuntime/Private/MuR/ModelPrivate.h
+++ b/Engine/Plugins/Experimental/Mutable/Source/MutableRuntime/Private/MuR/ModelPrivate.h
@@ -331,7 +331,9 @@ namespace mu
 			bool bCannotBeScaled = pImage->m_flags & Image::IF_CANNOT_BE_SCALED;
 			if (bCannotBeScaled)
 			{
-				MipsToStore = 1;
+				// Store only the mips that we have already calculated. We assume we have calculated them correctly.
+				//MipsToStore = 1;
+				MipsToStore = pImage->GetLODCount();
 			}
 
 			// TODO:
diff --git a/Engine/Plugins/Experimental/Mutable/Source/MutableRuntime/Private/MuR/OpImageDisplace.h b/Engine/Plugins/Experimental/Mutable/Source/MutableRuntime/Private/MuR/OpImageDisplace.h
--- a/Engine/Plugins/Experimental/Mutable/Source/MutableRuntime/Private/MuR/OpImageDisplace.h
+++ b/Engine/Plugins/Experimental/Mutable/Source/MutableRuntime/Private/MuR/OpImageDisplace.h
@@ -11,17 +11,17 @@
 namespace mu
 {
 
-    inline uint8_t MutableEncodeOffset( int x, int y )
+    inline uint8 MutableEncodeOffset( int x, int y )
 	{
-        uint8_t c = uint8_t( ((7)<<4) | 7 );
+        uint8 c = uint8( ((7)<<4) | 7 );
 		if ( x<8 && x>-8 && y<8 && y>-8 )
 		{
-            c = uint8_t( ((x+7)<<4) | (y+7) );
+            c = uint8( ((x+7)<<4) | (y+7) );
 		}
 		return c;
 	}
 
-    inline void MutableDecodeOffset( uint8_t c, int& x, int& y )
+    inline void MutableDecodeOffset( uint8 c, int& x, int& y )
 	{
 		x = int(c>>4) - 7;
 		y = int(c&0xf) - 7;
@@ -30,86 +30,97 @@ namespace mu
 	//---------------------------------------------------------------------------------------------
 	//!
 	//---------------------------------------------------------------------------------------------
-	inline void ImageMakeGrowMap( Image* pResult, const Image* pMask, int border )
+	inline void ImageMakeGrowMap( Image* pResult, const Image* pMask, int InBorder )
 	{
-		MUTABLE_CPUPROFILER_SCOPE(ImageMakeGrowMap)
+		MUTABLE_CPUPROFILER_SCOPE(ImageMakeGrowMap);
 
-		check( pMask->GetFormat() == EImageFormat::IF_L_UBYTE );
-		check( pResult->GetFormat() == EImageFormat::IF_L_UBYTE );
-		check( pResult->GetSizeX() == pMask->GetSizeX() );
-		check( pResult->GetSizeY() == pMask->GetSizeY() );
+		check(pMask->GetFormat() == EImageFormat::IF_L_UBYTE);
+		check(pResult->GetFormat() == EImageFormat::IF_L_UBYTE);
+		check(pResult->GetSizeX() == pMask->GetSizeX());
+		check(pResult->GetSizeY() == pMask->GetSizeY());
+		check(pResult->GetLODCount() == pMask->GetLODCount());
 
-		int sizeX = pResult->GetSizeX();
-		int sizeY = pResult->GetSizeY();
-		check( sizeX>0 && sizeY>0 );
+		int32 MipCount = pResult->GetLODCount();
+		int32 SizeX = pResult->GetSizeX();
+		int32 SizeY = pResult->GetSizeY();
+		check( SizeX>0 && SizeY>0 );
 
-        if (sizeX<4 || sizeY<4)
+        if (SizeX<4 || SizeY<4)
         {
             return;
         }
 
-		memset( pResult->GetData(), MutableEncodeOffset( 0, 0 ), sizeX*sizeY );
+		FMemory::Memset( pResult->GetData(), MutableEncodeOffset( 0, 0 ), SizeX*SizeY );
 
-		ImagePtrConst pThisMask = pMask;
-		for ( int b=0; b<border; ++b )
+		Ptr<const Image> pThisMask = pMask;
+		for ( int b=0; b<InBorder; ++b )
 		{
-			ImagePtr pNextMask = pThisMask->Clone();            
+			Ptr<Image> pNextMask = pThisMask->Clone();
 
-			//for ( int y=0; y<sizeY; ++y )
-			const auto& ProcessRow = [
-				pThisMask, pNextMask, sizeX, sizeY, pResult
-			] (int32 y)
-			{
-				uint8_t* pNextData = pNextMask->GetData() + sizeX*y;
+			for( int CurrentMip=0; CurrentMip<MipCount; ++CurrentMip)
+			{ 
+				FIntVector2 MipSize = pResult->CalculateMipSize(CurrentMip);
 
-				const uint8_t* pS0 = y > 0 ? pThisMask->GetData() + (y - 1) * sizeX : nullptr;
-				const uint8_t* pS1 = pThisMask->GetData() + y * sizeX;
-				bool bNotLastRow = y < sizeY - 1;
-				const uint8_t* pS2 = bNotLastRow ? pThisMask->GetData() + (y + 1) * sizeX : nullptr;
+				uint8* ResultData = pResult->GetMipData(CurrentMip);
+				const uint8* ThisMaskData = pThisMask->GetMipData(CurrentMip);
+				uint8* NextMaskData = pNextMask->GetMipData(CurrentMip);
 
-				const uint8_t* pR0 = y > 0 ? pResult->GetData() + (y - 1) * sizeX : 0;
-				uint8_t* pR1 = pResult->GetData() + y * sizeX;
+				//for ( int y=0; y<sizeY; ++y )
+				const auto ProcessRow = [
+					ThisMaskData, NextMaskData, MipSize, ResultData
+				] (int32 y)
+				{
+					uint8* pNextData = NextMaskData + MipSize.X * y;
 
-				const uint8_t* pR2 = bNotLastRow ? pResult->GetData() + (y + 1) * sizeX : nullptr;
+					const uint8* pS0 = y > 0 ? ThisMaskData + (y - 1) * MipSize.X : nullptr;
+					const uint8* pS1 = ThisMaskData + y * MipSize.X;
+					bool bNotLastRow = y < MipSize.Y - 1;
+					const uint8* pS2 = bNotLastRow ? ThisMaskData + (y + 1) * MipSize.X : nullptr;
 
-				for (int x = 0; x < sizeX; ++x)
-				{
-					int dx, dy;
-					if (pS1[x])
-					{
-						*pNextData = 255;
-					}
-					else if (y > 0 && pS0[x])
-					{
-						MutableDecodeOffset(pR0[x], dx, dy);
-						pR1[x] = MutableEncodeOffset(dx, dy - 1);
-						*pNextData = 255;
-					}
-					else if (bNotLastRow && pS2 && pR2 && pS2[x])
-					{
-						MutableDecodeOffset(pR2[x], dx, dy);
-						pR1[x] = MutableEncodeOffset(dx, dy + 1);
-						*pNextData = 255;
-					}
-					else if (x < sizeX - 1 && pS1[x + 1])
-					{
-						MutableDecodeOffset(pR1[x + 1], dx, dy);
-						pR1[x] = MutableEncodeOffset(dx + 1, dy);
-						*pNextData = 255;
-					}
-					else if (x > 0 && pS1[x - 1])
+					const uint8* pR0 = y > 0 ? ResultData + (y - 1) * MipSize.X : 0;
+					uint8* pR1 = ResultData + y * MipSize.X;
+
+					const uint8* pR2 = bNotLastRow ? ResultData + (y + 1) * MipSize.X : nullptr;
+
+					for (int x = 0; x < MipSize.X; ++x)
 					{
-						MutableDecodeOffset(pR1[x - 1], dx, dy);
-						pR1[x] = MutableEncodeOffset(dx - 1, dy);
-						*pNextData = 255;
+						int dx, dy;
+						if (pS1[x])
+						{
+							*pNextData = 255;
+						}
+						else if (y > 0 && pS0[x])
+						{
+							MutableDecodeOffset(pR0[x], dx, dy);
+							pR1[x] = MutableEncodeOffset(dx, dy - 1);
+							*pNextData = 255;
+						}
+						else if (bNotLastRow && pS2 && pR2 && pS2[x])
+						{
+							MutableDecodeOffset(pR2[x], dx, dy);
+							pR1[x] = MutableEncodeOffset(dx, dy + 1);
+							*pNextData = 255;
+						}
+						else if (x < MipSize.X - 1 && pS1[x + 1])
+						{
+							MutableDecodeOffset(pR1[x + 1], dx, dy);
+							pR1[x] = MutableEncodeOffset(dx + 1, dy);
+							*pNextData = 255;
+						}
+						else if (x > 0 && pS1[x - 1])
+						{
+							MutableDecodeOffset(pR1[x - 1], dx, dy);
+							pR1[x] = MutableEncodeOffset(dx - 1, dy);
+							*pNextData = 255;
+						}
+
+						++pNextData;
 					}
 
-					++pNextData;
-				}
-
-			};
+				};
 
-			ParallelFor(sizeY, ProcessRow);
+				ParallelFor(MipSize.Y, ProcessRow);
+			}
 
 			pThisMask = pNextMask;
 		}
@@ -133,15 +144,19 @@ namespace mu
 		check( pResult->GetSizeX() == pMap->GetSizeX() );
 		check( pResult->GetSizeY() == pMap->GetSizeY() );
 
-		int sizeX = pResult->GetSizeX();
-		int sizeY = pResult->GetSizeY();
+		int32 LODCount = pResult->GetLODCount();
+		check(LODCount <= pMap->GetLODCount());
+		check(LODCount <= pSource->GetLODCount());
+
+		int32 sizeX = pResult->GetSizeX();
+		int32 sizeY = pResult->GetSizeY();
 
         if (sizeX<4 || sizeY<4)
         {
             return;
         }
 
-        const uint8_t* pSourceData = pSource->GetData();
+        const uint8* pSourceData = pSource->GetData();
 
         auto mapFormat = pMap->GetFormat();
         bool isUncompressed = ( mapFormat == EImageFormat::IF_L_UBYTE );
@@ -154,12 +169,12 @@ namespace mu
 			case EImageFormat::IF_L_UBYTE:
 			{
 				//for ( int y=0; y<sizeY; ++y )
-				const auto& ProcessRow = [
+				const auto ProcessRow = [
 					pSourceData, pMap, pResult, sizeX, sizeY
 				] (int32 y)
 				{
-					const uint8_t* pMapData = pMap->GetData() + y * sizeX;
-					uint8_t* pResultData = pResult->GetData() + y * sizeX;
+					const uint8* pMapData = pMap->GetData() + y * sizeX;
+					uint8* pResultData = pResult->GetData() + y * sizeX;
 
 					for (int x = 0; x < sizeX; ++x)
 					{
@@ -169,7 +184,7 @@ namespace mu
 						check(y + dy >= 0 && y + dy < sizeY);
 
 						int offset = ((y + dy) * sizeX + (x + dx));
-						memcpy(pResultData, &pSourceData[offset], 1);
+						FMemory::Memcpy(pResultData, &pSourceData[offset], 1);
 
 						++pResultData;
 						++pMapData;
@@ -184,12 +199,12 @@ namespace mu
 			case EImageFormat::IF_RGB_UBYTE:
 			{
 				//for (int y = 0; y < sizeY; ++y)
-				const auto& ProcessRow = [
+				const auto ProcessRow = [
 					pSourceData, pMap, pResult, sizeX, sizeY
 				] (int32 y)
 				{
-					const uint8_t* pMapData = pMap->GetData() + y * sizeX;
-					uint8_t* pResultData = pResult->GetData() + y * sizeX*3;
+					const uint8* pMapData = pMap->GetData() + y * sizeX;
+					uint8* pResultData = pResult->GetData() + y * sizeX*3;
 
 					for (int x = 0; x < sizeX; ++x)
 					{
@@ -199,7 +214,7 @@ namespace mu
 						check(y + dy >= 0 && y + dy < sizeY);
 
 						int offset = ((y + dy) * sizeX + (x + dx)) * 3;
-						memcpy(pResultData, &pSourceData[offset], 3);
+						FMemory::Memcpy(pResultData, &pSourceData[offset], 3);
 
 						pResultData += 3;
 						++pMapData;
@@ -216,12 +231,12 @@ namespace mu
 				// Generic
 				int pixelSize = GetImageFormatData( pResult->GetFormat() ).m_bytesPerBlock;
 				//for ( int y=0; y<sizeY; ++y )
-				const auto& ProcessRow = [
+				const auto ProcessRow = [
 					pSourceData, pMap, pResult, sizeX, sizeY, pixelSize
 				] (int32 y)
 				{
-					const uint8_t* pMapData = pMap->GetData() + y * sizeX;
-					uint8_t* pResultData = pResult->GetData() + y * sizeX * pixelSize;
+					const uint8* pMapData = pMap->GetData() + y * sizeX;
+					uint8* pResultData = pResult->GetData() + y * sizeX * pixelSize;
 
 					for (int x = 0; x < sizeX; ++x)
 					{
@@ -231,7 +246,7 @@ namespace mu
 						check(y + dy >= 0 && y + dy < sizeY);
 
 						int offset = ((y + dy) * sizeX + (x + dx)) * pixelSize;
-						memcpy(pResultData, &pSourceData[offset], pixelSize);
+						FMemory::Memcpy(pResultData, &pSourceData[offset], pixelSize);
 
 						pResultData += pixelSize;
 						++pMapData;
@@ -249,14 +264,14 @@ namespace mu
 		{
             int pixelSize = GetImageFormatData( pResult->GetFormat() ).m_bytesPerBlock;
 
-			const uint8_t* pMapData = pMap->GetData();
-			uint8_t* pResultData = pResult->GetData();
+			const uint8* pMapData = pMap->GetData();
+			uint8* pResultData = pResult->GetData();
 
 			// TODO: ParallelFor this
 
-            for (int lod=0;lod<pMap->GetLODCount(); ++lod)
+            for (int lod=0;lod<LODCount; ++lod)
             {
-                pMapData += 4+sizeY*sizeof(uint32_t);
+                pMapData += 4+sizeY*sizeof(uint32);
 
                 for ( int y=0; y<sizeY; ++y )
                 {
@@ -267,10 +282,10 @@ namespace mu
                         uint16 equal = *(const uint16*)pMapData;
                         pMapData += 2;
 
-                        uint8_t different = *pMapData;
+                        uint8 different = *pMapData;
                         ++pMapData;
 
-                        uint8_t equalPixel = *pMapData;
+                        uint8 equalPixel = *pMapData;
                         ++pMapData;
 
                         // Equal pixels
@@ -283,7 +298,7 @@ namespace mu
 							if (dx >= 0 && dx < sizeX && dy >= 0 && dy < sizeY)
 							{
 								int offset = (dy * sizeX + dx) * pixelSize;
-								memcpy(pResultData, &pSourceData[offset], equal * pixelSize);
+								FMemory::Memcpy(pResultData, &pSourceData[offset], equal * pixelSize);
 							}
 
                             pResultData += equal*pixelSize;
@@ -302,7 +317,7 @@ namespace mu
 							if (dx>=0 && dx<sizeX && dy>=0 && dy<sizeY)
 							{
 								int offset = (dy * sizeX + dx) * pixelSize;
-								memcpy(pResultData, &pSourceData[offset], pixelSize);
+								FMemory::Memcpy(pResultData, &pSourceData[offset], pixelSize);
 							}
 
                             pResultData += pixelSize;
diff --git a/Engine/Plugins/Experimental/Mutable/Source/MutableRuntime/Private/MuR/OpImageMipmap.h b/Engine/Plugins/Experimental/Mutable/Source/MutableRuntime/Private/MuR/OpImageMipmap.h
--- a/Engine/Plugins/Experimental/Mutable/Source/MutableRuntime/Private/MuR/OpImageMipmap.h
+++ b/Engine/Plugins/Experimental/Mutable/Source/MutableRuntime/Private/MuR/OpImageMipmap.h
@@ -28,13 +28,13 @@ namespace mu
 			float* inBuf = &inOutKernelStorage[0];
 			float* outBuf = &temp[0];
 
-			for (int32_t i = 0; i < ITERS; ++i)
+			for (int32 i = 0; i < ITERS; ++i)
 			{
 				constexpr float oneOverThree = 1.0f / 3.0f;
 
 				outBuf[0] = (inBuf[0] + inBuf[1]) * oneOverThree;
 
-				for (int32_t j = 1; j < N - 1; ++j)
+				for (int32 j = 1; j < N - 1; ++j)
 				{
 					outBuf[j] = (inBuf[j - 1] + inBuf[j] + inBuf[j + 1]) * oneOverThree;
 				}
@@ -47,7 +47,7 @@ namespace mu
 
 			if (ITERS % 2 == 1)
 			{
-				for (uint32_t i = 0; i < N; ++i)
+				for (uint32 i = 0; i < N; ++i)
 				{
 					inOutKernelStorage[i] = temp[i];
 				}
@@ -68,7 +68,7 @@ namespace mu
 
 			const float center = static_cast<float>(N) * 0.5f;
 			float sum = 0;
-			for (int32_t i = 0; i < N; ++i)
+			for (int32 i = 0; i < N; ++i)
 			{
 				const float dist = normalDist(static_cast<float>(i) - center + 0.5f);
 				kernel[i] = dist;
@@ -76,7 +76,7 @@ namespace mu
 			}
 
 			const float normFactor = 1.0f / sum;
-			for (int32_t i = 0; i < N; ++i)
+			for (int32 i = 0; i < N; ++i)
 			{
 				kernel[i] *= normFactor;
 			}
@@ -85,7 +85,7 @@ namespace mu
 		template<size_t N>
 		inline void BuildSharpenKernel(float(&kernel)[N], float factor)
 		{
-			for (uint32_t i = 0; i < N; ++i)
+			for (uint32 i = 0; i < N; ++i)
 			{
 				kernel[i] = 0.0f;
 			}
@@ -99,9 +99,9 @@ namespace mu
 		{
 			static_assert(N == SharpenKernelStorage8::SIZE);
 
-			for (int32_t y = 0; y < N; ++y)
+			for (int32 y = 0; y < N; ++y)
 			{
-				for (int32_t x = 0; x < N; ++x)
+				for (int32 x = 0; x < N; ++x)
 				{
 					outKernel.m_storage[y * N + x] = kernel1D[y] * kernel1D[x];
 				}
@@ -131,7 +131,7 @@ namespace mu
 			BlurKernel<1>(kernel1DStorage0);
 			BlurKernel<3>(kernel1DStorage1);
 
-			for (int32_t i = 0; i < SharpenKernelStorage8::SIZE; ++i)
+			for (int32 i = 0; i < SharpenKernelStorage8::SIZE; ++i)
 			{
 				kernel1DStorage0[i] += kernel1DStorage1[i];
 			}
@@ -141,78 +141,78 @@ namespace mu
 			return kernel;
 		}
 
-		template<int32_t PIXEL_SIZE, typename CHANNEL_TYPE, EAddressMode AD_MODE = EAddressMode::AM_NONE>
+		template<int32 PIXEL_SIZE, typename CHANNEL_TYPE, EAddressMode AD_MODE = EAddressMode::AM_NONE>
 		class ImageAccessor
 		{
-			const vec2<int32_t> m_dim;
+			const FIntVector2 m_dim;
 			CHANNEL_TYPE const* const m_pImageBuf;
 
 		public:
-			ImageAccessor(const CHANNEL_TYPE* const pImageBuf, vec2<uint32_t> dim)
+			ImageAccessor(const CHANNEL_TYPE* const pImageBuf, FIntVector2 dim)
 				: m_dim(dim)
 				, m_pImageBuf(pImageBuf)
 			{
 			}
 
-			uint8_t operator()(int32_t x, int32_t y, int32_t c) const
+			uint8 operator()(int32 x, int32 y, int32 c) const
 			{
 				if constexpr (AD_MODE == EAddressMode::AM_WRAP)
 				{
-					// (y & (m_dim.y() - 1)) * m_dim.y() + (x & (m_dim.x() - 1))
-					return m_pImageBuf[((y % m_dim.y()) * m_dim.x() + (x % m_dim.x())) * PIXEL_SIZE + c];
+					// (y & (m_dim.Y - 1)) * m_dim.Y + (x & (m_dim.X - 1))
+					return m_pImageBuf[((y % m_dim.Y) * m_dim.X + (x % m_dim.X)) * PIXEL_SIZE + c];
 				}
 				else if (AD_MODE == EAddressMode::AM_CLAMP)
 				{
-					return m_pImageBuf[(mu::clamp(y, 0, m_dim.y() - 1) * m_dim.x() +
-						mu::clamp(x, 0, m_dim.x() - 1))
+					return m_pImageBuf[(mu::clamp(y, 0, m_dim.Y - 1) * m_dim.X +
+						mu::clamp(x, 0, m_dim.X - 1))
 						* PIXEL_SIZE + c];
 				}
 				else if (AD_MODE == EAddressMode::AM_BLACK_BORDER)
 				{
-					return ((x < m_dim.x()) && (y < m_dim.y()) && (x >= 0) && (y >= 0))
-						? m_pImageBuf[(y * m_dim.x() + x) * PIXEL_SIZE + c]
+					return ((x < m_dim.X) && (y < m_dim.Y) && (x >= 0) && (y >= 0))
+						? m_pImageBuf[(y * m_dim.X + x) * PIXEL_SIZE + c]
 						: 0;
 				}
 				else
 				{
-					return m_pImageBuf[(y * m_dim.x() + x) * PIXEL_SIZE + c];
+					return m_pImageBuf[(y * m_dim.X + x) * PIXEL_SIZE + c];
 				}
 			}
 		};
 
 
-		template<int32_t PIXEL_SIZE, EAddressMode AD, typename PIXEL_TYPE>
+		template<int32 PIXEL_SIZE, EAddressMode AD, typename PIXEL_TYPE>
 		void GenerateMipSharpenedRegion(
 			const SharpenKernelStorage8& kernel,
 			const ImageAccessor<PIXEL_SIZE, PIXEL_TYPE, AD>& sourceAccessor,
-			PIXEL_TYPE* pDest, vec2<int32_t> destSize,
-			vec2<int32_t> regionStart, vec2<int32_t> regionEnd)
+			PIXEL_TYPE* pDest, FIntVector2 destSize,
+			FIntVector2 regionStart, FIntVector2 regionEnd)
 		{
 			// Only kernels of size 8 supported.
-			constexpr int32_t kernelSize = SharpenKernelStorage8::SIZE;
+			constexpr int32 kernelSize = SharpenKernelStorage8::SIZE;
 			static_assert(kernelSize == 8);
 
-			constexpr int32_t kernelCenter = (kernelSize >> 1) - 1;
-			for (int32_t y = regionStart.y(); y < regionEnd.y(); ++y)
+			constexpr int32 kernelCenter = (kernelSize >> 1) - 1;
+			for (int32 y = regionStart.Y; y < regionEnd.Y; ++y)
 			{
-				for (int32_t x = regionStart.x(); x < regionEnd.x(); ++x)
+				for (int32 x = regionStart.X; x < regionEnd.X; ++x)
 				{
 					float stagingStorage[kernelSize * kernelSize * PIXEL_SIZE];
 
-					const vec2<uint32_t> sourceCoord = vec2<uint32_t>(x << 1, y << 1);
+					const FIntVector2 sourceCoord = FIntVector2(x << 1, y << 1);
 
 
 					// Convert Image kernel stamp to float. This conversion will be done multiple times.
 					// and should be done once before applying the kernel, or operate directly in unorm8/snorm8.
-					for (int32_t ky = 0; ky < kernelSize; ++ky)
+					for (int32 ky = 0; ky < kernelSize; ++ky)
 					{
-						for (int32_t kx = 0; kx < kernelSize; ++kx)
+						for (int32 kx = 0; kx < kernelSize; ++kx)
 						{
-							for (int32_t c = 0; c < PIXEL_SIZE; ++c)
+							for (int32 c = 0; c < PIXEL_SIZE; ++c)
 							{
 								stagingStorage[(ky * kernelSize + kx) + (kernelSize * kernelSize * c)] = static_cast<float>(
-									sourceAccessor(sourceCoord.x() + kx - kernelCenter,
-										sourceCoord.y() + ky - kernelCenter,
+									sourceAccessor(sourceCoord.X + kx - kernelCenter,
+										sourceCoord.Y + ky - kernelCenter,
 										c)) / 255.0f;
 							}
 						}
@@ -223,17 +223,17 @@ namespace mu
 					// Since here we are only applying to half the image and the kernel used is small (8x8) the benefits are not 
 					// clear.
 
-					for (int32_t c = 0; c < PIXEL_SIZE; ++c)
+					for (int32 c = 0; c < PIXEL_SIZE; ++c)
 					{
 						float sampleValue = 0.0f;
 						float sum = 0.0f;
-						for (int32_t k = 0; k < kernelSize * kernelSize; ++k)
+						for (int32 k = 0; k < kernelSize * kernelSize; ++k)
 						{
 							sampleValue += kernel.m_storage[k] * stagingStorage[k + (kernelSize * kernelSize * c)];
 							sum += kernel.m_storage[k];
 						}
 
-						pDest[(y * destSize.x() + x) * PIXEL_SIZE + c] = static_cast<uint8_t>(mu::clamp(sampleValue, 0.0f, 1.0f) * 255.0f);
+						pDest[(y * destSize.X + x) * PIXEL_SIZE + c] = static_cast<uint8>(mu::clamp(sampleValue, 0.0f, 1.0f) * 255.0f);
 					}
 				}
 			}
@@ -241,8 +241,8 @@ namespace mu
 
 		template<size_t PIXEL_SIZE, EAddressMode AD_MODE>
 		inline void GenerateMipmapUint8Sharpen(int mips,
-			const uint8_t* pSource, uint8_t* pDest,
-			vec2<int32_t> sourceSize,
+			const uint8* pSource, uint8* pDest,
+			FIntVector2 sourceSize,
 			float sharpeningFactor)
 		{
 			SharpenKernelStorage8 kernel = MakeMipGaussianSharpenKernel(sharpeningFactor);
@@ -251,10 +251,10 @@ namespace mu
 
 			for (; mips >= 0; --mips)
 			{
-				ImageAccessor<PIXEL_SIZE, uint8_t, AD_MODE> sourceSampler(pSource, sourceSize);
-				ImageAccessor<PIXEL_SIZE, uint8_t, EAddressMode::AM_NONE> sourceSamplerAddressModeNone(pSource, sourceSize);
+				ImageAccessor<PIXEL_SIZE, uint8, AD_MODE> sourceSampler(pSource, sourceSize);
+				ImageAccessor<PIXEL_SIZE, uint8, EAddressMode::AM_NONE> sourceSamplerAddressModeNone(pSource, sourceSize);
 
-				vec2<int32_t> destSize = vec2<int32_t>(FMath::DivideAndRoundUp(sourceSize.x(), 2), FMath::DivideAndRoundUp(sourceSize.y(), 2));
+				FIntVector2 destSize = FIntVector2(FMath::DivideAndRoundUp(sourceSize.X, 2), FMath::DivideAndRoundUp(sourceSize.Y, 2));
 
 				// Core image, will never sample outside the image so we don't need to care about address mode.
 				// Only sample source pixels where the kernel can be fully applied.
@@ -263,81 +263,81 @@ namespace mu
 					kernel,
 					sourceSamplerAddressModeNone,
 					pDest, destSize,
-					vec2<int32_t>(2, 2), vec2<int32_t>(destSize.x() - 2, destSize.y() - 2));
+					FIntVector2(2, 2), FIntVector2(destSize.X - 2, destSize.Y - 2));
 
 				// Horizontal borders.
 				GenerateMipSharpenedRegion(
 					kernel,
 					sourceSampler,
 					pDest, destSize,
-					vec2<int32_t>(0, 0), vec2<int32_t>(destSize.x(), FMath::Min(2, destSize.y())));
+					FIntVector2(0, 0), FIntVector2(destSize.X, FMath::Min(2, destSize.Y)));
 
 				GenerateMipSharpenedRegion(
 					kernel,
 					sourceSampler,
 					pDest, destSize,
-					vec2<int32_t>(0, FMath::Max(destSize.y() - 2, 0)), vec2<int32_t>(destSize.x(), destSize.y()));
+					FIntVector2(0, FMath::Max(destSize.Y - 2, 0)), FIntVector2(destSize.X, destSize.Y));
 
 				// Vertical borders
 				GenerateMipSharpenedRegion(
 					kernel,
 					sourceSampler,
 					pDest, destSize,
-					vec2<int32_t>(0, FMath::Min(2, destSize.y())), vec2<int32_t>(FMath::Min(2, destSize.x()), destSize.y() - 2));
+					FIntVector2(0, FMath::Min(2, destSize.Y)), FIntVector2(FMath::Min(2, destSize.X), destSize.Y - 2));
 
 				GenerateMipSharpenedRegion(
 					kernel,
 					sourceSampler,
 					pDest, destSize,
-					vec2<int32_t>(destSize.x() - 2, FMath::Min(2, destSize.y())), vec2<int32_t>(destSize.x(), destSize.y() - 2));
+					FIntVector2(destSize.X - 2, FMath::Min(2, destSize.Y)), FIntVector2(destSize.X, destSize.Y - 2));
 
 				sourceSize = destSize;
 
 				pSource = pDest;
-				pDest = pDest + destSize.x() * destSize.y() * PIXEL_SIZE;
+				pDest = pDest + destSize.X * destSize.Y * PIXEL_SIZE;
 
-				totalDestSize += sourceSize.x() * sourceSize.y() * PIXEL_SIZE;
+				totalDestSize += sourceSize.X * sourceSize.Y * PIXEL_SIZE;
 			}
 		}
 
-		template<int32_t PIXEL_SIZE>
+		template<int32 PIXEL_SIZE>
 		inline void GenerateMipmapUint8Unfiltered(
 			int mips,
-			const uint8_t* pSource, uint8_t* pDest,
-			vec2<int> sourceSize)
+			const uint8* pSource, uint8* pDest,
+			FIntVector2 sourceSize)
 		{
 			for (; mips >= 0; --mips)
 			{
-				vec2<int32_t> destSize = vec2<int32_t>(FMath::DivideAndRoundUp(sourceSize.x(), 2), FMath::DivideAndRoundUp(sourceSize.y(), 2));
+				FIntVector2 destSize = FIntVector2(FMath::DivideAndRoundUp(sourceSize.X, 2), FMath::DivideAndRoundUp(sourceSize.Y, 2));
 
-				for (int32_t y = 0; y < destSize.y(); ++y)
+				for (int32 y = 0; y < destSize.Y; ++y)
 				{
-					for (int32_t x = 0; x < destSize.x(); ++x)
+					for (int32 x = 0; x < destSize.X; ++x)
 					{
-						for (int32_t c = 0; c < PIXEL_SIZE; ++c)
+						for (int32 c = 0; c < PIXEL_SIZE; ++c)
 						{
-							pDest[(y * destSize.x() + x) * PIXEL_SIZE + c] =
-								pSource[((y << 1) * sourceSize.x() + (x << 1)) * PIXEL_SIZE + c];
+							pDest[(y * destSize.X + x) * PIXEL_SIZE + c] =
+								pSource[((y << 1) * sourceSize.X + (x << 1)) * PIXEL_SIZE + c];
 						}
 					}
 				}
 
 				sourceSize = destSize;
 				pSource = pDest;
-				pDest = pDest + destSize.x() * destSize.y() * PIXEL_SIZE;
+				pDest = pDest + destSize.X * destSize.Y * PIXEL_SIZE;
 			}
 		}
 
-		template<int32_t PIXEL_SIZE>
+		template<int32 PIXEL_SIZE>
 		inline void GenerateMipmapsUint8SimpleAverage(
 			int mips,
-			const uint8_t* pSource, uint8_t* pDest,
-			vec2<int> sourceSize)
+			const uint8* pSource, uint8* pDest,
+			FIntVector2 sourceSize)
 		{
-			const uint8_t* pMipSource = pSource;
-			uint8_t* pMipDest = pDest;
+			const uint8* pMipSource = pSource;
+			uint8* pMipDest = pDest;
 
-			vec2<int> destSize = sourceSize;
+			FIntVector2 destSize = sourceSize;
 			for (int m = 0; m < mips; ++m)
 			{
 				check(destSize[0] > 1 || destSize[1] > 1);
@@ -355,13 +355,13 @@ namespace mu
 				int sourceStride = sourceSize[0] * PIXEL_SIZE;
 				int destStride = destSize[0] * PIXEL_SIZE;
 
-				const auto& ProcessRow = [
+				const auto ProcessRow = [
 					pMipDest, pMipSource, fullColumns, strayColumn, sourceStride, destStride
 				] (uint32 y)
 				{
-					const uint8_t* pSourceRow0 = pMipSource + 2 * y * sourceStride;
-					const uint8_t* pSourceRow1 = pSourceRow0 + sourceStride;
-					uint8_t* pDestRow = pMipDest + y * destStride;
+					const uint8* pSourceRow0 = pMipSource + 2 * y * sourceStride;
+					const uint8* pSourceRow1 = pSourceRow0 + sourceStride;
+					uint8* pDestRow = pMipDest + y * destStride;
 
 					for (int x = 0; x < fullColumns; ++x)
 					{
@@ -369,7 +369,7 @@ namespace mu
 						{
 							int p = pSourceRow0[c] + pSourceRow0[PIXEL_SIZE + c] + pSourceRow1[c] +
 								pSourceRow1[PIXEL_SIZE + c];
-							pDestRow[c] = (uint8_t)(p >> 2);
+							pDestRow[c] = (uint8)(p >> 2);
 						}
 
 						pSourceRow0 += 2 * PIXEL_SIZE;
@@ -382,7 +382,7 @@ namespace mu
 						for (int c = 0; c < PIXEL_SIZE; ++c)
 						{
 							int p = pSourceRow0[c] + pSourceRow1[c];
-							pDestRow[c] = (uint8_t)(p >> 1);
+							pDestRow[c] = (uint8)(p >> 1);
 						}
 					}
 				};
@@ -402,16 +402,16 @@ namespace mu
 
 					if (strayRow)
 					{
-						const uint8_t* pSourceRow0 = pMipSource + 2 * fullRows * sourceStride;
-						const uint8_t* pSourceRow1 = pSourceRow0 + sourceStride;
-						uint8_t* pDestRow = pMipDest + fullRows * destStride;
+						const uint8* pSourceRow0 = pMipSource + 2 * fullRows * sourceStride;
+						const uint8* pSourceRow1 = pSourceRow0 + sourceStride;
+						uint8* pDestRow = pMipDest + fullRows * destStride;
 
 						for (int x = 0; x < fullColumns; ++x)
 						{
 							for (int c = 0; c < PIXEL_SIZE; ++c)
 							{
 								int p = pSourceRow0[c] + pSourceRow0[PIXEL_SIZE + c];
-								pDestRow[c] = (uint8_t)(p >> 1);
+								pDestRow[c] = (uint8)(p >> 1);
 							}
 
 							pSourceRow0 += 2 * PIXEL_SIZE;
@@ -441,10 +441,10 @@ namespace mu
 		//! Generate the mipmaps for byte-based images of whatever number of channels.
 		//! \param mips number of additional levels to build from the source.
 		//---------------------------------------------------------------------------------------------
-	template<int32_t PIXEL_SIZE>
+	template<int32 PIXEL_SIZE>
 	inline void GenerateMipmapsUint8(int mips,
-		const uint8_t* pSource, uint8_t* pDest,
-		vec2<int32_t> sourceSize,
+		const uint8* pSource, uint8* pDest,
+		FIntVector2 sourceSize,
 		const FMipmapGenerationSettings& settings)
 	{
 		using namespace OpImageMipmap_Detail;
@@ -543,28 +543,27 @@ namespace mu
         case EImageFormat::IF_ASTC_4x4_RG_LDR:
         {
 		// Uncompress the last mip that we already have
-		vec2<int> uncompressedSize = pBase->CalculateMipSize(startLevel);
-		scratch->pUncompressed = new Image(
-			(uint16)uncompressedSize[0],
-			(uint16)uncompressedSize[1],
-			1,
-			EImageFormat::IF_RGBA_UBYTE);
-
-
-		FImageSize uncompressedMipsSize = pBase->CalculateMipSize(startLevel + 1);
-        // Generate the mipmaps from there on
-        scratch->pUncompressedMips = new Image(
-			(uint16)uncompressedMipsSize[0],
-			(uint16)uncompressedMipsSize[1],
-			FMath::Max(1, levelCount - startLevel - 1),
-			EImageFormat::IF_RGBA_UBYTE);
-
-		// Compress the mipmapped image
-		scratch->pCompressedMips = new Image(
-			(uint16)uncompressedMipsSize[0],
-			(uint16)uncompressedMipsSize[1],
-			scratch->pUncompressedMips->GetLODCount(),
-			pBase->GetFormat());
+			FIntVector2 uncompressedSize = pBase->CalculateMipSize(startLevel);
+			scratch->pUncompressed = new Image(
+				(uint16)uncompressedSize[0],
+				(uint16)uncompressedSize[1],
+				1,
+				EImageFormat::IF_RGBA_UBYTE);
+
+			FIntVector2 uncompressedMipsSize = pBase->CalculateMipSize(startLevel + 1);
+			// Generate the mipmaps from there on
+			scratch->pUncompressedMips = new Image(
+				(uint16)uncompressedMipsSize[0],
+				(uint16)uncompressedMipsSize[1],
+				FMath::Max(1, levelCount - startLevel - 1),
+				EImageFormat::IF_RGBA_UBYTE);
+
+			// Compress the mipmapped image
+			scratch->pCompressedMips = new Image(
+				(uint16)uncompressedMipsSize[0],
+				(uint16)uncompressedMipsSize[1],
+				scratch->pUncompressedMips->GetLODCount(),
+				pBase->GetFormat());
 
             break;
         }
@@ -573,7 +572,7 @@ namespace mu
         case EImageFormat::IF_L_UBYTE_RLE:
         {
             // Uncompress the last mip that we already have
-            vec2<int> uncompressedSize = pBase->CalculateMipSize(startLevel);
+			FIntVector2 uncompressedSize = pBase->CalculateMipSize(startLevel);
             scratch->pUncompressed = new Image(
                 (uint16)uncompressedSize[0],
                 (uint16)uncompressedSize[1],
@@ -581,7 +580,7 @@ namespace mu
                 EImageFormat::IF_L_UBYTE);
 
 
-			FImageSize uncompressedMipsSize = pBase->CalculateMipSize(startLevel + 1);
+			FIntVector2 uncompressedMipsSize = pBase->CalculateMipSize(startLevel + 1);
             // Generate the mipmaps from there on
             scratch->pUncompressedMips = new Image(
                 (uint16)uncompressedMipsSize[0],
@@ -639,8 +638,8 @@ namespace mu
 			check(
 				[&]() -> bool
 				{
-					const FImageSize BaseImageNextMipSize = pBase->CalculateMipSize(startLevel + 1);
-					return BaseImageNextMipSize.x() == pDest->GetSizeX() && BaseImageNextMipSize.y() == pDest->GetSizeY();
+					const FIntVector2 BaseImageNextMipSize = pBase->CalculateMipSize(startLevel + 1);
+					return BaseImageNextMipSize.X == pDest->GetSizeX() && BaseImageNextMipSize.Y == pDest->GetSizeY();
 				}());
 		}
 
@@ -661,11 +660,11 @@ namespace mu
             return;
         }
 
-        const uint8_t* pSourceBuf = !bGenerateOnlyTail ? pDest->GetMipData(startLevel) : pBase->GetMipData(startLevel);
+        const uint8* pSourceBuf = !bGenerateOnlyTail ? pDest->GetMipData(startLevel) : pBase->GetMipData(startLevel);
 
-        uint8_t* pDestBuf = !bGenerateOnlyTail ? pDest->GetMipData(startLevel + 1) : pDest->GetMipData(0);
+        uint8* pDestBuf = !bGenerateOnlyTail ? pDest->GetMipData(startLevel + 1) : pDest->GetMipData(0);
 	
-        vec2<int> sourceSize = pBase->CalculateMipSize(startLevel);
+		FIntVector2 sourceSize = pBase->CalculateMipSize(startLevel);
 
         switch (pBase->GetFormat())
         {
@@ -756,10 +755,10 @@ namespace mu
 			return;
 		}
 
-		const uint8_t* pSourceBuf = InBase->GetMipData(StartLevel);
-		uint8_t* pDestBuf = InBase->GetMipData(StartLevel + 1);
+		const uint8* pSourceBuf = InBase->GetMipData(StartLevel);
+		uint8* pDestBuf = InBase->GetMipData(StartLevel + 1);
 
-		vec2<int> sourceSize = InBase->CalculateMipSize(StartLevel);
+		FIntVector2 sourceSize = InBase->CalculateMipSize(StartLevel);
 
 		switch (InBase->GetFormat())
 		{
diff --git a/Engine/Plugins/Experimental/Mutable/Source/MutableRuntime/Private/MuR/OpImagePixelFormat.cpp b/Engine/Plugins/Experimental/Mutable/Source/MutableRuntime/Private/MuR/OpImagePixelFormat.cpp
--- a/Engine/Plugins/Experimental/Mutable/Source/MutableRuntime/Private/MuR/OpImagePixelFormat.cpp
+++ b/Engine/Plugins/Experimental/Mutable/Source/MutableRuntime/Private/MuR/OpImagePixelFormat.cpp
@@ -24,7 +24,7 @@ ImagePtr ImagePixelFormat( int imageCompressionQuality, const Image* pBase,
 {
 	MUTABLE_CPUPROFILER_SCOPE(ImagePixelFormat);
 
-    vec2<int> resultSize;
+	FIntVector2 resultSize;
     int resultLODCount = 0;
 
     if (onlyLOD==-1)
@@ -66,7 +66,7 @@ bool ImagePixelFormatInPlace( int imageCompressionQuality, Image* pResult, const
 
     bool success = true;
 
-    vec2<int> resultSize;
+	FIntVector2 resultSize;
     int resultLODCount = 0;
     const uint8_t* pBaseBuf = 0;
     int pixelCount = 0;
@@ -454,7 +454,7 @@ bool ImagePixelFormatInPlace( int imageCompressionQuality, Image* pResult, const
             {
                 for ( int m=0; m<resultLODCount; ++m )
                 {
-                    vec2<int> mipSize = pResult->CalculateMipSize( m );
+					FIntVector2 mipSize = pResult->CalculateMipSize( m );
                     miro::BC1_to_RGB( mipSize[0], mipSize[1],
                                         pBase->GetMipData(baseLOD+m),
                                         pResult->GetMipData(m) );
@@ -466,7 +466,7 @@ bool ImagePixelFormatInPlace( int imageCompressionQuality, Image* pResult, const
             {
                 for ( int m=0; m<resultLODCount; ++m )
                 {
-                    vec2<int> mipSize = pResult->CalculateMipSize( m );
+					FIntVector2 mipSize = pResult->CalculateMipSize( m );
                     miro::BC2_to_RGB( mipSize[0], mipSize[1],
                                         pBase->GetMipData(baseLOD+m),
                                         pResult->GetMipData(m) );
@@ -478,7 +478,7 @@ bool ImagePixelFormatInPlace( int imageCompressionQuality, Image* pResult, const
             {
                 for ( int m=0; m<resultLODCount; ++m )
                 {
-                    vec2<int> mipSize = pResult->CalculateMipSize( m );
+					FIntVector2 mipSize = pResult->CalculateMipSize( m );
                     miro::BC3_to_RGB( mipSize[0], mipSize[1],
                                         pBase->GetMipData(baseLOD+m),
                                         pResult->GetMipData(m) );
@@ -490,7 +490,7 @@ bool ImagePixelFormatInPlace( int imageCompressionQuality, Image* pResult, const
             {
                 for ( int m=0; m<resultLODCount; ++m )
                 {
-                    vec2<int> mipSize = pResult->CalculateMipSize( m );
+					FIntVector2 mipSize = pResult->CalculateMipSize( m );
                     miro::BC4_to_RGB( mipSize[0], mipSize[1],
                                       pBase->GetMipData(baseLOD+m),
                                       pResult->GetMipData(m) );
@@ -502,7 +502,7 @@ bool ImagePixelFormatInPlace( int imageCompressionQuality, Image* pResult, const
             {
                 for ( int m=0; m<resultLODCount; ++m )
                 {
-                    vec2<int> mipSize = pResult->CalculateMipSize( m );
+					FIntVector2 mipSize = pResult->CalculateMipSize( m );
                     miro::BC5_to_RGB( mipSize[0], mipSize[1],
                                       pBase->GetMipData(baseLOD+m),
                                       pResult->GetMipData(m) );
@@ -514,7 +514,7 @@ bool ImagePixelFormatInPlace( int imageCompressionQuality, Image* pResult, const
             {
                 for ( int m=0; m<resultLODCount; ++m )
                 {
-                    vec2<int> mipSize = pResult->CalculateMipSize( m );
+					FIntVector2 mipSize = pResult->CalculateMipSize( m );
                     miro::ASTC4x4RGBL_to_RGB( mipSize[0], mipSize[1],
                                       pBase->GetMipData(baseLOD+m),
                                       pResult->GetMipData(m) );
@@ -526,7 +526,7 @@ bool ImagePixelFormatInPlace( int imageCompressionQuality, Image* pResult, const
             {
                 for ( int m=0; m<resultLODCount; ++m )
                 {
-                    vec2<int> mipSize = pResult->CalculateMipSize( m );
+					FIntVector2 mipSize = pResult->CalculateMipSize( m );
                     miro::ASTC4x4RGBAL_to_RGB( mipSize[0], mipSize[1],
                                       pBase->GetMipData(baseLOD+m),
                                       pResult->GetMipData(m) );
@@ -538,7 +538,7 @@ bool ImagePixelFormatInPlace( int imageCompressionQuality, Image* pResult, const
             {
                 for ( int m=0; m<resultLODCount; ++m )
                 {
-                    vec2<int> mipSize = pResult->CalculateMipSize( m );
+					FIntVector2 mipSize = pResult->CalculateMipSize( m );
                     miro::ASTC4x4RGL_to_RGB( mipSize[0], mipSize[1],
                                       pBase->GetMipData(baseLOD+m),
                                       pResult->GetMipData(m) );
@@ -633,7 +633,7 @@ bool ImagePixelFormatInPlace( int imageCompressionQuality, Image* pResult, const
             {
                 for ( int m=0; m<resultLODCount; ++m )
                 {
-                    vec2<int> mipSize = pResult->CalculateMipSize( m );
+					FIntVector2 mipSize = pResult->CalculateMipSize( m );
                     miro::BC1_to_RGBA( mipSize[0], mipSize[1],
                                         pBase->GetMipData(baseLOD+m),
                                         pResult->GetMipData(m) );
@@ -645,7 +645,7 @@ bool ImagePixelFormatInPlace( int imageCompressionQuality, Image* pResult, const
             {
                 for ( int m=0; m<resultLODCount; ++m )
                 {
-                    vec2<int> mipSize = pResult->CalculateMipSize( m );
+					FIntVector2 mipSize = pResult->CalculateMipSize( m );
                     miro::BC2_to_RGBA( mipSize[0], mipSize[1],
                                         pBase->GetMipData(baseLOD+m),
                                         pResult->GetMipData(m) );
@@ -657,7 +657,7 @@ bool ImagePixelFormatInPlace( int imageCompressionQuality, Image* pResult, const
             {
                 for ( int m=0; m<resultLODCount; ++m )
                 {
-                    vec2<int> mipSize = pResult->CalculateMipSize( m );
+					FIntVector2 mipSize = pResult->CalculateMipSize( m );
                     miro::BC3_to_RGBA( mipSize[0], mipSize[1],
                                         pBase->GetMipData(baseLOD+m),
                                         pResult->GetMipData(m) );
@@ -669,7 +669,7 @@ bool ImagePixelFormatInPlace( int imageCompressionQuality, Image* pResult, const
             {
                 for ( int m=0; m<resultLODCount; ++m )
                 {
-                    vec2<int> mipSize = pResult->CalculateMipSize( m );
+					FIntVector2 mipSize = pResult->CalculateMipSize( m );
                     miro::BC4_to_RGBA( mipSize[0], mipSize[1],
                                        pBase->GetMipData(baseLOD+m),
                                        pResult->GetMipData(m) );
@@ -681,7 +681,7 @@ bool ImagePixelFormatInPlace( int imageCompressionQuality, Image* pResult, const
             {
                 for ( int m=0; m<resultLODCount; ++m )
                 {
-                    vec2<int> mipSize = pResult->CalculateMipSize( m );
+                    FIntVector2 mipSize = pResult->CalculateMipSize( m );
                     miro::BC5_to_RGBA( mipSize[0], mipSize[1],
                                         pBase->GetMipData(baseLOD+m),
                                         pResult->GetMipData(m) );
@@ -693,7 +693,7 @@ bool ImagePixelFormatInPlace( int imageCompressionQuality, Image* pResult, const
             {
                 for ( int m=0; m<resultLODCount; ++m )
                 {
-                    vec2<int> mipSize = pResult->CalculateMipSize( m );
+                    FIntVector2 mipSize = pResult->CalculateMipSize( m );
                     miro::ASTC4x4RGBL_to_RGBA( mipSize[0], mipSize[1],
                                         pBase->GetMipData(baseLOD+m),
                                         pResult->GetMipData(m) );
@@ -705,7 +705,7 @@ bool ImagePixelFormatInPlace( int imageCompressionQuality, Image* pResult, const
             {
                 for ( int m=0; m<resultLODCount; ++m )
                 {
-                    vec2<int> mipSize = pResult->CalculateMipSize( m );
+                    FIntVector2 mipSize = pResult->CalculateMipSize( m );
                     miro::ASTC4x4RGBAL_to_RGBA( mipSize[0], mipSize[1],
                                         pBase->GetMipData(baseLOD+m),
                                         pResult->GetMipData(m) );
@@ -717,7 +717,7 @@ bool ImagePixelFormatInPlace( int imageCompressionQuality, Image* pResult, const
             {
                 for ( int m=0; m<resultLODCount; ++m )
                 {
-                    vec2<int> mipSize = pResult->CalculateMipSize( m );
+                    FIntVector2 mipSize = pResult->CalculateMipSize( m );
                     miro::ASTC4x4RGL_to_RGBA( mipSize[0], mipSize[1],
                                         pBase->GetMipData(baseLOD+m),
                                         pResult->GetMipData(m) );
@@ -866,7 +866,7 @@ bool ImagePixelFormatInPlace( int imageCompressionQuality, Image* pResult, const
             case EImageFormat::IF_RGB_UBYTE:
                 for ( int m=0; m<resultLODCount; ++m )
                 {
-                    vec2<int> mipSize = pResult->CalculateMipSize( m );
+                    FIntVector2 mipSize = pResult->CalculateMipSize( m );
                     miro::RGB_to_BC1(
                         mipSize[0], mipSize[1], pBase->GetMipData( baseLOD + m ),
                         pResult->GetMipData( m ), imageCompressionQuality );
@@ -875,7 +875,7 @@ bool ImagePixelFormatInPlace( int imageCompressionQuality, Image* pResult, const
             case EImageFormat::IF_RGBA_UBYTE:
                 for ( int m=0; m<resultLODCount; ++m )
                 {
-                    vec2<int> mipSize = pResult->CalculateMipSize( m );
+                    FIntVector2 mipSize = pResult->CalculateMipSize( m );
                     miro::RGBA_to_BC1(
                         mipSize[0], mipSize[1], pBase->GetMipData( baseLOD + m ),
                         pResult->GetMipData( m ), imageCompressionQuality );
@@ -887,7 +887,7 @@ bool ImagePixelFormatInPlace( int imageCompressionQuality, Image* pResult, const
                 ImagePtr pTemp = ImagePixelFormat( imageCompressionQuality, pBase, EImageFormat::IF_RGB_UBYTE );
                 for ( int m=0; m<resultLODCount; ++m )
                 {
-                    vec2<int> mipSize = pResult->CalculateMipSize( m );
+                    FIntVector2 mipSize = pResult->CalculateMipSize( m );
                     miro::RGB_to_BC1(
                         mipSize[0], mipSize[1], pTemp->GetMipData( baseLOD + m ),
                         pResult->GetMipData( m ), imageCompressionQuality );
@@ -900,7 +900,7 @@ bool ImagePixelFormatInPlace( int imageCompressionQuality, Image* pResult, const
                 ImagePtr pTemp = ImagePixelFormat( imageCompressionQuality, pBase, EImageFormat::IF_RGBA_UBYTE );
                 for (int m = 0; m < resultLODCount; ++m)
                 {
-                    vec2<int> mipSize = pResult->CalculateMipSize(m);
+                    FIntVector2 mipSize = pResult->CalculateMipSize(m);
                     miro::RGBA_to_BC1(
                         mipSize[0], mipSize[1], pTemp->GetMipData( baseLOD + m ),
                         pResult->GetMipData( m ), imageCompressionQuality );
@@ -924,7 +924,7 @@ bool ImagePixelFormatInPlace( int imageCompressionQuality, Image* pResult, const
             case EImageFormat::IF_RGBA_UBYTE:
                 for ( int m=0; m<resultLODCount; ++m )
                 {
-                    vec2<int> mipSize = pResult->CalculateMipSize( m );
+                    FIntVector2 mipSize = pResult->CalculateMipSize( m );
                     miro::RGBA_to_BC2(
                         mipSize[0], mipSize[1], pBase->GetMipData( baseLOD + m ),
                         pResult->GetMipData( m ), imageCompressionQuality );
@@ -947,7 +947,7 @@ bool ImagePixelFormatInPlace( int imageCompressionQuality, Image* pResult, const
             case EImageFormat::IF_RGBA_UBYTE:
                 for ( int m=0; m<resultLODCount; ++m )
                 {
-                    vec2<int> mipSize = pResult->CalculateMipSize( m );
+                    FIntVector2 mipSize = pResult->CalculateMipSize( m );
                     miro::RGBA_to_BC3(
                         mipSize[0], mipSize[1], pBase->GetMipData( baseLOD + m ),
                         pResult->GetMipData( m ), imageCompressionQuality );
@@ -957,7 +957,7 @@ bool ImagePixelFormatInPlace( int imageCompressionQuality, Image* pResult, const
             case EImageFormat::IF_RGB_UBYTE:          
                 for ( int m=0; m<resultLODCount; ++m )
                 {
-                    vec2<int> mipSize = pResult->CalculateMipSize( m );
+                    FIntVector2 mipSize = pResult->CalculateMipSize( m );
                     miro::RGB_to_BC3(
                         mipSize[0], mipSize[1], pBase->GetMipData( baseLOD + m ),
                         pResult->GetMipData( m ), imageCompressionQuality );
@@ -970,7 +970,7 @@ bool ImagePixelFormatInPlace( int imageCompressionQuality, Image* pResult, const
 
                 for ( int m=0; m<resultLODCount; ++m )
                 {
-                    vec2<int> mipSize = pResult->CalculateMipSize( m );
+                    FIntVector2 mipSize = pResult->CalculateMipSize( m );
                     miro::BC1_to_BC3(
                         mipSize[0], mipSize[1], pBase->GetMipData( baseLOD + m ),
                         pResult->GetMipData( m ), imageCompressionQuality );
@@ -994,7 +994,7 @@ bool ImagePixelFormatInPlace( int imageCompressionQuality, Image* pResult, const
             case EImageFormat::IF_L_UBYTE:
                 for ( int m=0; m<resultLODCount; ++m )
                 {
-                    vec2<int> mipSize = pResult->CalculateMipSize( m );
+                    FIntVector2 mipSize = pResult->CalculateMipSize( m );
                     miro::L_to_BC4(
                         mipSize[0], mipSize[1], pBase->GetMipData( baseLOD + m ),
                         pResult->GetMipData( m ), imageCompressionQuality );
@@ -1017,7 +1017,7 @@ bool ImagePixelFormatInPlace( int imageCompressionQuality, Image* pResult, const
             case EImageFormat::IF_RGBA_UBYTE:
                 for ( int m=0; m<resultLODCount; ++m )
                 {
-                    vec2<int> mipSize = pResult->CalculateMipSize( m );
+                    FIntVector2 mipSize = pResult->CalculateMipSize( m );
                     miro::RGBA_to_BC5(
                         mipSize[0], mipSize[1], pBase->GetMipData( baseLOD + m ),
                         pResult->GetMipData( m ), imageCompressionQuality );
@@ -1027,7 +1027,7 @@ bool ImagePixelFormatInPlace( int imageCompressionQuality, Image* pResult, const
             case EImageFormat::IF_RGB_UBYTE:
                 for ( int m=0; m<resultLODCount; ++m )
                 {
-                    vec2<int> mipSize = pResult->CalculateMipSize( m );
+                    FIntVector2 mipSize = pResult->CalculateMipSize( m );
                     miro::RGB_to_BC5(
                         mipSize[0], mipSize[1], pBase->GetMipData( baseLOD + m ),
                         pResult->GetMipData( m ), imageCompressionQuality );
@@ -1050,7 +1050,7 @@ bool ImagePixelFormatInPlace( int imageCompressionQuality, Image* pResult, const
             case EImageFormat::IF_RGBA_UBYTE:
                 for ( int m=0; m<resultLODCount; ++m )
                 {
-                    vec2<int> mipSize = pResult->CalculateMipSize( m );
+                    FIntVector2 mipSize = pResult->CalculateMipSize( m );
                     miro::RGBA_to_ASTC4x4RGBL(
                         mipSize[0], mipSize[1], pBase->GetMipData( baseLOD + m ),
                         pResult->GetMipData( m ), imageCompressionQuality );
@@ -1060,7 +1060,7 @@ bool ImagePixelFormatInPlace( int imageCompressionQuality, Image* pResult, const
             case EImageFormat::IF_RGB_UBYTE:
                 for ( int m=0; m<resultLODCount; ++m )
                 {
-                    vec2<int> mipSize = pResult->CalculateMipSize( m );
+                    FIntVector2 mipSize = pResult->CalculateMipSize( m );
                     miro::RGB_to_ASTC4x4RGBL(
                         mipSize[0], mipSize[1], pBase->GetMipData( baseLOD + m ),
                         pResult->GetMipData( m ), imageCompressionQuality );
@@ -1070,7 +1070,7 @@ bool ImagePixelFormatInPlace( int imageCompressionQuality, Image* pResult, const
             case EImageFormat::IF_ASTC_4x4_RGBA_LDR:
                 for ( int m=0; m<resultLODCount; ++m )
                 {
-                    vec2<int> mipSize = pResult->CalculateMipSize( m );
+                    FIntVector2 mipSize = pResult->CalculateMipSize( m );
                     miro::ASTC4x4RGBAL_to_ASTC4x4RGBL( mipSize[0], mipSize[1],
                                        pBase->GetMipData(baseLOD+m),
                                        pResult->GetMipData(m) );
@@ -1080,7 +1080,7 @@ bool ImagePixelFormatInPlace( int imageCompressionQuality, Image* pResult, const
 			case EImageFormat::IF_L_UBYTE:
 				for (int m = 0; m < resultLODCount; ++m)
 				{
-					vec2<int> mipSize = pResult->CalculateMipSize(m);
+					FIntVector2 mipSize = pResult->CalculateMipSize(m);
 					miro::L_to_ASTC4x4RGBL(mipSize[0], mipSize[1],
 						pBase->GetMipData(baseLOD + m),
 						pResult->GetMipData(m), imageCompressionQuality);
@@ -1103,7 +1103,7 @@ bool ImagePixelFormatInPlace( int imageCompressionQuality, Image* pResult, const
             case EImageFormat::IF_RGBA_UBYTE:
                 for ( int m=0; m<resultLODCount; ++m )
                 {
-                    vec2<int> mipSize = pResult->CalculateMipSize( m );
+                    FIntVector2 mipSize = pResult->CalculateMipSize( m );
                     miro::RGBA_to_ASTC4x4RGBAL(
                         mipSize[0], mipSize[1], pBase->GetMipData( baseLOD + m ),
                         pResult->GetMipData( m ), imageCompressionQuality );
@@ -1113,7 +1113,7 @@ bool ImagePixelFormatInPlace( int imageCompressionQuality, Image* pResult, const
             case EImageFormat::IF_RGB_UBYTE:
                 for ( int m=0; m<resultLODCount; ++m )
                 {
-                    vec2<int> mipSize = pResult->CalculateMipSize( m );
+                    FIntVector2 mipSize = pResult->CalculateMipSize( m );
                     miro::RGB_to_ASTC4x4RGBAL(
                         mipSize[0], mipSize[1], pBase->GetMipData( baseLOD + m ),
                         pResult->GetMipData( m ), imageCompressionQuality );
@@ -1123,7 +1123,7 @@ bool ImagePixelFormatInPlace( int imageCompressionQuality, Image* pResult, const
             case EImageFormat::IF_ASTC_4x4_RGB_LDR:
                 for ( int m=0; m<resultLODCount; ++m )
                 {
-                    vec2<int> mipSize = pResult->CalculateMipSize( m );
+                    FIntVector2 mipSize = pResult->CalculateMipSize( m );
                     miro::ASTC4x4RGBL_to_ASTC4x4RGBAL( mipSize[0], mipSize[1],
                                        pBase->GetMipData(baseLOD+m),
                                        pResult->GetMipData(m) );
@@ -1146,7 +1146,7 @@ bool ImagePixelFormatInPlace( int imageCompressionQuality, Image* pResult, const
             case EImageFormat::IF_RGBA_UBYTE:
                 for ( int m=0; m<resultLODCount; ++m )
                 {
-                    vec2<int> mipSize = pResult->CalculateMipSize( m );
+                    FIntVector2 mipSize = pResult->CalculateMipSize( m );
                     miro::RGBA_to_ASTC4x4RGL(
                         mipSize[0], mipSize[1], pBase->GetMipData( baseLOD + m ),
                         pResult->GetMipData( m ), imageCompressionQuality );
@@ -1156,7 +1156,7 @@ bool ImagePixelFormatInPlace( int imageCompressionQuality, Image* pResult, const
             case EImageFormat::IF_RGB_UBYTE:
                 for ( int m=0; m<resultLODCount; ++m )
                 {
-                    vec2<int> mipSize = pResult->CalculateMipSize( m );
+                    FIntVector2 mipSize = pResult->CalculateMipSize( m );
                     miro::RGB_to_ASTC4x4RGL(
                         mipSize[0], mipSize[1], pBase->GetMipData( baseLOD + m ),
                         pResult->GetMipData( m ), imageCompressionQuality );
@@ -1166,7 +1166,7 @@ bool ImagePixelFormatInPlace( int imageCompressionQuality, Image* pResult, const
             case EImageFormat::IF_ASTC_4x4_RG_LDR:
                 for ( int m=0; m<resultLODCount; ++m )
                 {
-                    vec2<int> mipSize = pResult->CalculateMipSize( m );
+                    FIntVector2 mipSize = pResult->CalculateMipSize( m );
                     // Hack that actually works because of block size.
                     miro::ASTC4x4RGBAL_to_ASTC4x4RGBL( mipSize[0], mipSize[1],
                                        pBase->GetMipData(baseLOD+m),
diff --git a/Engine/Plugins/Experimental/Mutable/Source/MutableRuntime/Private/MuR/OpImageResize.h b/Engine/Plugins/Experimental/Mutable/Source/MutableRuntime/Private/MuR/OpImageResize.h
--- a/Engine/Plugins/Experimental/Mutable/Source/MutableRuntime/Private/MuR/OpImageResize.h
+++ b/Engine/Plugins/Experimental/Mutable/Source/MutableRuntime/Private/MuR/OpImageResize.h
@@ -623,7 +623,7 @@ namespace mu
 
         // Linear filtering
         //for ( int y=0; y<destSizeY; ++y )
-		const auto& ProcessRow = [
+		const auto ProcessRow = [
 			pDest, pBase, sizeX, destSizeY
 		] (uint32 y)
 		{
diff --git a/Engine/Plugins/Experimental/Mutable/Source/MutableRuntime/Private/MuR/Operations.cpp b/Engine/Plugins/Experimental/Mutable/Source/MutableRuntime/Private/MuR/Operations.cpp
--- a/Engine/Plugins/Experimental/Mutable/Source/MutableRuntime/Private/MuR/Operations.cpp
+++ b/Engine/Plugins/Experimental/Mutable/Source/MutableRuntime/Private/MuR/Operations.cpp
@@ -156,6 +156,7 @@ namespace mu
     //---------------------------------------------------------------------------------------------
     void ForEachReference( OP& op, const TFunctionRef<void(OP::ADDRESS*)> f )
     {
+		// Only operations that still use ASTOpFixed should be handled here.
         switch ( op.type )
         {
         case OP_TYPE::NONE:
@@ -356,10 +357,6 @@ namespace mu
             f(&op.args.ImageRasterMesh.projector );
             break;
 
-        case OP_TYPE::IM_MAKEGROWMAP:
-            f(&op.args.ImageMakeGrowMap.mask );
-            break;
-
         case OP_TYPE::IM_DISPLACE:
             f(&op.args.ImageDisplace.source );
             f(&op.args.ImageDisplace.displacementMap );
diff --git a/Engine/Plugins/Experimental/Mutable/Source/MutableRuntime/Private/MuR/Operations.h b/Engine/Plugins/Experimental/Mutable/Source/MutableRuntime/Private/MuR/Operations.h
--- a/Engine/Plugins/Experimental/Mutable/Source/MutableRuntime/Private/MuR/Operations.h
+++ b/Engine/Plugins/Experimental/Mutable/Source/MutableRuntime/Private/MuR/Operations.h
@@ -1094,7 +1094,6 @@ namespace mu
             ImageGPUArgs ImageGPU;
             ImageCropArgs ImageCrop;
             ImageRasterMeshArgs ImageRasterMesh;
-            ImageMakeGrowMapArgs ImageMakeGrowMap;
             ImageDisplaceArgs ImageDisplace;
 			ImageInvertArgs ImageInvert;
 
diff --git a/Engine/Plugins/Experimental/Mutable/Source/MutableRuntime/Public/MuR/Image.h b/Engine/Plugins/Experimental/Mutable/Source/MutableRuntime/Public/MuR/Image.h
--- a/Engine/Plugins/Experimental/Mutable/Source/MutableRuntime/Public/MuR/Image.h
+++ b/Engine/Plugins/Experimental/Mutable/Source/MutableRuntime/Public/MuR/Image.h
@@ -343,7 +343,7 @@ namespace mu
 
 		//! Calculate the size in pixels of a particular mipmap of this image. The size doesn't
 		//! include pixels necessary for completing blocks in block-compressed formats.
-		vec2<int> CalculateMipSize(int lod) const;
+		FIntVector2 CalculateMipSize(int lod) const;
 
 		//! Return a pointer to the beginning of the data for a particular mip.
 		uint8* GetMipData(int mip);
diff --git a/Engine/Plugins/Experimental/Mutable/Source/MutableTools/Private/MuT/AST.cpp b/Engine/Plugins/Experimental/Mutable/Source/MutableTools/Private/MuT/AST.cpp
--- a/Engine/Plugins/Experimental/Mutable/Source/MutableTools/Private/MuT/AST.cpp
+++ b/Engine/Plugins/Experimental/Mutable/Source/MutableTools/Private/MuT/AST.cpp
@@ -1234,31 +1234,12 @@ FImageDesc ASTOpFixed::GetImageDesc( bool returnBestOption, FGetImageDescContext
         res.m_format = EImageFormat::IF_L_UBYTE;
         break;
 
-    case OP_TYPE::IM_GPU:
-    {
-        check(false);
-        // The size comes from a special instruction
-//                const GPU_PROGRAM& gpuProg = program.m_gpuPrograms[ op.args.ImageGPU.program ];
-//                if (gpuProg.m_size)
-//                {
-//                    res = children[op.args.VolumeLayer.source]->GetImageDesc( returnBestOption, context );
-//                }
-//                // The format is fixed in the program description
-//                res.m_format = gpuProg.m_format;
-        break;
-    }
-
     case OP_TYPE::IM_RASTERMESH:
         res = GetImageDesc( op.args.ImageRasterMesh.image, returnBestOption, context );
         res.m_size[0]=op.args.ImageRasterMesh.sizeX;
         res.m_size[1]=op.args.ImageRasterMesh.sizeY;
         break;
 
-    case OP_TYPE::IM_MAKEGROWMAP:
-        res = GetImageDesc( op.args.ImageMakeGrowMap.mask, returnBestOption, context );
-        res.m_format = EImageFormat::IF_L_UBYTE;
-        break;
-
     case OP_TYPE::IM_DISPLACE:
         res = GetImageDesc( op.args.ImageDisplace.source, returnBestOption, context );
         break;
@@ -1786,13 +1767,6 @@ mu::Ptr<ImageSizeExpression> ASTOpFixed::GetImageSizeExpression() const
         pRes->size[1] = op.args.ImageGradient.size[1];
         break;
 
-    case OP_TYPE::IM_MAKEGROWMAP:
-        if ( children[op.args.ImageMakeGrowMap.mask] )
-        {
-            pRes = children[op.args.ImageMakeGrowMap.mask].child()->GetImageSizeExpression();
-        }
-        break;
-
     case OP_TYPE::IM_DISPLACE:
         if ( children[op.args.ImageDisplace.source] )
         {
diff --git a/Engine/Plugins/Experimental/Mutable/Source/MutableTools/Private/MuT/AST.h b/Engine/Plugins/Experimental/Mutable/Source/MutableTools/Private/MuT/AST.h
--- a/Engine/Plugins/Experimental/Mutable/Source/MutableTools/Private/MuT/AST.h
+++ b/Engine/Plugins/Experimental/Mutable/Source/MutableTools/Private/MuT/AST.h
@@ -195,6 +195,9 @@ namespace mu
         //! Store resource data in disk instead of memory
         bool m_useDiskCache = false;
 
+		/** Compile optimizing for the generation of smaller mipmaps of every image. */
+		bool bEnableProgressiveImages = false;
+
         // Additional advanced fine-tuning parameters
         //---------------------------------------------------------------------
 
@@ -273,7 +276,7 @@ namespace mu
 	public:
 
 		// \TODO This is recursive and may cause stack overflows in big models.
-		Ptr<ASTOp> Apply(const ASTOp* root);
+		Ptr<ASTOp> Apply(const ASTOp* Root);
 
 	protected:
 
@@ -292,7 +295,7 @@ namespace mu
 	{
 	public:
 
-		Ptr<ASTOp> Apply(const ASTOp* root);
+		Ptr<ASTOp> Apply(const class ASTOpImageMipmap* Root);
 
 	protected:
 
diff --git a/Engine/Plugins/Experimental/Mutable/Source/MutableTools/Private/MuT/ASTOpImageMakeGrowMap.cpp b/Engine/Plugins/Experimental/Mutable/Source/MutableTools/Private/MuT/ASTOpImageMakeGrowMap.cpp
new file mode 100644
--- /dev/null
+++ b/Engine/Plugins/Experimental/Mutable/Source/MutableTools/Private/MuT/ASTOpImageMakeGrowMap.cpp
@@ -0,0 +1,227 @@
+// Copyright Epic Games, Inc. All Rights Reserved.
+
+#include "MuT/ASTOpImageMakeGrowMap.h"
+
+#include "Containers/Map.h"
+#include "HAL/UnrealMemory.h"
+#include "Math/UnrealMathSSE.h"
+#include "Misc/AssertionMacros.h"
+#include "MuR/ImagePrivate.h"
+#include "MuR/MemoryPrivate.h"
+#include "MuR/ModelPrivate.h"
+#include "MuR/MutableMath.h"
+#include "MuR/RefCounted.h"
+#include "MuR/Types.h"
+#include "MuT/ASTOpConditional.h"
+#include "MuT/ASTOpImageCompose.h"
+#include "MuT/ASTOpImagePatch.h"
+#include "MuT/ASTOpImagePixelFormat.h"
+#include "MuT/ASTOpSwitch.h"
+#include "MuT/StreamsPrivate.h"
+
+#include <memory>
+#include <utility>
+
+
+namespace mu
+{
+
+	ASTOpImageMakeGrowMap::ASTOpImageMakeGrowMap()
+		: Mask(this)
+	{
+	}
+
+
+	ASTOpImageMakeGrowMap::~ASTOpImageMakeGrowMap()
+	{
+		// Explicit call needed to avoid recursive destruction
+		ASTOp::RemoveChildren();
+	}
+
+
+	bool ASTOpImageMakeGrowMap::IsEqual(const ASTOp& otherUntyped) const
+	{
+		if (const ASTOpImageMakeGrowMap* other = dynamic_cast<const ASTOpImageMakeGrowMap*>(&otherUntyped))
+		{
+			return Mask == other->Mask &&
+				Border == other->Border;
+		}
+		return false;
+	}
+
+
+	uint64 ASTOpImageMakeGrowMap::Hash() const
+	{
+		uint64 res = std::hash<void*>()(Mask.child().get());
+		hash_combine(res, Border);
+		return res;
+	}
+
+
+	mu::Ptr<ASTOp> ASTOpImageMakeGrowMap::Clone(MapChildFuncRef mapChild) const
+	{
+		mu::Ptr<ASTOpImageMakeGrowMap> n = new ASTOpImageMakeGrowMap();
+		n->Mask = mapChild(Mask.child());
+		n->Border = Border;
+		return n;
+	}
+
+
+	void ASTOpImageMakeGrowMap::ForEachChild(const TFunctionRef<void(ASTChild&)> f)
+	{
+		f(Mask);
+	}
+
+
+	void ASTOpImageMakeGrowMap::Link(FProgram& program, const FLinkerOptions*)
+	{
+		// Already linked?
+		if (!linkedAddress)
+		{
+			OP::ImageMakeGrowMapArgs args;
+			FMemory::Memzero(&args, sizeof(args));
+
+			args.border = Border;
+			if (Mask) args.mask = Mask->linkedAddress;
+
+			linkedAddress = (OP::ADDRESS)program.m_opAddress.Num();
+			//program.m_code.push_back(op);
+			program.m_opAddress.Add((uint32_t)program.m_byteCode.Num());
+			AppendCode(program.m_byteCode, GetOpType());
+			AppendCode(program.m_byteCode, args);
+		}
+	}
+
+
+	mu::Ptr<ASTOp> ASTOpImageMakeGrowMap::OptimiseSink(const FModelOptimizationOptions& options, FOptimizeSinkContext& context) const
+	{
+		mu::Ptr<ASTOp> at;
+
+		switch (Mask.child()->GetOpType())
+		{
+
+		case OP_TYPE::IM_CONDITIONAL:
+		{
+			// We move the format down the two paths
+			Ptr<ASTOpConditional> nop = mu::Clone<ASTOpConditional>(Mask.child());
+
+			Ptr<ASTOpImageMakeGrowMap> aOp = mu::Clone<ASTOpImageMakeGrowMap>(this);
+			aOp->Mask = nop->yes.child();
+			nop->yes = aOp;
+
+			Ptr<ASTOpImageMakeGrowMap> bOp = mu::Clone<ASTOpImageMakeGrowMap>(this);
+			bOp->Mask = nop->no.child();
+			nop->no = bOp;
+
+			at = nop;
+			break;
+		}
+
+		case OP_TYPE::IM_SWITCH:
+		{
+			// Move the format down all the paths
+			Ptr<ASTOpSwitch> nop = mu::Clone<ASTOpSwitch>(Mask.child());
+
+			if (nop->def)
+			{
+				Ptr<ASTOpImageMakeGrowMap> defOp = mu::Clone<ASTOpImageMakeGrowMap>(this);
+				defOp->Mask = nop->def.child();
+				nop->def = defOp;
+			}
+
+			// We need to copy the options because we change them
+			for (size_t v = 0; v < nop->cases.Num(); ++v)
+			{
+				if (nop->cases[v].branch)
+				{
+					Ptr<ASTOpImageMakeGrowMap> bOp = mu::Clone<ASTOpImageMakeGrowMap>(this);
+					bOp->Mask = nop->cases[v].branch.child();
+					nop->cases[v].branch = bOp;
+				}
+			}
+
+			at = nop;
+			break;
+		}
+
+		default:
+			break;
+		}
+
+
+		return at;
+	}
+
+
+	//!
+	FImageDesc ASTOpImageMakeGrowMap::GetImageDesc(bool returnBestOption, FGetImageDescContext* context) const
+	{
+		FImageDesc res;
+
+		// Local context in case it is necessary
+		FGetImageDescContext localContext;
+		if (!context)
+		{
+			context = &localContext;
+		}
+		else
+		{
+			// Cached result?
+			FImageDesc* PtrValue = context->m_results.Find(this);
+			if (PtrValue)
+			{
+				return *PtrValue;
+			}
+		}
+
+		if (Mask.child())
+		{
+			res = Mask.child()->GetImageDesc(returnBestOption, context);
+		}
+
+		// Cache the result
+		context->m_results.Add(this, res);
+
+		return res;
+	}
+
+
+	void ASTOpImageMakeGrowMap::GetLayoutBlockSize(int* pBlockX, int* pBlockY)
+	{
+		if (Mask.child())
+		{
+			// Assume the block size of the biggest mip
+			Mask.child()->GetLayoutBlockSize(pBlockX, pBlockY);
+		}
+	}
+
+
+	bool ASTOpImageMakeGrowMap::IsImagePlainConstant(FVector4f& colour) const
+	{
+		bool res = false;
+		if (Mask.child())
+		{
+			Mask.child()->IsImagePlainConstant(colour);
+		}
+		return res;
+	}
+
+
+	mu::Ptr<ImageSizeExpression> ASTOpImageMakeGrowMap::GetImageSizeExpression() const
+	{
+		mu::Ptr<ImageSizeExpression> pRes;
+
+		if (Mask.child())
+		{
+			pRes = Mask.child()->GetImageSizeExpression();
+		}
+		else
+		{
+			pRes = new ImageSizeExpression;
+		}
+
+		return pRes;
+	}
+
+}
+
diff --git a/Engine/Plugins/Experimental/Mutable/Source/MutableTools/Private/MuT/ASTOpImageMakeGrowMap.h b/Engine/Plugins/Experimental/Mutable/Source/MutableTools/Private/MuT/ASTOpImageMakeGrowMap.h
new file mode 100644
--- /dev/null
+++ b/Engine/Plugins/Experimental/Mutable/Source/MutableTools/Private/MuT/ASTOpImageMakeGrowMap.h
@@ -0,0 +1,45 @@
+// Copyright Epic Games, Inc. All Rights Reserved.
+
+#pragma once
+
+#include "HAL/PlatformMath.h"
+#include "MuR/Image.h"
+#include "MuR/Operations.h"
+#include "MuR/Ptr.h"
+#include "MuT/AST.h"
+
+
+namespace mu
+{
+	struct FProgram;
+
+	class ASTOpImageMakeGrowMap : public ASTOp
+	{
+	public:
+
+		ASTChild Mask;
+
+		uint32 Border = 0;
+
+	public:
+
+		ASTOpImageMakeGrowMap();
+		ASTOpImageMakeGrowMap(const ASTOpImageMakeGrowMap&) = delete;
+		~ASTOpImageMakeGrowMap();
+
+		OP_TYPE GetOpType() const override { return OP_TYPE::IM_MAKEGROWMAP; }
+		uint64 Hash() const override;
+		bool IsEqual(const ASTOp& otherUntyped) const override;
+		Ptr<ASTOp> Clone(MapChildFuncRef mapChild) const override;
+		void ForEachChild(const TFunctionRef<void(ASTChild&)>) override;
+		void Link(FProgram& program, const FLinkerOptions* Options) override;
+		Ptr<ASTOp> OptimiseSink(const FModelOptimizationOptions& options, FOptimizeSinkContext& context) const override;
+		FImageDesc GetImageDesc(bool returnBestOption, FGetImageDescContext* context) const override;
+		void GetLayoutBlockSize(int* pBlockX, int* pBlockY) override;
+		bool IsImagePlainConstant(FVector4f& colour) const override;
+		Ptr<ImageSizeExpression> GetImageSizeExpression() const override;
+
+	};
+
+}
+
diff --git a/Engine/Plugins/Experimental/Mutable/Source/MutableTools/Private/MuT/ASTOpImageMipmap.cpp b/Engine/Plugins/Experimental/Mutable/Source/MutableTools/Private/MuT/ASTOpImageMipmap.cpp
--- a/Engine/Plugins/Experimental/Mutable/Source/MutableTools/Private/MuT/ASTOpImageMipmap.cpp
+++ b/Engine/Plugins/Experimental/Mutable/Source/MutableTools/Private/MuT/ASTOpImageMipmap.cpp
@@ -26,426 +26,425 @@
 namespace mu
 {
 
-ASTOpImageMipmap::ASTOpImageMipmap()
-    : Source(this)
-{
-}
-
-
-ASTOpImageMipmap::~ASTOpImageMipmap()
-{
-    // Explicit call needed to avoid recursive destruction
-    ASTOp::RemoveChildren();
-}
-
-
-bool ASTOpImageMipmap::IsEqual(const ASTOp& otherUntyped) const
-{
-    if (const ASTOpImageMipmap* other = dynamic_cast<const ASTOpImageMipmap*>(&otherUntyped) )
-    {
-        return Source==other->Source && 
-			Levels == other->Levels &&
-			BlockLevels == other->BlockLevels &&
-			bOnlyTail == other->bOnlyTail &&
-			SharpenFactor == other->SharpenFactor &&
-			AddressMode == other->AddressMode &&
-			FilterType == other->FilterType &&
-			DitherMipmapAlpha == other->DitherMipmapAlpha;
+	ASTOpImageMipmap::ASTOpImageMipmap()
+		: Source(this)
+	{
 	}
-    return false;
-}
 
 
-uint64 ASTOpImageMipmap::Hash() const
-{
-	uint64 res = std::hash<void*>()(Source.child().get() );
-    hash_combine( res, Levels);
-    return res;
-}
+	ASTOpImageMipmap::~ASTOpImageMipmap()
+	{
+		// Explicit call needed to avoid recursive destruction
+		ASTOp::RemoveChildren();
+	}
 
 
-mu::Ptr<ASTOp> ASTOpImageMipmap::Clone(MapChildFuncRef mapChild) const
-{
-	mu::Ptr<ASTOpImageMipmap> n = new ASTOpImageMipmap();
-    n->Source = mapChild(Source.child());
-	n->Levels = Levels;
-	n->BlockLevels = BlockLevels;
-	n->bOnlyTail = bOnlyTail;
-	n->SharpenFactor = SharpenFactor;
-	n->AddressMode = AddressMode;
-	n->FilterType = FilterType;
-	n->DitherMipmapAlpha = DitherMipmapAlpha;
+	bool ASTOpImageMipmap::IsEqual(const ASTOp& otherUntyped) const
+	{
+		if (const ASTOpImageMipmap* other = dynamic_cast<const ASTOpImageMipmap*>(&otherUntyped))
+		{
+			return Source == other->Source &&
+				Levels == other->Levels &&
+				BlockLevels == other->BlockLevels &&
+				bOnlyTail == other->bOnlyTail &&
+				bPreventSplitTail == other->bPreventSplitTail &&
+				SharpenFactor == other->SharpenFactor &&
+				AddressMode == other->AddressMode &&
+				FilterType == other->FilterType &&
+				DitherMipmapAlpha == other->DitherMipmapAlpha;
+		}
+		return false;
+	}
 
-    return n;
-}
 
+	uint64 ASTOpImageMipmap::Hash() const
+	{
+		uint64 res = std::hash<void*>()(Source.child().get());
+		hash_combine(res, Levels);
+		return res;
+	}
 
-void ASTOpImageMipmap::ForEachChild(const TFunctionRef<void(ASTChild&)> f )
-{
-    f( Source );
-}
 
+	mu::Ptr<ASTOp> ASTOpImageMipmap::Clone(MapChildFuncRef mapChild) const
+	{
+		mu::Ptr<ASTOpImageMipmap> n = new ASTOpImageMipmap();
+		n->Source = mapChild(Source.child());
+		n->Levels = Levels;
+		n->BlockLevels = BlockLevels;
+		n->bOnlyTail = bOnlyTail;
+		n->bPreventSplitTail = bPreventSplitTail;
+		n->SharpenFactor = SharpenFactor;
+		n->AddressMode = AddressMode;
+		n->FilterType = FilterType;
+		n->DitherMipmapAlpha = DitherMipmapAlpha;
+
+		return n;
+	}
 
-void ASTOpImageMipmap::Link( FProgram& program, const FLinkerOptions* )
-{
-    // Already linked?
-    if (!linkedAddress)
-    {
-        OP::ImageMipmapArgs args;
-        FMemory::Memzero( &args, sizeof(args) );
-
-		args.levels = Levels;
-		args.blockLevels = BlockLevels;
-		args.onlyTail = bOnlyTail;
-		args.sharpenFactor = SharpenFactor;
-		args.addressMode = AddressMode;
-		args.filterType = FilterType;
-		args.ditherMipmapAlpha = DitherMipmapAlpha;
-		if (Source) args.source = Source->linkedAddress;
-
-        linkedAddress = (OP::ADDRESS)program.m_opAddress.Num();
-        //program.m_code.push_back(op);
-        program.m_opAddress.Add((uint32_t)program.m_byteCode.Num());
-        AppendCode(program.m_byteCode,OP_TYPE::IM_MIPMAP);
-        AppendCode(program.m_byteCode,args);
-    }
-
-}
-
-
-mu::Ptr<ASTOp> ASTOpImageMipmap::OptimiseSink(const FModelOptimizationOptions& options, FOptimizeSinkContext& context) const
-{
-	mu::Ptr<ASTOp> at;
 
-	mu::Ptr<ASTOp> sourceAt = Source.child();
-	switch (sourceAt->GetOpType())
+	void ASTOpImageMipmap::ForEachChild(const TFunctionRef<void(ASTChild&)> f)
 	{
+		f(Source);
+	}
+
 
-	case OP_TYPE::IM_BLANKLAYOUT:
+	void ASTOpImageMipmap::Link(FProgram& program, const FLinkerOptions*)
 	{
-		// Set the mipmap generation on the blank layout operation
-		auto mop = mu::Clone<ASTOpFixed>(sourceAt);
-		mop->op.args.ImageBlankLayout.generateMipmaps = 1;    // true
-		mop->op.args.ImageBlankLayout.mipmapCount = Levels;
-		at = mop;
-		break;
+		// Already linked?
+		if (!linkedAddress)
+		{
+			OP::ImageMipmapArgs args;
+			FMemory::Memzero(&args, sizeof(args));
+
+			args.levels = Levels;
+			args.blockLevels = BlockLevels;
+			args.onlyTail = bOnlyTail;
+			args.sharpenFactor = SharpenFactor;
+			args.addressMode = AddressMode;
+			args.filterType = FilterType;
+			args.ditherMipmapAlpha = DitherMipmapAlpha;
+			if (Source) args.source = Source->linkedAddress;
+
+			linkedAddress = (OP::ADDRESS)program.m_opAddress.Num();
+			//program.m_code.push_back(op);
+			program.m_opAddress.Add((uint32_t)program.m_byteCode.Num());
+			AppendCode(program.m_byteCode, GetOpType());
+			AppendCode(program.m_byteCode, args);
+		}
 	}
 
-	case OP_TYPE::IM_PIXELFORMAT:
+
+	mu::Ptr<ASTOp> ASTOpImageMipmap::OptimiseSink(const FModelOptimizationOptions& options, FOptimizeSinkContext& context) const
 	{
-		// Swap unless the mipmap operation builds only the tail or is compressed.
-		// Otherwise, we could fall in a loop of swapping mipmaps and pixelformats.
-		mu::Ptr<ASTOpImageMipmap> mop = mu::Clone<ASTOpImageMipmap>(this);
-		mu::Ptr<ASTOpImagePixelFormat> fop = mu::Clone<ASTOpImagePixelFormat>(sourceAt);
-		bool isCompressedFormat = IsCompressedFormat(fop->Format);
-		if (isCompressedFormat && !mop->bOnlyTail)
+		mu::Ptr<ASTOp> at;
+
+		mu::Ptr<ASTOp> sourceAt = Source.child();
+		switch (sourceAt->GetOpType())
 		{
-			mop->Source = fop->Source.child();
-			fop->Source = mop;
 
-			at = fop;
+		case OP_TYPE::IM_BLANKLAYOUT:
+		{
+			// Set the mipmap generation on the blank layout operation
+			auto mop = mu::Clone<ASTOpFixed>(sourceAt);
+			mop->op.args.ImageBlankLayout.generateMipmaps = 1;    // true
+			mop->op.args.ImageBlankLayout.mipmapCount = Levels;
+			at = mop;
+			break;
 		}
-		break;
-	}
 
-	default:
-	{
-		at = context.ImageMipmapSinker.Apply(this);
+		case OP_TYPE::IM_PIXELFORMAT:
+		{
+			// Swap unless the mipmap operation builds only the tail or is compressed.
+			// Otherwise, we could fall in a loop of swapping mipmaps and pixelformats.
+			mu::Ptr<ASTOpImageMipmap> mop = mu::Clone<ASTOpImageMipmap>(this);
+			mu::Ptr<ASTOpImagePixelFormat> fop = mu::Clone<ASTOpImagePixelFormat>(sourceAt);
+			bool isCompressedFormat = IsCompressedFormat(fop->Format);
+			if (isCompressedFormat && !mop->bOnlyTail)
+			{
+				mop->Source = fop->Source.child();
+				fop->Source = mop;
 
-		break;
-	} // mipmap source default
+				at = fop;
+			}
+			break;
+		}
 
-	} // mipmap source type switch
+		default:
+		{
+			at = context.ImageMipmapSinker.Apply(this);
 
-	
-	return at;
-}
+			break;
+		} // mipmap source default
 
+		} // mipmap source type switch
 
-//!
-FImageDesc ASTOpImageMipmap::GetImageDesc(bool returnBestOption, FGetImageDescContext* context) const
-{
-	FImageDesc res;
 
-	// Local context in case it is necessary
-	FGetImageDescContext localContext;
-	if (!context)
-	{
-		context = &localContext;
-	}
-	else
-	{
-		// Cached result?
-		FImageDesc* PtrValue = context->m_results.Find(this);
-		if (PtrValue)
-		{
-			return *PtrValue;
-		}
+		return at;
 	}
 
-	if (Source.child())
+
+	//!
+	FImageDesc ASTOpImageMipmap::GetImageDesc(bool returnBestOption, FGetImageDescContext* context) const
 	{
-		res = Source.child()->GetImageDesc( returnBestOption, context);
-	}
+		FImageDesc res;
 
-	int mipLevels = Levels;
+		// Local context in case it is necessary
+		FGetImageDescContext localContext;
+		if (!context)
+		{
+			context = &localContext;
+		}
+		else
+		{
+			// Cached result?
+			FImageDesc* PtrValue = context->m_results.Find(this);
+			if (PtrValue)
+			{
+				return *PtrValue;
+			}
+		}
 
-	if (mipLevels == 0)
-	{
-		mipLevels = FMath::Max(
-			(int)ceilf(logf((float)res.m_size[0]) / logf(2.0f)),
-			(int)ceilf(logf((float)res.m_size[1]) / logf(2.0f))
-		);
-	}
+		if (Source.child())
+		{
+			res = Source.child()->GetImageDesc(returnBestOption, context);
+		}
 
-	res.m_lods = FMath::Max(res.m_lods, (uint8)mipLevels);
+		int mipLevels = Levels;
 
+		if (mipLevels == 0)
+		{
+			mipLevels = FMath::Max(
+				(int)ceilf(logf((float)res.m_size[0]) / logf(2.0f)),
+				(int)ceilf(logf((float)res.m_size[1]) / logf(2.0f))
+			);
+		}
 
-	// Cache the result
-	if (context)
-	{
-		context->m_results.Add(this, res);
-	}
+		res.m_lods = FMath::Max(res.m_lods, (uint8)mipLevels);
 
-	return res;
-}
 
+		// Cache the result
+		if (context)
+		{
+			context->m_results.Add(this, res);
+		}
 
-void ASTOpImageMipmap::GetLayoutBlockSize(int* pBlockX, int* pBlockY)
-{
-	if (Source.child())
-	{
-		// Assume the block size of the biggest mip
-		Source.child()->GetLayoutBlockSize(pBlockX, pBlockY);
+		return res;
 	}
-}
 
 
-bool ASTOpImageMipmap::IsImagePlainConstant(FVector4f& colour) const
-{
-	bool res = false;
-	if (Source.child())
+	void ASTOpImageMipmap::GetLayoutBlockSize(int* pBlockX, int* pBlockY)
 	{
-		Source.child()->IsImagePlainConstant(colour);
+		if (Source.child())
+		{
+			// Assume the block size of the biggest mip
+			Source.child()->GetLayoutBlockSize(pBlockX, pBlockY);
+		}
 	}
-	return res;
-}
-
 
-mu::Ptr<ImageSizeExpression> ASTOpImageMipmap::GetImageSizeExpression() const
-{
-	mu::Ptr<ImageSizeExpression> pRes;
 
-	if (Source.child())
-	{
-		pRes = Source.child()->GetImageSizeExpression();
-	}
-	else
+	bool ASTOpImageMipmap::IsImagePlainConstant(FVector4f& colour) const
 	{
-		pRes = new ImageSizeExpression;
+		bool res = false;
+		if (Source.child())
+		{
+			Source.child()->IsImagePlainConstant(colour);
+		}
+		return res;
 	}
 
-	return pRes;
-}
-
 
-//---------------------------------------------------------------------------------------------
-//---------------------------------------------------------------------------------------------
-//---------------------------------------------------------------------------------------------
-mu::Ptr<ASTOp> Sink_ImageMipmapAST::Apply(const ASTOp* root)
-{
-	check(root->GetOpType() == OP_TYPE::IM_MIPMAP);
+	mu::Ptr<ImageSizeExpression> ASTOpImageMipmap::GetImageSizeExpression() const
+	{
+		mu::Ptr<ImageSizeExpression> pRes;
 
-	m_root = dynamic_cast<const ASTOpImageMipmap*>(root);
-	OldToNew.Empty();
+		if (Source.child())
+		{
+			pRes = Source.child()->GetImageSizeExpression();
+		}
+		else
+		{
+			pRes = new ImageSizeExpression;
+		}
 
-	if (m_root->bOnlyTail)
-	{
-		return nullptr;
+		return pRes;
 	}
 
-	mu::Ptr<ASTOp> newSource;
-	m_initialSource = m_root->Source.child();
-
-	// Before sinking, see if it is worth splitting into miptail and mip.
-	if (!m_root->bOnlyTail
-		// \TODO: Review this: It seems it fails for a BC1 8x4 texture having blockSize=4, and level=0
-			//&&
-			//m_root->Levels
-			//!=
-			//m_root->BlockLevels 
-		)
+
+	//---------------------------------------------------------------------------------------------
+	//---------------------------------------------------------------------------------------------
+	//---------------------------------------------------------------------------------------------
+	mu::Ptr<ASTOp> Sink_ImageMipmapAST::Apply(const ASTOpImageMipmap* InRoot)
 	{
-		// the block mipmaps can be done before composition is done.
-		mu::Ptr<ASTOpImageMipmap> newMip = mu::Clone<ASTOpImageMipmap>(m_root);
-		newMip->Levels = m_root->BlockLevels;
-		newMip->BlockLevels = m_root->BlockLevels;
-		newMip->bOnlyTail = false;
+		m_root = InRoot;
+		OldToNew.Empty();
 
-		// the smallest mipmaps after the composition is done.
-		mu::Ptr<ASTOpImageMipmap> topMipOp = mu::Clone<ASTOpImageMipmap>(m_root);
-		topMipOp->bOnlyTail = true;
+		if (m_root->bOnlyTail)
+		{
+			return nullptr;
+		}
 
-		// Proceed
-		topMipOp->Source = Visit(m_initialSource, newMip.get());
+		mu::Ptr<ASTOp> newSource;
+		m_initialSource = m_root->Source.child();
 
-		newSource = topMipOp;
-	}
-	else
-	{
-		// Proceed
-		newSource = Visit(m_initialSource, m_root);
-	}
+		// Before sinking, see if it is worth splitting into miptail and mip.
+		if (!m_root->bPreventSplitTail
+			// \TODO: Review this: It seems it fails for a BC1 8x4 texture having blockSize=4, and level=0
+				//&&
+				//m_root->Levels
+				//!=
+				//m_root->BlockLevels 
+			)
+		{
+			// the block mipmaps can be done before composition is done.
+			mu::Ptr<ASTOpImageMipmap> newMip = mu::Clone<ASTOpImageMipmap>(m_root);
+			newMip->Levels = m_root->BlockLevels;
+			newMip->BlockLevels = m_root->BlockLevels;
+			newMip->bOnlyTail = false;
 
-	m_root = nullptr;
+			// the smallest mipmaps after the composition is done.
+			mu::Ptr<ASTOpImageMipmap> topMipOp = mu::Clone<ASTOpImageMipmap>(m_root);
+			topMipOp->bOnlyTail = true;
 
-	// If there is any change, it is the new root.
-	if (newSource != m_initialSource)
-	{
-		return newSource;
-	}
+			// Proceed
+			topMipOp->Source = Visit(m_initialSource, newMip.get());
 
-	return nullptr;
-}
+			newSource = topMipOp;
+		}
+		else
+		{
+			// Proceed
+			newSource = Visit(m_initialSource, m_root);
+		}
 
+		m_root = nullptr;
 
-//---------------------------------------------------------------------------------------------
-//---------------------------------------------------------------------------------------------
-//---------------------------------------------------------------------------------------------
-mu::Ptr<ASTOp> Sink_ImageMipmapAST::Visit(mu::Ptr<ASTOp> at, const ASTOpImageMipmap* currentMipmapOp)
-{
-	if (!at) return nullptr;
+		// If there is any change, it is the new root.
+		if (newSource != m_initialSource)
+		{
+			return newSource;
+		}
 
-	// Already visited?
-	const Ptr<ASTOp>* Cached = OldToNew.Find({ at,currentMipmapOp });
-	if (Cached)
-	{
-		return *Cached;
+		return nullptr;
 	}
 
-	mu::Ptr<ASTOp> newAt = at;
-	switch (at->GetOpType())
-	{
 
-	case OP_TYPE::IM_CONDITIONAL:
+	//---------------------------------------------------------------------------------------------
+	//---------------------------------------------------------------------------------------------
+	//---------------------------------------------------------------------------------------------
+	mu::Ptr<ASTOp> Sink_ImageMipmapAST::Visit(mu::Ptr<ASTOp> at, const ASTOpImageMipmap* currentMipmapOp)
 	{
-		// We move the mask creation down the two paths
-		auto newOp = mu::Clone<ASTOpConditional>(at);
-		newOp->yes = Visit(newOp->yes.child(), currentMipmapOp);
-		newOp->no = Visit(newOp->no.child(), currentMipmapOp);
-		newAt = newOp;
-		break;
-	}
+		if (!at) return nullptr;
 
-	case OP_TYPE::IM_SWITCH:
-	{
-		// We move the mask creation down all the paths
-		auto newOp = mu::Clone<ASTOpSwitch>(at);
-		newOp->def = Visit(newOp->def.child(), currentMipmapOp);
-		for (auto& c : newOp->cases)
+		// Already visited?
+		const Ptr<ASTOp>* Cached = OldToNew.Find({ at,currentMipmapOp });
+		if (Cached)
 		{
-			c.branch = Visit(c.branch.child(), currentMipmapOp);
+			return *Cached;
 		}
-		newAt = newOp;
-		break;
-	}
 
-	case OP_TYPE::IM_COMPOSE:
-	{
-		const ASTOpImageCompose* typedAt = dynamic_cast<const ASTOpImageCompose*>(at.get());
-		if (!currentMipmapOp->bOnlyTail
-			&&
-			// Don't move the mipmapping if we are composing with a mask.
-			// TODO: allow mipmapping in the masks, RLE formats, etc.
-			!typedAt->Mask
-			)
+		mu::Ptr<ASTOp> newAt = at;
+		switch (at->GetOpType())
 		{
-			auto newOp = mu::Clone<ASTOpImageCompose>(at);
-
-			auto baseOp = newOp->Base.child();
-			newOp->Base = Visit(baseOp, currentMipmapOp);
-
-			auto blockOp = newOp->BlockImage.child();
-			newOp->BlockImage = Visit(blockOp, currentMipmapOp);
 
+		case OP_TYPE::IM_CONDITIONAL:
+		{
+			// We move the op down the two paths
+			auto newOp = mu::Clone<ASTOpConditional>(at);
+			newOp->yes = Visit(newOp->yes.child(), currentMipmapOp);
+			newOp->no = Visit(newOp->no.child(), currentMipmapOp);
 			newAt = newOp;
+			break;
 		}
 
-		break;
-	}
+		case OP_TYPE::IM_SWITCH:
+		{
+			// We move the op down all the paths
+			auto newOp = mu::Clone<ASTOpSwitch>(at);
+			newOp->def = Visit(newOp->def.child(), currentMipmapOp);
+			for (auto& c : newOp->cases)
+			{
+				c.branch = Visit(c.branch.child(), currentMipmapOp);
+			}
+			newAt = newOp;
+			break;
+		}
 
-	case OP_TYPE::IM_PATCH:
-	{
-		if (!currentMipmapOp->bOnlyTail)
+		case OP_TYPE::IM_COMPOSE:
 		{
-			// Special case: we propagate the mipmapping down the patch up to
-			// the level allowed by the patch size and placement. We then leave
-			// a top-level  mipmapping operation to generate the smallest
-			// mipmaps after the patch is done.
+			const ASTOpImageCompose* typedAt = dynamic_cast<const ASTOpImageCompose*>(at.get());
+			if (!currentMipmapOp->bOnlyTail
+				&&
+				// Don't move the mipmapping if we are composing with a mask.
+				// TODO: allow mipmapping in the masks, RLE formats, etc.
+				!typedAt->Mask
+				)
+			{
+				auto newOp = mu::Clone<ASTOpImageCompose>(at);
+
+				auto baseOp = newOp->Base.child();
+				newOp->Base = Visit(baseOp, currentMipmapOp);
 
-			const ASTOpImagePatch* typedSource = dynamic_cast<const ASTOpImagePatch*>(at.get());
-			auto rectOp = typedSource->patch.child();
-			ASTOp::FGetImageDescContext context;
-			auto patchDesc = rectOp->GetImageDesc(false, &context);
+				auto blockOp = newOp->BlockImage.child();
+				newOp->BlockImage = Visit(blockOp, currentMipmapOp);
 
-			uint8_t blockLevels = 0;
+				newAt = newOp;
+			}
+
+			break;
+		}
+
+		case OP_TYPE::IM_PATCH:
+		{
+			if (!currentMipmapOp->bOnlyTail)
 			{
-				uint16 minX = typedSource->location[0];
-				uint16 minY = typedSource->location[1];
-				uint16 sizeX = patchDesc.m_size[0];
-				uint16 sizeY = patchDesc.m_size[1];
-				while (minX && minY && sizeX && sizeY
-					&&
-					minX % 2 == 0 && minY % 2 == 0 && sizeX % 2 == 0 && sizeY % 2 == 0)
+				// Special case: we propagate the mipmapping down the patch up to
+				// the level allowed by the patch size and placement. We then leave
+				// a top-level  mipmapping operation to generate the smallest
+				// mipmaps after the patch is done.
+
+				const ASTOpImagePatch* typedSource = dynamic_cast<const ASTOpImagePatch*>(at.get());
+				auto rectOp = typedSource->patch.child();
+				ASTOp::FGetImageDescContext context;
+				auto patchDesc = rectOp->GetImageDesc(false, &context);
+
+				uint8_t blockLevels = 0;
 				{
-					blockLevels++;
-					minX /= 2;
-					minY /= 2;
-					sizeX /= 2;
-					sizeY /= 2;
+					uint16 minX = typedSource->location[0];
+					uint16 minY = typedSource->location[1];
+					uint16 sizeX = patchDesc.m_size[0];
+					uint16 sizeY = patchDesc.m_size[1];
+					while (minX && minY && sizeX && sizeY
+						&&
+						minX % 2 == 0 && minY % 2 == 0 && sizeX % 2 == 0 && sizeY % 2 == 0)
+					{
+						blockLevels++;
+						minX /= 2;
+						minY /= 2;
+						sizeX /= 2;
+						sizeY /= 2;
+					}
 				}
-			}
 
-			mu::Ptr<ASTOpImageMipmap> newMip = mu::Clone<ASTOpImageMipmap>(currentMipmapOp);
-			newMip->Levels = blockLevels;
-			newMip->BlockLevels = blockLevels;
-			newMip->bOnlyTail = false;
+				mu::Ptr<ASTOpImageMipmap> newMip = mu::Clone<ASTOpImageMipmap>(currentMipmapOp);
+				newMip->Levels = blockLevels;
+				newMip->BlockLevels = blockLevels;
+				newMip->bOnlyTail = false;
 
-			auto newOp = mu::Clone<ASTOpImagePatch>(at);
-			newOp->base = Visit(newOp->base.child(), newMip.get());
-			newOp->patch = Visit(newOp->patch.child(), newMip.get());
-			newAt = newOp;
+				auto newOp = mu::Clone<ASTOpImagePatch>(at);
+				newOp->base = Visit(newOp->base.child(), newMip.get());
+				newOp->patch = Visit(newOp->patch.child(), newMip.get());
+				newAt = newOp;
 
-			// We need to add a mipmap on top to finish the mipmapping
-			{
-				mu::Ptr<ASTOpImageMipmap> topMipOp = mu::Clone<ASTOpImageMipmap>(currentMipmapOp);
-				topMipOp->Source = newOp;
-				topMipOp->bOnlyTail = true;
-				newAt = topMipOp;
+				// We need to add a mipmap on top to finish the mipmapping
+				{
+					mu::Ptr<ASTOpImageMipmap> topMipOp = mu::Clone<ASTOpImageMipmap>(currentMipmapOp);
+					topMipOp->Source = newOp;
+					topMipOp->bOnlyTail = true;
+					newAt = topMipOp;
+				}
 			}
-		}
 
-		break;
-	}
+			break;
+		}
 
-	default:
-		break;
-	}
+		default:
+			break;
+		}
 
-	// end on line, replace with mipmap
-	if (at == newAt && at != m_initialSource)
-	{
-		mu::Ptr<ASTOpImageMipmap> newOp = mu::Clone<ASTOpImageMipmap>(currentMipmapOp);
-		check(newOp->GetOpType() == OP_TYPE::IM_MIPMAP);
+		// end on line, replace with mipmap
+		if (at == newAt && at != m_initialSource)
+		{
+			mu::Ptr<ASTOpImageMipmap> newOp = mu::Clone<ASTOpImageMipmap>(currentMipmapOp);
+			check(newOp->GetOpType() == OP_TYPE::IM_MIPMAP);
 
-		newOp->Source = at;
+			newOp->Source = at;
 
-		newAt = newOp;
-	}
+			newAt = newOp;
+		}
 
-	OldToNew.Add({ at,currentMipmapOp }, newAt);
+		OldToNew.Add({ at,currentMipmapOp }, newAt);
 
-	return newAt;
-}
+		return newAt;
+	}
 
 }
\ No newline at end of file
diff --git a/Engine/Plugins/Experimental/Mutable/Source/MutableTools/Private/MuT/ASTOpImageMipmap.h b/Engine/Plugins/Experimental/Mutable/Source/MutableTools/Private/MuT/ASTOpImageMipmap.h
--- a/Engine/Plugins/Experimental/Mutable/Source/MutableTools/Private/MuT/ASTOpImageMipmap.h
+++ b/Engine/Plugins/Experimental/Mutable/Source/MutableTools/Private/MuT/ASTOpImageMipmap.h
@@ -11,8 +11,7 @@
 
 namespace mu
 {
-struct FProgram;
-template <class SCALAR> class vec4;
+	struct FProgram;
 
 	class ASTOpImageMipmap : public ASTOp
 	{
@@ -20,15 +19,18 @@ template <class SCALAR> class vec4;
 
 		ASTChild Source;
 
-		uint8_t Levels = 0;
+		uint8 Levels = 0;
 
 		//! Number of mipmaps that can be generated for a single layout block.
-		uint8_t BlockLevels = 0;
+		uint8 BlockLevels = 0;
 
 		//! This is true if this operation is supposed to build only the tail mipmaps.
 		//! It is used during the code optimisation phase, and to validate the code.
 		bool bOnlyTail = false;
 
+		/** If this is enabled, at optimize time, the mip operation will not be split in top and bottom mip (for compose tails). */
+		bool bPreventSplitTail = false;
+
 		//! Mipmap generation settings. 
 		float SharpenFactor = 0.0f;
 		EAddressMode AddressMode = EAddressMode::AM_NONE;
diff --git a/Engine/Plugins/Experimental/Mutable/Source/MutableTools/Private/MuT/ASTOpSwitch.cpp b/Engine/Plugins/Experimental/Mutable/Source/MutableTools/Private/MuT/ASTOpSwitch.cpp
--- a/Engine/Plugins/Experimental/Mutable/Source/MutableTools/Private/MuT/ASTOpSwitch.cpp
+++ b/Engine/Plugins/Experimental/Mutable/Source/MutableTools/Private/MuT/ASTOpSwitch.cpp
@@ -2,6 +2,7 @@
 
 #include "MuT/ASTOpSwitch.h"
 
+#include "MuT/ASTOpParameter.h"
 #include "Containers/Map.h"
 #include "Misc/AssertionMacros.h"
 #include "MuR/ImagePrivate.h"
@@ -415,7 +416,7 @@ namespace mu
 		// Constant condition?
 		if (variable->GetOpType() == OP_TYPE::NU_CONSTANT)
 		{
-			Ptr<ASTOp> branch = def.child();
+			Ptr<ASTOp> Branch = def.child();
 
 			auto typedCondition = dynamic_cast<const ASTOpFixed*>(variable.child().get());
 			for (int32 o = 0; o < cases.Num(); ++o)
@@ -425,55 +426,58 @@ namespace mu
 					==
 					(int)cases[o].condition)
 				{
-					branch = cases[o].branch.child();
+					Branch = cases[o].branch.child();
 					break;
 				}
 			}
 
-			return branch;
+			return Branch;
 		}
 
-		//else
-		//{
-			// If all the possible branches are the same, remove the instruction
-			// TODO: If the variables is a parameter, checking that all the possible
-			// values of this parameter are handled.
-
-	//            Disable this: I think it doesn't deal correctly with def. If all branches are the same but the
-	//              value is none of the options, "def" should be selected.
-	//            OP::ADDRESS switchAt = at;
-	//            bool same = true;
-	//            OP::ADDRESS value = 0;
-	//            //vector< UINT16 > options;
-
-	//            const auto& options = program.m_constantSwitches[program.m_code[switchAt].args.Switch.options].m_options;
-	//            for ( size_t o=0; o<options.size(); ++o )
-	//            {
-	//                if ( options[o].at )
-	//                {
-	//                    if (!value)
-	//                    {
-	//                        value = options[o].at;
-	//                    }
-	//                    else
-	//                    {
-	//                        same &= ( value == options[o].at );
-	//                    }
-
-	//                    //options.push_back
-	//                    //		( program.m_code[switchAt].args.Switch.conditions[o] );
-	//                }
-	//            }
-
-
-	//            if (same)
-	//            {
-	//                m_modified = true;
-
-	//                at = value;
-	//            }
-
-		//}
+		else if (variable->GetOpType() == OP_TYPE::NU_PARAMETER)
+		{
+			// If all the branches for the possible values are the same op remove the instruction
+			const ASTOpParameter* ParamOp = dynamic_cast<const ASTOpParameter*>(variable.child().get());
+			check(ParamOp);
+			check(!ParamOp->parameter.m_possibleValues.IsEmpty());
+
+			bool bFirstValue = true;
+			bool bAllSame = true;
+			Ptr<ASTOp> SameBranch = nullptr;
+			for (const FParameterDesc::INT_VALUE_DESC& Value : ParamOp->parameter.m_possibleValues)
+			{
+				// Look for the switch branch it would take
+				Ptr<ASTOp> Branch = def.child();
+				for (const FCase& Case : cases)
+				{
+					if (Case.condition == Value.m_value)
+					{
+						Branch = Case.branch.child();
+						break;
+					}
+				}
+
+				if (bFirstValue)
+				{
+					bFirstValue = false;
+					SameBranch = Branch;
+				}
+				else
+				{
+					if (SameBranch != Branch)
+					{
+						bAllSame = false;
+						SameBranch = nullptr;
+						break;
+					}
+				}
+			}
+
+	        if (bAllSame)
+	        {
+				return SameBranch;
+	        }
+		}
 
 		return nullptr;
 	}
diff --git a/Engine/Plugins/Experimental/Mutable/Source/MutableTools/Private/MuT/CodeGenerator_Image.cpp b/Engine/Plugins/Experimental/Mutable/Source/MutableTools/Private/MuT/CodeGenerator_Image.cpp
--- a/Engine/Plugins/Experimental/Mutable/Source/MutableTools/Private/MuT/CodeGenerator_Image.cpp
+++ b/Engine/Plugins/Experimental/Mutable/Source/MutableTools/Private/MuT/CodeGenerator_Image.cpp
@@ -26,6 +26,7 @@
 #include "MuT/ASTOpImageNormalComposite.h"
 #include "MuT/ASTOpImagePixelFormat.h"
 #include "MuT/ASTOpImageTransform.h"
+#include "MuT/ASTOpImageMakeGrowMap.h"
 #include "MuT/ASTOpMeshExtractLayoutBlocks.h"
 #include "MuT/ASTOpMeshFormat.h"
 #include "MuT/ASTOpParameter.h"
@@ -1606,14 +1607,22 @@ namespace mu
         rasterop->op.args.ImageRasterMesh.sizeX = op->op.args.ImageRasterMesh.sizeX;
         rasterop->op.args.ImageRasterMesh.sizeY = op->op.args.ImageRasterMesh.sizeY;
 
-        Ptr<ASTOpFixed> mapop = new ASTOpFixed();
-        mapop->op.type = OP_TYPE::IM_MAKEGROWMAP;
-        mapop->SetChild( mapop->op.args.ImageMakeGrowMap.mask, rasterop );
-        mapop->op.args.ImageMakeGrowMap.border = 2;
+        Ptr<ASTOpImageMakeGrowMap> MakeGrowMapOp = new ASTOpImageMakeGrowMap();
+		MakeGrowMapOp->Mask = rasterop;
+		MakeGrowMapOp->Border = 2;
+		
+		// If we want to be able to generate progressive mips efficiently, we need mipmaps for the "displacement map".
+		if (m_compilerOptions->m_optimisationOptions.bEnableProgressiveImages)
+		{
+			Ptr<ASTOpImageMipmap> MipMask = new ASTOpImageMipmap;
+			MipMask->Source = MakeGrowMapOp->Mask.child();
+			MipMask->bPreventSplitTail = true;
+			MakeGrowMapOp->Mask = MipMask;
+		}
 
         Ptr<ASTOpFixed> disop = new ASTOpFixed();
         disop->op.type = OP_TYPE::IM_DISPLACE;
-        disop->SetChild( disop->op.args.ImageDisplace.displacementMap, mapop );
+        disop->SetChild( disop->op.args.ImageDisplace.displacementMap, MakeGrowMapOp);
         disop->SetChild( disop->op.args.ImageDisplace.source, op );
 
         result.op = disop;
diff --git a/Engine/Plugins/Experimental/Mutable/Source/MutableTools/Private/MuT/CodeOptimiser_Semantic.cpp b/Engine/Plugins/Experimental/Mutable/Source/MutableTools/Private/MuT/CodeOptimiser_Semantic.cpp
--- a/Engine/Plugins/Experimental/Mutable/Source/MutableTools/Private/MuT/CodeOptimiser_Semantic.cpp
+++ b/Engine/Plugins/Experimental/Mutable/Source/MutableTools/Private/MuT/CodeOptimiser_Semantic.cpp
@@ -19,6 +19,7 @@
 #include "MuT/ASTOpImageMultiLayer.h"
 #include "MuT/ASTOpImagePatch.h"
 #include "MuT/ASTOpImagePixelFormat.h"
+#include "MuT/ASTOpImageMakeGrowMap.h"
 #include "MuT/ASTOpMeshMorph.h"
 #include "MuT/ASTOpSwitch.h"
 #include "MuT/CodeOptimiser.h"
@@ -1703,69 +1704,6 @@ namespace mu
         }
 
 
-        //-----------------------------------------------------------------------------------------
-        //-----------------------------------------------------------------------------------------
-        //-----------------------------------------------------------------------------------------
-        case OP_TYPE::IM_MAKEGROWMAP:
-        {
-            auto sourceAt = children[op.args.ImageMakeGrowMap.mask].child();
-
-            switch ( sourceAt->GetOpType() )
-            {
-
-            case OP_TYPE::IM_CONDITIONAL:
-            {
-                // We move the format down the two paths
-                auto nop = mu::Clone<ASTOpConditional>(sourceAt);
-
-                auto aOp = mu::Clone<ASTOpFixed>(this);
-                aOp->SetChild( aOp->op.args.ImageMakeGrowMap.mask, nop->yes );
-                nop->yes = aOp;
-
-                auto bOp = mu::Clone<ASTOpFixed>(this);
-                bOp->SetChild( bOp->op.args.ImageMakeGrowMap.mask, nop->no);
-                nop->no = bOp;
-
-                at = nop;
-                break;
-            }
-
-            case OP_TYPE::IM_SWITCH:
-            {
-                // Move the format down all the paths
-                auto nop = mu::Clone<ASTOpSwitch>(sourceAt);
-
-                if (nop->def)
-                {
-                    auto defOp = mu::Clone<ASTOpFixed>(this);
-                    defOp->SetChild( defOp->op.args.ImageMakeGrowMap.mask, nop->def );
-                    nop->def = defOp;
-                }
-
-                // We need to copy the options because we change them
-                for ( size_t v=0; v<nop->cases.Num(); ++v )
-                {
-                    if ( nop->cases[v].branch )
-                    {
-                        auto bOp = mu::Clone<ASTOpFixed>(this);
-                        bOp->SetChild( bOp->op.args.ImageMakeGrowMap.mask, nop->cases[v].branch );
-                        nop->cases[v].branch = bOp;
-                    }
-                }
-
-                at = nop;
-                break;
-            }
-
-            default:
-                break;
-            }
-
-            break;
-        }
-
-
-
         //-----------------------------------------------------------------------------------------
         //-----------------------------------------------------------------------------------------
         //-----------------------------------------------------------------------------------------
@@ -2176,9 +2114,9 @@ namespace mu
 
 			case OP_TYPE::IM_MAKEGROWMAP:
 			{
-				auto nop = mu::Clone<ASTOpFixed>(at);
-				auto maskOp = nop->children[nop->op.args.ImageMakeGrowMap.mask].child();
-				nop->SetChild(nop->op.args.ImageMakeGrowMap.mask, Visit(maskOp, currentSinkingOp));
+				auto nop = mu::Clone<ASTOpImageMakeGrowMap>(at);
+				auto maskOp = nop->Mask.child();
+				nop->Mask = Visit(maskOp, currentSinkingOp);
 				newAt = nop;
 				break;
 			}
@@ -2474,13 +2412,12 @@ namespace mu
 					newOp->SetChild(newOp->op.args.ImageDisplace.source, baseOp);
 
 					// Clone the map op and replace its children
-					Ptr<ASTOpFixed> mapOp = mu::Clone<ASTOpFixed>(OriginalDisplacementMapOp);
+					Ptr<ASTOpImageMakeGrowMap> mapOp = mu::Clone<ASTOpImageMakeGrowMap>(OriginalDisplacementMapOp);
 					newOp->SetChild(newOp->op.args.ImageDisplace.displacementMap, mapOp);
 
 					Ptr<ASTOpFixed> mapSourceOp = mu::Clone<ASTOpFixed>(this);
-					mapSourceOp->SetChild(mapSourceOp->op.args.ImageResize.source,
-						mapOp->children[mapOp->op.args.ImageMakeGrowMap.mask]);
-					mapOp->SetChild(mapOp->op.args.ImageDisplace.displacementMap, mapSourceOp);
+					mapSourceOp->SetChild(mapSourceOp->op.args.ImageResize.source, mapOp->Mask);
+					mapOp->Mask = mapSourceOp;
 
 					at = newOp;
 				}
diff --git a/Engine/Plugins/Experimental/Mutable/Source/MutableTools/Private/MuT/Compiler.cpp b/Engine/Plugins/Experimental/Mutable/Source/MutableTools/Private/MuT/Compiler.cpp
--- a/Engine/Plugins/Experimental/Mutable/Source/MutableTools/Private/MuT/Compiler.cpp
+++ b/Engine/Plugins/Experimental/Mutable/Source/MutableTools/Private/MuT/Compiler.cpp
@@ -163,6 +163,13 @@ namespace mu
 	}
 
 
+	//---------------------------------------------------------------------------------------------
+	void CompilerOptions::SetEnableProgressiveImages(bool bEnabled)
+	{
+		m_pD->m_optimisationOptions.bEnableProgressiveImages = bEnabled;
+	}
+
+
     //---------------------------------------------------------------------------------------------
     //---------------------------------------------------------------------------------------------
     //---------------------------------------------------------------------------------------------
diff --git a/Engine/Plugins/Experimental/Mutable/Source/MutableTools/Public/MuT/Compiler.h b/Engine/Plugins/Experimental/Mutable/Source/MutableTools/Public/MuT/Compiler.h
--- a/Engine/Plugins/Experimental/Mutable/Source/MutableTools/Public/MuT/Compiler.h
+++ b/Engine/Plugins/Experimental/Mutable/Source/MutableTools/Public/MuT/Compiler.h
@@ -90,6 +90,9 @@ namespace mu
         //! 
         void SetDataPackingStrategy( int32 minRomSize, int32 MinTextureResidentMipCount );
 
+		/** If enabled it will make sure that the object is compile to generate smaller mips of the images. */
+		void SetEnableProgressiveImages(bool bEnabled);
+
         //! Different data packing strategies
         enum class TextureLayoutStrategy : std::uint8_t
         {
