#rnx

diff --git a/Engine/Plugins/Experimental/Mutable/Source/CustomizableObject/Private/MuCO/CustomizableObjectSystem.cpp b/Engine/Plugins/Experimental/Mutable/Source/CustomizableObject/Private/MuCO/CustomizableObjectSystem.cpp
--- a/Engine/Plugins/Experimental/Mutable/Source/CustomizableObject/Private/MuCO/CustomizableObjectSystem.cpp
+++ b/Engine/Plugins/Experimental/Mutable/Source/CustomizableObject/Private/MuCO/CustomizableObjectSystem.cpp
@@ -684,7 +684,7 @@ void FCustomizableObjectSystemPrivate::AddReferencedObjects(FReferenceCollector&
 }
 
 
-int32 FCustomizableObjectSystemPrivate::EnableMutableProgressiveMipStreaming = 0;
+int32 FCustomizableObjectSystemPrivate::EnableMutableProgressiveMipStreaming = 1;
 
 
 // Warning! If this is enabled, do not get references to the textures generated by Mutable! They are owned by Mutable and could become invalid at any moment
