[mutable] Added to mutable debugger the amount of operations for each navigable operation type.
Now, alongside with the operation type displayed on the navigation dropdown you will be able to check out the amount of times an operation of the type has been found on the tree on a non duplicated address.
- Some small tweaks.

diff --git a/Engine/Plugins/Experimental/Mutable/Source/CustomizableObjectEditor/Private/MuCOE/SMutableCodeViewer.cpp b/Engine/Plugins/Experimental/Mutable/Source/CustomizableObjectEditor/Private/MuCOE/SMutableCodeViewer.cpp
--- a/Engine/Plugins/Experimental/Mutable/Source/CustomizableObjectEditor/Private/MuCOE/SMutableCodeViewer.cpp
+++ b/Engine/Plugins/Experimental/Mutable/Source/CustomizableObjectEditor/Private/MuCOE/SMutableCodeViewer.cpp
@@ -357,7 +357,7 @@ void SMutableCodeViewer::SetCurrentModel(const mu::ModelPtr& InMutableModel)
 	ExpandedElements.Empty();
 	FoundModelOperationTypeElements.Empty();
 	ModelOperationTypes.Empty();
-	ModelOperationTypeStrings.Empty();
+	ModelOperationTypeNames.Empty();
 	NavigationElements.Empty();
 
 	// Generate all elements before starting the tree UI so we have a deterministic set of unique and duplicated elements
@@ -365,7 +365,7 @@ void SMutableCodeViewer::SetCurrentModel(const mu::ModelPtr& InMutableModel)
 
 	// Setup Navigation system
 	{
-		// Store the addresses of the root nodes for later search operations
+		// Store the addresses of the root nodes so they can be used by operation search methods
 		CacheRootNodeAddresses();
 
 		// Cache the operation types that are present on the model
@@ -387,7 +387,7 @@ void SMutableCodeViewer::SetCurrentModel(const mu::ModelPtr& InMutableModel)
 void SMutableCodeViewer::Construct(const FArguments& InArgs, const mu::ModelPtr& InMutableModel/*, const TSharedPtr<SDockTab>& ConstructUnderMajorTab*/)
 {
 	// Min width allowed for the column. Needed to avoid having issues with the constants space being to small
-	// and then getting too tall on the y axis crashing the drawer.
+	// and then getting too tall on the y axis crashing the UI drawer.
 	constexpr float MinParametersCollWidth = 400;
 
 	SetCurrentModel(InMutableModel);
@@ -706,16 +706,17 @@ void SMutableCodeViewer::GenerateNavigationDropdownElements()
 	for	(int32 OperationTypeIndex = 0; OperationTypeIndex < OperationTypesCount;  OperationTypeIndex++)
 	{
 		// Get the type as a string to be able to print it on the UI
-		TSharedPtr<FString> OperationTypeAsString = ModelOperationTypeStrings[OperationTypeIndex];
-		FText OperationTypeName = FText::FromString(*OperationTypeAsString);
+		const FText OperationTypeName = FText::FromString(ModelOperationTypeNames[OperationTypeIndex]);
 
-		// Get the Color to be used on the text that will represent the operation on the dropdown
-		const mu::OP_TYPE RepresentedType = ModelOperationTypes[OperationTypeIndex];
+		const mu::OP_TYPE RepresentedType = ModelOperationTypes[OperationTypeIndex].Key;
 		
+		// Get the Color to be used on the text that will represent the operation on the dropdown
 		const FSlateColor OperationColor = ColorPerComputationalCost[StaticCast<uint8>(GetOperationTypeComputationalCost(RepresentedType))];
+
+		const uint32 OperationTypeInstancesCount = ModelOperationTypes[OperationTypeIndex].Value;
 		
 		// Generate an element to be used by the ComboBox handling the selection of the type to be used during navigation
-		TSharedPtr<FMutableOperationElement> OperationElement = MakeShared<FMutableOperationElement>(RepresentedType,OperationTypeName,OperationColor);
+		TSharedPtr<FMutableOperationElement> OperationElement = MakeShared<FMutableOperationElement>(RepresentedType, OperationTypeName,OperationTypeInstancesCount,OperationColor);
 		FoundModelOperationTypeElements.Add(OperationElement);
 	}
 }
@@ -723,9 +724,9 @@ void SMutableCodeViewer::GenerateNavigationDropdownElements()
 TSharedRef<SWidget> SMutableCodeViewer::OnGenerateOpNavigationDropDownWidget(
 	TSharedPtr<FMutableOperationElement> MutableOperationElement) const
 {
-	return SNew(STextBlock).
-			Text(MutableOperationElement->OperationTypeText).
-			ColorAndOpacity(MutableOperationElement->OperationTextColor);
+	return SNew(STextBlock)
+			.Text(MutableOperationElement->OperationTypeText)
+			.ColorAndOpacity(MutableOperationElement->OperationTextColor);
 }
 
 void SMutableCodeViewer::OnNavigationSelectedOperationChanged(
@@ -763,15 +764,18 @@ void SMutableCodeViewer::OnNavigationSelectedOperationChanged(
 void SMutableCodeViewer::GenerateNavigationOpTypeStrings()
 {
 	// Grab only the names from the operation types located during the caching of operation types of the model
-	for (const mu::OP_TYPE LocatedOperationType : ModelOperationTypes)
+	for (const TTuple<mu::OP_TYPE, uint32>& LocatedOperationType : ModelOperationTypes)
 	{
-		// Find the name
-		const uint16 OperationIndex = static_cast<uint16>(LocatedOperationType);
+		// Find the name of the Operation type
+		const uint16 OperationIndex = static_cast<uint16>(LocatedOperationType.Key);
 		const TCHAR* OpName = mu::s_opNames[OperationIndex];
-		
-		// Operation name to be saved
-		TSharedPtr<FString> OperationNamePointer = MakeShared<FString>(OpName);
-		ModelOperationTypeStrings.Add( OperationNamePointer);
+
+		// Remove trailing whitespaces adding noise and messing up concatenations with other strings
+		FString OperationNameString{OpName};
+		OperationNameString.RemoveSpacesInline();
+
+		// Save the name
+		ModelOperationTypeNames.Add( OperationNameString);
 	}
 }
 
@@ -784,10 +788,16 @@ void SMutableCodeViewer::OnSelectedOperationTypeFromTree()
 	
 	const mu::OP_TYPE OperationType =
 		MutableModel->GetPrivate()->m_program.GetOpType(ReferenceOperationElement->MutableOperation);
-	const int32 OperationIndex = ModelOperationTypes.IndexOfByKey(OperationType);
 
+	// Search for the index of the provided operation type inside the array of operation types found on the model
+	const int32 OperationIndex = ModelOperationTypes.IndexOfByPredicate(
+		[OperationType](const TPair<mu::OP_TYPE, uint32>& Current)
+		{
+			return Current.Key == OperationType;
+		});
+	
 	// Failing the next check would mean that we are not caching all the types present on the current operation's tree
-	check (OperationIndex != -1);
+	check (OperationIndex != INDEX_NONE);
 	
 	// Set the type operation type to be looking for -> Will invoke OnOptionTypeSelectionChanged
 	TargetedTypeSelector->SetSelectedItem(FoundModelOperationTypeElements[OperationIndex]);
@@ -797,7 +807,6 @@ void SMutableCodeViewer::OnSelectedOperationTypeFromTree()
 }
 
 
-
 void SMutableCodeViewer::SortNavigationElements()
 {
 	// Sort the array from lower index to bigger index (0 , 1 , 2 ...)
@@ -875,25 +884,56 @@ void SMutableCodeViewer::GetOperationsOfType(const mu::OP_TYPE& TargetOperationT
 void SMutableCodeViewer::CacheOperationTypesPresentOnModel()
 {
 	check(MutableModel)
-	check(RootNodeAddresses.Num());
 
-	// Set of operation types to be filed 
-	TSet<mu::OP_TYPE> OperationTypes;
+	// Initialize NodeOperationTypes with empty TPair for each possible mutable operation type
+	{
+		constexpr uint32 OperationTypesCount = static_cast<uint16>(mu::OP_TYPE::COUNT);
+		ModelOperationTypes.Empty(OperationTypesCount);
+		for (uint32 Index = 0; Index < OperationTypesCount ; Index++)
+		{
+			mu::OP_TYPE TargetType = StaticCast<mu::OP_TYPE>(Index);
+			ModelOperationTypes.Add(TPair<mu::OP_TYPE,uint32>{TargetType,0});
+		}
+	}
 	
-	const mu::FProgram& Program = MutableModel->GetPrivate()->m_program;
-	TSet<mu::OP::ADDRESS> AlreadyVisitedAddresses;
-	GetOperationTypesPresentOnModel(RootNodeAddresses,Program,OperationTypes,AlreadyVisitedAddresses);
+	// Locate all operation types found on the provided model program data structure and count how many instances of each
+	// there are
+	{
+		// Get the types and the amount of instances of each unique operation on the mutable model
+		const mu::FProgram& Program = MutableModel->GetPrivate()->m_program;
+		const uint32 ProgramAddressesCount = Program.m_opAddress.Num();
+		
+		// Ensure first operation type is NONE since we are skipping it due to it having to have that type
+		check (Program.GetOpType(Program.m_opAddress[0]) == mu::OP_TYPE::NONE);
 
-	// After using the set object save it as an array for later reference
-	ModelOperationTypes = OperationTypes.Array();
+		// Iterate over the addresses of the program and count how many instances each type has.
+		for (uint32 ProgramAddressesIndex = 1 ; ProgramAddressesIndex < ProgramAddressesCount; ProgramAddressesIndex++)
+		{
+			// Locate what is the position (index) of the operation type of the address on our collection of types found until now
+			const mu::OP_TYPE OperationType = Program.GetOpType(ProgramAddressesIndex);
+			
+			// Increase the counter for this operation type
+			const uint16 TypeAsInteger = StaticCast<uint16>(OperationType);;
+			ModelOperationTypes[TypeAsInteger].Value++;
+		}
+	}
 	
-	// Sort the contents of the array alphabetically
-	ModelOperationTypes.StableSort([&](const mu::OP_TYPE& A, const mu::OP_TYPE& B)
+	// Remove all operation types that do have no operations present on the model
+	{
+		ModelOperationTypes.RemoveAll(
+		[](const TPair<mu::OP_TYPE, uint32>& Current)
+			{
+				return Current.Value == 0;
+			});
+	}
+	
+	// Sort the contents of the array of mutable operation types alphabetically
+	ModelOperationTypes.StableSort([&](const TPair<mu::OP_TYPE,uint32>& A, const TPair<mu::OP_TYPE,uint32>& B)
 	{
 		// Find the name
 		FString AString;
 		{
-			const uint16 OperationIndex = static_cast<uint16>(A);
+			const uint16 OperationIndex = static_cast<uint16>(A.Key);
 			const TCHAR* OpName = mu::s_opNames[OperationIndex];
 			AString = FString(OpName);
 		}
@@ -901,7 +941,7 @@ void SMutableCodeViewer::CacheOperationTypesPresentOnModel()
 		// Find out the name of the first element
 		FString BString;
 		{
-			const uint16 OperationIndex = static_cast<uint16>(B);
+			const uint16 OperationIndex = static_cast<uint16>(B.Key);
 			const TCHAR* OpName = mu::s_opNames[OperationIndex];
 			BString = FString(OpName);
 		}
@@ -911,50 +951,14 @@ void SMutableCodeViewer::CacheOperationTypesPresentOnModel()
 	});	
 
 	// Add the "None" option as a default value at index 0 to the sorted array.
-	// the NONE option serves as a way of telling the user "hey, you are not navigating over a type of operation".
-	// We currently use it to also tell the user that if he is navigating it is over constant related operations
-	ModelOperationTypes.Insert(mu::OP_TYPE::NONE,0);
+	// the NONE option serves as a way of telling the user "hey, you are not navigating over operations of X type".
+	// We currently use it to tell the user that if he is navigating it is over constant related operations
+	ModelOperationTypes.Insert(TPair<mu::OP_TYPE,uint32>{mu::OP_TYPE::NONE,0},0);
+	ModelOperationTypes.Shrink();
 
 	// ModelOperationTypes is now an array starting with None and all the types found on the operations tree in alphabetical order
 }
 
-void SMutableCodeViewer::GetOperationTypesPresentOnModel(
-	const TArray< mu::OP::ADDRESS>& InParentAddresses, const mu::FProgram& InProgram, TSet<mu::OP_TYPE>& OutLocatedOperations,
-	TSet<mu::OP::ADDRESS>& AlreadyProcessedAddresses)
-{
-	TArray<mu::OP::ADDRESS> NextBatch;
-	
-	for (const mu::OP::ADDRESS& ParentAddress : InParentAddresses)
-	{
-		// Process parent 
-		if ( !AlreadyProcessedAddresses.Contains(ParentAddress))
-		{
-			// Cache the operation type for later usage
-			OutLocatedOperations.Add(InProgram.GetOpType(ParentAddress));
-			
-			// Cache to avoid processing it again later
-			AlreadyProcessedAddresses.Add(ParentAddress);
-
-			// Generic case for unnamed children traversal.
-			mu::ForEachReference(InProgram, ParentAddress, [this, &NextBatch]( mu::OP::ADDRESS ChildAddress)
-			{
-				// Avoid processing items more than once
-				if (ChildAddress)
-				{
-					// Add for expansion
-					NextBatch.Add(ChildAddress);
-				}
-			});
-		}
-		
-	}
-
-	if (NextBatch.Num())
-	{
-		// Process the children of this object
-		GetOperationTypesPresentOnModel(NextBatch,InProgram,OutLocatedOperations,AlreadyProcessedAddresses);
-	}
-}
 
 
 FText SMutableCodeViewer::OnPrintNavigableObjectAddressesCount() const
@@ -971,7 +975,7 @@ FText SMutableCodeViewer::OnPrintNavigableObjectAddressesCount() const
 		
 		OutputString.Append( FString::FromInt(NavigationElementsCount));
 
-		// Format -> 1 / 12 or 12
+		// Format : 1 / 12 or 12
 	}
 	
 	// Depending on the amount of navigable objects (addresses, not actual elements) display the amount there are
@@ -1069,8 +1073,8 @@ void SMutableCodeViewer::FocusViewOnNavigationTarget()
 void SMutableCodeViewer::GenerateAllTreeElements()
 {
 	// By generating all tree elements prior to usage we are able to :
-	//	- Compute the index of each one to aidd on navigation
-	//	- Remove indeterministic assignation of the "Duplicated" state of elements. It was due to user interaction with the tree
+	//	- Compute the index of each one to aid on navigation
+	//	- Remove non-deterministic assignation of the "Duplicated" state of elements. It was due to user interaction with the tree
 	//	Only unique elements, their children and duplicated elements will be generated. Children of duplicates will
 	// 	be ignored due to how we handle them when expanding and contracting elements (OnExpansionChanged)
 
diff --git a/Engine/Plugins/Experimental/Mutable/Source/CustomizableObjectEditor/Private/MuCOE/SMutableCodeViewer.h b/Engine/Plugins/Experimental/Mutable/Source/CustomizableObjectEditor/Private/MuCOE/SMutableCodeViewer.h
--- a/Engine/Plugins/Experimental/Mutable/Source/CustomizableObjectEditor/Private/MuCOE/SMutableCodeViewer.h
+++ b/Engine/Plugins/Experimental/Mutable/Source/CustomizableObjectEditor/Private/MuCOE/SMutableCodeViewer.h
@@ -279,11 +279,12 @@ private:
 	/** Operation type we are using to search for tree nodes. Driven primarily by the UI */
 	mu::OP_TYPE OperationTypeToSearch = mu::OP_TYPE::NONE;
 
-	/** Operation types present on the currently set mutable model. */
-	TArray<mu::OP_TYPE> ModelOperationTypes;
+	/** Operation types present on the currently set mutable model.
+	 * The value of the TPair represents the amount of times operations using it have been found (not duplicates) */
+	TArray< TPair< mu::OP_TYPE,uint32>> ModelOperationTypes;
 	
 	/** Array with all the names for each of the operations available on mutable. Used by the UI  */
-	TArray<TSharedPtr<FString>> ModelOperationTypeStrings;
+	TArray<FString> ModelOperationTypeNames;
 	
 	/** Stores a list of strings based on the possible types of operations mutable define to be used by the UI */
 	void GenerateNavigationOpTypeStrings();
@@ -297,16 +298,6 @@ private:
 	 * @note Method designed to be called once per model.*/
 	void CacheOperationTypesPresentOnModel();
 	
-	/** Method to scan over all the operations performed on the current model to produce a set of operation types present on it.
-	 * @param InParentAddresses Addresses of the parent objects. Required to be able to perform recursive calls to this method (to get the data from the children)
-	 * @param  InProgram Mutable program.
-	 * @param OutLocatedOperations Output of the method : Set with all the operation types present on the model.
-	 * @param AlreadyProcessedAddresses Addresses already processed. Used to avoid processing the same operation twice.
-	 */
-	void GetOperationTypesPresentOnModel(
-		const TArray< mu::OP::ADDRESS>& InParentAddresses, const mu::FProgram& InProgram, TSet<mu::OP_TYPE>& OutLocatedOperations,
-		TSet<mu::OP::ADDRESS>& AlreadyProcessedAddresses);
-
 	
 	/*
 	 * Navigation : UI callback methods
@@ -717,11 +708,20 @@ private:
 class FMutableOperationElement : public TSharedFromThis<FMutableOperationElement>
 {
 public:
-	FMutableOperationElement(mu::OP_TYPE InOperationType, FText OperationTypeBeingRepresented,FSlateColor OperationColor)
+	FMutableOperationElement(mu::OP_TYPE InOperationType,FText OperationTypeName,uint32 OperationTypeInstanceCount,FSlateColor OperationColor)
 	{
 		OperationType = InOperationType;
-		OperationTypeText = OperationTypeBeingRepresented;
 		OperationTextColor = OperationColor;
+
+		// Show the amount of instances of the operation type is found on model.
+		// @note This if block will not be triggered when working with an operation type that is not present on the model
+		// but added manually to the list of operation types (currently mu::OP_Type::None is an example of this)
+		FText OperationsCountText = FText::GetEmpty();
+		if (OperationTypeInstanceCount > 0)
+		{
+			OperationsCountText = FText::Format(INVTEXT(" - ({0})"), OperationTypeInstanceCount);
+		}
+		OperationTypeText = FText::Format(INVTEXT("{0}{1}"), OperationTypeName,OperationsCountText);
 	}
 
 public:
