#rnx

diff --git a/Engine/Plugins/Experimental/Mutable/Source/CustomizableObject/Classes/MuCO/CustomizableObject.h b/Engine/Plugins/Experimental/Mutable/Source/CustomizableObject/Classes/MuCO/CustomizableObject.h
--- a/Engine/Plugins/Experimental/Mutable/Source/CustomizableObject/Classes/MuCO/CustomizableObject.h
+++ b/Engine/Plugins/Experimental/Mutable/Source/CustomizableObject/Classes/MuCO/CustomizableObject.h
@@ -743,16 +743,46 @@ struct FMutableRefLODData
 };
 
 
+USTRUCT()
 struct FMutableRefSocket
 {
+	GENERATED_BODY()
+
+	UPROPERTY()
 	FName SocketName;
+	UPROPERTY()
 	FName BoneName;
 
+	UPROPERTY()
 	FVector RelativeLocation;
+	UPROPERTY()
 	FRotator RelativeRotation;
+	UPROPERTY()
 	FVector RelativeScale;
 
+	UPROPERTY()
 	bool bForceAlwaysAnimated = false;
+
+	// When two sockets have the same name, the one with higher priority will be picked and the other discarded
+	UPROPERTY()
+	int32 Priority = -1;
+
+	bool operator ==(const FMutableRefSocket& Other) const
+	{
+		if (
+			SocketName == Other.SocketName &&
+			BoneName == Other.BoneName &&
+			RelativeLocation == Other.RelativeLocation &&
+			RelativeRotation == Other.RelativeRotation &&
+			RelativeScale == Other.RelativeScale &&
+			bForceAlwaysAnimated == Other.bForceAlwaysAnimated &&
+			Priority == Other.Priority)
+		{
+			return true;
+		}
+
+		return false;
+	}
 	
 	friend FArchive& operator<<(FArchive& Ar, FMutableRefSocket& Data)
 	{
@@ -761,6 +791,8 @@ struct FMutableRefSocket
 		Ar << Data.RelativeLocation;
 		Ar << Data.RelativeRotation;
 		Ar << Data.RelativeScale;
+		Ar << Data.bForceAlwaysAnimated;
+		Ar << Data.Priority;
 
 		return Ar;
 	}
@@ -1323,6 +1355,10 @@ public:
 	UPROPERTY()
 	TMap<FString, TSoftClassPtr<UAnimInstance>> AnimBPAssetsMap;
 
+	UPROPERTY()
+	/** Stores the sockets provided by the part skeletal meshes, to be merged in the generated meshes */
+	TArray<FMutableRefSocket> SocketArray;
+
 	/** Stores the textures that will be used to mask-out areas in projectors. The cache isn't used for rendering, but for coverage testing */
 	UPROPERTY()
 	TSoftObjectPtr<UMutableMaskOutCache> MaskOutCache;
diff --git a/Engine/Plugins/Experimental/Mutable/Source/CustomizableObject/Private/MuCO/CustomizableObjectInstance.cpp b/Engine/Plugins/Experimental/Mutable/Source/CustomizableObject/Private/MuCO/CustomizableObjectInstance.cpp
--- a/Engine/Plugins/Experimental/Mutable/Source/CustomizableObject/Private/MuCO/CustomizableObjectInstance.cpp
+++ b/Engine/Plugins/Experimental/Mutable/Source/CustomizableObject/Private/MuCO/CustomizableObjectInstance.cpp
@@ -1477,6 +1477,7 @@ UPhysicsAsset* UCustomizableInstancePrivateData::BuildPhysicsAsset(
 	return Result;
 }
 
+
 bool UCustomizableInstancePrivateData::UpdateSkeletalMesh_PostBeginUpdate0(UCustomizableObjectInstance* Public, const TSharedPtr<FMutableOperationData>& OperationData )
 {
 	MUTABLE_CPUPROFILER_SCOPE(UCustomizableInstancePrivateData::UpdateSkeletalMesh_PostBeginUpdate0)
@@ -1605,6 +1606,7 @@ bool UCustomizableInstancePrivateData::UpdateSkeletalMesh_PostBeginUpdate0(UCust
 
 					TArray<USkeletalMeshSocket*>& Sockets = SkeletalMesh->GetMeshOnlySocketList();
 					Sockets.Empty(SocketCount);
+					TMap<FName, TTuple<int32, int32>> SocketMap; // Maps Socket name to Sockets Array index and priority
 					
 					for (uint32 SocketIndex = 0; SocketIndex < SocketCount; ++SocketIndex)
 					{
@@ -1620,7 +1622,62 @@ bool UCustomizableInstancePrivateData::UpdateSkeletalMesh_PostBeginUpdate0(UCust
 						Socket->RelativeScale = RefSocket.RelativeScale;
 
 						Socket->bForceAlwaysAnimated = RefSocket.bForceAlwaysAnimated;
-						Sockets.Add(Socket);
+						const int32 LastIndex = Sockets.Add(Socket);
+
+						SocketMap.Add(Socket->SocketName, TTuple<int32, int32>(LastIndex, RefSocket.Priority));
+					}
+
+					mu::MeshPtrConst MutableMesh = Component.Mesh;
+
+					for (int32 TagIndex = 0; TagIndex < MutableMesh->GetTagCount(); ++TagIndex)
+					{
+						FString Tag = MutableMesh->GetTag(TagIndex);
+
+						if (Tag.RemoveFromStart("__Socket:"))
+						{
+							check(Tag.IsNumeric());
+							const int32 MutableSocketIndex = FCString::Atoi(*Tag);
+							check(CustomizableObject->SocketArray.IsValidIndex(MutableSocketIndex));
+
+							if (CustomizableObject->SocketArray.IsValidIndex(MutableSocketIndex))
+							{
+								const FMutableRefSocket& MutableSocket = CustomizableObject->SocketArray[MutableSocketIndex];
+								int32 IndexToWriteSocket = -1;
+
+								if (TTuple<int32, int32>* FoundSocket = SocketMap.Find(MutableSocket.SocketName))
+								{
+									if (FoundSocket->Value < MutableSocket.Priority)
+									{
+										// Overwrite the existing socket because the new mesh part one is higher priority
+										IndexToWriteSocket = FoundSocket->Key;
+										FoundSocket->Value = MutableSocket.Priority;
+									}
+								}
+								else
+								{
+									// New Socket
+									USkeletalMeshSocket* Socket = NewObject<USkeletalMeshSocket>(SkeletalMesh);
+									IndexToWriteSocket = Sockets.Add(Socket);
+									SocketMap.Add(MutableSocket.SocketName, TTuple<int32, int32>(IndexToWriteSocket, MutableSocket.Priority));
+								}
+
+								if (IndexToWriteSocket >= 0)
+								{
+									check(Sockets.IsValidIndex(IndexToWriteSocket));
+
+									USkeletalMeshSocket* SocketToWrite = Sockets[IndexToWriteSocket];
+
+									SocketToWrite->SocketName = MutableSocket.SocketName;
+									SocketToWrite->BoneName = MutableSocket.BoneName;
+									
+									SocketToWrite->RelativeLocation = MutableSocket.RelativeLocation;
+									SocketToWrite->RelativeRotation = MutableSocket.RelativeRotation;
+									SocketToWrite->RelativeScale = MutableSocket.RelativeScale;
+									
+									SocketToWrite->bForceAlwaysAnimated = MutableSocket.bForceAlwaysAnimated;
+								}
+							}
+						}
 					}
 				}
 			}
diff --git a/Engine/Plugins/Experimental/Mutable/Source/CustomizableObjectEditor/Classes/MuCOE/Nodes/CustomizableObjectNodeObjectGroup.h b/Engine/Plugins/Experimental/Mutable/Source/CustomizableObjectEditor/Classes/MuCOE/Nodes/CustomizableObjectNodeObjectGroup.h
--- a/Engine/Plugins/Experimental/Mutable/Source/CustomizableObjectEditor/Classes/MuCOE/Nodes/CustomizableObjectNodeObjectGroup.h
+++ b/Engine/Plugins/Experimental/Mutable/Source/CustomizableObjectEditor/Classes/MuCOE/Nodes/CustomizableObjectNodeObjectGroup.h
@@ -38,6 +38,11 @@ public:
 	UPROPERTY(EditAnywhere, Category = UI, meta = (DisplayName = "Parameter UI Metadata"))
 	FMutableParamUIMetadata ParamUIMetadata;
 
+	// The sockets defined in meshes deriving from this node will inherit this socket priority. When in the generated merged mesh there
+	// are clashes with socket names, the one with higher priority will be kept and the other discarded.
+	UPROPERTY(EditAnywhere, Category = MeshSockets)
+	int32 SocketPriority = 0;
+
 	// UObject interface.
 	void Serialize(FArchive& Ar) override;
 	void PostEditChangeProperty(FPropertyChangedEvent& PropertyChangedEvent) override;
diff --git a/Engine/Plugins/Experimental/Mutable/Source/CustomizableObjectEditor/Private/MuCOE/CustomizableObjectCompiler.cpp b/Engine/Plugins/Experimental/Mutable/Source/CustomizableObjectEditor/Private/MuCOE/CustomizableObjectCompiler.cpp
--- a/Engine/Plugins/Experimental/Mutable/Source/CustomizableObjectEditor/Private/MuCOE/CustomizableObjectCompiler.cpp
+++ b/Engine/Plugins/Experimental/Mutable/Source/CustomizableObjectEditor/Private/MuCOE/CustomizableObjectCompiler.cpp
@@ -1165,6 +1165,17 @@ void FCustomizableObjectCompiler::CompileInternal(UCustomizableObject* Object, c
 	
 		Object->AnimBPAssetsMap = GenerationContext.AnimBPAssetsMap;
 
+		// Mark the object as modified, used to avoid missing assets in packages.
+		if (Object->SocketArray != GenerationContext.SocketArray)
+		{
+			if (!ParamNamesToSelectedOptions.Num()) // Don't mark the objects as modified because of a partial compilation
+			{
+				Object->MarkPackageDirty();
+			}
+		}
+
+		Object->SocketArray = GenerationContext.SocketArray;
+
 		// 
 		if (!ParamNamesToSelectedOptions.Num()) // Don't mark the objects as modified because of a partial compilation
 		{
diff --git a/Engine/Plugins/Experimental/Mutable/Source/CustomizableObjectEditor/Private/MuCOE/GenerateMutableSource/GenerateMutableSource.cpp b/Engine/Plugins/Experimental/Mutable/Source/CustomizableObjectEditor/Private/MuCOE/GenerateMutableSource/GenerateMutableSource.cpp
--- a/Engine/Plugins/Experimental/Mutable/Source/CustomizableObjectEditor/Private/MuCOE/GenerateMutableSource/GenerateMutableSource.cpp
+++ b/Engine/Plugins/Experimental/Mutable/Source/CustomizableObjectEditor/Private/MuCOE/GenerateMutableSource/GenerateMutableSource.cpp
@@ -896,6 +896,9 @@ mu::NodeObjectPtr GenerateMutableSource(const UEdGraphPin * Pin, FMutableGraphGe
 		mu::NodeObjectGroupPtr GroupNode = new mu::NodeObjectGroup();
 		Result = GroupNode;
 
+		// All sockets from all mesh parts plugged into this group node will have the following priority when there's a socket name clash
+		GenerationContext.SocketPriorityStack.Push(TypedNodeGroup->SocketPriority);
+
 		GenerationContext.AddParameterNameUnique(TypedNodeGroup, TypedNodeGroup->GroupName);
 		GroupNode->SetName(TCHAR_TO_ANSI(*TypedNodeGroup->GroupName));
 		GroupNode->SetUid(TCHAR_TO_ANSI(*TypedNodeGroup->NodeGuid.ToString()));
@@ -1106,6 +1109,10 @@ mu::NodeObjectPtr GenerateMutableSource(const UEdGraphPin * Pin, FMutableGraphGe
 		// Remove the projectors from this node
 		GenerationContext.ProjectorGroupMap.Remove(TypedNodeGroup);
 
+		// Go back to the parent group node's socket priority if it exists
+		ensure(GenerationContext.SocketPriorityStack.Num() > 0);
+		GenerationContext.SocketPriorityStack.Pop();
+
 		check(NumProjectorCountBeforeNode == GenerationContext.ProjectorGroupMap.Num());
 	}
 	else
@@ -1416,4 +1423,28 @@ int32 GetMaxTextureSize(const UTexture2D* ReferenceTexture, const FMutableGraphG
 	return 0;
 }
 
+
+void AddSocketTagsToMesh(const USkeletalMesh* SourceMesh, mu::MeshPtr MutableMesh, FMutableGraphGenerationContext& GenerationContext)
+{
+	for (int32 SocketIndex = 0; SocketIndex < SourceMesh->NumSockets(); ++SocketIndex)
+	{
+		USkeletalMeshSocket* Socket = SourceMesh->GetSocketByIndex(SocketIndex);
+
+		FMutableRefSocket MutableSocket;
+		MutableSocket.SocketName = Socket->SocketName;
+		MutableSocket.BoneName = Socket->BoneName;
+		MutableSocket.RelativeLocation = Socket->RelativeLocation;
+		MutableSocket.RelativeRotation = Socket->RelativeRotation;
+		MutableSocket.RelativeScale = Socket->RelativeScale;
+		MutableSocket.bForceAlwaysAnimated = Socket->bForceAlwaysAnimated;
+
+		MutableSocket.Priority = GenerationContext.SocketPriorityStack.IsEmpty() ? 0 : GenerationContext.SocketPriorityStack.Top();
+			
+		int32 SocketArrayIndex = GenerationContext.SocketArray.AddUnique(MutableSocket);
+		FString SocketTag = FString::Printf(TEXT("__Socket:%d"), SocketArrayIndex);
+
+		AddTagToMutableMeshUnique(*MutableMesh, SocketTag);
+	}
+}
+
 #undef LOCTEXT_NAMESPACE
diff --git a/Engine/Plugins/Experimental/Mutable/Source/CustomizableObjectEditor/Private/MuCOE/GenerateMutableSource/GenerateMutableSource.h b/Engine/Plugins/Experimental/Mutable/Source/CustomizableObjectEditor/Private/MuCOE/GenerateMutableSource/GenerateMutableSource.h
--- a/Engine/Plugins/Experimental/Mutable/Source/CustomizableObjectEditor/Private/MuCOE/GenerateMutableSource/GenerateMutableSource.h
+++ b/Engine/Plugins/Experimental/Mutable/Source/CustomizableObjectEditor/Private/MuCOE/GenerateMutableSource/GenerateMutableSource.h
@@ -711,6 +711,14 @@ struct FMutableGraphGenerationContext
 	/** Stores the anim BP assets gathered from the SkeletalMesh nodes during compilation, to be used in mesh generation in-game */
 	TMap<FString, TSoftClassPtr<UAnimInstance>> AnimBPAssetsMap;
 
+	/** Stores the sockets provided by the part skeletal meshes, to be merged in the generated meshes */
+	TArray<FMutableRefSocket> SocketArray;
+
+	/** Used to propagate the socket priority defined in group nodes to their child skeletal mesh nodes
+	* It's a stack because group nodes are recursive
+	*/
+	TArray<int32> SocketPriorityStack;
+
 	// Stores the textures that will be used to mask-out areas in the projection. The cache isn't used for rendering, but for coverage testing
 	TMap<FString, FString> MaskOutMaterialCache; // Maps a UMaterial's asset path to a UTexture's asset path
 	TMap<FString, FMaskOutTexture> MaskOutTextureCache; // Maps a UTexture's asset path to the cached mask-out texture data
@@ -790,6 +798,8 @@ mu::NodeMeshApplyPosePtr CreateNodeMeshApplyPose(mu::NodeMeshPtr InputMeshNode,
 /** Adds Tag to MutableMesh uniquely, returns the index were the tag has been inserted or the index where an intance of the tag has been found */
 int32 AddTagToMutableMeshUnique(mu::Mesh& MutableMesh, const FString& Tag);
 
+void AddSocketTagsToMesh(const USkeletalMesh* SourceMesh, mu::MeshPtr MutableMesh, FMutableGraphGenerationContext& GenerationContext);
+
 // Generates the tag for an animation instance
 FString GenerateAnimationInstanceTag(const FString& AnimInstance, int32 SlotIndex);
 
diff --git a/Engine/Plugins/Experimental/Mutable/Source/CustomizableObjectEditor/Private/MuCOE/GenerateMutableSource/GenerateMutableSourceMesh.cpp b/Engine/Plugins/Experimental/Mutable/Source/CustomizableObjectEditor/Private/MuCOE/GenerateMutableSource/GenerateMutableSourceMesh.cpp
--- a/Engine/Plugins/Experimental/Mutable/Source/CustomizableObjectEditor/Private/MuCOE/GenerateMutableSource/GenerateMutableSourceMesh.cpp
+++ b/Engine/Plugins/Experimental/Mutable/Source/CustomizableObjectEditor/Private/MuCOE/GenerateMutableSource/GenerateMutableSourceMesh.cpp
@@ -132,6 +132,7 @@
 #include "UObject/SoftObjectPtr.h"
 #include "UObject/UObjectGlobals.h"
 #include "UObject/UnrealNames.h"
+#include "Engine/SkeletalMeshSocket.h"
 
 #define LOCTEXT_NAMESPACE "CustomizableObjectEditor"
 
@@ -2601,6 +2602,7 @@ void GenerateMorphTarget(const UCustomizableObjectNode* Node, const UEdGraphPin*
 	}
 }
 
+
 /** Convert a CustomizableObject Source Graph into a mutable source graph  */
 mu::NodeMeshPtr GenerateMutableSourceMesh(const UEdGraphPin * Pin,
 	FMutableGraphGenerationContext & GenerationContext,
@@ -2746,11 +2748,12 @@ mu::NodeMeshPtr GenerateMutableSourceMesh(const UEdGraphPin * Pin,
 				for (const FGameplayTag& GamePlayTag : TypedNodeSkel->AnimationGameplayTags)
 				{
 					const FString AnimBPTag = GenerateGameplayTag(GamePlayTag.ToString());
-					const int32 CurrentTagCount = MutableMesh->GetTagCount();
 					
 					AddTagToMutableMeshUnique(*MutableMesh, AnimBPTag);
 				}
 
+				AddSocketTagsToMesh(TypedNodeSkel->SkeletalMesh, MutableMesh, GenerationContext);
+
 				MeshData.bHasVertexColors = TypedNodeSkel->SkeletalMesh->GetHasVertexColors();
 				FSkeletalMeshModel* importModel = Helper_GetImportedModel(TypedNodeSkel->SkeletalMesh);
 				MeshData.NumTexCoordChannels = importModel->LODModels[LOD].NumTexCoords;
@@ -3486,4 +3489,4 @@ mu::NodeMeshPtr GenerateMutableSourceMesh(const UEdGraphPin * Pin,
 	return Result;
 }
 
-#undef LOCTEXT_NAMESPACE
\ No newline at end of file
+#undef LOCTEXT_NAMESPACE
diff --git a/Engine/Plugins/Experimental/Mutable/Source/CustomizableObjectEditor/Private/MuCOE/GenerateMutableSource/GenerateMutableSourceTable.cpp b/Engine/Plugins/Experimental/Mutable/Source/CustomizableObjectEditor/Private/MuCOE/GenerateMutableSource/GenerateMutableSourceTable.cpp
--- a/Engine/Plugins/Experimental/Mutable/Source/CustomizableObjectEditor/Private/MuCOE/GenerateMutableSource/GenerateMutableSourceTable.cpp
+++ b/Engine/Plugins/Experimental/Mutable/Source/CustomizableObjectEditor/Private/MuCOE/GenerateMutableSource/GenerateMutableSourceTable.cpp
@@ -242,6 +242,8 @@ void FillTableColumn(const UEdGraphPin* Pin, mu::TablePtr MutableTable, FString
 							AddTagToMutableMeshUnique(*MutableMesh, GenerateGameplayTag(Tag.ToString()));
 						}
 
+						AddSocketTagsToMesh(SkeletalMesh, MutableMesh, GenerationContext);
+
 						MutableTable->SetCell(CurrentColumn, RowIdx, MutableMesh.get());
 					}
 					else
