[mutable] Fix texture compresion format on some platforms.

#rnx

diff --git a/Engine/Plugins/Experimental/Mutable/Source/CustomizableObjectEditor/Private/MuCOE/GenerateMutableSource/GenerateMutableSourceSurface.cpp b/Engine/Plugins/Experimental/Mutable/Source/CustomizableObjectEditor/Private/MuCOE/GenerateMutableSource/GenerateMutableSourceSurface.cpp
--- a/Engine/Plugins/Experimental/Mutable/Source/CustomizableObjectEditor/Private/MuCOE/GenerateMutableSource/GenerateMutableSourceSurface.cpp
+++ b/Engine/Plugins/Experimental/Mutable/Source/CustomizableObjectEditor/Private/MuCOE/GenerateMutableSource/GenerateMutableSourceSurface.cpp
@@ -679,61 +679,47 @@ mu::NodeSurfacePtr GenerateMutableSourceSurface(const UEdGraphPin * Pin, FMutabl
 							if (GenerationContext.Options.bTextureCompression)
 							{
 								check(PlatformFormats[0].Num() > LayerIndex);
-								FString FormatWithoutPrefix = PlatformFormats[0][LayerIndex].ToString();
 
-								FormatWithoutPrefix.RemoveFromStart("XBOXONE_");
-								FormatWithoutPrefix.RemoveFromStart("XBOXONEGDK_");
-								FormatWithoutPrefix.RemoveFromStart("XSX_");
-								FormatWithoutPrefix.RemoveFromStart("PS4_");
-								FormatWithoutPrefix.RemoveFromStart("PS5_");
-								FormatWithoutPrefix.RemoveFromStart("SWITCH_");
-								FormatWithoutPrefix.RemoveFromStart("OODLE_");
-								FormatWithoutPrefix.RemoveFromStart("TFO_");
-
-								// Special case for textures compressed with OODLE with the format "PlatformPrefix_OODLE_TextureFormat"
-								FString OodleString = TEXT("_OODLE_");
+								FString FormatWithoutPrefix = PlatformFormats[0][LayerIndex].ToString();
 								FString LeftSplit, RightSplit;
-								bool bSplitString = FormatWithoutPrefix.Split(OodleString, &LeftSplit, &RightSplit);
-
-								if (bSplitString)
-								{									
+								if (FormatWithoutPrefix.Split(TEXT("_"), &LeftSplit, &RightSplit, ESearchCase::IgnoreCase, ESearchDir::FromEnd))
+								{
 									FormatWithoutPrefix = RightSplit;
 								}
 
+								mu::EImageFormat mutableFormat = mu::EImageFormat::IF_RGBA_UBYTE;
+								mu::EImageFormat mutableFormatIfAlpha = mu::EImageFormat::IF_NONE;
+
 								if (FormatWithoutPrefix == TEXT("AutoDXT"))
 								{
-									FormatImage->SetFormat(mu::EImageFormat::IF_BC1, mu::EImageFormat::IF_BC3);
+									mutableFormat = mu::EImageFormat::IF_BC1;
+									mutableFormatIfAlpha = mu::EImageFormat::IF_BC3;
 								}
 								else if ((FormatWithoutPrefix == TEXT("AutoASTC")) || (FormatWithoutPrefix == TEXT("ASTC_RGBAuto")))
 								{
-									FormatImage->SetFormat(mu::EImageFormat::IF_ASTC_4x4_RGB_LDR, mu::EImageFormat::IF_ASTC_4x4_RGBA_LDR);
+									mutableFormat = mu::EImageFormat::IF_ASTC_4x4_RGB_LDR;
+									mutableFormatIfAlpha = mu::EImageFormat::IF_ASTC_4x4_RGBA_LDR;
 								}
+								else if (FormatWithoutPrefix == TEXT("DXT1")) mutableFormat = mu::EImageFormat::IF_BC1;
+								else if (FormatWithoutPrefix == TEXT("DXT3")) mutableFormat = mu::EImageFormat::IF_BC2;
+								else if (FormatWithoutPrefix == TEXT("DXT5")) mutableFormat = mu::EImageFormat::IF_BC3;
+								else if (FormatWithoutPrefix == TEXT("BC1")) mutableFormat = mu::EImageFormat::IF_BC1;
+								else if (FormatWithoutPrefix == TEXT("BC2")) mutableFormat = mu::EImageFormat::IF_BC2;
+								else if (FormatWithoutPrefix == TEXT("BC3")) mutableFormat = mu::EImageFormat::IF_BC3;
+								else if (FormatWithoutPrefix == TEXT("BC4")) mutableFormat = mu::EImageFormat::IF_BC4;
+								else if (FormatWithoutPrefix == TEXT("BC5")) mutableFormat = mu::EImageFormat::IF_BC5;
+								else if (FormatWithoutPrefix == TEXT("ASTC_RGB")) mutableFormat = mu::EImageFormat::IF_ASTC_4x4_RGB_LDR;
+								else if (FormatWithoutPrefix == TEXT("ASTC_RGBA")) mutableFormat = mu::EImageFormat::IF_ASTC_4x4_RGBA_LDR;
+								else if (FormatWithoutPrefix == TEXT("ASTC_NormalRG")) mutableFormat = mu::EImageFormat::IF_ASTC_4x4_RG_LDR;
+								else if (FormatWithoutPrefix == TEXT("G8")) mutableFormat = mu::EImageFormat::IF_L_UBYTE;
+								else if (FormatWithoutPrefix == TEXT("BGRA8")) mutableFormat = mu::EImageFormat::IF_RGBA_UBYTE;
 								else
 								{
-									mu::EImageFormat mutableFormat = mu::EImageFormat::IF_RGBA_UBYTE;
-
-									if (FormatWithoutPrefix == TEXT("DXT1")) mutableFormat = mu::EImageFormat::IF_BC1;
-									else if (FormatWithoutPrefix == TEXT("DXT3")) mutableFormat = mu::EImageFormat::IF_BC2;
-									else if (FormatWithoutPrefix == TEXT("DXT5")) mutableFormat = mu::EImageFormat::IF_BC3;
-									else if (FormatWithoutPrefix == TEXT("BC1")) mutableFormat = mu::EImageFormat::IF_BC1;
-									else if (FormatWithoutPrefix == TEXT("BC2")) mutableFormat = mu::EImageFormat::IF_BC2;
-									else if (FormatWithoutPrefix == TEXT("BC3")) mutableFormat = mu::EImageFormat::IF_BC3;
-									else if (FormatWithoutPrefix == TEXT("BC4")) mutableFormat = mu::EImageFormat::IF_BC4;
-									else if (FormatWithoutPrefix == TEXT("BC5")) mutableFormat = mu::EImageFormat::IF_BC5;
-									else if (FormatWithoutPrefix == TEXT("ASTC_RGB")) mutableFormat = mu::EImageFormat::IF_ASTC_4x4_RGB_LDR;
-									else if (FormatWithoutPrefix == TEXT("ASTC_RGBA")) mutableFormat = mu::EImageFormat::IF_ASTC_4x4_RGBA_LDR;
-									else if (FormatWithoutPrefix == TEXT("ASTC_NormalRG")) mutableFormat = mu::EImageFormat::IF_ASTC_4x4_RG_LDR;
-									else if (FormatWithoutPrefix == TEXT("G8")) mutableFormat = mu::EImageFormat::IF_L_UBYTE;
-									else if (FormatWithoutPrefix == TEXT("BGRA8")) mutableFormat = mu::EImageFormat::IF_RGBA_UBYTE;
-									else
-									{
-										UE_LOG(LogMutable, Warning, TEXT("Unexpected image format [%s]."), *FormatWithoutPrefix);
-									}
-
-									FormatImage->SetFormat(mutableFormat);
+									UE_LOG(LogMutable, Warning, TEXT("Unexpected image format [%s]."), *FormatWithoutPrefix);
 								}
-							}
 
+								FormatImage->SetFormat(mutableFormat, mutableFormatIfAlpha);
+							}
 						}
 
 						ImageNode = LastImage;
