[mutable] Fix bug with displacement image mips.
Extended debugger code tree with more information for certain instructions.
#rnx

diff --git a/Engine/Plugins/Experimental/Mutable/Source/CustomizableObject/Classes/MuCO/CustomizableObject.h b/Engine/Plugins/Experimental/Mutable/Source/CustomizableObject/Classes/MuCO/CustomizableObject.h
--- a/Engine/Plugins/Experimental/Mutable/Source/CustomizableObject/Classes/MuCO/CustomizableObject.h
+++ b/Engine/Plugins/Experimental/Mutable/Source/CustomizableObject/Classes/MuCO/CustomizableObject.h
@@ -1217,7 +1217,7 @@ private:
 	// This is a manual version number for the binary blobs in this asset.
 	// Increasing it invalidates all the previously compiled models.
 	// Warning: If while merging code both versions have changed, take the highest+1.
-	static const int32 CurrentSupportedVersion = 354;
+	static const int32 CurrentSupportedVersion = 355;
 
 public:
 
diff --git a/Engine/Plugins/Experimental/Mutable/Source/CustomizableObjectEditor/Private/MuCOE/SMutableCodeViewer.cpp b/Engine/Plugins/Experimental/Mutable/Source/CustomizableObjectEditor/Private/MuCOE/SMutableCodeViewer.cpp
--- a/Engine/Plugins/Experimental/Mutable/Source/CustomizableObjectEditor/Private/MuCOE/SMutableCodeViewer.cpp
+++ b/Engine/Plugins/Experimental/Mutable/Source/CustomizableObjectEditor/Private/MuCOE/SMutableCodeViewer.cpp
@@ -55,6 +55,7 @@
 #include "MuT/ErrorLogPrivate.h"
 #include "MuT/Streams.h"
 #include "MuT/Table.h"
+#include "MuT/TypeInfo.h"
 #include "SlotBase.h"
 #include "Styling/AppStyle.h"
 #include "Styling/CoreStyle.h"
@@ -128,10 +129,62 @@ public:
 		{
 			const mu::FProgram& Program = RowItem->MutableModel->GetPrivate()->m_program;
 			const mu::OP_TYPE Type = Program.GetOpType(RowItem->MutableOperation);
-			const TCHAR* OpName = mu::s_opNames[int32(Type)];
+			FString OpName = mu::s_opNames[int32(Type)];
+			OpName.TrimEndInline();
+
+			// See if the operation type accepts additional information in the label
+			switch ( Type )
+			{
+			case mu::OP_TYPE::IM_SWIZZLE:
+			{
+				mu::OP::ImageSwizzleArgs Args = Program.GetOpArgs<mu::OP::ImageSwizzleArgs>(RowItem->MutableOperation);
+				OpName += TEXT(" ");
+				OpName += ANSI_TO_TCHAR(mu::TypeInfo::s_imageFormatName[int32(Args.format)]);
+				break;
+			}
+
+			case mu::OP_TYPE::IM_PIXELFORMAT:
+			{
+				mu::OP::ImagePixelFormatArgs Args = Program.GetOpArgs<mu::OP::ImagePixelFormatArgs>(RowItem->MutableOperation);
+				OpName += TEXT(" ");
+				OpName += ANSI_TO_TCHAR(mu::TypeInfo::s_imageFormatName[int32(Args.format)]);
+				OpName += TEXT(" or ");
+				OpName += ANSI_TO_TCHAR(mu::TypeInfo::s_imageFormatName[int32(Args.formatIfAlpha)]);
+				break;
+			}
+
+			case mu::OP_TYPE::IM_MULTILAYER:
+			{
+				mu::OP::ImageMultiLayerArgs Args = Program.GetOpArgs<mu::OP::ImageMultiLayerArgs>(RowItem->MutableOperation);
+				OpName += TEXT(" ");
+				OpName += mu::TypeInfo::s_blendModeName[int32(Args.blendType)];
+				OpName += FString::Printf(TEXT(" range-id %d"), Args.rangeId);
+				break;
+			}
+
+			case mu::OP_TYPE::IM_LAYER:
+			{
+				mu::OP::ImageLayerArgs Args = Program.GetOpArgs<mu::OP::ImageLayerArgs>(RowItem->MutableOperation);
+				OpName += TEXT(" ");
+				OpName += mu::TypeInfo::s_blendModeName[int32(Args.blendType)];
+				OpName += FString::Printf(TEXT(" flags %d"),Args.flags);
+				break;
+			}
+
+			case mu::OP_TYPE::IM_LAYERCOLOUR:
+			{
+				mu::OP::ImageLayerColourArgs Args = Program.GetOpArgs<mu::OP::ImageLayerColourArgs>(RowItem->MutableOperation);
+				OpName += TEXT(" ");
+				OpName += mu::TypeInfo::s_blendModeName[int32(Args.blendType)];
+				break;
+			}
+
+			default:
+				break;
+			}
 
 			// Prepare the text shown on the UI side of the operation tree
-			FString MainLabel = FString::Printf(TEXT("%s (%d) : %s"), *RowItem->Caption, int32(RowItem->MutableOperation), OpName);
+			FString MainLabel = FString::Printf(TEXT("%s (%d) : %s"), *RowItem->Caption, int32(RowItem->MutableOperation), *OpName);
 
 #if UE_BUILD_DEBUG
 			FString IndexOnTree = FString::FromInt(RowItem->IndexOnTree);
@@ -1189,6 +1242,54 @@ void SMutableCodeViewer::GenerateElementRecursive(mu::OP::ADDRESS InParentAddres
 		break;
 	}
 
+	case mu::OP_TYPE::IM_SWIZZLE:
+	{
+		mu::OP::ImageSwizzleArgs Args = InProgram.GetOpArgs<mu::OP::ImageSwizzleArgs>(InParentAddress);
+		for (int32 Channel = 0; Channel < 4; ++Channel)
+		{
+			FString Caption = FString::Printf(TEXT("%d is %d from "), Channel, Args.sourceChannels[Channel]);
+			AddOpFunc(Args.sources[Channel], Caption);
+		}
+		break;
+	}
+
+	case mu::OP_TYPE::IM_LAYER:
+	{
+		mu::OP::ImageLayerArgs Args = InProgram.GetOpArgs<mu::OP::ImageLayerArgs>(InParentAddress);
+		AddOpFunc(Args.base, TEXT("base "));
+		if (Args.mask)
+		{
+			AddOpFunc(Args.mask, TEXT("mask "));
+		}
+		AddOpFunc(Args.blended, TEXT("blended "));
+		break;
+	}
+
+	case mu::OP_TYPE::IM_LAYERCOLOUR:
+	{
+		mu::OP::ImageLayerColourArgs Args = InProgram.GetOpArgs<mu::OP::ImageLayerColourArgs>(InParentAddress);
+		AddOpFunc(Args.base, TEXT("base "));
+		if (Args.mask)
+		{
+			AddOpFunc(Args.mask, TEXT("mask "));
+		}
+		AddOpFunc(Args.colour, TEXT("colour "));
+		break;
+	}
+
+	case mu::OP_TYPE::IM_MULTILAYER:
+	{
+		mu::OP::ImageMultiLayerArgs Args = InProgram.GetOpArgs<mu::OP::ImageMultiLayerArgs>(InParentAddress);
+		AddOpFunc(Args.rangeSize, TEXT("range "));
+		AddOpFunc(Args.base, TEXT("base "));
+		if (Args.mask)
+		{
+			AddOpFunc(Args.mask, TEXT("mask "));
+		}
+		AddOpFunc(Args.blended, TEXT("blended "));
+		break;
+	}
+
 	default:
 	{
 		// Generic list of child operations
diff --git a/Engine/Plugins/Experimental/Mutable/Source/MutableRuntime/Private/MuR/CodeRunner.cpp b/Engine/Plugins/Experimental/Mutable/Source/MutableRuntime/Private/MuR/CodeRunner.cpp
--- a/Engine/Plugins/Experimental/Mutable/Source/MutableRuntime/Private/MuR/CodeRunner.cpp
+++ b/Engine/Plugins/Experimental/Mutable/Source/MutableRuntime/Private/MuR/CodeRunner.cpp
@@ -3620,6 +3620,13 @@ namespace mu
                 Ptr<const Image> pA = GetMemory().GetImage( CACHE_ADDRESS(args.base,item) );
                 Ptr<const Image> pB = GetMemory().GetImage( CACHE_ADDRESS(args.patch,item) );
 
+				// Failsafe
+				if (!pA || !pB)
+				{
+					GetMemory().SetImage(item, pA);
+					break;
+				}
+
                 box<vec2<int>> rect;
                 rect.min[0] = args.minX;
                 rect.min[1] = args.minY;
@@ -3665,14 +3672,11 @@ namespace mu
 
 						if (pResult->GetFormat() != format)
 						{
-							pResult =
-								ImagePixelFormat(m_pSettings->GetPrivate()->m_imageCompressionQuality,
-									pResult.get(), format);
+							pResult = ImagePixelFormat(m_pSettings->GetPrivate()->m_imageCompressionQuality, pResult.get(), format);
 						}
 						if (pB->GetFormat() != format)
 						{
-							pB = ImagePixelFormat(m_pSettings->GetPrivate()->m_imageCompressionQuality,
-								pB.get(), format);
+							pB = ImagePixelFormat(m_pSettings->GetPrivate()->m_imageCompressionQuality, pB.get(), format);
 						}
 					}
 
diff --git a/Engine/Plugins/Experimental/Mutable/Source/MutableRuntime/Private/MuR/CodeRunnerTasked.cpp b/Engine/Plugins/Experimental/Mutable/Source/MutableRuntime/Private/MuR/CodeRunnerTasked.cpp
--- a/Engine/Plugins/Experimental/Mutable/Source/MutableRuntime/Private/MuR/CodeRunnerTasked.cpp
+++ b/Engine/Plugins/Experimental/Mutable/Source/MutableRuntime/Private/MuR/CodeRunnerTasked.cpp
@@ -418,6 +418,12 @@ namespace mu
 
 		// Warning: This runs in a worker thread.
 
+		if (!m_base)
+		{
+			// This shouldn't really happen if data is correct.
+			return;
+		}
+
 		// TODO: a bug may cause this (issue T97)
 		EImageFormat finalFormat = EImageFormat::IF_NONE;
 		if (IsCompressedFormat(m_base->GetFormat()))
diff --git a/Engine/Plugins/Experimental/Mutable/Source/MutableRuntime/Private/MuR/Image.cpp b/Engine/Plugins/Experimental/Mutable/Source/MutableRuntime/Private/MuR/Image.cpp
--- a/Engine/Plugins/Experimental/Mutable/Source/MutableRuntime/Private/MuR/Image.cpp
+++ b/Engine/Plugins/Experimental/Mutable/Source/MutableRuntime/Private/MuR/Image.cpp
@@ -92,17 +92,17 @@ namespace mu
 
 
     //---------------------------------------------------------------------------------------------
-    ImagePtr Image::StaticUnserialise( InputArchive& arch )
+    Ptr<Image> Image::StaticUnserialise( InputArchive& arch )
     {
 		LLM_SCOPE_BYNAME(TEXT("MutableRuntime"));
-        ImagePtr pResult = new Image();
+        Ptr<Image> pResult = new Image();
         arch >> *pResult;
         return pResult;
     }
 
 
 	//---------------------------------------------------------------------------------------------
-	ImagePtr Image::ExtractMip(int32 Mip) const
+	Ptr<Image> Image::ExtractMip(int32 Mip) const
 	{
 		if (Mip == 0 && m_lods == 1)
 		{
@@ -120,7 +120,8 @@ namespace mu
 
 			if (DataSize)
 			{
-				ImagePtr pResult = new Image(MipSize[0], MipSize[1], 1, m_format );
+				Ptr<Image> pResult = new Image(MipSize[0], MipSize[1], 1, m_format );
+				pResult->m_flags = m_flags;
 				uint8* DestData = pResult->GetData();
 				FMemory::Memcpy(DestData, SourceData, DataSize);
 				return pResult;
@@ -128,9 +129,10 @@ namespace mu
 			else
 			{
 				EImageFormat uncompressedFormat = GetUncompressedFormat(m_format);
-				ImagePtr pTemp = ImagePixelFormat(Quality, this, uncompressedFormat);
+				Ptr<Image> pTemp = ImagePixelFormat(Quality, this, uncompressedFormat);
 				pTemp = pTemp->ExtractMip(Mip);
-				ImagePtr pResult = ImagePixelFormat(Quality, pTemp.get(), m_format);
+				Ptr<Image> pResult = ImagePixelFormat(Quality, pTemp.get(), m_format);
+				pResult->m_flags = m_flags;
 				return pResult;
 			}
 		}
diff --git a/Engine/Plugins/Experimental/Mutable/Source/MutableRuntime/Private/MuR/ModelPrivate.h b/Engine/Plugins/Experimental/Mutable/Source/MutableRuntime/Private/MuR/ModelPrivate.h
--- a/Engine/Plugins/Experimental/Mutable/Source/MutableRuntime/Private/MuR/ModelPrivate.h
+++ b/Engine/Plugins/Experimental/Mutable/Source/MutableRuntime/Private/MuR/ModelPrivate.h
@@ -376,7 +376,15 @@ namespace mu
 				// Generate next mip if necessary
 				if (Mip + 1 < MipsToStore)
 				{
-					pMip = pMip->ExtractMip(1);
+					if (Mip > pImage->GetLODCount())
+					{
+						// Generate from the last mip.
+						pMip = pMip->ExtractMip(1);
+					}
+					else
+					{
+						pMip = pImage->ExtractMip(Mip+1);
+					}
 					check(pMip);
 				}
 			}
diff --git a/Engine/Plugins/Experimental/Mutable/Source/MutableRuntime/Private/MuR/OpImageDisplace.h b/Engine/Plugins/Experimental/Mutable/Source/MutableRuntime/Private/MuR/OpImageDisplace.h
--- a/Engine/Plugins/Experimental/Mutable/Source/MutableRuntime/Private/MuR/OpImageDisplace.h
+++ b/Engine/Plugins/Experimental/Mutable/Source/MutableRuntime/Private/MuR/OpImageDisplace.h
@@ -45,12 +45,12 @@ namespace mu
 		int32 SizeY = pResult->GetSizeY();
 		check( SizeX>0 && SizeY>0 );
 
-        if (SizeX<4 || SizeY<4)
-        {
-            return;
-        }
+        //if (SizeX<4 || SizeY<4)
+        //{
+        //    return;
+        //}
 
-		FMemory::Memset( pResult->GetData(), MutableEncodeOffset( 0, 0 ), SizeX*SizeY );
+		FMemory::Memset( pResult->GetData(), MutableEncodeOffset( 0, 0 ), pResult->CalculatePixelCount() );
 
 		Ptr<const Image> pThisMask = pMask;
 		for ( int b=0; b<InBorder; ++b )
@@ -65,7 +65,7 @@ namespace mu
 				const uint8* ThisMaskData = pThisMask->GetMipData(CurrentMip);
 				uint8* NextMaskData = pNextMask->GetMipData(CurrentMip);
 
-				//for ( int y=0; y<sizeY; ++y )
+				//for ( int y=0; y< MipSize.Y; ++y )
 				const auto ProcessRow = [
 					ThisMaskData, NextMaskData, MipSize, ResultData
 				] (int32 y)
@@ -151,10 +151,10 @@ namespace mu
 		int32 sizeX = pResult->GetSizeX();
 		int32 sizeY = pResult->GetSizeY();
 
-        if (sizeX<4 || sizeY<4)
-        {
-            return;
-        }
+        //if (sizeX<4 || sizeY<4)
+        //{
+        //    return;
+        //}
 
         const uint8* pSourceData = pSource->GetData();
 
@@ -343,3 +343,5 @@ namespace mu
 
 
 }
+
+
diff --git a/Engine/Plugins/Experimental/Mutable/Source/MutableRuntime/Private/MuR/Operations.h b/Engine/Plugins/Experimental/Mutable/Source/MutableRuntime/Private/MuR/Operations.h
--- a/Engine/Plugins/Experimental/Mutable/Source/MutableRuntime/Private/MuR/Operations.h
+++ b/Engine/Plugins/Experimental/Mutable/Source/MutableRuntime/Private/MuR/Operations.h
@@ -554,7 +554,6 @@ namespace mu
             ADDRESS base;
             ADDRESS mask;
             ADDRESS colour;
-            EImageFormat reformat;
 			uint8 blendType;	// One of BLEND_TYPE
         };
 
diff --git a/Engine/Plugins/Experimental/Mutable/Source/MutableRuntime/Public/MuR/Image.h b/Engine/Plugins/Experimental/Mutable/Source/MutableRuntime/Public/MuR/Image.h
--- a/Engine/Plugins/Experimental/Mutable/Source/MutableRuntime/Public/MuR/Image.h
+++ b/Engine/Plugins/Experimental/Mutable/Source/MutableRuntime/Public/MuR/Image.h
@@ -243,7 +243,7 @@ namespace mu
 			IF_IS_PLAIN_COLOUR = 1 << 1,
 
 			// If this is set, the image shouldn't be scaled: it's contents is resoultion-dependent.
-			IF_CANNOT_BE_SCALED = 2 << 1
+			IF_CANNOT_BE_SCALED = 1 << 2
 		} EImageFlags;
 
 		//! Persistent flags with some image properties. The meaning will depend of every
diff --git a/Engine/Plugins/Experimental/Mutable/Source/MutableRuntime/Public/MuR/Model.h b/Engine/Plugins/Experimental/Mutable/Source/MutableRuntime/Public/MuR/Model.h
--- a/Engine/Plugins/Experimental/Mutable/Source/MutableRuntime/Public/MuR/Model.h
+++ b/Engine/Plugins/Experimental/Mutable/Source/MutableRuntime/Public/MuR/Model.h
@@ -12,7 +12,7 @@
 //! data structures. Compiled models are not necessarily compatible when the runtime is updated,
 //! so this version number can be used externally to verify this. It is not used internally, and
 //! serializing models from different versions than this runtime will probably result in a crash.
-#define MUTABLE_COMPILED_MODEL_CODE_VERSION		uint32( 49 )
+#define MUTABLE_COMPILED_MODEL_CODE_VERSION		uint32( 50 )
 #define MUTABLE_PARAMETERS_VERSION              uint32( 1 )
 
 
diff --git a/Engine/Plugins/Experimental/Mutable/Source/MutableTools/Private/MuT/TypeInfo.cpp b/Engine/Plugins/Experimental/Mutable/Source/MutableTools/Private/MuT/TypeInfo.cpp
--- a/Engine/Plugins/Experimental/Mutable/Source/MutableTools/Private/MuT/TypeInfo.cpp
+++ b/Engine/Plugins/Experimental/Mutable/Source/MutableTools/Private/MuT/TypeInfo.cpp
@@ -130,4 +130,18 @@ namespace mu
 		"None",
 	};
 
+	const char* TypeInfo::s_blendModeName[uint32(EBlendType::_BT_COUNT)] =
+	{
+		"SoftLight",
+		"HardLight",
+		"Burn",
+		"Dodge",
+		"Screen",
+		"Overlay",
+		"Blend",
+		"Multiply",
+		"AlphaOverlay", // TODO: This name is not descriptive.
+		"NormalCombine",
+	};
+
 }
\ No newline at end of file
diff --git a/Engine/Plugins/Experimental/Mutable/Source/MutableTools/Public/MuT/TypeInfo.h b/Engine/Plugins/Experimental/Mutable/Source/MutableTools/Public/MuT/TypeInfo.h
--- a/Engine/Plugins/Experimental/Mutable/Source/MutableTools/Public/MuT/TypeInfo.h
+++ b/Engine/Plugins/Experimental/Mutable/Source/MutableTools/Public/MuT/TypeInfo.h
@@ -24,13 +24,15 @@ namespace mu
 
 		static const char* s_meshBufferFormatName[ MBF_COUNT ];
 
+		static const char* s_blendModeName[uint32(EBlendType::_BT_COUNT) ];
+
 		static const char* s_projectorTypeName [ static_cast<uint32>(mu::PROJECTOR_TYPE::COUNT) ];
 
-		static const char* s_curveInterpolationModeName [ static_cast<uint8_t>(mu::CurveKeyFrame::InterpMode::Count) ];
+		static const char* s_curveInterpolationModeName [ static_cast<uint32>(mu::CurveKeyFrame::InterpMode::Count) ];
 
-		static const char* s_curveTangentModeName [ static_cast<uint8_t> (mu::CurveKeyFrame::TangentMode::Count) ];
+		static const char* s_curveTangentModeName [ static_cast<uint32> (mu::CurveKeyFrame::TangentMode::Count) ];
 
-		static const char* s_curveTangentWeightModeName [ static_cast<uint8_t> (mu::CurveKeyFrame::TangentWeightMode::Count) ];
+		static const char* s_curveTangentWeightModeName [ static_cast<uint32> (mu::CurveKeyFrame::TangentWeightMode::Count) ];
 	};
 
 }
