[Mutable] Removed unwanted merged changes from another CL.

diff --git a/Engine/Plugins/Experimental/Mutable/Source/CustomizableObject/Classes/MuCO/CustomizableObjectInstance.h b/Engine/Plugins/Experimental/Mutable/Source/CustomizableObject/Classes/MuCO/CustomizableObjectInstance.h
--- a/Engine/Plugins/Experimental/Mutable/Source/CustomizableObject/Classes/MuCO/CustomizableObjectInstance.h
+++ b/Engine/Plugins/Experimental/Mutable/Source/CustomizableObject/Classes/MuCO/CustomizableObjectInstance.h
@@ -601,7 +601,7 @@ public:
 	void SetIsPlayerOrNearIt(bool NewValue);
 	float GetMinSquareDistToPlayer() const;
 	void SetMinSquareDistToPlayer(float NewValue);
-	void SetMinMaxLODToLoad(int32 NewMinLOD = 0, int32 NewMaxLOD = INT32_MAX, bool LimitLODUpgrades = true);
+	void SetMinMaxLODToLoad(int32 NewMinLOD, int32 NewMaxLOD, bool LimitLODUpgrades = true);
 	int32 GetMinLODToLoad() const;
 	int32 GetMaxLODToLoad() const;
 	int32 GetNumLODsAvailable() const;
diff --git a/Engine/Plugins/Experimental/Mutable/Source/CustomizableObject/Private/MuCO/CustomizableObjectInstance.cpp b/Engine/Plugins/Experimental/Mutable/Source/CustomizableObject/Private/MuCO/CustomizableObjectInstance.cpp
--- a/Engine/Plugins/Experimental/Mutable/Source/CustomizableObject/Private/MuCO/CustomizableObjectInstance.cpp
+++ b/Engine/Plugins/Experimental/Mutable/Source/CustomizableObject/Private/MuCO/CustomizableObjectInstance.cpp
@@ -259,11 +259,6 @@ void UCustomizableInstancePrivateData::SetMinMaxLODToLoad(UCustomizableObjectIns
 	
 	Public->Descriptor.MinLODToLoad = NewMinLOD;
 	Public->Descriptor.MaxLODToLoad = NewMaxLOD;
-
-	if (const FMutableQueueElem* QueueElem = UCustomizableObjectSystem::GetInstance()->GetPrivate()->MutableOperationQueue.Get(Public))
-	{
-		QueueElem->Operation->InstanceDescriptorRuntimeHash = Public->GetUpdateDescriptorRuntimeHash();
-	}
 }
 
 
@@ -549,7 +544,6 @@ void UCustomizableObjectInstance::PostLoad()
 	Super::PostLoad();
 
 	Descriptor.ReloadParameters();
-	SetMinMaxLODToLoad();
 }
 
 
@@ -686,7 +680,6 @@ void UCustomizableObjectInstance::SetObject(UCustomizableObject* InObject)
 {
 	Descriptor.SetCustomizableObject(*InObject);
 	PrivateData->ReloadParameters(this);
-	SetMinMaxLODToLoad();
 }
 
 
