#rnx

diff --git a/Engine/Plugins/Experimental/Mutable/Source/CustomizableObject/Private/MuCO/CustomizableObjectInstance.cpp b/Engine/Plugins/Experimental/Mutable/Source/CustomizableObject/Private/MuCO/CustomizableObjectInstance.cpp
--- a/Engine/Plugins/Experimental/Mutable/Source/CustomizableObject/Private/MuCO/CustomizableObjectInstance.cpp
+++ b/Engine/Plugins/Experimental/Mutable/Source/CustomizableObject/Private/MuCO/CustomizableObjectInstance.cpp
@@ -1637,7 +1637,7 @@ bool UCustomizableInstancePrivateData::UpdateSkeletalMesh_PostBeginUpdate0(UCust
 						{
 							check(Tag.IsNumeric());
 							const int32 MutableSocketIndex = FCString::Atoi(*Tag);
-							check(CustomizableObject->SocketArray.IsValidIndex(MutableSocketIndex));
+							//check(CustomizableObject->SocketArray.IsValidIndex(MutableSocketIndex));
 
 							if (CustomizableObject->SocketArray.IsValidIndex(MutableSocketIndex))
 							{
