[mutable] Fix bug preventing the removal of layout blocks.
#rnx

diff --git a/Engine/Plugins/Experimental/Mutable/Source/MutableRuntime/Private/MuR/CodeRunner.cpp b/Engine/Plugins/Experimental/Mutable/Source/MutableRuntime/Private/MuR/CodeRunner.cpp
--- a/Engine/Plugins/Experimental/Mutable/Source/MutableRuntime/Private/MuR/CodeRunner.cpp
+++ b/Engine/Plugins/Experimental/Mutable/Source/MutableRuntime/Private/MuR/CodeRunner.cpp
@@ -4994,12 +4994,7 @@ namespace mu
 			{
 				Ptr<const Mesh> Mesh = GetMemory().GetMesh(CACHE_ADDRESS(args.Mesh, item));
 
-				Ptr<const Layout> Result;
-
-				if (Mesh && Mesh->GetLayoutCount()>int32(args.LayoutIndex))
-				{
-					Result = Mesh->GetLayout( args.LayoutIndex );
-				}
+				Ptr<const Layout> Result = LayoutFromMesh_RemoveBlocks(Mesh.get(), args.LayoutIndex);
 
 				GetMemory().SetLayout(item, Result);
 				break;
diff --git a/Engine/Plugins/Experimental/Mutable/Source/MutableRuntime/Private/MuR/OpLayoutRemoveBlocks.h b/Engine/Plugins/Experimental/Mutable/Source/MutableRuntime/Private/MuR/OpLayoutRemoveBlocks.h
--- a/Engine/Plugins/Experimental/Mutable/Source/MutableRuntime/Private/MuR/OpLayoutRemoveBlocks.h
+++ b/Engine/Plugins/Experimental/Mutable/Source/MutableRuntime/Private/MuR/OpLayoutRemoveBlocks.h
@@ -11,17 +11,24 @@ namespace mu
 {
 
 	//---------------------------------------------------------------------------------------------
-	inline Ptr<Layout> LayoutRemoveBlocks_Deprecated(const Layout* pSource, const Mesh* pMesh, int layoutIndex)
+	inline Ptr<Layout> LayoutFromMesh_RemoveBlocks(const Mesh* InMesh, int32 InLayoutIndex)
 	{
+		if (!InMesh || InMesh->GetLayoutCount() <= InLayoutIndex)
+		{
+			return nullptr;
+		}
+
+		const Layout* Source = InMesh->GetLayout(InLayoutIndex);
+
 		// Create the list of blocks in the mesh
 		TArray<bool> blocksFound;
 		blocksFound.SetNumZeroed(1024);
 
-		UntypedMeshBufferIteratorConst itBlocks(pMesh->GetVertexBuffers(), MBS_LAYOUTBLOCK, layoutIndex);
+		UntypedMeshBufferIteratorConst itBlocks(InMesh->GetVertexBuffers(), MBS_LAYOUTBLOCK, InLayoutIndex);
 		if (itBlocks.GetFormat() == MBF_UINT16)
 		{
 			const uint16* pBlocks = reinterpret_cast<const uint16*>(itBlocks.ptr());
-			for (int i = 0; i < pMesh->GetVertexCount(); ++i)
+			for (int32 i = 0; i < InMesh->GetVertexCount(); ++i)
 			{
 				if (pBlocks[i] >= blocksFound.Num())
 				{
@@ -35,7 +42,7 @@ namespace mu
 		{
 			// This seems to happen.
 			// May this happen when entire meshes are removed?
-			return pSource->Clone();
+			return Source->Clone();
 		}
 		else
 		{
@@ -44,8 +51,8 @@ namespace mu
 		}
 
 		// Remove blocks that are not in the mesh
-		Ptr<Layout> pResult = pSource->Clone();
-		int dest = 0;
+		Ptr<Layout> pResult = Source->Clone();
+		int32 dest = 0;
 		for (int32 b = 0; b < pResult->m_blocks.Num(); ++b)
 		{
 			int blockIndex = pResult->m_blocks[b].m_id;
