#rnx

diff --git a/Engine/Plugins/Experimental/Mutable/Source/CustomizableObjectEditor/Private/MuCOE/GenerateMutableSource/GenerateMutableSource.h b/Engine/Plugins/Experimental/Mutable/Source/CustomizableObjectEditor/Private/MuCOE/GenerateMutableSource/GenerateMutableSource.h
--- a/Engine/Plugins/Experimental/Mutable/Source/CustomizableObjectEditor/Private/MuCOE/GenerateMutableSource/GenerateMutableSource.h
+++ b/Engine/Plugins/Experimental/Mutable/Source/CustomizableObjectEditor/Private/MuCOE/GenerateMutableSource/GenerateMutableSource.h
@@ -572,16 +572,22 @@ struct FMutableGraphGenerationContext
 		{
 			/** Source mesh data. */
 			const UObject* Mesh = nullptr;
-			int LOD = 0;
-			int MaterialIndex = 0;
+			int32 LOD = 0;
+			int32 MaterialIndex = 0;
 
 			/** Flag used to generate this mesh. Bit mask of EMutableMeshConversionFlags */
 			uint32 Flags = 0;
 
-			bool operator==( const FKey& k ) const
+			/** Tags added at the UE level that go through the Mutable core and are merged in the generated mesh.
+			 *  Only add the tags that make the mesh unique and require it not to be cached together with the 
+			 *  same exact mesh but with different tags.
+			*/
+			FString Tags;
+
+			bool operator==( const FKey& OtherKey ) const
 			{
-				return Mesh == k.Mesh && LOD == k.LOD && MaterialIndex == k.MaterialIndex 
-					&& Flags == k.Flags;
+				return Mesh == OtherKey.Mesh && LOD == OtherKey.LOD && MaterialIndex == OtherKey.MaterialIndex
+					&& Flags == OtherKey.Flags && Tags == OtherKey.Tags;
 			}
 		};
 
diff --git a/Engine/Plugins/Experimental/Mutable/Source/CustomizableObjectEditor/Private/MuCOE/GenerateMutableSource/GenerateMutableSourceMesh.cpp b/Engine/Plugins/Experimental/Mutable/Source/CustomizableObjectEditor/Private/MuCOE/GenerateMutableSource/GenerateMutableSourceMesh.cpp
--- a/Engine/Plugins/Experimental/Mutable/Source/CustomizableObjectEditor/Private/MuCOE/GenerateMutableSource/GenerateMutableSourceMesh.cpp
+++ b/Engine/Plugins/Experimental/Mutable/Source/CustomizableObjectEditor/Private/MuCOE/GenerateMutableSource/GenerateMutableSourceMesh.cpp
@@ -1819,12 +1819,14 @@ mu::MeshPtr ConvertStaticMeshToMutable(UStaticMesh* StaticMesh, int LOD, int Mat
 }
 
 
-// Convert a mesh constant to a mutable format.
-mu::MeshPtr GenerateMutableMesh(UObject * Mesh, int32 LOD, int32 MaterialIndex, FMutableGraphGenerationContext & GenerationContext, const UCustomizableObjectNode* CurrentNode)
+// Convert a Mesh constant to a mutable format. UniqueTags are the tags that make this Mesh unique that cannot be merged in the cache 
+//  with the exact same Mesh with other tags
+mu::MeshPtr GenerateMutableMesh(UObject * Mesh, int32 LOD, int32 MaterialIndex, const FString& UniqueTags, FMutableGraphGenerationContext & GenerationContext, const UCustomizableObjectNode* CurrentNode)
 {
 	// Get the mesh generation flags to use
 	uint32 CurrentFlags = GenerationContext.MeshGenerationFlags.Last();
-	FMutableGraphGenerationContext::FGeneratedMeshData::FKey Key = { Mesh, LOD, MaterialIndex, CurrentFlags };
+
+	FMutableGraphGenerationContext::FGeneratedMeshData::FKey Key = { Mesh, LOD, MaterialIndex, CurrentFlags, UniqueTags };
 	mu::MeshPtr MutableMesh = GenerationContext.FindGeneratedMesh(Key);
 	
 	if (!MutableMesh)
@@ -1929,7 +1931,7 @@ mu::MeshPtr BuildMorphedMutableMesh(const UEdGraphPin* BaseSourcePin, const FStr
 
 
 		// Get the base mesh
-		mu::MeshPtr BaseSourceMesh = GenerateMutableMesh(SkeletalMesh, LODIndex, SectionIndex, GenerationContext, Node);
+		mu::MeshPtr BaseSourceMesh = GenerateMutableMesh(SkeletalMesh, LODIndex, SectionIndex, FString(), GenerationContext, Node);
 		if (BaseSourceMesh)
 		{
 			// Clone it (it will probably be shared)
@@ -2679,7 +2681,30 @@ mu::NodeMeshPtr GenerateMutableSourceMesh(const UEdGraphPin * Pin,
 				}
 			}
 
-			mu::MeshPtr MutableMesh = GenerateMutableMesh(TypedNodeSkel->SkeletalMesh, LOD, SectionIndex, GenerationContext, TypedNodeSkel);
+			// First process the mesh tags that are going to make the mesh unique and affect whether it's repeated in 
+			// the mesh cache or not
+			FString MeshUniqueTags;
+			FString AnimBPAssetTag;
+
+			if (!TypedNodeSkel->AnimInstance.IsNull())
+			{
+				int32 SlotIndex = TypedNodeSkel->AnimBlueprintSlot;
+				GenerationContext.AnimBPAssetsMap.Add(TypedNodeSkel->AnimInstance.ToString(), TypedNodeSkel->AnimInstance);
+
+				AnimBPAssetTag = GenerateAnimationInstanceTag(TypedNodeSkel->AnimInstance.ToString(), SlotIndex);
+				MeshUniqueTags += AnimBPAssetTag;
+			}
+
+			TArray<FString> ArrayAnimBPTags;
+
+			for (const FGameplayTag& GamePlayTag : TypedNodeSkel->AnimationGameplayTags)
+			{
+				const FString AnimBPTag = GenerateGameplayTag(GamePlayTag.ToString());
+				ArrayAnimBPTags.Add(AnimBPTag);
+				MeshUniqueTags += AnimBPTag;
+			}
+
+			mu::MeshPtr MutableMesh = GenerateMutableMesh(TypedNodeSkel->SkeletalMesh, LOD, SectionIndex, MeshUniqueTags, GenerationContext, TypedNodeSkel);
 			if (MutableMesh)
 			{
 				MeshNode->SetValue(MutableMesh);
@@ -2737,19 +2762,12 @@ mu::NodeMeshPtr GenerateMutableSourceMesh(const UEdGraphPin * Pin,
 
 				if (!TypedNodeSkel->AnimInstance.IsNull())
 				{
-					int32 SlotIndex= TypedNodeSkel->AnimBlueprintSlot;
-					GenerationContext.AnimBPAssetsMap.Add(TypedNodeSkel->AnimInstance.ToString(), TypedNodeSkel->AnimInstance);
-
-					FString AnimBPAssetTag = GenerateAnimationInstanceTag(TypedNodeSkel->AnimInstance.ToString(), SlotIndex);
-
 					AddTagToMutableMeshUnique(*MutableMesh, AnimBPAssetTag);
 				}
 
-				for (const FGameplayTag& GamePlayTag : TypedNodeSkel->AnimationGameplayTags)
+				for (const FString& GamePlayTag : ArrayAnimBPTags)
 				{
-					const FString AnimBPTag = GenerateGameplayTag(GamePlayTag.ToString());
-					
-					AddTagToMutableMeshUnique(*MutableMesh, AnimBPTag);
+					AddTagToMutableMeshUnique(*MutableMesh, GamePlayTag);
 				}
 
 				AddSocketTagsToMesh(TypedNodeSkel->SkeletalMesh, MutableMesh, GenerationContext);
@@ -2868,7 +2886,7 @@ mu::NodeMeshPtr GenerateMutableSourceMesh(const UEdGraphPin * Pin,
 			}
 			check(MaterialIndex < TypedNodeStatic->LODs[LOD].Materials.Num());
 
-			mu::MeshPtr MutableMesh = GenerateMutableMesh(TypedNodeStatic->StaticMesh, LOD, MaterialIndex, GenerationContext, TypedNodeStatic);
+			mu::MeshPtr MutableMesh = GenerateMutableMesh(TypedNodeStatic->StaticMesh, LOD, MaterialIndex, FString(), GenerationContext, TypedNodeStatic);
 			if (MutableMesh)
 			{
 				MeshNode->SetValue(MutableMesh);
diff --git a/Engine/Plugins/Experimental/Mutable/Source/CustomizableObjectEditor/Private/MuCOE/GenerateMutableSource/GenerateMutableSourceMesh.h b/Engine/Plugins/Experimental/Mutable/Source/CustomizableObjectEditor/Private/MuCOE/GenerateMutableSource/GenerateMutableSourceMesh.h
--- a/Engine/Plugins/Experimental/Mutable/Source/CustomizableObjectEditor/Private/MuCOE/GenerateMutableSource/GenerateMutableSourceMesh.h
+++ b/Engine/Plugins/Experimental/Mutable/Source/CustomizableObjectEditor/Private/MuCOE/GenerateMutableSource/GenerateMutableSourceMesh.h
@@ -21,7 +21,7 @@ mu::MeshPtr ConvertSkeletalMeshToMutable(USkeletalMesh* InSkeletalMesh, int LOD,
 mu::MeshPtr ConvertStaticMeshToMutable(UStaticMesh* StaticMesh, int LOD, int MaterialIndex, FMutableGraphGenerationContext& GenerationContext, const UCustomizableObjectNode* CurrentNode);
 
 
-mu::MeshPtr GenerateMutableMesh(UObject * Mesh, int32 LOD, int32 MaterialIndex, FMutableGraphGenerationContext & GenerationContext, const UCustomizableObjectNode* CurrentNode);
+mu::MeshPtr GenerateMutableMesh(UObject * Mesh, int32 LOD, int32 MaterialIndex, const FString& MeshUniqueTags, FMutableGraphGenerationContext & GenerationContext, const UCustomizableObjectNode* CurrentNode);
 
 
 mu::MeshPtr BuildMorphedMutableMesh(const UEdGraphPin * BaseSourcePin, const FString& MorphTargetName, FMutableGraphGenerationContext & GenerationContext, const FName& RowName = "");
diff --git a/Engine/Plugins/Experimental/Mutable/Source/CustomizableObjectEditor/Private/MuCOE/GenerateMutableSource/GenerateMutableSourceTable.cpp b/Engine/Plugins/Experimental/Mutable/Source/CustomizableObjectEditor/Private/MuCOE/GenerateMutableSource/GenerateMutableSourceTable.cpp
--- a/Engine/Plugins/Experimental/Mutable/Source/CustomizableObjectEditor/Private/MuCOE/GenerateMutableSource/GenerateMutableSourceTable.cpp
+++ b/Engine/Plugins/Experimental/Mutable/Source/CustomizableObjectEditor/Private/MuCOE/GenerateMutableSource/GenerateMutableSourceTable.cpp
@@ -219,7 +219,23 @@ void FillTableColumn(const UEdGraphPin* Pin, mu::TablePtr MutableTable, FString
 						CurrentColumn = MutableTable->AddColumn(TCHAR_TO_ANSI(*MutableColumnName), mu::TABLE_COLUMN_TYPE::TCT_MESH);
 					}
 
-					mu::MeshPtr MutableMesh = GenerateMutableMesh(SkeletalMesh, CurrentLOD, MatIndex, GenerationContext, CustomNodeTable);
+					// First process the mesh tags that are going to make the mesh unique and affect whether it's repeated in 
+					// the mesh cache or not
+					FString MeshUniqueTags;
+
+					if (!AnimBPAssetTag.IsEmpty())
+					{
+						MeshUniqueTags += AnimBPAssetTag;
+					}
+
+					TArray<FString> ArrayAnimBPTags;
+
+					for (const FGameplayTag& Tag : GameplayTags)
+					{
+						MeshUniqueTags += GenerateGameplayTag(Tag.ToString());
+					}
+
+					mu::MeshPtr MutableMesh = GenerateMutableMesh(SkeletalMesh, CurrentLOD, MatIndex, MeshUniqueTags, GenerationContext, CustomNodeTable);
 
 					if (MutableMesh)
 					{
@@ -310,7 +326,7 @@ void FillTableColumn(const UEdGraphPin* Pin, mu::TablePtr MutableTable, FString
 						CurrentColumn = MutableTable->AddColumn(TCHAR_TO_ANSI(*MutableColumnName), mu::TABLE_COLUMN_TYPE::TCT_MESH);
 					}
 
-					mu::MeshPtr MutableMesh = GenerateMutableMesh(StaticMesh, CurrentLOD, MatIndex, GenerationContext, CustomNodeTable);
+					mu::MeshPtr MutableMesh = GenerateMutableMesh(StaticMesh, CurrentLOD, MatIndex, FString(), GenerationContext, CustomNodeTable);
 
 					if (MutableMesh)
 					{
