[mutable] Fix compilation warning: unexpected texture format

#rnx

diff --git a/Engine/Plugins/Experimental/Mutable/Source/CustomizableObjectEditor/Private/MuCOE/GenerateMutableSource/GenerateMutableSourceSurface.cpp b/Engine/Plugins/Experimental/Mutable/Source/CustomizableObjectEditor/Private/MuCOE/GenerateMutableSource/GenerateMutableSourceSurface.cpp
--- a/Engine/Plugins/Experimental/Mutable/Source/CustomizableObjectEditor/Private/MuCOE/GenerateMutableSource/GenerateMutableSourceSurface.cpp
+++ b/Engine/Plugins/Experimental/Mutable/Source/CustomizableObjectEditor/Private/MuCOE/GenerateMutableSource/GenerateMutableSourceSurface.cpp
@@ -680,14 +680,13 @@ mu::NodeSurfacePtr GenerateMutableSourceSurface(const UEdGraphPin * Pin, FMutabl
 							{
 								check(PlatformFormats[0].Num() > LayerIndex);
 
-								FString FormatWithoutPrefix = PlatformFormats[0][LayerIndex].ToString();
-								FString LeftSplit, RightSplit;
-								if (FormatWithoutPrefix.Split(TEXT("_"), &LeftSplit, &RightSplit, ESearchCase::IgnoreCase, ESearchDir::FromEnd))
-								{
-									FormatWithoutPrefix = RightSplit;
-								}
+								const FString PlatformFormat = PlatformFormats[0][LayerIndex].ToString();
+
+								// Remove platform prefix
+								FString FormatWithoutPrefix = PlatformFormat;
+								PlatformFormat.Split(TEXT("_"), nullptr, &FormatWithoutPrefix, ESearchCase::IgnoreCase, ESearchDir::FromEnd);
 
-								mu::EImageFormat mutableFormat = mu::EImageFormat::IF_RGBA_UBYTE;
+								mu::EImageFormat mutableFormat = mu::EImageFormat::IF_NONE;
 								mu::EImageFormat mutableFormatIfAlpha = mu::EImageFormat::IF_NONE;
 
 								if (FormatWithoutPrefix == TEXT("AutoDXT"))
@@ -695,11 +694,6 @@ mu::NodeSurfacePtr GenerateMutableSourceSurface(const UEdGraphPin * Pin, FMutabl
 									mutableFormat = mu::EImageFormat::IF_BC1;
 									mutableFormatIfAlpha = mu::EImageFormat::IF_BC3;
 								}
-								else if ((FormatWithoutPrefix == TEXT("AutoASTC")) || (FormatWithoutPrefix == TEXT("ASTC_RGBAuto")))
-								{
-									mutableFormat = mu::EImageFormat::IF_ASTC_4x4_RGB_LDR;
-									mutableFormatIfAlpha = mu::EImageFormat::IF_ASTC_4x4_RGBA_LDR;
-								}
 								else if (FormatWithoutPrefix == TEXT("DXT1")) mutableFormat = mu::EImageFormat::IF_BC1;
 								else if (FormatWithoutPrefix == TEXT("DXT3")) mutableFormat = mu::EImageFormat::IF_BC2;
 								else if (FormatWithoutPrefix == TEXT("DXT5")) mutableFormat = mu::EImageFormat::IF_BC3;
@@ -708,13 +702,25 @@ mu::NodeSurfacePtr GenerateMutableSourceSurface(const UEdGraphPin * Pin, FMutabl
 								else if (FormatWithoutPrefix == TEXT("BC3")) mutableFormat = mu::EImageFormat::IF_BC3;
 								else if (FormatWithoutPrefix == TEXT("BC4")) mutableFormat = mu::EImageFormat::IF_BC4;
 								else if (FormatWithoutPrefix == TEXT("BC5")) mutableFormat = mu::EImageFormat::IF_BC5;
-								else if (FormatWithoutPrefix == TEXT("ASTC_RGB")) mutableFormat = mu::EImageFormat::IF_ASTC_4x4_RGB_LDR;
-								else if (FormatWithoutPrefix == TEXT("ASTC_RGBA")) mutableFormat = mu::EImageFormat::IF_ASTC_4x4_RGBA_LDR;
-								else if (FormatWithoutPrefix == TEXT("ASTC_NormalRG")) mutableFormat = mu::EImageFormat::IF_ASTC_4x4_RG_LDR;
 								else if (FormatWithoutPrefix == TEXT("G8")) mutableFormat = mu::EImageFormat::IF_L_UBYTE;
 								else if (FormatWithoutPrefix == TEXT("BGRA8")) mutableFormat = mu::EImageFormat::IF_RGBA_UBYTE;
-								else
+								else if (PlatformFormat.Contains(TEXT("ASTC")))
 								{
+									if ((FormatWithoutPrefix == TEXT("AutoASTC")) || (FormatWithoutPrefix == TEXT("RGBAuto")))
+									{
+										mutableFormat = mu::EImageFormat::IF_ASTC_4x4_RGB_LDR;
+										mutableFormatIfAlpha = mu::EImageFormat::IF_ASTC_4x4_RGBA_LDR;
+									}
+									else if (FormatWithoutPrefix == TEXT("RGB")) mutableFormat = mu::EImageFormat::IF_ASTC_4x4_RGB_LDR;
+									else if (FormatWithoutPrefix == TEXT("RGBA")) mutableFormat = mu::EImageFormat::IF_ASTC_4x4_RGBA_LDR;
+									else if (FormatWithoutPrefix == TEXT("NormalRG")) mutableFormat = mu::EImageFormat::IF_ASTC_4x4_RG_LDR;
+								}
+
+								if (mutableFormat == mu::EImageFormat::IF_NONE)
+								{
+									// Format not supported by Mutable, use RBGA_UBYTE as default.
+									mutableFormat = mu::EImageFormat::IF_RGBA_UBYTE;
+
 									UE_LOG(LogMutable, Warning, TEXT("Unexpected image format [%s]."), *FormatWithoutPrefix);
 								}
 
