[mutable] Added option to draw bone hierarchy in the CO viewport.
+ Reorganized the view options into a common menu.
+ Added option to eddit the camera speed.

#rnx

diff --git a/Engine/Plugins/Experimental/Mutable/Source/CustomizableObjectEditor/Private/MuCOE/CustomizableObjectEditorViewportClient.cpp b/Engine/Plugins/Experimental/Mutable/Source/CustomizableObjectEditor/Private/MuCOE/CustomizableObjectEditorViewportClient.cpp
--- a/Engine/Plugins/Experimental/Mutable/Source/CustomizableObjectEditor/Private/MuCOE/CustomizableObjectEditorViewportClient.cpp
+++ b/Engine/Plugins/Experimental/Mutable/Source/CustomizableObjectEditor/Private/MuCOE/CustomizableObjectEditorViewportClient.cpp
@@ -97,6 +97,7 @@
 #include "SceneView.h"
 #include "ScopedTransaction.h"
 #include "ShowFlags.h"
+#include "SkeletalDebugRendering.h"
 #include "SlotBase.h"
 #include "StaticMeshResources.h"
 #include "Styling/AppStyle.h"
@@ -502,6 +503,11 @@ FCustomizableObjectEditorViewportClient::FCustomizableObjectEditorViewportClient
 	bActivateOrbitalCamera = true;
 	bSetOrbitalOnPerspectiveMode = true;
 
+	const int32 CameraSpeed = 3;
+	SetCameraSpeedSetting(CameraSpeed);
+
+	bShowBones = false;
+
 	MaterialToDrawInUVs = 0;
 	MaterialToDrawInUVsLOD = 0;
 	MaterialToDrawInUVsIndex = 0;
@@ -902,6 +908,17 @@ void FCustomizableObjectEditorViewportClient::Draw(const FSceneView* View, FPrim
 			}
 		}
 	}
+
+	if (bShowBones)
+	{
+		for (TWeakObjectPtr<UDebugSkelMeshComponent> SkeletalMeshComponent : SkeletalMeshComponents)
+		{
+			if (SkeletalMeshComponent.IsValid())
+			{
+				DrawMeshBones(SkeletalMeshComponent.Get(), PDI);
+			}
+		}
+	}
 }
 
 
@@ -3317,6 +3334,84 @@ void FCustomizableObjectEditorViewportClient::SetCameraMode(bool Value)
 	UpdateCameraSetup();
 }
 
+
+void FCustomizableObjectEditorViewportClient::SetShowBones()
+{
+	bShowBones = !bShowBones;
+}
+
+
+bool FCustomizableObjectEditorViewportClient::IsShowingBones() const
+{
+	return bShowBones;
+}
+
+
+void FCustomizableObjectEditorViewportClient::DrawMeshBones(UDebugSkelMeshComponent* MeshComponent, FPrimitiveDrawInterface* PDI)
+{
+	if (!MeshComponent ||
+		!MeshComponent->GetSkeletalMeshAsset() ||
+		MeshComponent->GetNumDrawTransform() == 0 ||
+		MeshComponent->SkeletonDrawMode == ESkeletonDrawMode::Hidden)
+	{
+		return;
+	}
+
+	TArray<FTransform> WorldTransforms;
+	WorldTransforms.AddUninitialized(MeshComponent->GetNumDrawTransform());
+
+	TArray<FLinearColor> BoneColors;
+	BoneColors.AddUninitialized(MeshComponent->GetNumDrawTransform());
+
+	const FLinearColor BoneColor = GetDefault<UPersonaOptions>()->DefaultBoneColor;
+	const FLinearColor VirtualBoneColor = GetDefault<UPersonaOptions>()->VirtualBoneColor;
+	const TArray<FBoneIndexType>& DrawBoneIndices = MeshComponent->GetDrawBoneIndices();
+	
+	for (int32 Index = 0; Index < DrawBoneIndices.Num(); ++Index)
+	{
+		const int32 BoneIndex = DrawBoneIndices[Index];
+		WorldTransforms[BoneIndex] = MeshComponent->GetDrawTransform(BoneIndex) * MeshComponent->GetComponentTransform();
+		BoneColors[BoneIndex] = BoneColor;
+	}
+
+	// color virtual bones
+	for (int16 VirtualBoneIndex : MeshComponent->GetReferenceSkeleton().GetRequiredVirtualBones())
+	{
+		BoneColors[VirtualBoneIndex] = VirtualBoneColor;
+	}
+
+	constexpr bool bForceDraw = false;
+
+	// don't allow selection if the skeleton draw mode is greyed out
+	//const bool bAddHitProxy = MeshComponent->SkeletonDrawMode != ESkeletonDrawMode::GreyedOut;
+
+	FSkelDebugDrawConfig DrawConfig;
+	DrawConfig.BoneDrawMode = EBoneDrawMode::All;
+	DrawConfig.BoneDrawSize = 1.0f;
+	DrawConfig.bAddHitProxy = false;
+	DrawConfig.bForceDraw = bForceDraw;
+	DrawConfig.DefaultBoneColor = GetMutableDefault<UPersonaOptions>()->DefaultBoneColor;
+	DrawConfig.AffectedBoneColor = GetMutableDefault<UPersonaOptions>()->AffectedBoneColor;
+	DrawConfig.SelectedBoneColor = GetMutableDefault<UPersonaOptions>()->SelectedBoneColor;
+	DrawConfig.ParentOfSelectedBoneColor = GetMutableDefault<UPersonaOptions>()->ParentOfSelectedBoneColor;
+
+	//No user interaction right now
+	TArray<TRefCountPtr<HHitProxy>> HitProxies;
+
+	SkeletalDebugRendering::DrawBones(
+		PDI,
+		MeshComponent->GetComponentLocation(),
+		DrawBoneIndices,
+		MeshComponent->GetReferenceSkeleton(),
+		WorldTransforms,
+		MeshComponent->BonesOfInterest,
+		BoneColors,
+		HitProxies,
+		DrawConfig
+	);
+}
+
+
 /////////////////////////////////////////////////
 // select folder dialog \todo: move to its own file
 //////////////////////////////////////////////////
diff --git a/Engine/Plugins/Experimental/Mutable/Source/CustomizableObjectEditor/Private/MuCOE/CustomizableObjectEditorViewportLODCommands.cpp b/Engine/Plugins/Experimental/Mutable/Source/CustomizableObjectEditor/Private/MuCOE/CustomizableObjectEditorViewportLODCommands.cpp
--- a/Engine/Plugins/Experimental/Mutable/Source/CustomizableObjectEditor/Private/MuCOE/CustomizableObjectEditorViewportLODCommands.cpp
+++ b/Engine/Plugins/Experimental/Mutable/Source/CustomizableObjectEditor/Private/MuCOE/CustomizableObjectEditorViewportLODCommands.cpp
@@ -18,6 +18,7 @@ void FCustomizableObjectEditorViewportLODCommands::RegisterCommands()
 	UI_COMMAND(HighResScreenshot, "High Resolution Screenshot...", "Opens the control panel for high resolution screenshots", EUserInterfaceActionType::Button, FInputChord());
 	UI_COMMAND(OrbitalCamera, "Orbital Camera", "Sets the camera in orbital mode", EUserInterfaceActionType::RadioButton, FInputChord());
 	UI_COMMAND(FreeCamera, "Free Camera", "Sets the camera in free mode", EUserInterfaceActionType::RadioButton, FInputChord());
+	UI_COMMAND(ShowBones, "Show Bones", "Shows the bone hierarchy of the skeletal mesh", EUserInterfaceActionType::ToggleButton, FInputChord());
 }
 
 #undef LOCTEXT_NAMESPACE
diff --git a/Engine/Plugins/Experimental/Mutable/Source/CustomizableObjectEditor/Private/MuCOE/SCustomizableObjectEditorViewport.cpp b/Engine/Plugins/Experimental/Mutable/Source/CustomizableObjectEditor/Private/MuCOE/SCustomizableObjectEditorViewport.cpp
--- a/Engine/Plugins/Experimental/Mutable/Source/CustomizableObjectEditor/Private/MuCOE/SCustomizableObjectEditorViewport.cpp
+++ b/Engine/Plugins/Experimental/Mutable/Source/CustomizableObjectEditor/Private/MuCOE/SCustomizableObjectEditorViewport.cpp
@@ -618,6 +618,14 @@ void SCustomizableObjectEditorViewportTabBody::BindCommands()
 		FIsActionChecked::CreateSP(this, &SCustomizableObjectEditorViewportTabBody::IsCameraModeActive,1)
 	);
 
+	// Bones Mode
+	CommandList.MapAction(
+		ViewportLODMenuCommands.ShowBones,
+		FExecuteAction::CreateSP(this, &SCustomizableObjectEditorViewportTabBody::SetShowBones),
+		FCanExecuteAction(),
+		FIsActionChecked::CreateSP(this, &SCustomizableObjectEditorViewportTabBody::IsShowingBones)
+	);
+
 	const FEditorViewportCommands& ViewportCommands = FEditorViewportCommands::Get();
 
 	// Camera Views
@@ -720,6 +728,30 @@ void SCustomizableObjectEditorViewportTabBody::SetDrawDefaultUVMaterial()
 }
 
 
+void SCustomizableObjectEditorViewportTabBody::SetShowBones()
+{
+	LevelViewportClient->SetShowBones();
+}
+
+
+bool SCustomizableObjectEditorViewportTabBody::IsShowingBones()
+{
+	return LevelViewportClient->IsShowingBones();
+}
+
+
+void SCustomizableObjectEditorViewportTabBody::SetViewportCameraSpeed(const int32 Speed)
+{
+	LevelViewportClient->SetCameraSpeedSetting(Speed);
+}
+
+
+int32 SCustomizableObjectEditorViewportTabBody::GetViewportCameraSpeed()
+{
+	return LevelViewportClient->GetCameraSpeedSetting();
+}
+
+
 TSharedRef<SWidget> SCustomizableObjectEditorViewportTabBody::BuildToolBar()
 {
 	//FToolBarBuilder SaveThumbnailToolbarBuilder(GetToolkitCommands());
diff --git a/Engine/Plugins/Experimental/Mutable/Source/CustomizableObjectEditor/Private/MuCOE/SCustomizableObjectEditorViewport.h b/Engine/Plugins/Experimental/Mutable/Source/CustomizableObjectEditor/Private/MuCOE/SCustomizableObjectEditorViewport.h
--- a/Engine/Plugins/Experimental/Mutable/Source/CustomizableObjectEditor/Private/MuCOE/SCustomizableObjectEditorViewport.h
+++ b/Engine/Plugins/Experimental/Mutable/Source/CustomizableObjectEditor/Private/MuCOE/SCustomizableObjectEditorViewport.h
@@ -264,6 +264,18 @@ public:
 	/** Sets the first material of the instance as the default section to draw in the UVs Overlay (when there is no material selected in the combobox)*/
 	void SetDrawDefaultUVMaterial();
 
+	/** Sets the bones visibility */
+	void SetShowBones();
+
+	/** Returns true if the skeletal mesh bones are visible in the viewport */
+	bool IsShowingBones();
+
+	/** Sets the speed of the viewport camera */
+	void SetViewportCameraSpeed(const int32 Speed);
+
+	/** Returns the current speed of the viewport camera */
+	int32 GetViewportCameraSpeed();
+
 private:
 	/** Determines the visibility of the viewport. */
 	bool IsVisible() const;
diff --git a/Engine/Plugins/Experimental/Mutable/Source/CustomizableObjectEditor/Private/MuCOE/SCustomizableObjectEditorViewportToolBar.cpp b/Engine/Plugins/Experimental/Mutable/Source/CustomizableObjectEditor/Private/MuCOE/SCustomizableObjectEditorViewportToolBar.cpp
--- a/Engine/Plugins/Experimental/Mutable/Source/CustomizableObjectEditor/Private/MuCOE/SCustomizableObjectEditorViewportToolBar.cpp
+++ b/Engine/Plugins/Experimental/Mutable/Source/CustomizableObjectEditor/Private/MuCOE/SCustomizableObjectEditorViewportToolBar.cpp
@@ -59,6 +59,7 @@
 #include "Widgets/Layout/SUniformGridPanel.h"
 #include "Widgets/SBoxPanel.h"
 #include "Widgets/SCompoundWidget.h"
+#include "Widgets/Input/SSlider.h"
 #include "Widgets/Text/STextBlock.h"
 
 class SWidget;
@@ -391,16 +392,16 @@ void SCustomizableObjectEditorViewportToolBar::Construct(const FArguments& InArg
 				.OnGetMenuContent(this, &SCustomizableObjectEditorViewportToolBar::GenerateLODMenu)
 		]
 
-	// Camera Mode menu (Free or Orbital camera)
-	+ SHorizontalBox::Slot()
+	// View Options Menu (Camera options, Bones...)
+	+SHorizontalBox::Slot()
 		.AutoWidth()
 		.Padding(2.0f, 2.0f)
 		[
-			//Camera Mode
+			//Show Bones
 			SNew(SEditorViewportToolbarMenu)
-				.ParentToolBar(SharedThis(this))
-				.Label(this, &SCustomizableObjectEditorViewportToolBar::GetCameraModeMenuLabel)
-				.OnGetMenuContent(this, &SCustomizableObjectEditorViewportToolBar::GenerateCameraModeMenu)
+			.ParentToolBar(SharedThis(this))
+			.Label(LOCTEXT("ViewOptionsMenuLabel","View Options"))
+			.OnGetMenuContent(this, &SCustomizableObjectEditorViewportToolBar::GenerateViewportOptionsMenu)
 		];
 	
 	TSharedRef<SWidget> RTSButtons = GenerateRTSButtons();
@@ -827,7 +828,7 @@ TSharedRef<SWidget> SCustomizableObjectEditorViewportToolBar::GenerateViewMenu()
 //	}
 //	MenuBuilder.EndSection();
 //}
-//
+
 //void SCustomizableObjectEditorViewportToolBar::FillShowOverlayDrawMenu(FMenuBuilder& MenuBuilder) const
 //{
 //	const FAnimViewportShowCommands& Actions = FAnimViewportShowCommands::Get();
@@ -1090,7 +1091,7 @@ void SCustomizableObjectEditorViewportToolBar::GenerateSceneSetupMenu(FMenuBuild
 	//MenuBuilder.AddWidget(DetailsView, FText(), true);
 }
 
-TSharedRef<SWidget> SCustomizableObjectEditorViewportToolBar::GenerateCameraModeMenu() const
+TSharedRef<SWidget> SCustomizableObjectEditorViewportToolBar::GenerateViewportOptionsMenu() const
 {
 	const FCustomizableObjectEditorViewportLODCommands& Actions = FCustomizableObjectEditorViewportLODCommands::Get();
 
@@ -1098,37 +1099,65 @@ TSharedRef<SWidget> SCustomizableObjectEditorViewportToolBar::GenerateCameraMode
 
 	FMenuBuilder ShowMenuBuilder(bInShouldCloseWindowAfterMenuSelection, Viewport.Pin()->GetCommandList());
 	{
-		// LOD Models
-		ShowMenuBuilder.BeginSection("NAME_None");
+		ShowMenuBuilder.AddSubMenu(LOCTEXT("OptionsMenu_CameraOptions", "Camera Mode"),
+			LOCTEXT("OptionsMenu_CameraOptionsTooltip", "Select the camera mode"),
+		FNewMenuDelegate::CreateLambda([this, &Actions](FMenuBuilder& SubMenuBuilder)
 		{
-			ShowMenuBuilder.AddMenuEntry(Actions.OrbitalCamera);
-			ShowMenuBuilder.AddMenuEntry(Actions.FreeCamera);
-		}
-		ShowMenuBuilder.EndSection();
-	}
-
-	return ShowMenuBuilder.MakeWidget();
-}
+			SubMenuBuilder.BeginSection("Camera");
+			{
+				TSharedPtr<SWidget> BoneSizeWidget = SNew(SVerticalBox)
+				+ SVerticalBox::Slot().AutoHeight()
+				.HAlign(HAlign_Left)
+				.Padding(FMargin(20.0f, 5.0f, 0.0f, 0.0f))
+				[
+					SNew(STextBlock)
+					.Text(LOCTEXT("OptionsMenu_CameraOptions_CameraSpeed_Text", "Camera Speed"))
+					.Font(FAppStyle::GetFontStyle(TEXT("MenuItem.Font")))
+				]
 
-FText SCustomizableObjectEditorViewportToolBar::GetCameraModeMenuLabel() const
-{
-	FText Label = FText::FromString("Camera Mode");
-	if (Viewport.IsValid())
-	{
-		bool CameraMode = Viewport.Pin()->GetViewportClient()->IsOrbitalCameraActive();
+				+SVerticalBox::Slot().AutoHeight()
+				.HAlign(HAlign_Left)
+				.Padding(FMargin(20.0f, 0.0f, 0.0f, 0.0f))
+				[
+					SNew(SHorizontalBox)
+					+ SHorizontalBox::Slot().AutoWidth()
+					[
+						SNew(SBox).WidthOverride(100.0f)
+						[
+							SNew(SSlider)
+							.Value_Lambda([=]() {return Viewport.Pin().Get()->GetViewportCameraSpeed(); })
+							.MinValue(1).MaxValue(4)
+							.OnValueChanged_Lambda([=](int32 Value) { Viewport.Pin().Get()->SetViewportCameraSpeed(Value); })
+						]
+					]
+					+ SHorizontalBox::Slot().AutoWidth()
+					[
+						SNew(STextBlock).Text_Lambda([=]() { return FText::AsNumber(Viewport.Pin().Get()->GetViewportCameraSpeed()); })
+					]
+				];
+				
+				SubMenuBuilder.AddMenuEntry(Actions.OrbitalCamera);
+				SubMenuBuilder.AddMenuEntry(Actions.FreeCamera);
+				SubMenuBuilder.AddWidget(BoneSizeWidget.ToSharedRef(), LOCTEXT("OptionMenu_CameraOptions_CameraSpeed", ""));
+			}
+			SubMenuBuilder.EndSection();
+		}));
 
-		if (CameraMode)
-		{
-			Label = FText::FromString("Camera Mode: Orbital");
-		}
-		else
+		ShowMenuBuilder.AddSubMenu(LOCTEXT("OptionsMenu_BoneOptions", "Bones"), LOCTEXT("OptionsMenu_BoneOptionsTooltip", "Show/hide bone hierarchy"),
+		FNewMenuDelegate::CreateLambda([&Actions](FMenuBuilder& SubMenuBuilder)
 		{
-			Label = FText::FromString("Camera Mode: Free");
-		}
+			SubMenuBuilder.BeginSection("Bones");
+			{
+				SubMenuBuilder.AddMenuEntry(Actions.ShowBones);
+			}
+			SubMenuBuilder.EndSection();
+		}));
 	}
-	return Label;
+
+	return ShowMenuBuilder.MakeWidget();
 }
 
+
 //TSharedRef<class IDetailCustomization> SCustomizableObjectEditorViewportToolBar::CustomizePreviewSceneDescription()
 //{
 //	return MakeShareable(new FPreviewSceneDescriptionCustomization(FAssetData(&Viewport.Pin()->GetSkeletonTree()->GetEditableSkeleton()->GetSkeleton()).GetExportTextName(), Viewport.Pin()->GetPreviewScene()->GetPersonaToolkit()));
diff --git a/Engine/Plugins/Experimental/Mutable/Source/CustomizableObjectEditor/Public/MuCOE/CustomizableObjectEditorViewportClient.h b/Engine/Plugins/Experimental/Mutable/Source/CustomizableObjectEditor/Public/MuCOE/CustomizableObjectEditorViewportClient.h
--- a/Engine/Plugins/Experimental/Mutable/Source/CustomizableObjectEditor/Public/MuCOE/CustomizableObjectEditorViewportClient.h
+++ b/Engine/Plugins/Experimental/Mutable/Source/CustomizableObjectEditor/Public/MuCOE/CustomizableObjectEditorViewportClient.h
@@ -418,6 +418,14 @@ public:
 		return FullNameMaterialToDrawInUVs;
 	}
 
+	/** Sets the skeletal mesh bones visibility */
+	void SetShowBones();
+
+	/** Returns true if bones are vibile in viewport */
+	bool IsShowingBones() const;
+
+	/** Draws Mesh Bones in foreground (From: FAnimationViewportClient) */
+	void DrawMeshBones(UDebugSkelMeshComponent* MeshComponent, FPrimitiveDrawInterface* PDI);
 
 private:
 	/** Component for the static/skeletal mesh. */
@@ -519,4 +527,7 @@ private:
 
 	/** bool to return the Camera mode to Orbital when changing the Camera view to Perspective */
 	bool bSetOrbitalOnPerspectiveMode;
+
+	/** Flag to control the bones visibility in the viewport */
+	bool bShowBones;
 };
\ No newline at end of file
diff --git a/Engine/Plugins/Experimental/Mutable/Source/CustomizableObjectEditor/Public/MuCOE/CustomizableObjectEditorViewportLODCommands.h b/Engine/Plugins/Experimental/Mutable/Source/CustomizableObjectEditor/Public/MuCOE/CustomizableObjectEditorViewportLODCommands.h
--- a/Engine/Plugins/Experimental/Mutable/Source/CustomizableObjectEditor/Public/MuCOE/CustomizableObjectEditorViewportLODCommands.h
+++ b/Engine/Plugins/Experimental/Mutable/Source/CustomizableObjectEditor/Public/MuCOE/CustomizableObjectEditorViewportLODCommands.h
@@ -56,6 +56,9 @@ public:
 	/** Sets the camera mode to free camera */
 	TSharedPtr< FUICommandInfo > FreeCamera;
 
+	/** Shows/Hide the bones of the skeletal mesh */
+	TSharedPtr< FUICommandInfo > ShowBones;
+
 public:
 	/** Registers our commands with the binding system */
 	virtual void RegisterCommands() override;
diff --git a/Engine/Plugins/Experimental/Mutable/Source/CustomizableObjectEditor/Public/MuCOE/SCustomizableObjectEditorViewportToolBar.h b/Engine/Plugins/Experimental/Mutable/Source/CustomizableObjectEditor/Public/MuCOE/SCustomizableObjectEditorViewportToolBar.h
--- a/Engine/Plugins/Experimental/Mutable/Source/CustomizableObjectEditor/Public/MuCOE/SCustomizableObjectEditorViewportToolBar.h
+++ b/Engine/Plugins/Experimental/Mutable/Source/CustomizableObjectEditor/Public/MuCOE/SCustomizableObjectEditorViewportToolBar.h
@@ -116,18 +116,11 @@ private:
 	void GenerateSceneSetupMenu(FMenuBuilder& MenuBuilder);
 
 	/**
-	* Generates the toolbar Camera Mode menu content
+	* Generates the toolbar Options Mode menu content
 	*
-	* @return The widget containing the Camera Mode menu content
+	* @return The widget containing the Options menu content
 	*/
-	TSharedRef<SWidget> GenerateCameraModeMenu() const;
-
-	/**
-	* Returns the label for the Camera Mode tool bar menu, which changes depending on the current Camera Mode selection
-	*
-	* @return	Label to use for this Camera Mode label
-	*/
-	FText GetCameraModeMenuLabel() const;
+	TSharedRef<SWidget> GenerateViewportOptionsMenu() const;
 
 
 	/** Customize the details of the scene setup object */
