[mutable] Fix editor crash when updating the COI properties widget due to a missing parameter.

#rnx

diff --git a/Engine/Plugins/Experimental/Mutable/Source/CustomizableObjectEditor/Private/MuCOE/SCustomizableInstanceProperties.cpp b/Engine/Plugins/Experimental/Mutable/Source/CustomizableObjectEditor/Private/MuCOE/SCustomizableInstanceProperties.cpp
--- a/Engine/Plugins/Experimental/Mutable/Source/CustomizableObjectEditor/Private/MuCOE/SCustomizableInstanceProperties.cpp
+++ b/Engine/Plugins/Experimental/Mutable/Source/CustomizableObjectEditor/Private/MuCOE/SCustomizableInstanceProperties.cpp
@@ -1010,68 +1010,82 @@ void SCustomizableInstanceProperties::AddParameter(int32 ParamIndexInObject)
 			else
 			{
 				TArray<FCustomizableObjectProjectorParameterValue>& ProjectorParameters = CustomInstance->GetProjectorParameters();
-				int32 ProjectorParamIndex = CustomInstance->FindProjectorParameterNameIndex(ParamName);
+				const int32 ProjectorParamIndex = CustomInstance->FindProjectorParameterNameIndex(ParamName);
 				check(ProjectorParamIndex < ProjectorParameters.Num());
 
+				// Selected Pose UI
 				FString PoseSwitchEnumParamName = ParamName + FMultilayerProjector::POSE_PARAMETER_POSTFIX;
 				FString PoseSwitchEnumParamNameWithRange = PoseSwitchEnumParamName + FString::Printf(TEXT("__%d"), -1);
-				int32 PoseSwitchEnumParamIndexInObject = CustomizableObject->FindParameter(PoseSwitchEnumParamName);
-				check(PoseSwitchEnumParamIndexInObject >= 0);
+				const int32 PoseSwitchEnumParamIndexInObject = CustomizableObject->FindParameter(PoseSwitchEnumParamName);
 
-				int32 NumPoseValues = CustomizableObject->GetIntParameterNumOptions(PoseSwitchEnumParamIndexInObject);
-
-				//TSharedPtr< TArray< TSharedPtr<FString> > > PoseOptionNames = MakeShareable(new TArray< TSharedPtr<FString> >());
-				//IntOptionNames.Add(PoseOptionNames);
-				
-				TArray<FString> PoseOptionNamesAttribute;
-				FString PoseValue = GetIntParameterValue(PoseSwitchEnumParamName, -1);
-				int32 PoseValueIndex = 0;
-				
-				for (int32 j = 0; j < NumPoseValues; ++j)
+				if (PoseSwitchEnumParamIndexInObject != INDEX_NONE)
 				{
-					FString PossibleValue = CustomizableObject->GetIntParameterAvailableOption(PoseSwitchEnumParamIndexInObject, j);
-					if (PossibleValue == PoseValue)
+					const int32 NumPoseValues = CustomizableObject->GetIntParameterNumOptions(PoseSwitchEnumParamIndexInObject);
+
+					TArray<FString> PoseOptionNamesAttribute;
+					FString PoseValue = GetIntParameterValue(PoseSwitchEnumParamName, -1);
+					int32 PoseValueIndex = 0;
+
+					for (int32 j = 0; j < NumPoseValues; ++j)
 					{
-						PoseValueIndex = j;
+						const FString PossibleValue = CustomizableObject->GetIntParameterAvailableOption(PoseSwitchEnumParamIndexInObject, j);
+						if (PossibleValue == PoseValue)
+						{
+							PoseValueIndex = j;
+						}
+
+						PoseOptionNamesAttribute.Add(PossibleValue);
 					}
-					//PoseOptionNames->Add(MakeShareable(new FString(CustomizableObject->GetIntParameterAvailableOption(PoseSwitchEnumParamIndexInObject, j))));
-					PoseOptionNamesAttribute.Add(CustomizableObject->GetIntParameterAvailableOption(PoseSwitchEnumParamIndexInObject, j));
-				}
 
-				TSharedPtr<SHorizontalBox> Box;
+					ParameterBox->AddSlot()
+						.HAlign(HAlign_Right)
+						.VAlign(VAlign_Fill)
+						.FillWidth(10.f)
+						[
+							SNew(SHorizontalBox)
+							+ SHorizontalBox::Slot()
+								.HAlign(HAlign_Fill)
+								.VAlign(VAlign_Center)
+								.FillWidth(0.45f)
+								[
+									SNew(SMutableTextSearchBox)
+									.PossibleSuggestions(PoseOptionNamesAttribute)
+									.InitialText(FText::FromString(PoseOptionNamesAttribute[PoseValueIndex]))
+									.MustMatchPossibleSuggestions(TAttribute<bool>(true))
+									.SuggestionListPlacement(EMenuPlacement::MenuPlacement_ComboBox)
+									.OnTextCommitted(this, &SCustomizableInstanceProperties::OnProjectorTextureParameterComboBoxChanged, PoseSwitchEnumParamNameWithRange)
+									.ToolTipText(LOCTEXT("Pose selector tooltip", "Select the skeletal mesh pose used for projection. This does not control the actual visual mesh pose in the viewport (or during gameplay for that matter). It has to be manually set. You can drag&drop a pose onto the preview viewport."))
+								]
+
+							+ SHorizontalBox::Slot()
+								.HAlign(HAlign_Fill)
+								.VAlign(VAlign_Center)
+								.FillWidth(0.3f)
+								[
+									SNew(SButton)
+									.ToolTipText(LOCTEXT("Add Layer", "Add Layer"))
+									.Text(LOCTEXT("Add Layer", "Add Layer"))
+									.OnClicked(this, &SCustomizableInstanceProperties::OnProjectorLayerAdded, ParamName)
+									.HAlign(HAlign_Fill)
+								]
+						];
+
+				}
+				else
+				{
+					ParameterBox->AddSlot()
+						.HAlign(HAlign_Right)
+						.VAlign(VAlign_Fill)
+						.FillWidth(10.f)
+						[
+							SNew(SButton)
+							.ToolTipText(LOCTEXT("Add Layer", "Add Layer"))
+							.Text(LOCTEXT("Add Layer", "Add Layer"))
+							.OnClicked(this, &SCustomizableInstanceProperties::OnProjectorLayerAdded, ParamName)
+							.HAlign(HAlign_Fill)
+						];
+				}
 
-				ParameterBox->AddSlot()
-				.HAlign(HAlign_Right)
-				.VAlign(VAlign_Fill)
-				.FillWidth(10.f)
-				[
-					SAssignNew(Box, SHorizontalBox)
-					+ SHorizontalBox::Slot()
-					.HAlign(HAlign_Fill)
-					.VAlign(VAlign_Center)
-					.FillWidth(0.45f)
-					[
-						SNew(SMutableTextSearchBox)
-						.PossibleSuggestions(PoseOptionNamesAttribute)
-						.InitialText(FText::FromString(PoseOptionNamesAttribute[PoseValueIndex]))
-						.MustMatchPossibleSuggestions(TAttribute<bool>(true))
-						.SuggestionListPlacement(EMenuPlacement::MenuPlacement_ComboBox)
-						.OnTextCommitted(this, &SCustomizableInstanceProperties::OnProjectorTextureParameterComboBoxChanged, PoseSwitchEnumParamNameWithRange)
-						.ToolTipText(LOCTEXT("Pose selector tooltip", "Select the skeletal mesh pose used for projection. This does not control the actual visual mesh pose in the viewport (or during gameplay for that matter). It has to be manually set. You can drag&drop a pose onto the preview viewport."))
-					]
-					
-					+ SHorizontalBox::Slot()
-					.HAlign(HAlign_Fill)
-					.VAlign(VAlign_Center)
-					.FillWidth(0.3f)
-					[
-						SNew(SButton)
-						.ToolTipText(LOCTEXT("Add Layer", "Add Layer"))
-						.Text(LOCTEXT("Add Layer", "Add Layer"))
-						.OnClicked(this, &SCustomizableInstanceProperties::OnProjectorLayerAdded, ParamName)
-						.HAlign(HAlign_Fill)
-					]
-				];
 
 				FString TextureSwitchEnumParamName = ParamName + FMultilayerProjector::IMAGE_PARAMETER_POSTFIX;
 
