[mutable] Changing group names updates their parameter names in states that referenced the old name.

#rnx

diff --git a/Engine/Plugins/Experimental/Mutable/Source/CustomizableObjectEditor/Classes/MuCOE/Nodes/CustomizableObjectNodeObjectGroup.h b/Engine/Plugins/Experimental/Mutable/Source/CustomizableObjectEditor/Classes/MuCOE/Nodes/CustomizableObjectNodeObjectGroup.h
--- a/Engine/Plugins/Experimental/Mutable/Source/CustomizableObjectEditor/Classes/MuCOE/Nodes/CustomizableObjectNodeObjectGroup.h
+++ b/Engine/Plugins/Experimental/Mutable/Source/CustomizableObjectEditor/Classes/MuCOE/Nodes/CustomizableObjectNodeObjectGroup.h
@@ -65,5 +65,13 @@ public:
 	{
 		return FindPin(TEXT("Projectors"));
 	}
+
+	UEdGraphPin* GroupPin() const
+	{
+		return FindPin(TEXT("Group"));
+	}
+
+private:
+	FString LastGroupName;
 };
 
diff --git a/Engine/Plugins/Experimental/Mutable/Source/CustomizableObjectEditor/Private/MuCOE/Nodes/CustomizableObjectNodeObjectGroup.cpp b/Engine/Plugins/Experimental/Mutable/Source/CustomizableObjectEditor/Private/MuCOE/Nodes/CustomizableObjectNodeObjectGroup.cpp
--- a/Engine/Plugins/Experimental/Mutable/Source/CustomizableObjectEditor/Private/MuCOE/Nodes/CustomizableObjectNodeObjectGroup.cpp
+++ b/Engine/Plugins/Experimental/Mutable/Source/CustomizableObjectEditor/Private/MuCOE/Nodes/CustomizableObjectNodeObjectGroup.cpp
@@ -6,6 +6,8 @@
 #include "Internationalization/Internationalization.h"
 #include "MuCO/CustomizableObjectCustomVersion.h"
 #include "MuCOE/EdGraphSchema_CustomizableObject.h"
+#include "MuCOE/GraphTraversal.h"
+#include "MuCOE/Nodes/CustomizableObjectNodeObject.h"
 #include "Serialization/Archive.h"
 #include "UObject/NameTypes.h"
 
@@ -34,11 +36,52 @@ void UCustomizableObjectNodeObjectGroup::Serialize(FArchive& Ar)
 	{
 		groupPin->PinType.PinCategory = UEdGraphSchema_CustomizableObject::PC_GroupProjector;
 	}
+
+	LastGroupName = GroupName;
 }
 
+
 void UCustomizableObjectNodeObjectGroup::PostEditChangeProperty(FPropertyChangedEvent& PropertyChangedEvent)
 {
 	Super::PostEditChangeProperty(PropertyChangedEvent);
+	for (UEdGraphPin* LinkedPin : FollowOutputPinArray(*GroupPin()))
+	{
+		UCustomizableObjectNode* Root = Cast<UCustomizableObjectNode>(LinkedPin->GetOwningNode());
+		check(Root); // All nodes inherit from UCustomizableObjectNode.
+
+		if (UCustomizableObjectNodeObject* CurrentRootNode = Cast<UCustomizableObjectNodeObject>(Root))
+		{
+			if (CurrentRootNode->ParentObject)
+			{
+				TArray<UCustomizableObject*> VisitedObjects;
+				CurrentRootNode = GetFullGraphRootNodeObject(CurrentRootNode, VisitedObjects);
+			}
+			if (!CurrentRootNode->ParentObject)
+			{
+				for (FCustomizableObjectState& State : CurrentRootNode->States)
+				{
+					for (int p = State.RuntimeParameters.Num() - 1; p >= 0; --p)
+					{
+						if (State.RuntimeParameters[p].Equals(LastGroupName))
+						{
+							State.RuntimeParameters.RemoveAt(p);
+							State.RuntimeParameters.Add(GroupName);
+						}
+					}
+					if (State.ForcedParameterValues.Contains(GroupName))
+					{
+						// Forced parameter already contains the NEW name of the parameter. TODO: Warning?
+					}
+					else if (State.ForcedParameterValues.Contains(LastGroupName))
+					{
+						FString LastForcedValue = State.ForcedParameterValues.FindAndRemoveChecked(LastGroupName);
+						State.ForcedParameterValues.Emplace(GroupName, LastForcedValue);
+					}
+				}
+			}
+		}
+	}
+	LastGroupName = GroupName;
 }
 
 
