[mutable] Fix crash when generating skeletons from skeletons with virtual bones

#rnx

diff --git a/Engine/Plugins/Experimental/Mutable/Source/CustomizableObject/Private/MuCO/CustomizableObjectInstance.cpp b/Engine/Plugins/Experimental/Mutable/Source/CustomizableObject/Private/MuCO/CustomizableObjectInstance.cpp
--- a/Engine/Plugins/Experimental/Mutable/Source/CustomizableObject/Private/MuCO/CustomizableObjectInstance.cpp
+++ b/Engine/Plugins/Experimental/Mutable/Source/CustomizableObject/Private/MuCO/CustomizableObjectInstance.cpp
@@ -3072,7 +3072,7 @@ bool UCustomizableInstancePrivateData::BuildSkeletalMeshSkeletonData(const TShar
 
 		// The second step is to update the pose transforms in the ref skeleton from the BasesInvMatrix
 		FReferenceSkeletonModifier SkeletonModifier(ReferenceSkeleton, Skeleton);
-		for (int32 RefSkelBoneIndex = 0; RefSkelBoneIndex < ReferenceSkeleton.GetNum(); ++RefSkelBoneIndex)
+		for (int32 RefSkelBoneIndex = 0; RefSkelBoneIndex < ReferenceSkeleton.GetRawBoneNum(); ++RefSkelBoneIndex)
 		{
 			int32 ParentBoneIndex = ReferenceSkeleton.GetParentIndex(RefSkelBoneIndex);
 			if (ParentBoneIndex >= 0)
diff --git a/Engine/Plugins/Experimental/Mutable/Source/CustomizableObject/Private/MuCO/UnrealConversionUtils.cpp b/Engine/Plugins/Experimental/Mutable/Source/CustomizableObject/Private/MuCO/UnrealConversionUtils.cpp
--- a/Engine/Plugins/Experimental/Mutable/Source/CustomizableObject/Private/MuCO/UnrealConversionUtils.cpp
+++ b/Engine/Plugins/Experimental/Mutable/Source/CustomizableObject/Private/MuCO/UnrealConversionUtils.cpp
@@ -130,7 +130,7 @@ namespace UnrealConversionUtils
 		, const FReferenceSkeleton& InSourceReferenceSkeleton, const TArray<bool>& InUsedBones,
 		FReferenceSkeleton& InRefSkeleton, const USkeleton* InSkeleton)
 	{
-		const int32 SourceBoneCount = InSourceReferenceSkeleton.GetNum();
+		const int32 SourceBoneCount = InSourceReferenceSkeleton.GetRawBoneNum();
 		const TArray<FTransform>& SourceRawMeshBonePose = InSourceReferenceSkeleton.GetRawRefBonePose();
 
 		MUTABLE_CPUPROFILER_SCOPE(BuildSkeletonData_BuildRefSkeleton);
diff --git a/Engine/Plugins/Experimental/Mutable/Source/CustomizableObjectEditor/Private/MuCOE/GenerateMutableSource/GenerateMutableSourceMesh.cpp b/Engine/Plugins/Experimental/Mutable/Source/CustomizableObjectEditor/Private/MuCOE/GenerateMutableSource/GenerateMutableSourceMesh.cpp
--- a/Engine/Plugins/Experimental/Mutable/Source/CustomizableObjectEditor/Private/MuCOE/GenerateMutableSource/GenerateMutableSourceMesh.cpp
+++ b/Engine/Plugins/Experimental/Mutable/Source/CustomizableObjectEditor/Private/MuCOE/GenerateMutableSource/GenerateMutableSourceMesh.cpp
@@ -2149,7 +2149,7 @@ bool GetAndValidateReshapeBonesToDeform(
 	{
 		for (const USkeletalMesh* SkeletalMesh : SkeletalMeshes)
 		{
-			int32 NumBonesToDeform = SkeletalMesh->GetRefSkeleton().GetNum();
+			int32 NumBonesToDeform = SkeletalMesh->GetRefSkeleton().GetRawBoneNum();
 
 			for (int32 BoneIndex = 0; BoneIndex < NumBonesToDeform; ++BoneIndex)
 			{
@@ -2178,7 +2178,7 @@ bool GetAndValidateReshapeBonesToDeform(
 	{
 		// Getting reference skeleton from the reference skeletal mesh of the current component
 		const FReferenceSkeleton RefSkeleton = GenerationContext.ComponentInfos[GenerationContext.CurrentMeshComponent].RefSkeletalMesh->GetRefSkeleton();
-		int32 NumBones = RefSkeleton.GetNum();
+		int32 NumBones = RefSkeleton.GetRawBoneNum();
 
 		for (int32 BoneIndex = 0; BoneIndex < NumBones; ++BoneIndex)
 		{
@@ -2196,7 +2196,7 @@ bool GetAndValidateReshapeBonesToDeform(
 
 		for (const USkeletalMesh* SkeletalMesh : SkeletalMeshes)
 		{
-			int32 NumBones = SkeletalMesh->GetRefSkeleton().GetNum();
+			int32 NumBones = SkeletalMesh->GetRefSkeleton().GetRawBoneNum();
 
 			for (int32 BoneIndex = 0; BoneIndex < NumBones; ++BoneIndex)
 			{
@@ -2235,7 +2235,7 @@ bool GetAndValidateReshapePhysicsToDeform(
 		const FReferenceSkeleton& RefSkeleton =
 			GenerationContext.ComponentInfos[GenerationContext.CurrentMeshComponent].RefSkeletalMesh->GetRefSkeleton();
 
-		const int32 RefSkeletonNumBones = RefSkeleton.GetNum();
+		const int32 RefSkeletonNumBones = RefSkeleton.GetRawBoneNum();
 		BoneNamesInUserSelection.SetNum(RefSkeletonNumBones);
 		for (int32 I = 0; I < RefSkeletonNumBones; ++I)
 		{
diff --git a/Engine/Plugins/Experimental/Mutable/Source/CustomizableObjectEditor/Private/MuCOE/Nodes/CustomizableObjectNodeAnimationPose.cpp b/Engine/Plugins/Experimental/Mutable/Source/CustomizableObjectEditor/Private/MuCOE/Nodes/CustomizableObjectNodeAnimationPose.cpp
--- a/Engine/Plugins/Experimental/Mutable/Source/CustomizableObjectEditor/Private/MuCOE/Nodes/CustomizableObjectNodeAnimationPose.cpp
+++ b/Engine/Plugins/Experimental/Mutable/Source/CustomizableObjectEditor/Private/MuCOE/Nodes/CustomizableObjectNodeAnimationPose.cpp
@@ -114,9 +114,9 @@ void UCustomizableObjectNodeAnimationPose::StaticRetrievePoseInformation(UPoseAs
 
 		// Use all bones from the pose's skeleton as RequiredBones
 		TArray<FBoneIndexType> RequiredBones;
-		RequiredBones.SetNumUninitialized(PoseRefSkeleton.GetNum());
+		RequiredBones.SetNumUninitialized(PoseRefSkeleton.GetRawBoneNum());
 
-		for (uint16 BoneIndex = 0; BoneIndex < PoseRefSkeleton.GetNum(); ++BoneIndex)
+		for (uint16 BoneIndex = 0; BoneIndex < PoseRefSkeleton.GetRawBoneNum(); ++BoneIndex)
 		{
 			RequiredBones[BoneIndex] = BoneIndex;
 		}
@@ -156,7 +156,7 @@ void UCustomizableObjectNodeAnimationPose::StaticRetrievePoseInformation(UPoseAs
 
 		// Extract final pose by combining the OutPose with the RefSkeleton pose (for bones missing in the PoseRefSkeleton)
 		const FReferenceSkeleton& RefSkeleton = RefSkeletalMesh->GetRefSkeleton();
-		for (int32 BoneIndex = 0; BoneIndex < RefSkeleton.GetNum(); ++BoneIndex)
+		for (int32 BoneIndex = 0; BoneIndex < RefSkeleton.GetRawBoneNum(); ++BoneIndex)
 		{
 			FName BoneName = RefSkeleton.GetBoneName(BoneIndex);
 
diff --git a/Engine/Plugins/Experimental/Mutable/Source/CustomizableObjectEditor/Private/MuCOE/Nodes/CustomizableObjectNodeMeshClipMorphDetails.cpp b/Engine/Plugins/Experimental/Mutable/Source/CustomizableObjectEditor/Private/MuCOE/Nodes/CustomizableObjectNodeMeshClipMorphDetails.cpp
--- a/Engine/Plugins/Experimental/Mutable/Source/CustomizableObjectEditor/Private/MuCOE/Nodes/CustomizableObjectNodeMeshClipMorphDetails.cpp
+++ b/Engine/Plugins/Experimental/Mutable/Source/CustomizableObjectEditor/Private/MuCOE/Nodes/CustomizableObjectNodeMeshClipMorphDetails.cpp
@@ -108,7 +108,7 @@ void FCustomizableObjectNodeMeshClipMorphDetails::CustomizeDetails(IDetailLayout
 		{
 			BoneComboOptions.Empty();
 				
-			for (int32 i = 0; i < SkeletalMesh->GetRefSkeleton().GetNum(); ++i)
+			for (int32 i = 0; i < SkeletalMesh->GetRefSkeleton().GetRawBoneNum(); ++i)
 			{
 				FName BoneName = SkeletalMesh->GetRefSkeleton().GetBoneName(i);
 				BoneComboOptions.Add(MakeShareable(new FString(BoneName.ToString())));
@@ -252,7 +252,7 @@ void FCustomizableObjectNodeMeshClipMorphDetails::OnBoneComboBoxSelectionChanged
 
 			//			int32 NextBone = -1;
 
-			//			for (int32 i = 0; i < SkeletalMesh->RefSkeleton.GetNum(); ++i)
+			//			for (int32 i = 0; i < SkeletalMesh->RefSkeleton.GetRawBoneNum(); ++i)
 			//			{
 			//				if (SkeletalMesh->RefSkeleton.GetParentIndex(i) == CurrentBone)
 			//				{
@@ -281,7 +281,7 @@ void FCustomizableObjectNodeMeshClipMorphDetails::OnBoneComboBoxSelectionChanged
 
 				FVector ChildLocation = FVector::ForwardVector;
 
-				for (int32 i = 0; i < SkeletalMesh->GetRefSkeleton().GetNum(); ++i)
+				for (int32 i = 0; i < SkeletalMesh->GetRefSkeleton().GetRawBoneNum(); ++i)
 				{
 					if (SkeletalMesh->GetRefSkeleton().GetParentIndex(i) == ParentIndex)
 					{
diff --git a/Engine/Plugins/Experimental/Mutable/Source/CustomizableObjectEditor/Private/MuCOE/Nodes/CustomizableObjectNodeProjectorParameterDetails.cpp b/Engine/Plugins/Experimental/Mutable/Source/CustomizableObjectEditor/Private/MuCOE/Nodes/CustomizableObjectNodeProjectorParameterDetails.cpp
--- a/Engine/Plugins/Experimental/Mutable/Source/CustomizableObjectEditor/Private/MuCOE/Nodes/CustomizableObjectNodeProjectorParameterDetails.cpp
+++ b/Engine/Plugins/Experimental/Mutable/Source/CustomizableObjectEditor/Private/MuCOE/Nodes/CustomizableObjectNodeProjectorParameterDetails.cpp
@@ -155,7 +155,7 @@ void FCustomizableObjectNodeProjectorParameterDetails::CustomizeDetails( IDetail
 	{
 		BoneComboOptions.Empty();
 			
-		for (int32 i = 0; i < SkeletalMesh->GetRefSkeleton().GetNum(); ++i)
+		for (int32 i = 0; i < SkeletalMesh->GetRefSkeleton().GetRawBoneNum(); ++i)
 		{
 			FName BoneName = SkeletalMesh->GetRefSkeleton().GetBoneName(i);
 			BoneComboOptions.Add(MakeShareable(new FString(BoneName.ToString())));
@@ -272,7 +272,7 @@ void FCustomizableObjectNodeProjectorParameterDetails::OnBoneComboBoxSelectionCh
 
 				FVector ChildLocation = FVector::ForwardVector;
 
-				for (int32 i = 0; i < SkeletalMesh->GetRefSkeleton().GetNum(); ++i)
+				for (int32 i = 0; i < SkeletalMesh->GetRefSkeleton().GetRawBoneNum(); ++i)
 				{
 					if (SkeletalMesh->GetRefSkeleton().GetParentIndex(i) == ParentIndex)
 					{
