[mutable] Use the max UVChannel density of the RefSkeletalMesh as reference for the generated meshes.

#rnx

diff --git a/Engine/Plugins/Experimental/Mutable/Source/CustomizableObject/Classes/MuCO/CustomizableObject.h b/Engine/Plugins/Experimental/Mutable/Source/CustomizableObject/Classes/MuCO/CustomizableObject.h
--- a/Engine/Plugins/Experimental/Mutable/Source/CustomizableObject/Classes/MuCO/CustomizableObject.h
+++ b/Engine/Plugins/Experimental/Mutable/Source/CustomizableObject/Classes/MuCO/CustomizableObject.h
@@ -803,9 +803,12 @@ struct FMutableRefSkeletalMeshSettings
 {
 	bool bEnablePerPolyCollision = false;
 
+	float DefaultUVChannelDensity = 0.f;
+
 	friend FArchive& operator<<(FArchive& Ar, FMutableRefSkeletalMeshSettings& Data)
 	{
 		Ar << Data.bEnablePerPolyCollision;
+		Ar << Data.DefaultUVChannelDensity;
 
 		return Ar;
 	}
@@ -1217,7 +1220,7 @@ private:
 	// This is a manual version number for the binary blobs in this asset.
 	// Increasing it invalidates all the previously compiled models.
 	// Warning: If while merging code both versions have changed, take the highest+1.
-	static const int32 CurrentSupportedVersion = 355;
+	static const int32 CurrentSupportedVersion = 356;
 
 public:
 
diff --git a/Engine/Plugins/Experimental/Mutable/Source/CustomizableObject/Private/MuCO/CustomizableObjectInstance.cpp b/Engine/Plugins/Experimental/Mutable/Source/CustomizableObject/Private/MuCO/CustomizableObjectInstance.cpp
--- a/Engine/Plugins/Experimental/Mutable/Source/CustomizableObject/Private/MuCO/CustomizableObjectInstance.cpp
+++ b/Engine/Plugins/Experimental/Mutable/Source/CustomizableObject/Private/MuCO/CustomizableObjectInstance.cpp
@@ -826,24 +826,6 @@ void UCustomizableInstancePrivateData::PostEditChangePropertyWithoutEditor(USkel
 		// reinitialize resource
 		InSkeletalMesh->InitResources();
 	}
-
-	{
-		MUTABLE_CPUPROFILER_SCOPE(UpdateSkeletalMesh_LocalUVDensities);
-
-		for (int32 MaterialIndex = 0; MaterialIndex < InSkeletalMesh->GetMaterials().Num(); ++MaterialIndex)
-		{
-			FMeshUVChannelInfo& UVChannelData = InSkeletalMesh->GetMaterials()[MaterialIndex].UVChannelData;
-
-			UVChannelData.bInitialized = true;
-			UVChannelData.bOverrideDensities = false;
-
-			for (int32 i = 0; i < TEXSTREAM_MAX_NUM_UVCHANNELS; ++i)
-			{
-				// TODO
-				UVChannelData.LocalUVDensities[i] = 200.f;
-			}
-		}
-	}
 }
 
 
@@ -2939,6 +2921,20 @@ void UCustomizableInstancePrivateData::ConvertImage(UTexture2D* Texture, mu::Ima
 }
 
 
+void SetMeshUVChannelDensity(FMeshUVChannelInfo& UVChannelInfo, float Density = 0.f)
+{
+	Density = Density > 0.f ? Density : 150.f;
+
+	UVChannelInfo.bInitialized = true;
+	UVChannelInfo.bOverrideDensities = false;
+
+	for (int32 i = 0; i < TEXSTREAM_MAX_NUM_UVCHANNELS; ++i)
+	{
+		UVChannelInfo.LocalUVDensities[i] = Density;
+	}
+}
+
+
 void UCustomizableInstancePrivateData::BuildSkeletalMeshElementData(const TSharedPtr<FMutableOperationData>& OperationData, USkeletalMesh* SkeletalMesh, UCustomizableObjectInstance* CustomizableObjectInstance, int32 ComponentIndex)
 {
 	MUTABLE_CPUPROFILER_SCOPE(UCustomizableInstancePrivateData::BuildSkeletalMeshElementData);
@@ -2948,6 +2944,9 @@ void UCustomizableInstancePrivateData::BuildSkeletalMeshElementData(const TShare
 	SkeletalMesh->GetMaterials().SetNum(1);
 	SkeletalMesh->GetMaterials()[0] = UnrealMaterial;
 
+	// Default density
+	SetMeshUVChannelDensity(SkeletalMesh->GetMaterials()[0].UVChannelData);
+
 	int32 MeshLODIndex = HasCOInstanceFlags(ReduceLODs) ? 0 : OperationData->CurrentMinLOD;
 
 	int32 LODCount = OperationData->InstanceUpdateData.LODs.Num();
@@ -4706,14 +4705,9 @@ void UCustomizableInstancePrivateData::BuildMaterials(const TSharedPtr<FMutableO
 	for (int32 ComponentIndex = 0; ComponentIndex < NumComponents; ++ComponentIndex)
 	{
 		USkeletalMesh* SkeletalMesh = Public->SkeletalMeshes[ComponentIndex];
-
+		if (SkeletalMesh)
 		{
-			MUTABLE_CPUPROFILER_SCOPE(Init);
-
-			if (SkeletalMesh)
-			{
-				SkeletalMesh->GetMaterials().Reset();
-			}
+			SkeletalMesh->GetMaterials().Reset();
 		}
 
 		// Maps serializations of FMutableMaterialPlaceholder to Created Dynamic Material instances, used to reuse materials across LODs
@@ -4746,6 +4740,9 @@ void UCustomizableInstancePrivateData::BuildMaterials(const TSharedPtr<FMutableO
 					continue;
 				}
 
+				const FMutableRefSkeletalMeshData* RefSkeletalMeshData = CustomizableObject->GetRefSkeletalMeshData(Component.Id);
+				check(RefSkeletalMeshData);
+
 				for (int32 SurfaceIndex = 0; SurfaceIndex < Component.SurfaceCount; ++SurfaceIndex)
 				{
 					const FInstanceUpdateData::FSurface& Surface = OperationData->InstanceUpdateData.Surfaces[Component.FirstSurface + SurfaceIndex];
@@ -5066,12 +5063,16 @@ void UCustomizableInstancePrivateData::BuildMaterials(const TSharedPtr<FMutableO
 
 							if (SkeletalMesh)
 							{
-								int32 MatIndex = SkeletalMesh->GetMaterials().Add(ActualMaterialInterface);
-								SkeletalMesh->GetMaterials()[MatIndex].MaterialSlotName = CustomizableObject->ReferencedMaterialSlotNames[Surface.MaterialIndex];
-								int32 LODMaterialIndex = Helper_GetLODInfoArray(SkeletalMesh)[MeshLODIndex].LODMaterialMap.Add(MatIndex);
-								Helper_GetLODRenderSections(SkeletalMesh, MeshLODIndex)[SurfaceIndex].MaterialIndex = LODMaterialIndex;
-
+								const int32 MatIndex = SkeletalMesh->GetMaterials().Num();
 								MutableMaterialPlaceholder.MatIndex = MatIndex;
+
+								// Set up SkeletalMaterial data
+								FSkeletalMaterial& SkeletalMaterial = SkeletalMesh->GetMaterials().Add_GetRef(ActualMaterialInterface);
+								SkeletalMaterial.MaterialSlotName = CustomizableObject->ReferencedMaterialSlotNames[Surface.MaterialIndex];
+								SetMeshUVChannelDensity(SkeletalMaterial.UVChannelData, RefSkeletalMeshData->Settings.DefaultUVChannelDensity);
+
+								const int32 LODMaterialIndex = Helper_GetLODInfoArray(SkeletalMesh)[MeshLODIndex].LODMaterialMap.Add(MatIndex);
+								Helper_GetLODRenderSections(SkeletalMesh, MeshLODIndex)[SurfaceIndex].MaterialIndex = LODMaterialIndex;
 							}
 
 							if (MaterialInstance)
diff --git a/Engine/Plugins/Experimental/Mutable/Source/CustomizableObjectEditor/Private/MuCOE/GenerateMutableSource/GenerateMutableSource.cpp b/Engine/Plugins/Experimental/Mutable/Source/CustomizableObjectEditor/Private/MuCOE/GenerateMutableSource/GenerateMutableSource.cpp
--- a/Engine/Plugins/Experimental/Mutable/Source/CustomizableObjectEditor/Private/MuCOE/GenerateMutableSource/GenerateMutableSource.cpp
+++ b/Engine/Plugins/Experimental/Mutable/Source/CustomizableObjectEditor/Private/MuCOE/GenerateMutableSource/GenerateMutableSource.cpp
@@ -1311,6 +1311,18 @@ void PopulateReferenceSkeletalMeshesData(FMutableGraphGenerationContext& Generat
 		// Additional Settings
 		Data.Settings.bEnablePerPolyCollision = RefSkeletalMesh->GetEnablePerPolyCollision();
 
+		const TArray<FSkeletalMaterial>& Materials = RefSkeletalMesh->GetMaterials();
+		for (const FSkeletalMaterial& Material : Materials)
+		{
+			if (Material.UVChannelData.bInitialized)
+			{
+				for (int32 UVIndex = 0; UVIndex < TEXSTREAM_MAX_NUM_UVCHANNELS; ++UVIndex)
+				{
+					Data.Settings.DefaultUVChannelDensity = FMath::Max(Data.Settings.DefaultUVChannelDensity, Material.UVChannelData.LocalUVDensities[UVIndex]);
+				}
+			}
+		}
+
 		// Skeleton
 		if(const USkeleton* Skeleton = RefSkeletalMesh->GetSkeleton())
 		{
