[Mutable] Disable CO Instance Editor using ingame bounds by default.

#rnx

diff --git a/Engine/Plugins/Experimental/Mutable/Source/CustomizableObjectEditor/Private/MuCOE/CustomizableObjectInstanceEditor.cpp b/Engine/Plugins/Experimental/Mutable/Source/CustomizableObjectEditor/Private/MuCOE/CustomizableObjectInstanceEditor.cpp
--- a/Engine/Plugins/Experimental/Mutable/Source/CustomizableObjectEditor/Private/MuCOE/CustomizableObjectInstanceEditor.cpp
+++ b/Engine/Plugins/Experimental/Mutable/Source/CustomizableObjectEditor/Private/MuCOE/CustomizableObjectInstanceEditor.cpp
@@ -350,7 +350,6 @@ void FCustomizableObjectInstanceEditor::CreatePreviewInstance()
 
 				if (PreviewSkeletalMeshComponents[ComponentIndex])
 				{
-					PreviewSkeletalMeshComponents[ComponentIndex]->UseInGameBounds(true);
 					PreviewCustomizableSkeletalComponents[ComponentIndex]->AttachToComponent(PreviewSkeletalMeshComponents[ComponentIndex], FAttachmentTransformRules::KeepRelativeTransform);
 				}
 				else
