[Mutable] Fix Instance Properties not showing Texture Parameters correctly.

diff --git a/Engine/Plugins/Experimental/Mutable/Source/CustomizableObject/Classes/MuCO/CustomizableObjectInstance.h b/Engine/Plugins/Experimental/Mutable/Source/CustomizableObject/Classes/MuCO/CustomizableObjectInstance.h
--- a/Engine/Plugins/Experimental/Mutable/Source/CustomizableObject/Classes/MuCO/CustomizableObjectInstance.h
+++ b/Engine/Plugins/Experimental/Mutable/Source/CustomizableObject/Classes/MuCO/CustomizableObjectInstance.h
@@ -353,18 +353,18 @@ public:
 	void SetFloatParameterSelectedOption(const FString& FloatParamName, float FloatValue, int32 RangeIndex = -1);
 	
 	/** Gets the value of a texture parameter with name "TextureParamName". */
-	uint64 GetTextureParameterSelectedOption(const FString& TextureParamName, int32 RangeIndex) const;
+	uint64 GetTextureParameterSelectedOption(const FString& TextureParamName, int32 RangeIndex = -1) const;
 
 	/** Gets the texture of a texture parameter with name "TextureParamName". */
 	UFUNCTION(BlueprintCallable, Category = CustomizableObjectInstance)
-	UTexture2D* GetTextureParameterSelectedOptionT(const FString& TextureParamName, int32 RangeIndex) const;
+	UTexture2D* GetTextureParameterSelectedOptionT(const FString& TextureParamName, int32 RangeIndex = -1) const;
 
 	/** Sets the texture value "TextureValue" of a texture parameter with index "TextureParamIndex". */
-	void SetTextureParameterSelectedOption(const FString& TextureParamName, uint64 TextureValue, int32 RangeIndex);
+	void SetTextureParameterSelectedOption(const FString& TextureParamName, uint64 TextureValue, int32 RangeIndex = -1);
 
 	/** Sets the texture "Texture" of a texture parameter with index "TextureParamIndex". */
 	UFUNCTION(BlueprintCallable, Category = CustomizableObjectInstance)
-	void SetTextureParameterSelectedOptionT(const FString& TextureParamName, UTexture2D* TextureValue, int32 RangeIndex);
+	void SetTextureParameterSelectedOptionT(const FString& TextureParamName, UTexture2D* TextureValue, int32 RangeIndex = -1);
 	
 	// Gets the value of a color parameter with name "ColorParamName"
 	UFUNCTION(BlueprintCallable, Category = CustomizableObjectInstance)
diff --git a/Engine/Plugins/Experimental/Mutable/Source/CustomizableObjectEditor/Private/MuCOE/SCustomizableInstanceProperties.cpp b/Engine/Plugins/Experimental/Mutable/Source/CustomizableObjectEditor/Private/MuCOE/SCustomizableInstanceProperties.cpp
--- a/Engine/Plugins/Experimental/Mutable/Source/CustomizableObjectEditor/Private/MuCOE/SCustomizableInstanceProperties.cpp
+++ b/Engine/Plugins/Experimental/Mutable/Source/CustomizableObjectEditor/Private/MuCOE/SCustomizableInstanceProperties.cpp
@@ -127,7 +127,7 @@ void SCustomizableInstanceProperties::Construct(const FArguments& InArgs)
 
 		// Store the texture parameters data required for the ui.
 		TextureParameterValueNames.Add(MakeShareable(new FString("None")));
-		TextureParameterValues.Add(-1);
+		TextureParameterValues.Add(UINT64_MAX);
 		TArray<FCustomizableObjectExternalTexture> Textures = UCustomizableObjectSystem::GetInstance()->GetTextureParameterValues();
 		for (int i = 0; i < Textures.Num(); ++i)
 		{
@@ -844,28 +844,17 @@ void SCustomizableInstanceProperties::AddParameter(int32 ParamIndexInObject)
 
 	case EMutableParameterType::Texture:
 	{
-		TArray<FCustomizableObjectTextureParameterValue>& TextureParameters = CustomInstance->GetTextureParameters();
-		TSharedPtr<FString> InitiallySelected = TextureParameterValueNames[0];
-		for (int i = 0; i < TextureParameters.Num(); ++i)
+		const TArray<FCustomizableObjectTextureParameterValue>& TextureParameters = CustomInstance->GetTextureParameters();
+		TSharedPtr<FString> InitiallySelected = TextureParameterValueNames[0]; // First index is always the None option.
+
+		const uint64 ParameterValue = CustomInstance->GetTextureParameterSelectedOption(ParamName);
+		
+		// Look for the value index
+		for (int32 ValueIndex = 0; ValueIndex < TextureParameterValueNames.Num(); ++ValueIndex)
 		{
-			FString LocalParamName = TextureParameters[i].ParameterName;
-			uint64 ParamValue = TextureParameters[i].ParameterValue;
-			if (TextureParameters[i].ParameterName == LocalParamName)
+			if (ParameterValue == TextureParameterValues[ValueIndex])
 			{
-				// Look for the value index
-				for (int ValueIndex = 0; ValueIndex < TextureParameterValueNames.Num(); ++ValueIndex)
-				{
-					if (TextureParameterValues[ValueIndex] == ParamValue)
-					{
-						InitiallySelected = TextureParameterValueNames[ValueIndex];
-
-						/*if (CustomInstance->ProfileParameterDat.)
-						{
-							CustomInstance->ProfileParameterDat->TextureParameters[i].ParameterValue = ParamValue;
-						}*/	
-					}
-				}
-
+				InitiallySelected = TextureParameterValueNames[ValueIndex];
 				break;
 			}
 		}
diff --git a/Engine/Plugins/Experimental/Mutable/Source/CustomizableObjectEditor/Public/MuCOE/SCustomizableInstanceProperties.h b/Engine/Plugins/Experimental/Mutable/Source/CustomizableObjectEditor/Public/MuCOE/SCustomizableInstanceProperties.h
--- a/Engine/Plugins/Experimental/Mutable/Source/CustomizableObjectEditor/Public/MuCOE/SCustomizableInstanceProperties.h
+++ b/Engine/Plugins/Experimental/Mutable/Source/CustomizableObjectEditor/Public/MuCOE/SCustomizableInstanceProperties.h
@@ -110,7 +110,7 @@ private:
 	// These arrays store the textures available for texture parameters of the model.
 	// These come from the texture generators registered in the CustomizableObjectSystem
 	TArray< TSharedPtr<FString> > TextureParameterValueNames;
-	TArray< int64_t > TextureParameterValues;
+	TArray< uint64_t > TextureParameterValues;
 
 	struct FSliderData
 	{
