[mutable] Optimization of the code generated for multilayer projector groups.
Optimized one image clone from the multilayer operation.
Refactored image blend functions and their files.
Renamed "Alpha Overlay" blende type to "Lighten"
Moved swizzle to its own ASTOp class.
#rnx

diff --git a/Engine/Plugins/Experimental/Mutable/Source/CustomizableObject/Classes/MuCO/CustomizableObject.h b/Engine/Plugins/Experimental/Mutable/Source/CustomizableObject/Classes/MuCO/CustomizableObject.h
--- a/Engine/Plugins/Experimental/Mutable/Source/CustomizableObject/Classes/MuCO/CustomizableObject.h
+++ b/Engine/Plugins/Experimental/Mutable/Source/CustomizableObject/Classes/MuCO/CustomizableObject.h
@@ -1220,7 +1220,7 @@ private:
 	// This is a manual version number for the binary blobs in this asset.
 	// Increasing it invalidates all the previously compiled models.
 	// Warning: If while merging code both versions have changed, take the highest+1.
-	static const int32 CurrentSupportedVersion = 356;
+	static const int32 CurrentSupportedVersion = 357;
 
 public:
 
diff --git a/Engine/Plugins/Experimental/Mutable/Source/CustomizableObjectEditor/Classes/MuCOE/Nodes/CustomizableObjectNodeTextureLayer.h b/Engine/Plugins/Experimental/Mutable/Source/CustomizableObjectEditor/Classes/MuCOE/Nodes/CustomizableObjectNodeTextureLayer.h
--- a/Engine/Plugins/Experimental/Mutable/Source/CustomizableObjectEditor/Classes/MuCOE/Nodes/CustomizableObjectNodeTextureLayer.h
+++ b/Engine/Plugins/Experimental/Mutable/Source/CustomizableObjectEditor/Classes/MuCOE/Nodes/CustomizableObjectNodeTextureLayer.h
@@ -33,7 +33,7 @@ enum ECustomizableObjectTextureLayerEffect
 	COTLE_BURN					UMETA(DisplayName = "BURN"),
 	COTLE_SCREEN				UMETA(DisplayName = "SCREEN"),
 	COTLE_OVERLAY				UMETA(DisplayName = "OVERLAY"),
-	COTLE_ALPHA_OVERLAY			UMETA(DisplayName = "ALPHA_OVERLAY"),
+	COTLE_LIGHTEN				UMETA(DisplayName = "LIGHTEN"),
 	COTLE_NORMAL_COMBINE		UMETA(DisplayName = "COMBINE")
 };
 
diff --git a/Engine/Plugins/Experimental/Mutable/Source/CustomizableObjectEditor/Private/MuCOE/GenerateMutableSource/GenerateMutableSourceGroupProjector.cpp b/Engine/Plugins/Experimental/Mutable/Source/CustomizableObjectEditor/Private/MuCOE/GenerateMutableSource/GenerateMutableSourceGroupProjector.cpp
--- a/Engine/Plugins/Experimental/Mutable/Source/CustomizableObjectEditor/Private/MuCOE/GenerateMutableSource/GenerateMutableSourceGroupProjector.cpp
+++ b/Engine/Plugins/Experimental/Mutable/Source/CustomizableObjectEditor/Private/MuCOE/GenerateMutableSource/GenerateMutableSourceGroupProjector.cpp
@@ -384,7 +384,7 @@ mu::NodeImagePtr GenerateMutableGroupProjection(const int32 NodeLOD, const int32
 
 		mu::NodeImageMultiLayerPtr BaseAlphaMultiLayerNode = new mu::NodeImageMultiLayer;
 		BaseAlphaMultiLayerNode->SetRange(ImageNodes_ProjectorTempData[i].NodeRange);
-		BaseAlphaMultiLayerNode->SetBlendType(mu::EBlendType::BT_ALPHA_OVERLAY);
+		BaseAlphaMultiLayerNode->SetBlendType(mu::EBlendType::BT_LIGHTEN);
 		BaseAlphaMultiLayerNode->SetBase(ResultAlpha);
 		BaseAlphaMultiLayerNode->SetBlended(MultiplySwizzleNode);
 		//BaseAlphaMultiLayerNode->SetMask(MultiplySwizzleNode); // No mask needed
diff --git a/Engine/Plugins/Experimental/Mutable/Source/CustomizableObjectEditor/Private/MuCOE/GenerateMutableSource/GenerateMutableSourceImage.cpp b/Engine/Plugins/Experimental/Mutable/Source/CustomizableObjectEditor/Private/MuCOE/GenerateMutableSource/GenerateMutableSourceImage.cpp
--- a/Engine/Plugins/Experimental/Mutable/Source/CustomizableObjectEditor/Private/MuCOE/GenerateMutableSource/GenerateMutableSourceImage.cpp
+++ b/Engine/Plugins/Experimental/Mutable/Source/CustomizableObjectEditor/Private/MuCOE/GenerateMutableSource/GenerateMutableSourceImage.cpp
@@ -329,7 +329,7 @@ mu::NodeImagePtr GenerateMutableSourceImage(const UEdGraphPin* Pin, FMutableGrap
 				case COTLE_BURN: Type = mu::EBlendType::BT_BURN; break;
 				case COTLE_SCREEN: Type = mu::EBlendType::BT_SCREEN; break;
 				case COTLE_OVERLAY: Type = mu::EBlendType::BT_OVERLAY; break;
-				case COTLE_ALPHA_OVERLAY: Type = mu::EBlendType::BT_ALPHA_OVERLAY; break;
+				case COTLE_LIGHTEN: Type = mu::EBlendType::BT_LIGHTEN; break;
 				case COTLE_NORMAL_COMBINE: Type = mu::EBlendType::BT_NORMAL_COMBINE; break;
 				default:
 					GenerationContext.Compiler->CompilerLog(LOCTEXT("UnsupportedImageEffect", "Texture layer effect not supported. Setting to 'Blend'."), Node);
diff --git a/Engine/Plugins/Experimental/Mutable/Source/CustomizableObjectEditor/Private/MuCOE/SMutableCodeViewer.cpp b/Engine/Plugins/Experimental/Mutable/Source/CustomizableObjectEditor/Private/MuCOE/SMutableCodeViewer.cpp
--- a/Engine/Plugins/Experimental/Mutable/Source/CustomizableObjectEditor/Private/MuCOE/SMutableCodeViewer.cpp
+++ b/Engine/Plugins/Experimental/Mutable/Source/CustomizableObjectEditor/Private/MuCOE/SMutableCodeViewer.cpp
@@ -156,9 +156,15 @@ public:
 			case mu::OP_TYPE::IM_MULTILAYER:
 			{
 				mu::OP::ImageMultiLayerArgs Args = Program.GetOpArgs<mu::OP::ImageMultiLayerArgs>(RowItem->MutableOperation);
-				OpName += TEXT(" ");
+				OpName += TEXT(" rgb: ");
 				OpName += mu::TypeInfo::s_blendModeName[int32(Args.blendType)];
-				OpName += FString::Printf(TEXT(" range-id %d"), Args.rangeId);
+				if (Args.blendTypeAlpha != int8(mu::EBlendType::BT_NONE))
+				{
+					OpName += TEXT(", a: ");
+					OpName += mu::TypeInfo::s_blendModeName[int32(Args.blendTypeAlpha)];
+				}
+				OpName += FString::Printf(TEXT(" range-id: %d"), Args.rangeId);
+				OpName += FString::Printf(TEXT(" mask-from-alpha: %d"), int32(Args.bUseMaskFromBlended));
 				break;
 			}
 
@@ -174,8 +180,10 @@ public:
 			case mu::OP_TYPE::IM_LAYERCOLOUR:
 			{
 				mu::OP::ImageLayerColourArgs Args = Program.GetOpArgs<mu::OP::ImageLayerColourArgs>(RowItem->MutableOperation);
-				OpName += TEXT(" ");
+				OpName += TEXT(" rgb: ");
 				OpName += mu::TypeInfo::s_blendModeName[int32(Args.blendType)];
+				OpName += TEXT(" a: ");
+				OpName += mu::TypeInfo::s_blendModeName[int32(Args.blendTypeAlpha)];
 				break;
 			}
 
@@ -1253,6 +1261,17 @@ void SMutableCodeViewer::GenerateElementRecursive(mu::OP::ADDRESS InParentAddres
 		break;
 	}
 
+	case mu::OP_TYPE::CO_SWIZZLE:
+	{
+		mu::OP::ColourSwizzleArgs Args = InProgram.GetOpArgs<mu::OP::ColourSwizzleArgs>(InParentAddress);
+		for (int32 Channel = 0; Channel < 4; ++Channel)
+		{
+			FString Caption = FString::Printf(TEXT("%d is %d from "), Channel, Args.sourceChannels[Channel]);
+			AddOpFunc(Args.sources[Channel], Caption);
+		}
+		break;
+	}
+
 	case mu::OP_TYPE::IM_LAYER:
 	{
 		mu::OP::ImageLayerArgs Args = InProgram.GetOpArgs<mu::OP::ImageLayerArgs>(InParentAddress);
diff --git a/Engine/Plugins/Experimental/Mutable/Source/MutableRuntime/Private/MuR/CodeRunner.cpp b/Engine/Plugins/Experimental/Mutable/Source/MutableRuntime/Private/MuR/CodeRunner.cpp
--- a/Engine/Plugins/Experimental/Mutable/Source/MutableRuntime/Private/MuR/CodeRunner.cpp
+++ b/Engine/Plugins/Experimental/Mutable/Source/MutableRuntime/Private/MuR/CodeRunner.cpp
@@ -22,34 +22,26 @@
 #include "MuR/MutableString.h"
 #include "MuR/MutableTrace.h"
 #include "MuR/OpEvalCurve.h"
-#include "MuR/OpImageAlphaOverlay.h"
 #include "MuR/OpImageApplyComposite.h"
 #include "MuR/OpImageBinarise.h"
 #include "MuR/OpImageBlend.h"
-#include "MuR/OpImageBurn.h"
 #include "MuR/OpImageColourMap.h"
 #include "MuR/OpImageCompose.h"
 #include "MuR/OpImageCrop.h"
 #include "MuR/OpImageDifference.h"
 #include "MuR/OpImageDisplace.h"
-#include "MuR/OpImageDodge.h"
 #include "MuR/OpImageGradient.h"
-#include "MuR/OpImageHardLight.h"
 #include "MuR/OpImageInterpolate.h"
 #include "MuR/OpImageInvert.h"
 #include "MuR/OpImageLuminance.h"
 #include "MuR/OpImageMipmap.h"
-#include "MuR/OpImageMultiply.h"
 #include "MuR/OpImageNormalCombine.h"
-#include "MuR/OpImageOverlay.h"
 #include "MuR/OpImagePixelFormat.h"
 #include "MuR/OpImageProject.h"
 #include "MuR/OpImageRasterMesh.h"
 #include "MuR/OpImageResize.h"
 #include "MuR/OpImageSaturate.h"
-#include "MuR/OpImageScreen.h"
 #include "MuR/OpImageSelectColour.h"
-#include "MuR/OpImageSoftLight.h"
 #include "MuR/OpImageTransform.h"
 #include "MuR/OpLayoutPack.h"
 #include "MuR/OpLayoutRemoveBlocks.h"
@@ -1040,7 +1032,7 @@ namespace mu
 
                 if (baseValid)
                 {
-                    SCHEDULED_OP_DATA HeapData;
+                    FScheduledOpData HeapData;
 					HeapData.bifactor = bifactor;
 					HeapData.min = FMath::Clamp(min, 0, NumTargets - 1);
 					HeapData.max = FMath::Clamp(max, 0, NumTargets - 1);
@@ -1081,7 +1073,7 @@ namespace mu
                 Ptr<const Mesh> pBase = GetMemory().GetMesh( CACHE_ADDRESS(BaseAt,item) );
 
                 // Factor from 0 to 1 between the two targets
-                const SCHEDULED_OP_DATA& HeapData = m_heapData[ (size_t)item.customState ];
+                const FScheduledOpData& HeapData = m_heapData[ (size_t)item.customState ];
                 float bifactor = HeapData.bifactor;
                 int min = HeapData.min;
                 int max = HeapData.max;
@@ -1249,7 +1241,7 @@ namespace mu
                 // Factor from 0 to 1 between the two targets
                 float bifactor = factor/delta - min;
 
-                SCHEDULED_OP_DATA data;
+                FScheduledOpData data;
                 data.bifactor = bifactor;
 				data.min = FMath::Clamp(min, 0, count - 1);
 				data.max = FMath::Clamp(max, 0, count - 1);
@@ -1318,7 +1310,7 @@ namespace mu
                     count++;
                 }
 
-                const SCHEDULED_OP_DATA& data = m_heapData[ (size_t)item.customState ];
+                const FScheduledOpData& data = m_heapData[ (size_t)item.customState ];
 
                 // Factor from 0 to 1 between the two targets
                 float bifactor = data.bifactor;
@@ -2453,49 +2445,48 @@ namespace mu
             {
             case 0:
                 AddOp( SCHEDULED_OP( item.at, item, 1),
-                       SCHEDULED_OP( args.rangeSize, item ),
-                       SCHEDULED_OP( args.base, item) );
+                       SCHEDULED_OP( args.rangeSize, item ) );
                 break;
 
             case 1:
             {
                 // We now know the number of iterations
-                int iterations = 0;
+                int32 Iterations = 0;
                 if (args.rangeSize)
                 {
-                    CACHE_ADDRESS address(args.rangeSize,item);
+                    CACHE_ADDRESS RangeAddress(args.rangeSize,item);
 
                     // We support both integers and scalars here, which is not common.
                     // \todo: review if this is necessary or we can enforce it at compile time.
-                    DATATYPE rangeSizeType = GetOpDataType( pModel->GetPrivate()->m_program.GetOpType(args.rangeSize) );
-                    if (rangeSizeType==DT_INT)
+                    DATATYPE RangeSizeType = GetOpDataType( pModel->GetPrivate()->m_program.GetOpType(args.rangeSize) );
+                    if (RangeSizeType == DT_INT)
                     {
-                        iterations = GetMemory().GetInt( address );
+						Iterations = GetMemory().GetInt(RangeAddress);
                     }
-                    else if (rangeSizeType==DT_SCALAR)
+                    else if (RangeSizeType == DT_SCALAR)
                     {
-                        iterations = int( GetMemory().GetScalar( address ) );
+						Iterations = int32( GetMemory().GetScalar(RangeAddress) );
                     }
                 }
 
                 // \todo Check that we are not overwriting the index (it shouldn't be set when we reach here)
-                SCHEDULED_OP itemCopy = item;
-                ExecutionIndex index = GetMemory().GetRageIndex( item.executionIndex );
-
-				TArray<SCHEDULED_OP> deps;
-				deps.SetNumUninitialized(iterations * 2 + 1);
-                int d=0;
-                deps[d++] = SCHEDULED_OP( args.base, item );
-                for (int i=0; i<iterations; ++i)
+                SCHEDULED_OP ItemCopy = item;
+                ExecutionIndex Index = GetMemory().GetRageIndex( item.executionIndex );
+
+				TArray<SCHEDULED_OP, TInlineAllocator<48>> Deps;
+				Deps.SetNumUninitialized(Iterations * 2 + 1);
+                int32 d=0;
+				Deps[d++] = SCHEDULED_OP( args.base, item );
+                for (int i=0; i< Iterations; ++i)
                 {
-                    index.SetFromModelRangeIndex( args.rangeId, i );
+					Index.SetFromModelRangeIndex( args.rangeId, i );
 
-                    itemCopy.executionIndex = GetMemory().GetRageIndexIndex(index);
+					ItemCopy.executionIndex = GetMemory().GetRageIndexIndex(Index);
 
-                    deps[d++] = SCHEDULED_OP( args.blended, itemCopy );
-                    deps[d++] = SCHEDULED_OP( args.mask, itemCopy );
+					Deps[d++] = SCHEDULED_OP( args.blended, ItemCopy);
+					Deps[d++] = SCHEDULED_OP( args.mask, ItemCopy);
                 }
-                AddOp( SCHEDULED_OP( item.at, item, 2, iterations), deps );
+                AddOp( SCHEDULED_OP( item.at, item, 2, Iterations), Deps);
 
                 break;
             }
@@ -2504,12 +2495,11 @@ namespace mu
             {
                 MUTABLE_CPUPROFILER_SCOPE(ImageMultiLayer);
 
-				int32 iterations = item.customState;
+				int32 Iterations = item.customState;
 
-                // TODO: Reuse base if possible
                 Ptr<const Image> pBase = GetMemory().GetImage( CACHE_ADDRESS(args.base,item) );
 
-                if (iterations==0)
+                if (Iterations ==0)
                 {
 					// Early out
 					GetMemory().SetImage(item, pBase);
@@ -2537,7 +2527,7 @@ namespace mu
 				// This could happen in the base image has mips, but one of the blended one doesn't.
 				bool bBlendOnlyOneMip = false;
 
-                for (int i=0; i<iterations; ++i)
+                for (int i=0; i< Iterations; ++i)
                 {
                     index.SetFromModelRangeIndex( args.rangeId, i );
                     itemCopy.executionIndex = GetMemory().GetRageIndexIndex(index);
@@ -2566,7 +2556,28 @@ namespace mu
 						bBlendOnlyOneMip = true;
 					}
 
-                    if (args.mask)
+					bool bApplyColorBlendToAlpha = false;
+
+					bool bDone = false;
+
+					if (!args.mask && args.bUseMaskFromBlended
+						&&
+						args.blendType == uint8(EBlendType::BT_BLEND)
+						&&
+						args.blendTypeAlpha == uint8(EBlendType::BT_LIGHTEN) )
+					{
+						// This is a frequent critical-path case because of multilayer projectors.
+						bDone = true;
+						
+						if (pBlended->GetLODCount() < pBase->GetLODCount())
+						{
+							bBlendOnlyOneMip = true;
+						}
+
+						BufferLayerComposite<BlendChannelMasked, LightenChannel, false>(pNew.get(), pBlended.get(), bBlendOnlyOneMip);
+					}
+
+                    if (!bDone && args.mask)
                     {
                         Ptr<const Image> pMask = GetMemory().GetImage( CACHE_ADDRESS(args.mask,itemCopy) );
 
@@ -2577,7 +2588,7 @@ namespace mu
 							pMask = ImageResizeLinear(0, pMask.get(), ResultSize);
 						}
 
-                        if (pMask->GetLODCount()<pBase->GetLODCount())
+                        if (pMask->GetLODCount()<pBase->GetLODCount() || pBlended->GetLODCount() < pBase->GetLODCount())
                         {
 							bBlendOnlyOneMip = true;
 						}
@@ -2591,20 +2602,44 @@ namespace mu
 							pBase = pNew;
 							pNew = new Image(pBase->GetSizeX(), pBase->GetSizeY(), pBase->GetLODCount(), pBase->GetFormat());
 							break;
-                        case EBlendType::BT_SOFTLIGHT: ImageLayerOnBase<SoftLightChannelMasked, SoftLightChannel, false>( pNew.get(), pMask.get(), pBlended.get(), false, bBlendOnlyOneMip); break;
-                        case EBlendType::BT_HARDLIGHT: ImageLayerOnBase<HardLightChannelMasked, HardLightChannel, false>( pNew.get(), pMask.get(), pBlended.get(), false, bBlendOnlyOneMip); break;
-                        case EBlendType::BT_BURN: ImageLayerOnBase<BurnChannelMasked, BurnChannel, false>( pNew.get(), pMask.get(), pBlended.get(), false, bBlendOnlyOneMip); break;
-                        case EBlendType::BT_DODGE: ImageLayerOnBase<DodgeChannelMasked, DodgeChannel, false>( pNew.get(), pMask.get(), pBlended.get(), false, bBlendOnlyOneMip); break;
-                        case EBlendType::BT_SCREEN: ImageLayerOnBase<ScreenChannelMasked, ScreenChannel, false>( pNew.get(), pMask.get(), pBlended.get(), false, bBlendOnlyOneMip); break;
-                        case EBlendType::BT_OVERLAY: ImageLayerOnBase<OverlayChannelMasked, OverlayChannel, false>( pNew.get(), pMask.get(), pBlended.get(), false, bBlendOnlyOneMip); break;
-                        case EBlendType::BT_ALPHA_OVERLAY: ImageLayerOnBase<AlphaOverlayChannelMasked, AlphaOverlayChannel, false>( pNew.get(), pMask.get(), pBlended.get(), false, bBlendOnlyOneMip); break;
-                        case EBlendType::BT_MULTIPLY: ImageLayerOnBase<MultiplyChannelMasked, MultiplyChannel, false>( pNew.get(), pMask.get(), pBlended.get(), false, bBlendOnlyOneMip); break;
-                        case EBlendType::BT_BLEND: ImageLayerOnBase<BlendChannelMasked, BlendChannel, false>( pNew.get(), pMask.get(), pBlended.get(), false, bBlendOnlyOneMip); break;
+                        case EBlendType::BT_SOFTLIGHT: BufferLayer<SoftLightChannelMasked, SoftLightChannel, false>( pNew->GetData(), pNew.get(), pMask.get(), pBlended.get(), bApplyColorBlendToAlpha, bBlendOnlyOneMip); break;
+                        case EBlendType::BT_HARDLIGHT: BufferLayer<HardLightChannelMasked, HardLightChannel, false>(pNew->GetData(), pNew.get(), pMask.get(), pBlended.get(), bApplyColorBlendToAlpha, bBlendOnlyOneMip); break;
+                        case EBlendType::BT_BURN: BufferLayer<BurnChannelMasked, BurnChannel, false>(pNew->GetData(), pNew.get(), pMask.get(), pBlended.get(), bApplyColorBlendToAlpha, bBlendOnlyOneMip); break;
+                        case EBlendType::BT_DODGE: BufferLayer<DodgeChannelMasked, DodgeChannel, false>(pNew->GetData(), pNew.get(), pMask.get(), pBlended.get(), bApplyColorBlendToAlpha, bBlendOnlyOneMip); break;
+                        case EBlendType::BT_SCREEN: BufferLayer<ScreenChannelMasked, ScreenChannel, false>(pNew->GetData(), pNew.get(), pMask.get(), pBlended.get(), bApplyColorBlendToAlpha, bBlendOnlyOneMip); break;
+                        case EBlendType::BT_OVERLAY: BufferLayer<OverlayChannelMasked, OverlayChannel, false>(pNew->GetData(), pNew.get(), pMask.get(), pBlended.get(), bApplyColorBlendToAlpha, bBlendOnlyOneMip); break;
+                        case EBlendType::BT_LIGHTEN: BufferLayer<LightenChannelMasked, LightenChannel, false>(pNew->GetData(), pNew.get(), pMask.get(), pBlended.get(), bApplyColorBlendToAlpha, bBlendOnlyOneMip); break;
+                        case EBlendType::BT_MULTIPLY: BufferLayer<MultiplyChannelMasked, MultiplyChannel, false>(pNew->GetData(), pNew.get(), pMask.get(), pBlended.get(), bApplyColorBlendToAlpha, bBlendOnlyOneMip); break;
+                        case EBlendType::BT_BLEND: BufferLayer<BlendChannelMasked, BlendChannel, false>(pNew->GetData(), pNew.get(), pMask.get(), pBlended.get(), bApplyColorBlendToAlpha, bBlendOnlyOneMip); break;
                         default: check(false);
                         }
 
                     }
-                    else
+					else if (!bDone && args.bUseMaskFromBlended)
+					{
+						if (pBlended->GetLODCount() < pBase->GetLODCount())
+						{
+							bBlendOnlyOneMip = true;
+						}
+
+						switch (EBlendType(args.blendType))
+						{
+						case EBlendType::BT_NORMAL_COMBINE:
+							check(false);
+							break;
+						case EBlendType::BT_SOFTLIGHT: BufferLayerEmbeddedMask<SoftLightChannelMasked, SoftLightChannel, false>(pNew->GetData(), pNew.get(), pBlended.get(), bApplyColorBlendToAlpha, bBlendOnlyOneMip); break;
+						case EBlendType::BT_HARDLIGHT: BufferLayerEmbeddedMask<HardLightChannelMasked, HardLightChannel, false>(pNew->GetData(), pNew.get(), pBlended.get(), bApplyColorBlendToAlpha, bBlendOnlyOneMip); break;
+						case EBlendType::BT_BURN: BufferLayerEmbeddedMask<BurnChannelMasked, BurnChannel, false>(pNew->GetData(), pNew.get(), pBlended.get(), bApplyColorBlendToAlpha, bBlendOnlyOneMip); break;
+						case EBlendType::BT_DODGE: BufferLayerEmbeddedMask<DodgeChannelMasked, DodgeChannel, false>(pNew->GetData(), pNew.get(), pBlended.get(), bApplyColorBlendToAlpha, bBlendOnlyOneMip); break;
+						case EBlendType::BT_SCREEN: BufferLayerEmbeddedMask<ScreenChannelMasked, ScreenChannel, false>(pNew->GetData(), pNew.get(), pBlended.get(), bApplyColorBlendToAlpha, bBlendOnlyOneMip); break;
+						case EBlendType::BT_OVERLAY: BufferLayerEmbeddedMask<OverlayChannelMasked, OverlayChannel, false>(pNew->GetData(), pNew.get(), pBlended.get(), bApplyColorBlendToAlpha, bBlendOnlyOneMip); break;
+						case EBlendType::BT_LIGHTEN: BufferLayerEmbeddedMask<LightenChannelMasked, LightenChannel, false>(pNew->GetData(), pNew.get(), pBlended.get(), bApplyColorBlendToAlpha, bBlendOnlyOneMip); break;
+						case EBlendType::BT_MULTIPLY: BufferLayerEmbeddedMask<MultiplyChannelMasked, MultiplyChannel, false>(pNew->GetData(), pNew.get(), pBlended.get(), bApplyColorBlendToAlpha, bBlendOnlyOneMip); break;
+						case EBlendType::BT_BLEND: BufferLayerEmbeddedMask<BlendChannelMasked, BlendChannel, false>(pNew->GetData(), pNew.get(), pBlended.get(), bApplyColorBlendToAlpha, bBlendOnlyOneMip); break;
+						default: check(false);
+						}
+					}
+                    else if (!bDone)
                     {
                         switch (EBlendType(args.blendType))
                         {
@@ -2615,18 +2650,37 @@ namespace mu
 							pBase = pNew;
 							pNew = new Image(pBase->GetSizeX(), pBase->GetSizeY(), pBase->GetLODCount(), pBase->GetFormat());
 							break;
-                        case EBlendType::BT_SOFTLIGHT: ImageLayerOnBase<SoftLightChannel, false>( pNew.get(), pBlended.get(), false, bBlendOnlyOneMip); break;
-                        case EBlendType::BT_HARDLIGHT: ImageLayerOnBase<HardLightChannel, false>( pNew.get(), pBlended.get(), false, bBlendOnlyOneMip); break;
-                        case EBlendType::BT_BURN: ImageLayerOnBase<BurnChannel, false>( pNew.get(), pBlended.get(), false, bBlendOnlyOneMip); break;
-                        case EBlendType::BT_DODGE: ImageLayerOnBase<DodgeChannel, false>( pNew.get(), pBlended.get(), false, bBlendOnlyOneMip); break;
-                        case EBlendType::BT_SCREEN: ImageLayerOnBase<ScreenChannel, false>( pNew.get(),  pBlended.get(), false, bBlendOnlyOneMip); break;
-                        case EBlendType::BT_OVERLAY: ImageLayerOnBase<OverlayChannel, false>( pNew.get(), pBlended.get(), false, bBlendOnlyOneMip); break;
-                        case EBlendType::BT_ALPHA_OVERLAY: ImageLayerOnBase<AlphaOverlayChannel, false>( pNew.get(), pBlended.get(), false, bBlendOnlyOneMip); break;
-                        case EBlendType::BT_MULTIPLY: ImageLayerOnBase<MultiplyChannel, false>( pNew.get(), pBlended.get(), false, bBlendOnlyOneMip); break;
-                        case EBlendType::BT_BLEND: pNew = pBlended->Clone(); break;
+                        case EBlendType::BT_SOFTLIGHT: BufferLayer<SoftLightChannel, false>(pNew.get(), pNew.get(), pBlended.get(), bApplyColorBlendToAlpha, bBlendOnlyOneMip); break;
+                        case EBlendType::BT_HARDLIGHT: BufferLayer<HardLightChannel, false>(pNew.get(), pNew.get(), pBlended.get(), bApplyColorBlendToAlpha, bBlendOnlyOneMip); break;
+                        case EBlendType::BT_BURN: BufferLayer<BurnChannel, false>(pNew.get(), pNew.get(), pBlended.get(), bApplyColorBlendToAlpha, bBlendOnlyOneMip); break;
+                        case EBlendType::BT_DODGE: BufferLayer<DodgeChannel, false>(pNew.get(), pNew.get(), pBlended.get(), bApplyColorBlendToAlpha, bBlendOnlyOneMip); break;
+                        case EBlendType::BT_SCREEN: BufferLayer<ScreenChannel, false>(pNew.get(), pNew.get(),  pBlended.get(), bApplyColorBlendToAlpha, bBlendOnlyOneMip); break;
+                        case EBlendType::BT_OVERLAY: BufferLayer<OverlayChannel, false>(pNew.get(), pNew.get(), pBlended.get(), bApplyColorBlendToAlpha, bBlendOnlyOneMip); break;
+                        case EBlendType::BT_LIGHTEN: BufferLayer<LightenChannel, false>(pNew.get(), pNew.get(), pBlended.get(), bApplyColorBlendToAlpha, bBlendOnlyOneMip); break;
+                        case EBlendType::BT_MULTIPLY: BufferLayer<MultiplyChannel, false>(pNew.get(), pNew.get(), pBlended.get(), bApplyColorBlendToAlpha, bBlendOnlyOneMip); break;
+                        case EBlendType::BT_BLEND: BufferLayer<BlendChannel, false>(pNew.get(), pNew.get(), pBlended.get(), bApplyColorBlendToAlpha, bBlendOnlyOneMip); break;
                         default: check(false);
                         }
                     }
+
+					// Apply the separate blend operation for alpha
+					if (!bDone && !bApplyColorBlendToAlpha && args.blendTypeAlpha != uint8(EBlendType::BT_NONE) )
+					{
+						// Separate alpha operation ignores the mask.
+						switch (EBlendType(args.blendTypeAlpha))
+						{
+						case EBlendType::BT_SOFTLIGHT: BufferLayerInPlace<SoftLightChannel, false, 1>(pNew.get(), pBlended.get(), bBlendOnlyOneMip, 3, 3); break;
+						case EBlendType::BT_HARDLIGHT: BufferLayerInPlace<HardLightChannel, false, 1>(pNew.get(), pBlended.get(), bBlendOnlyOneMip, 3, 3); break;
+						case EBlendType::BT_BURN: BufferLayerInPlace<BurnChannel, false, 1>(pNew.get(), pBlended.get(), bBlendOnlyOneMip, 3, 3); break;
+						case EBlendType::BT_DODGE: BufferLayerInPlace<DodgeChannel, false, 1>(pNew.get(), pBlended.get(), bBlendOnlyOneMip, 3, 3); break;
+						case EBlendType::BT_SCREEN: BufferLayerInPlace<ScreenChannel, false, 1>(pNew.get(), pBlended.get(), bBlendOnlyOneMip, 3, 3); break;
+						case EBlendType::BT_OVERLAY: BufferLayerInPlace<OverlayChannel, false, 1>(pNew.get(), pBlended.get(), bBlendOnlyOneMip, 3, 3); break;
+						case EBlendType::BT_LIGHTEN: BufferLayerInPlace<LightenChannel, false, 1>(pNew.get(), pBlended.get(), bBlendOnlyOneMip, 3, 3); break;
+						case EBlendType::BT_MULTIPLY: BufferLayerInPlace<MultiplyChannel, false, 1>(pNew.get(), pBlended.get(), bBlendOnlyOneMip, 3, 3); break;
+						case EBlendType::BT_BLEND: BufferLayerInPlace<BlendChannel, false, 1>(pNew.get(), pBlended.get(), bBlendOnlyOneMip, 3, 3); break;
+						default: check(false);
+						}
+					}
 				}
 
 				if (bBlendOnlyOneMip)
@@ -3060,7 +3114,7 @@ namespace mu
                 Ptr<const Layout> pLayout =
                     GetMemory().GetLayout( CACHE_ADDRESS( args.layout, SCHEDULED_OP::FromOpAndOptions(args.layout, item, 0)) );
 
-                SCHEDULED_OP_DATA data;
+                FScheduledOpData data;
                 data.layout = pLayout;
 				int32 dataPos = m_heapData.Add( data );
 
@@ -3121,7 +3175,7 @@ namespace mu
 
                 float bifactor = factor/delta - min;
 
-                SCHEDULED_OP_DATA data;
+                FScheduledOpData data;
                 data.bifactor = bifactor;
 				data.min = FMath::Clamp(min, 0, count - 1);
 				data.max = FMath::Clamp(max, 0, count - 1);
@@ -3158,7 +3212,7 @@ namespace mu
                 }
 
                 // Factor from 0 to 1 between the two targets
-                const SCHEDULED_OP_DATA& data = m_heapData[(size_t)item.customState];
+                const FScheduledOpData& data = m_heapData[(size_t)item.customState];
                 float bifactor = data.bifactor;
                 int min = data.min;
                 int max = data.max;
@@ -5475,7 +5529,7 @@ namespace mu
 
 			case 1:
 			{
-				const SCHEDULED_OP_DATA& SecondStageData = m_heapData[ item.customState ];
+				const FScheduledOpData& SecondStageData = m_heapData[ item.customState ];
 				FImageDesc& ResultAndBaseDesc = m_heapImageDesc[SecondStageData.min];
 				const FImageDesc& SourceDesc = m_heapImageDesc[SecondStageData.max];
 				ResultAndBaseDesc.m_size = SourceDesc.m_size;
diff --git a/Engine/Plugins/Experimental/Mutable/Source/MutableRuntime/Private/MuR/CodeRunner.h b/Engine/Plugins/Experimental/Mutable/Source/MutableRuntime/Private/MuR/CodeRunner.h
--- a/Engine/Plugins/Experimental/Mutable/Source/MutableRuntime/Private/MuR/CodeRunner.h
+++ b/Engine/Plugins/Experimental/Mutable/Source/MutableRuntime/Private/MuR/CodeRunner.h
@@ -54,7 +54,7 @@ class RangeIndex;
 
         //! Type of data sometimes stored in the code runner heap to pass info between
         //! operation stages.
-        struct SCHEDULED_OP_DATA
+        struct FScheduledOpData
         {
             float bifactor=0.0f;
             int32 min=0, max=0;
@@ -92,7 +92,7 @@ class RangeIndex;
 
         //! Heap of intermediate data pushed by some instructions and referred by others.
         //! It is not released until no operations are pending.
-		TArray< SCHEDULED_OP_DATA > m_heapData;
+		TArray< FScheduledOpData > m_heapData;
 		TArray< FImageDesc > m_heapImageDesc;
 
 	public:
diff --git a/Engine/Plugins/Experimental/Mutable/Source/MutableRuntime/Private/MuR/CodeRunnerTasked.cpp b/Engine/Plugins/Experimental/Mutable/Source/MutableRuntime/Private/MuR/CodeRunnerTasked.cpp
--- a/Engine/Plugins/Experimental/Mutable/Source/MutableRuntime/Private/MuR/CodeRunnerTasked.cpp
+++ b/Engine/Plugins/Experimental/Mutable/Source/MutableRuntime/Private/MuR/CodeRunnerTasked.cpp
@@ -20,20 +20,12 @@
 #include "MuR/ModelPrivate.h"
 #include "MuR/MutableMath.h"
 #include "MuR/MutableTrace.h"
-#include "MuR/OpImageAlphaOverlay.h"
 #include "MuR/OpImageBlend.h"
-#include "MuR/OpImageBurn.h"
 #include "MuR/OpImageCompose.h"
-#include "MuR/OpImageDodge.h"
-#include "MuR/OpImageHardLight.h"
 #include "MuR/OpImageMipmap.h"
-#include "MuR/OpImageMultiply.h"
 #include "MuR/OpImageNormalCombine.h"
-#include "MuR/OpImageOverlay.h"
 #include "MuR/OpImagePixelFormat.h"
 #include "MuR/OpImageResize.h"
-#include "MuR/OpImageScreen.h"
-#include "MuR/OpImageSoftLight.h"
 #include "MuR/OpImageSwizzle.h"
 #include "MuR/Operations.h"
 #include "MuR/Platform.h"
@@ -440,7 +432,7 @@ namespace mu
 				ImagePixelFormat(m_imageCompressionQuality, m_blended.get(), m_base->GetFormat());
 		}
 
-		bool applyToAlpha = (Args.flags & OP::ImageLayerArgs::F_APPLY_TO_ALPHA) != 0;
+		bool bApplyToAlpha = (Args.flags & OP::ImageLayerArgs::F_APPLY_TO_ALPHA) != 0;
 
 		ImagePtr pNew = new Image(m_base->GetSizeX(), m_base->GetSizeY(),
 			m_base->GetLODCount(), m_base->GetFormat());
@@ -478,6 +470,8 @@ namespace mu
 				m_blended = pDest;
 			}
 
+			bool bOnlyOneMip = false;
+
 			if (Args.mask)
 			{
 				if (m_base->GetSize() != m_mask->GetSize())
@@ -508,16 +502,16 @@ namespace mu
 
 				switch (EBlendType(Args.blendType))
 				{
-				case EBlendType::BT_NORMAL_COMBINE: ImageNormalCombine(pNew.get(), m_base.get(), m_mask.get(), m_blended.get(), false); break;
-				case EBlendType::BT_SOFTLIGHT: ImageSoftLight(pNew.get(), m_base.get(), m_mask.get(), m_blended.get(), applyToAlpha); break;
-				case EBlendType::BT_HARDLIGHT: ImageHardLight(pNew.get(), m_base.get(), m_mask.get(), m_blended.get(), applyToAlpha); break;
-				case EBlendType::BT_BURN: ImageBurn(pNew.get(), m_base.get(), m_mask.get(), m_blended.get(), applyToAlpha); break;
-				case EBlendType::BT_DODGE: ImageDodge(pNew.get(), m_base.get(), m_mask.get(), m_blended.get(), applyToAlpha); break;
-				case EBlendType::BT_SCREEN: ImageScreen(pNew.get(), m_base.get(), m_mask.get(), m_blended.get(), applyToAlpha); break;
-				case EBlendType::BT_OVERLAY: ImageOverlay(pNew.get(), m_base.get(), m_mask.get(), m_blended.get(), applyToAlpha); break;
-				case EBlendType::BT_ALPHA_OVERLAY: ImageAlphaOverlay(pNew.get(), m_base.get(), m_mask.get(), m_blended.get(), applyToAlpha); break;
-				case EBlendType::BT_MULTIPLY: ImageMultiply(pNew.get(), m_base.get(), m_mask.get(), m_blended.get(), applyToAlpha); break;
-				case EBlendType::BT_BLEND: ImageBlend(pNew.get(), m_base.get(), m_mask.get(), m_blended.get(), applyToAlpha); break;
+				case EBlendType::BT_NORMAL_COMBINE: ImageNormalCombine(pNew.get(), m_base.get(), m_mask.get(), m_blended.get(), bOnlyOneMip); break;
+				case EBlendType::BT_SOFTLIGHT: BufferLayer<SoftLightChannelMasked, SoftLightChannel, true>(pNew->GetData(), m_base.get(), m_mask.get(), m_blended.get(), bApplyToAlpha, bOnlyOneMip); break;
+				case EBlendType::BT_HARDLIGHT: BufferLayer<HardLightChannelMasked, HardLightChannel, true>(pNew->GetData(), m_base.get(), m_mask.get(), m_blended.get(), bApplyToAlpha, bOnlyOneMip); break;
+				case EBlendType::BT_BURN: BufferLayer<BurnChannelMasked, BurnChannel, true>(pNew->GetData(), m_base.get(), m_mask.get(), m_blended.get(), bApplyToAlpha, bOnlyOneMip); break;
+				case EBlendType::BT_DODGE: BufferLayer<DodgeChannelMasked, DodgeChannel, true>(pNew->GetData(), m_base.get(), m_mask.get(), m_blended.get(), bApplyToAlpha, bOnlyOneMip); break;
+				case EBlendType::BT_SCREEN: BufferLayer<ScreenChannelMasked, ScreenChannel, true>(pNew->GetData(), m_base.get(), m_mask.get(), m_blended.get(), bApplyToAlpha, bOnlyOneMip); break;
+				case EBlendType::BT_OVERLAY: BufferLayer<OverlayChannelMasked, OverlayChannel, true>(pNew->GetData(), m_base.get(), m_mask.get(), m_blended.get(), bApplyToAlpha, bOnlyOneMip); break;
+				case EBlendType::BT_LIGHTEN: BufferLayer<LightenChannelMasked, LightenChannel, true>(pNew->GetData(), m_base.get(), m_mask.get(), m_blended.get(), bApplyToAlpha, bOnlyOneMip); break;
+				case EBlendType::BT_MULTIPLY: BufferLayer<MultiplyChannelMasked, MultiplyChannel, true>(pNew->GetData(), m_base.get(), m_mask.get(), m_blended.get(), bApplyToAlpha, bOnlyOneMip); break;
+				case EBlendType::BT_BLEND: BufferLayer<BlendChannelMasked, BlendChannel, true>(pNew->GetData(), m_base.get(), m_mask.get(), m_blended.get(), bApplyToAlpha, bOnlyOneMip); break;
 				default: check(false);
 				}
 
@@ -526,16 +520,16 @@ namespace mu
 			{
 				switch (EBlendType(Args.blendType))
 				{
-				case EBlendType::BT_NORMAL_COMBINE: ImageNormalCombine(pNew.get(), m_base.get(), m_blended.get(), false); break;
-				case EBlendType::BT_SOFTLIGHT: ImageSoftLight(pNew.get(), m_base.get(), m_blended.get(), applyToAlpha); break;
-				case EBlendType::BT_HARDLIGHT: ImageHardLight(pNew.get(), m_base.get(), m_blended.get(), applyToAlpha); break;
-				case EBlendType::BT_BURN: ImageBurn(pNew.get(), m_base.get(), m_blended.get(), applyToAlpha); break;
-				case EBlendType::BT_DODGE: ImageDodge(pNew.get(), m_base.get(), m_blended.get(), applyToAlpha); break;
-				case EBlendType::BT_SCREEN: ImageScreen(pNew.get(), m_base.get(), m_blended.get(), applyToAlpha); break;
-				case EBlendType::BT_OVERLAY: ImageOverlay(pNew.get(), m_base.get(), m_blended.get(), applyToAlpha); break;
-				case EBlendType::BT_ALPHA_OVERLAY: ImageAlphaOverlay(pNew.get(), m_base.get(), m_blended.get(), applyToAlpha); break;
-				case EBlendType::BT_MULTIPLY: ImageMultiply(pNew.get(), m_base.get(), m_blended.get(), applyToAlpha); break;
-				case EBlendType::BT_BLEND: ImageBlend(pNew.get(), m_base.get(), m_blended.get(), applyToAlpha); break;
+				case EBlendType::BT_NORMAL_COMBINE: ImageNormalCombine(pNew.get(), m_base.get(), m_blended.get(), bOnlyOneMip); break;
+				case EBlendType::BT_SOFTLIGHT: BufferLayer<SoftLightChannel, true>(pNew.get(), m_base.get(), m_blended.get(), bApplyToAlpha, bOnlyOneMip); break;
+				case EBlendType::BT_HARDLIGHT: BufferLayer<HardLightChannel, true>(pNew.get(), m_base.get(), m_blended.get(), bApplyToAlpha, bOnlyOneMip); break;
+				case EBlendType::BT_BURN: BufferLayer<BurnChannel, true>(pNew.get(), m_base.get(), m_blended.get(), bApplyToAlpha, bOnlyOneMip); break;
+				case EBlendType::BT_DODGE: BufferLayer<DodgeChannel, true>(pNew.get(), m_base.get(), m_blended.get(), bApplyToAlpha, bOnlyOneMip); break;
+				case EBlendType::BT_SCREEN: BufferLayer<ScreenChannel, true>(pNew.get(), m_base.get(), m_blended.get(), bApplyToAlpha, bOnlyOneMip); break;
+				case EBlendType::BT_OVERLAY: BufferLayer<OverlayChannel, true>(pNew.get(), m_base.get(), m_blended.get(), bApplyToAlpha, bOnlyOneMip); break;
+				case EBlendType::BT_LIGHTEN: BufferLayer<LightenChannel, true>(pNew.get(), m_base.get(), m_blended.get(), bApplyToAlpha, bOnlyOneMip); break;
+				case EBlendType::BT_MULTIPLY: BufferLayer<MultiplyChannel, true>(pNew.get(), m_base.get(), m_blended.get(), bApplyToAlpha, bOnlyOneMip); break;
+				case EBlendType::BT_BLEND: BufferLayer<BlendChannel, true>(pNew.get(), m_base.get(), m_blended.get(), bApplyToAlpha, bOnlyOneMip); break;
 				default: check(false);
 				}
 			}
@@ -605,10 +599,8 @@ namespace mu
 		MUTABLE_CPUPROFILER_SCOPE(FImageLayerColourTask);
 
 		// Warning: This runs in a worker thread.
-
-		ImagePtr pNew = new Image(m_base->GetSizeX(), m_base->GetSizeY(),
-			m_base->GetLODCount(), m_base->GetFormat());
-
+		ImagePtr pNew;
+		
 		if (Args.mask)
 		{
 			if (m_base->GetSize() != m_mask->GetSize())
@@ -635,40 +627,104 @@ namespace mu
 					&scratch, settings);
 				m_mask = pDest;
 			}
+		}
+
+		bool bOnlyOneMip = false;
+
+		// Does it apply to colour?
+		if (EBlendType(Args.blendType) != EBlendType::BT_NONE)
+		{
+			pNew = new Image(m_base->GetSizeX(), m_base->GetSizeY(), m_base->GetLODCount(), m_base->GetFormat());
+
+			if (Args.mask)
+			{			
+				// \todo: precalculated tables for softlight
+				switch (EBlendType(Args.blendType))
+				{
+				case EBlendType::BT_NORMAL_COMBINE: ImageNormalCombine(pNew.get(), m_base.get(), m_mask.get(), vec3f(m_col)); break;
+				case EBlendType::BT_SOFTLIGHT: BufferLayerColour<SoftLightChannelMasked, SoftLightChannel, false>(pNew->GetData(), m_base.get(), m_mask.get(), vec3f(m_col)); break;
+				case EBlendType::BT_HARDLIGHT: BufferLayerColour<HardLightChannelMasked, HardLightChannel, false>(pNew->GetData(), m_base.get(), m_mask.get(), vec3f(m_col)); break;
+				case EBlendType::BT_BURN: BufferLayerColour<BurnChannelMasked, BurnChannel, false>(pNew->GetData(), m_base.get(), m_mask.get(), vec3f(m_col)); break;
+				case EBlendType::BT_DODGE: BufferLayerColour<DodgeChannelMasked, DodgeChannel, false>(pNew->GetData(), m_base.get(), m_mask.get(), vec3f(m_col)); break;
+				case EBlendType::BT_SCREEN: BufferLayerColour<ScreenChannelMasked, ScreenChannel, false>(pNew->GetData(), m_base.get(), m_mask.get(), vec3f(m_col)); break;
+				case EBlendType::BT_OVERLAY: BufferLayerColour<OverlayChannelMasked, OverlayChannel, false>(pNew->GetData(), m_base.get(), m_mask.get(), vec3f(m_col)); break;
+				case EBlendType::BT_LIGHTEN: BufferLayerColour<LightenChannelMasked, LightenChannel, false>(pNew->GetData(), m_base.get(), m_mask.get(), vec3f(m_col)); break;
+				case EBlendType::BT_MULTIPLY: BufferLayerColour<MultiplyChannelMasked, MultiplyChannel, false>(pNew->GetData(), m_base.get(), m_mask.get(), vec3f(m_col)); break;
+				case EBlendType::BT_BLEND: BufferLayerColour<BlendChannelMasked, BlendChannel, false>(pNew->GetData(), m_base.get(), m_mask.get(), vec3f(m_col)); break;
+				default: check(false);
+				}
 
-			// \todo: precalculated tables for softlight
-			switch (EBlendType(Args.blendType))
+			}
+			else
 			{
-			case EBlendType::BT_NORMAL_COMBINE: ImageNormalCombine(pNew.get(), m_base.get(), m_mask.get(), vec3f(m_col)); break;
-			case EBlendType::BT_SOFTLIGHT: ImageSoftLight(pNew.get(), m_base.get(), m_mask.get(), vec3f(m_col), nullptr); break;
-			case EBlendType::BT_HARDLIGHT: ImageHardLight(pNew.get(), m_base.get(), m_mask.get(), vec3f(m_col)); break;
-			case EBlendType::BT_BURN: ImageBurn(pNew.get(), m_base.get(), m_mask.get(), vec3f(m_col)); break;
-			case EBlendType::BT_DODGE: ImageDodge(pNew.get(), m_base.get(), m_mask.get(), vec3f(m_col)); break;
-			case EBlendType::BT_SCREEN: ImageScreen(pNew.get(), m_base.get(), m_mask.get(), vec3f(m_col)); break;
-			case EBlendType::BT_OVERLAY: ImageOverlay(pNew.get(), m_base.get(), m_mask.get(), vec3f(m_col)); break;
-			case EBlendType::BT_ALPHA_OVERLAY: ImageAlphaOverlay(pNew.get(), m_base.get(), m_mask.get(), vec3f(m_col)); break;
-			case EBlendType::BT_MULTIPLY: ImageMultiply(pNew.get(), m_base.get(), m_mask.get(), vec3f(m_col)); break;
-			case EBlendType::BT_BLEND: ImageBlend(pNew.get(), m_base.get(), m_mask.get(), vec3f(m_col)); break;
-			default: check(false);
+				switch (EBlendType(Args.blendType))
+				{
+				case EBlendType::BT_NORMAL_COMBINE: ImageNormalCombine(pNew.get(), m_base.get(), vec3f(m_col)); break;
+				case EBlendType::BT_SOFTLIGHT: BufferLayerColour<SoftLightChannel, false>(pNew.get(), m_base.get(), vec3f(m_col)); break;
+				case EBlendType::BT_HARDLIGHT: BufferLayerColour<HardLightChannel, false>(pNew.get(), m_base.get(), vec3f(m_col)); break;
+				case EBlendType::BT_BURN: BufferLayerColour<BurnChannel, false>(pNew.get(), m_base.get(), vec3f(m_col)); break;
+				case EBlendType::BT_DODGE: BufferLayerColour<DodgeChannel, false>(pNew.get(), m_base.get(), vec3f(m_col)); break;
+				case EBlendType::BT_SCREEN: BufferLayerColour<ScreenChannel, false>(pNew.get(), m_base.get(), vec3f(m_col)); break;
+				case EBlendType::BT_OVERLAY: BufferLayerColour<OverlayChannel, false>(pNew.get(), m_base.get(), vec3f(m_col)); break;
+				case EBlendType::BT_LIGHTEN: BufferLayerColour<LightenChannel, false>(pNew.get(), m_base.get(), vec3f(m_col)); break;
+				case EBlendType::BT_MULTIPLY: BufferLayerColour<MultiplyChannel, false>(pNew.get(), m_base.get(), vec3f(m_col)); break;
+				case EBlendType::BT_BLEND: FillPlainColourImage(pNew.get(), m_col); break;
+				default: check(false);
+				}
 			}
-
 		}
-		else
+
+		// Does it apply to alpha?
+		if (EBlendType(Args.blendTypeAlpha) != EBlendType::BT_NONE)
 		{
-			switch (EBlendType(Args.blendType))
+			if (!pNew)
+			{
+				pNew = mu::CloneOrTakeOver<mu::Image>(m_base.get());
+			}
+
+			uint32 AlphaOffset = 3;
+			if (pNew->GetFormat() == EImageFormat::IF_L_UBYTE)
 			{
-			case EBlendType::BT_NORMAL_COMBINE: ImageNormalCombine(pNew.get(), m_base.get(), vec3f(m_col)); break;
-			case EBlendType::BT_SOFTLIGHT: ImageSoftLight(pNew.get(), m_base.get(), vec3f(m_col)); break;
-			case EBlendType::BT_HARDLIGHT: ImageHardLight(pNew.get(), m_base.get(), vec3f(m_col)); break;
-			case EBlendType::BT_BURN: ImageBurn(pNew.get(), m_base.get(), vec3f(m_col)); break;
-			case EBlendType::BT_DODGE: ImageDodge(pNew.get(), m_base.get(), vec3f(m_col)); break;
-			case EBlendType::BT_SCREEN: ImageScreen(pNew.get(), m_base.get(), vec3f(m_col)); break;
-			case EBlendType::BT_OVERLAY: ImageOverlay(pNew.get(), m_base.get(), vec3f(m_col)); break;
-			case EBlendType::BT_ALPHA_OVERLAY: ImageAlphaOverlay(pNew.get(), m_base.get(), vec3f(m_col)); break;
-			case EBlendType::BT_MULTIPLY: ImageMultiply(pNew.get(), m_base.get(), vec3f(m_col)); break;
-			case EBlendType::BT_BLEND: FillPlainColourImage(pNew.get(), m_col); break;
-			default: check(false);
+				AlphaOffset = 1;
 			}
+
+			if (Args.mask)
+			{
+				// \todo: precalculated tables for softlight
+				switch (EBlendType(Args.blendTypeAlpha))
+				{
+				case EBlendType::BT_NORMAL_COMBINE: check(false); break;
+				case EBlendType::BT_SOFTLIGHT: BufferLayerColourInPlace<SoftLightChannelMasked, SoftLightChannel, false, 1>(pNew.get(), m_mask.get(), vec3f(m_col), bOnlyOneMip, 3); break;
+				case EBlendType::BT_HARDLIGHT: BufferLayerColourInPlace<HardLightChannelMasked, HardLightChannel, false, 1>(pNew.get(), m_mask.get(), vec3f(m_col), bOnlyOneMip, 3); break;
+				case EBlendType::BT_BURN: BufferLayerColourInPlace<BurnChannelMasked, BurnChannel, false, 1>(pNew.get(), m_mask.get(), vec3f(m_col), bOnlyOneMip, 3); break;
+				case EBlendType::BT_DODGE: BufferLayerColourInPlace<DodgeChannelMasked, DodgeChannel, false, 1>(pNew.get(), m_mask.get(), vec3f(m_col), bOnlyOneMip, 3); break;
+				case EBlendType::BT_SCREEN: BufferLayerColourInPlace<ScreenChannelMasked, ScreenChannel, false, 1>(pNew.get(), m_mask.get(), vec3f(m_col), bOnlyOneMip, 3); break;
+				case EBlendType::BT_OVERLAY: BufferLayerColourInPlace<OverlayChannelMasked, OverlayChannel, false, 1>(pNew.get(), m_mask.get(), vec3f(m_col), bOnlyOneMip, 3); break;
+				case EBlendType::BT_LIGHTEN: BufferLayerColourInPlace<LightenChannelMasked, LightenChannel, false, 1>(pNew.get(), m_mask.get(), vec3f(m_col), bOnlyOneMip, 3); break;
+				case EBlendType::BT_MULTIPLY: BufferLayerColourInPlace<MultiplyChannelMasked, MultiplyChannel, false, 1>(pNew.get(), m_mask.get(), vec3f(m_col), bOnlyOneMip, 3); break;
+				case EBlendType::BT_BLEND: BufferLayerColourInPlace<BlendChannelMasked, BlendChannel, false, 1>(pNew.get(), m_mask.get(), vec3f(m_col), bOnlyOneMip, 3); break;
+				default: check(false);
+				}
+
+			}
+			else
+			{
+				switch (EBlendType(Args.blendTypeAlpha))
+				{
+				case EBlendType::BT_NORMAL_COMBINE: check(false); break;
+				case EBlendType::BT_SOFTLIGHT: BufferLayerColourInPlace<SoftLightChannel, false, 1>(pNew.get(), vec3f(m_col), bOnlyOneMip, 3); break;
+				case EBlendType::BT_HARDLIGHT: BufferLayerColourInPlace<HardLightChannel, false, 1>(pNew.get(), vec3f(m_col), bOnlyOneMip, 3); break;
+				case EBlendType::BT_BURN: BufferLayerColourInPlace<BurnChannel, false, 1>(pNew.get(), vec3f(m_col), bOnlyOneMip, 3); break;
+				case EBlendType::BT_DODGE: BufferLayerColourInPlace<DodgeChannel, false, 1>(pNew.get(), vec3f(m_col), bOnlyOneMip, 3); break;
+				case EBlendType::BT_SCREEN: BufferLayerColourInPlace<ScreenChannel, false, 1>(pNew.get(), vec3f(m_col), bOnlyOneMip, 3); break;
+				case EBlendType::BT_OVERLAY: BufferLayerColourInPlace<OverlayChannel, false, 1>(pNew.get(), vec3f(m_col), bOnlyOneMip, 3); break;
+				case EBlendType::BT_LIGHTEN: BufferLayerColourInPlace<LightenChannel, false, 1>(pNew.get(), vec3f(m_col), bOnlyOneMip, 3); break;
+				case EBlendType::BT_MULTIPLY: BufferLayerColourInPlace<MultiplyChannel, false, 1>(pNew.get(), vec3f(m_col), bOnlyOneMip, 3); break;
+				case EBlendType::BT_BLEND: BufferLayerColourInPlace<BlendChannel, false, 1>(pNew.get(), vec3f(m_col), bOnlyOneMip, 3); break;
+				default: check(false);
+				}
+			}
+
 		}
 
 		m_base = nullptr;
diff --git a/Engine/Plugins/Experimental/Mutable/Source/MutableRuntime/Private/MuR/OpImageAlphaOverlay.h b/Engine/Plugins/Experimental/Mutable/Source/MutableRuntime/Private/MuR/OpImageAlphaOverlay.h
deleted file mode 100644
diff --git a/Engine/Plugins/Experimental/Mutable/Source/MutableRuntime/Private/MuR/OpImageBlend.h b/Engine/Plugins/Experimental/Mutable/Source/MutableRuntime/Private/MuR/OpImageBlend.h
--- a/Engine/Plugins/Experimental/Mutable/Source/MutableRuntime/Private/MuR/OpImageBlend.h
+++ b/Engine/Plugins/Experimental/Mutable/Source/MutableRuntime/Private/MuR/OpImageBlend.h
@@ -9,64 +9,168 @@
 namespace mu
 {
 
-	inline unsigned BlendChannel( unsigned // base
-		, unsigned blended )
+	inline uint32 BlendChannel( uint32 // base
+		, uint32 blended )
 	{
 		return blended;
 	}
 
-	inline unsigned BlendChannelMasked( unsigned base, unsigned blended, unsigned mask )
+	inline uint32 BlendChannelMasked( uint32 base, uint32 blended, uint32 mask )
 	{
-		unsigned blend = blended;
-		unsigned masked = ( ( ( 255 - mask ) * base ) + ( mask * blend ) ) / 255;
+		uint32 blend = blended;
+		uint32 masked = ( ( ( 255 - mask ) * base ) + ( mask * blend ) ) / 255;
 		return masked;
 	}
 
+	//! Blend a subimage on the base using a mask.
+	inline void ImageBlendOnBaseNoAlpha( Image* pBase,
+										 const Image* pMask,
+										 const Image* pBlended,
+										 const box< vec2<int> >& rect )
+	{
+		ImageLayerOnBaseNoAlpha< BlendChannel, false>
+				( pBase, pMask, pBlended, rect );
+	}
+
+
+	//---------------------------------------------------------------------------------------------
+	inline uint32 ScreenChannelMasked(uint32 base, uint32 blended, uint32 mask)
+	{
+		// R = 1 - (1-Base) × (1-Blend)
+		uint32 screen = 255 - (((255 - base) * (255 - blended)) >> 8);
+		uint32 masked = (((255 - mask) * base) + (mask * screen)) >> 8;
+		return masked;
+	}
+
+	inline uint32 ScreenChannel(uint32 base, uint32 blended)
+	{
+		// R = 1 - (1-Base) × (1-Blend)
+		uint32 screen = 255 - (((255 - base) * (255 - blended)) >> 8);
+		return screen;
+	}
 
-    inline void ImageBlend( Image* pResult, const Image* pBase, const Image* pMask, vec3<float> col )
+	//---------------------------------------------------------------------------------------------
+	inline uint32 SoftLightChannel(uint32 base, uint32 blended)
 	{
-        ImageLayerColour<BlendChannelMasked, BlendChannel, false>( pResult, pBase, pMask, col );
+		// gimp-like
+		uint32 mix = (base * blended) >> 8;
+		uint32 screen = 255 - (((255 - base) * (255 - blended)) >> 8);
+		uint32 softlight = (((255 - base) * mix) + (base * screen)) >> 8;
+		return softlight;
 	}
 
 
-    inline void ImageBlend( Image* pResult, const Image* pBase, const Image* pMask,
-                            const Image* pBlended,
-                            bool applyToAlpha )
+	inline uint32 SoftLightChannelMasked(uint32 base, uint32 blended, uint32 mask)
 	{
-        ImageLayer<BlendChannelMasked, BlendChannel, false>
-                ( pResult, pBase, pMask, pBlended, applyToAlpha );
+		uint32 softlight = SoftLightChannel(base, blended);
+		uint32 masked = (((255 - mask) * base) + (mask * softlight)) >> 8;
+		return masked;
 	}
 
+	//---------------------------------------------------------------------------------------------
+	inline uint32 HardLightChannel(uint32 base, uint32 blended)
+	{
+		// gimp-like
+
+
+		// photoshop-like
+		// if (Blend > ½) R = 1 - (1-Base) × (1-2×(Blend-½))
+		// if (Blend <= ½) R = Base × (2×Blend)
+		uint32 hardlight = blended > 128
+			? 255 - (((255 - base) * (255 - 2 * (blended - 128))) >> 8)
+			: (base * (2 * blended)) >> 8;
 
-    inline void ImageBlend( Image* pResult, const Image* pBase, const Image* pBlended,
-                            bool applyToAlpha )
-    {
-        return ImageLayer<BlendChannel,false>( pResult, pBase, pBlended, applyToAlpha );
-    }
+		hardlight = FMath::Min(255u, hardlight);
 
+		return hardlight;
+	}
 
-	inline void ImageBlendOnBase( Image* pBase, const Image* pMask, vec3<float> col )
+	inline uint32 HardLightChannelMasked(uint32 base, uint32 blended, uint32 mask)
 	{
-		ImageLayerColourOnBase<BlendChannelMasked, BlendChannel, false>( pBase, pMask, col );
+		uint32 hardlight = HardLightChannel(base, blended);
+
+		uint32 masked = (((255 - mask) * base) + (mask * hardlight)) >> 8;
+		return masked;
 	}
 
+	//---------------------------------------------------------------------------------------------
+	inline uint32 BurnChannel(uint32 base, uint32 blended)
+	{
+		// R = 1 - (1-Base) / Blend
+		uint32 burn =
+			FMath::Min(255,
+				FMath::Max(0,
+					255 - (((255 - (int)base) << 8) / ((int)blended + 1))
+				)
+			);
+		return burn;
+	}
 
-    inline void ImageBlendOnBase( Image* pBase, const Image* pMask, const Image* pBlended,
-                                  bool applyToAlpha, bool bOnlyFirstLOD )
+	inline uint32 BurnChannelMasked(uint32 base, uint32 blended, uint32 mask)
 	{
-        ImageLayerOnBase<BlendChannelMasked, BlendChannel, false>
-                 ( pBase, pMask, pBlended, applyToAlpha, bOnlyFirstLOD );
+		uint32 burn = BurnChannel(base, blended);
+		uint32 masked = (((255 - mask) * base) + (mask * burn)) >> 8;
+		return masked;
 	}
 
-	//! Blend a subimage on the base using a mask.
-	inline void ImageBlendOnBaseNoAlpha( Image* pBase,
-										 const Image* pMask,
-										 const Image* pBlended,
-										 const box< vec2<int> >& rect )
+	//---------------------------------------------------------------------------------------------
+	inline uint32 DodgeChannelMasked(uint32 base, uint32 blended, uint32 mask)
 	{
-		ImageLayerOnBaseNoAlpha< BlendChannel, false>
-				( pBase, pMask, pBlended, rect );
+		// R = Base / (1-Blend)
+		uint32 dodge = (base << 8) / (256 - blended);
+		uint32 masked = (((255 - mask) * base) + (mask * dodge)) >> 8;
+		return masked;
+	}
+
+	inline uint32 DodgeChannel(uint32 base, uint32 blended)
+	{
+		// R = Base / (1-Blend)
+		uint32 dodge = (base << 8) / (256 - blended);
+		return dodge;
+	}
+
+	//---------------------------------------------------------------------------------------------
+	inline uint32 LightenChannelMasked(uint32 base, uint32 blended, uint32 mask)
+	{
+		uint32 overlay = base + (blended * uint32(255 - base) >> 8);
+		uint32 masked = (((255 - mask) * base) + (mask * overlay)) >> 8;
+		return masked;
+	}
+
+
+	inline uint32 LightenChannel(uint32 base, uint32 blended)
+	{
+		uint32 overlay = base + (blended * uint32(255 - base) >> 8);
+		return overlay;
+	}
+
+	//---------------------------------------------------------------------------------------------
+	inline uint32 MultiplyChannelMasked(uint32 base, uint32 blended, uint32 mask)
+	{
+		uint32 multiply = (base * blended) >> 8;
+		uint32 masked = (((255 - mask) * base) + (mask * multiply)) >> 8;
+		return masked;
+	}
+
+	inline uint32 MultiplyChannel(uint32 base, uint32 blended)
+	{
+		uint32 multiply = (base * blended) >> 8;
+		return multiply;
 	}
 
+	//---------------------------------------------------------------------------------------------
+	inline uint32 OverlayChannelMasked(uint32 base, uint32 blended, uint32 mask)
+	{
+		uint32 overlay = (base * (base + ((2 * blended * (255 - base)) >> 8))) >> 8;
+		uint32 masked = (((255 - mask) * base) + (mask * overlay)) >> 8;
+		return masked;
+	}
+
+
+	inline uint32 OverlayChannel(uint32 base, uint32 blended)
+	{
+		uint32 overlay = (base * (base + ((2 * blended * (255 - base)) >> 8))) >> 8;
+		return overlay;
+	}
 
 }
diff --git a/Engine/Plugins/Experimental/Mutable/Source/MutableRuntime/Private/MuR/OpImageBurn.h b/Engine/Plugins/Experimental/Mutable/Source/MutableRuntime/Private/MuR/OpImageBurn.h
deleted file mode 100644
diff --git a/Engine/Plugins/Experimental/Mutable/Source/MutableRuntime/Private/MuR/OpImageDodge.h b/Engine/Plugins/Experimental/Mutable/Source/MutableRuntime/Private/MuR/OpImageDodge.h
deleted file mode 100644
diff --git a/Engine/Plugins/Experimental/Mutable/Source/MutableRuntime/Private/MuR/OpImageHardLight.h b/Engine/Plugins/Experimental/Mutable/Source/MutableRuntime/Private/MuR/OpImageHardLight.h
deleted file mode 100644
diff --git a/Engine/Plugins/Experimental/Mutable/Source/MutableRuntime/Private/MuR/OpImageLayer.h b/Engine/Plugins/Experimental/Mutable/Source/MutableRuntime/Private/MuR/OpImageLayer.h
--- a/Engine/Plugins/Experimental/Mutable/Source/MutableRuntime/Private/MuR/OpImageLayer.h
+++ b/Engine/Plugins/Experimental/Mutable/Source/MutableRuntime/Private/MuR/OpImageLayer.h
@@ -16,7 +16,7 @@ namespace mu
 	//---------------------------------------------------------------------------------------------
 	template< unsigned (*BLEND_FUNC)(unsigned,unsigned),
 			  bool CLAMP >
-    inline void ImageLayerColour( Image* pResult, const Image* pBase, vec3<float> col )
+    inline void BufferLayerColour( Image* pResult, const Image* pBase, vec3<float> col )
 	{
         check( pResult->GetFormat() == pBase->GetFormat() );
         check( pResult->GetSizeX() == pBase->GetSizeX() );
@@ -186,29 +186,41 @@ namespace mu
 
 	//---------------------------------------------------------------------------------------------
 	template< unsigned (*BLEND_FUNC_MASKED)(unsigned,unsigned,unsigned),
-			  unsigned (*BLEND_FUNC)(unsigned,unsigned),
-			  bool CLAMP,
-			  int NC >
+		unsigned (*BLEND_FUNC)(unsigned,unsigned),
+		bool CLAMP,
+		// Number of total channels to actually process
+		uint32 CHANNELS_TO_BLEND,
+		// Number of total channels in the base image
+		int32 BASE_CHANNEL_STRIDE >
     inline void BufferLayerColourFormat( uint8* pStartDestBuf,
-								   const Image* pBase,
-								   const Image* pMask,
-								   vec3<float> col )
+		const Image* pBase,
+		const Image* pMask,
+		vec3<float> col,
+		uint32 BaseOffset,
+		bool bOnlyOneMip )
 	{
+		check(CHANNELS_TO_BLEND+BaseOffset<=BASE_CHANNEL_STRIDE);
+
 		unsigned top[3];
 		top[0] = (unsigned)(255 * col[0]);
 		top[1] = (unsigned)(255 * col[1]);
 		top[2] = (unsigned)(255 * col[2]);
 
-		uint8* pDestBuf = pStartDestBuf;
+		uint8* pDestBuf = pStartDestBuf+ BaseOffset;
         const uint8* pMaskBuf = pMask->GetData();
-        const uint8* pBaseBuf = pBase->GetData();
+        const uint8* pBaseBuf = pBase->GetData()+ BaseOffset;
 		EImageFormat maskFormat = pMask->GetFormat();
 
 		int32 LODCount = pBase->GetLODCount();
 
+		if (bOnlyOneMip)
+		{
+			LODCount = 1;
+		}
+
         bool isUncompressed = ( maskFormat == EImageFormat::IF_L_UBYTE );
 
-		constexpr int32 NumColorChannels = FMath::Min(NC,3);
+		constexpr uint32 NumColorChannels = FMath::Min(CHANNELS_TO_BLEND,3u);
         if ( isUncompressed )
         {
             int32 pixelCount = pBase->CalculatePixelCount();
@@ -217,22 +229,22 @@ namespace mu
                 unsigned mask = pMaskBuf[i];
                 for ( int32 c=0; c<NumColorChannels; ++c )
                 {
-                    unsigned base = pBaseBuf[NC*i+c];
+                    unsigned base = pBaseBuf[BASE_CHANNEL_STRIDE *i+c];
                     unsigned result = BLEND_FUNC_MASKED( base, top[c], mask );
                     if ( CLAMP )
                     {
-                        pDestBuf[NC*i+c] = (uint8)FMath::Min( 255u, result );
+                        pDestBuf[BASE_CHANNEL_STRIDE *i+c] = (uint8)FMath::Min( 255u, result );
                     }
                     else
                     {
-                        pDestBuf[NC*i+c] = (uint8)result;
+                        pDestBuf[BASE_CHANNEL_STRIDE *i+c] = (uint8)result;
                     }
                 }
 
-                constexpr bool isNC4 = (NC==4);
+                constexpr bool isNC4 = (BASE_CHANNEL_STRIDE ==4);
                 if ( isNC4 )
                 {
-                    pDestBuf[NC*i+3] = pBaseBuf[NC*i+3];
+                    pDestBuf[BASE_CHANNEL_STRIDE *i+3] = pBaseBuf[BASE_CHANNEL_STRIDE *i+3];
                 }
             }
         }
@@ -252,7 +264,7 @@ namespace mu
 
                 for ( int r=0; r<rows; ++r )
                 {
-                    const uint8* pDestRowEnd = pDestBuf + width*NC;
+                    const uint8* pDestRowEnd = pDestBuf + width* BASE_CHANNEL_STRIDE;
                     while ( pDestBuf!=pDestRowEnd )
                     {
                         // Decode header
@@ -267,29 +279,29 @@ namespace mu
                         ++pMaskBuf;
 
                         // Equal pixels
-						check(pDestBuf + NC * equal <= pStartDestBuf + pBase->GetDataSize());
+						check(pDestBuf + BASE_CHANNEL_STRIDE * equal <= pStartDestBuf + pBase->GetDataSize());
                         if ( equalPixel==255 )
                         {
                             for ( int i=0; i<equal; ++i )
                             {
                                 for ( int32 c=0; c<NumColorChannels; ++c )
                                 {
-                                    unsigned base = pBaseBuf[NC*i+c];
+                                    unsigned base = pBaseBuf[BASE_CHANNEL_STRIDE *i+c];
                                     unsigned result = BLEND_FUNC( base, top[c] );
                                     if ( CLAMP )
                                     {
-                                        pDestBuf[NC*i+c] = (uint8)FMath::Min( 255u, result );
+                                        pDestBuf[BASE_CHANNEL_STRIDE *i+c] = (uint8)FMath::Min( 255u, result );
                                     }
                                     else
                                     {
-                                        pDestBuf[NC*i+c] = (uint8)result;
+                                        pDestBuf[BASE_CHANNEL_STRIDE *i+c] = (uint8)result;
                                     }
                                 }
 
-                                constexpr bool isNC4 = (NC==4);
+                                constexpr bool isNC4 = (BASE_CHANNEL_STRIDE ==4);
                                 if ( isNC4 )
                                 {
-                                    pDestBuf[NC*i+3] = pBaseBuf[NC*i+3];
+                                    pDestBuf[BASE_CHANNEL_STRIDE *i+3] = pBaseBuf[BASE_CHANNEL_STRIDE *i+3];
                                 }
                             }
                         }
@@ -299,22 +311,22 @@ namespace mu
                             {
                                 for ( int32 c=0; c<NumColorChannels; ++c )
                                 {
-                                    unsigned base = pBaseBuf[NC*i+c];
+                                    unsigned base = pBaseBuf[BASE_CHANNEL_STRIDE *i+c];
                                     unsigned result = BLEND_FUNC_MASKED( base, top[c], equalPixel );
                                     if ( CLAMP )
                                     {
-                                        pDestBuf[NC*i+c] = (uint8)FMath::Min( 255u, result );
+                                        pDestBuf[BASE_CHANNEL_STRIDE *i+c] = (uint8)FMath::Min( 255u, result );
                                     }
                                     else
                                     {
-                                        pDestBuf[NC*i+c] = (uint8)result;
+                                        pDestBuf[BASE_CHANNEL_STRIDE *i+c] = (uint8)result;
                                     }
                                 }
 
-                                constexpr bool isNC4 = (NC==4);
+                                constexpr bool isNC4 = (BASE_CHANNEL_STRIDE ==4);
                                 if ( isNC4 )
                                 {
-                                    pDestBuf[NC*i+3] = pBaseBuf[NC*i+3];
+                                    pDestBuf[BASE_CHANNEL_STRIDE *i+3] = pBaseBuf[BASE_CHANNEL_STRIDE *i+3];
                                 }
                             }
                         }
@@ -323,40 +335,40 @@ namespace mu
                             // It could happen if xxxxxOnBase
                             if (pDestBuf!=pBaseBuf)
                             {
-                                FMemory::Memmove( pDestBuf, pBaseBuf, NC*equal );
+                                FMemory::Memmove( pDestBuf, pBaseBuf, BASE_CHANNEL_STRIDE*equal );
                             }
                         }
-                        pDestBuf += NC*equal;
-                        pBaseBuf += NC*equal;
+                        pDestBuf += BASE_CHANNEL_STRIDE *equal;
+                        pBaseBuf += BASE_CHANNEL_STRIDE *equal;
 
                         // Different pixels
-						check(pDestBuf + NC * different <= pStartDestBuf + pBase->GetDataSize());
+						check(pDestBuf + BASE_CHANNEL_STRIDE * different <= pStartDestBuf + pBase->GetDataSize());
                         for ( int i=0; i<different; ++i )
                         {
                             for ( int32 c=0; c<NumColorChannels; ++c )
                             {
                                 unsigned mask = pMaskBuf[i];
-                                unsigned base = pBaseBuf[NC*i+c];
+                                unsigned base = pBaseBuf[BASE_CHANNEL_STRIDE *i+c];
                                 unsigned result = BLEND_FUNC_MASKED( base, top[c], mask );
                                 if ( CLAMP )
                                 {
-                                    pDestBuf[NC*i+c] = (uint8)FMath::Min( 255u, result );
+                                    pDestBuf[BASE_CHANNEL_STRIDE *i+c] = (uint8)FMath::Min( 255u, result );
                                 }
                                 else
                                 {
-                                    pDestBuf[NC*i+c] = (uint8)result;
+                                    pDestBuf[BASE_CHANNEL_STRIDE *i+c] = (uint8)result;
                                 }
                             }
 
-                            constexpr bool isNC4 = (NC==4);
+                            constexpr bool isNC4 = (BASE_CHANNEL_STRIDE ==4);
                             if ( isNC4 )
                             {
-                                pDestBuf[NC*i+3] = pBaseBuf[NC*i+3];
+                                pDestBuf[BASE_CHANNEL_STRIDE *i+3] = pBaseBuf[BASE_CHANNEL_STRIDE *i+3];
                             }
                         }
 
-                        pDestBuf += NC*different;
-                        pBaseBuf += NC*different;
+                        pDestBuf += BASE_CHANNEL_STRIDE *different;
+                        pBaseBuf += BASE_CHANNEL_STRIDE *different;
                         pMaskBuf += different;
                     }
                 }
@@ -372,6 +384,8 @@ namespace mu
 	}
 
 
+	//---------------------------------------------------------------------------------------------
+	//! Apply a blending function to an image with a colour source and a mask
 	//---------------------------------------------------------------------------------------------
 	template< unsigned (*BLEND_FUNC_MASKED)(unsigned,unsigned,unsigned),
 			  unsigned (*BLEND_FUNC)(unsigned,unsigned),
@@ -382,25 +396,44 @@ namespace mu
 								   const Image* pMask,
 								   vec3<float> col )
 	{
+		check(pBase->GetSizeX() == pMask->GetSizeX());
+		check(pBase->GetSizeY() == pMask->GetSizeY());
+		check(pMask->GetFormat() == EImageFormat::IF_L_UBYTE
+			||
+			pMask->GetFormat() == EImageFormat::IF_L_UBYTE_RLE);
+
+		bool bValid =
+			(pBase->GetSizeX() == pMask->GetSizeX())
+			&&
+			(pBase->GetSizeY() == pMask->GetSizeY())
+			&&
+			(pMask->GetFormat() == EImageFormat::IF_L_UBYTE
+				||
+				pMask->GetFormat() == EImageFormat::IF_L_UBYTE_RLE);
+		if (!bValid)
+		{
+			return;
+		}
+
 		EImageFormat baseFormat = pBase->GetFormat();
 		if ( baseFormat==EImageFormat::IF_RGB_UBYTE )
 		{
-			BufferLayerColourFormat<BLEND_FUNC_MASKED, BLEND_FUNC, CLAMP, 3>( pDestBuf, pBase, pMask, col );
+			BufferLayerColourFormat<BLEND_FUNC_MASKED, BLEND_FUNC, CLAMP, 3, 3>( pDestBuf, pBase, pMask, col, 0, false);
 		}
         else if ( baseFormat==EImageFormat::IF_RGBA_UBYTE )
         {
-            BufferLayerColourFormat<BLEND_FUNC_MASKED, BLEND_FUNC, CLAMP, 4>( pDestBuf, pBase, pMask, col );
+            BufferLayerColourFormat<BLEND_FUNC_MASKED, BLEND_FUNC, CLAMP, 3, 4>( pDestBuf, pBase, pMask, col, 0, false);
         }
         else if ( baseFormat==EImageFormat::IF_BGRA_UBYTE )
         {
             float temp = col[0];
             col[0] = col[2];
             col[2] = temp;
-            BufferLayerColourFormat<BLEND_FUNC_MASKED, BLEND_FUNC, CLAMP, 4>( pDestBuf, pBase, pMask, col );
+            BufferLayerColourFormat<BLEND_FUNC_MASKED, BLEND_FUNC, CLAMP, 3, 4>( pDestBuf, pBase, pMask, col, 0, false);
         }
         else if ( baseFormat==EImageFormat::IF_L_UBYTE )
 		{
-			BufferLayerColourFormat<BLEND_FUNC_MASKED, BLEND_FUNC, CLAMP, 1>( pDestBuf, pBase, pMask, col );
+			BufferLayerColourFormat<BLEND_FUNC_MASKED, BLEND_FUNC, CLAMP, 1, 1>( pDestBuf, pBase, pMask, col, 0, false);
 		}
 		else
 		{
@@ -408,81 +441,164 @@ namespace mu
 		}
 	}
 
+
 	//---------------------------------------------------------------------------------------------
-	//! Apply a blending function to an image with a colour source and a mask
+	//! Apply a blending function to an image with another image as blending layer
 	//---------------------------------------------------------------------------------------------
-	template< unsigned (*BLEND_FUNC_MASKED)(unsigned,unsigned,unsigned),
-			  unsigned (*BLEND_FUNC)(unsigned,unsigned),
-			  bool CLAMP >
-    inline void ImageLayerColour( Image* pResult,
-                                  const Image* pBase,
-                                  const Image* pMask,
-                                  vec3<float> col )
+	template< unsigned (*BLEND_FUNC)(unsigned, unsigned), bool CLAMP,
+		// Number of total channels to actually process
+		uint32 CHANNELS_TO_BLEND,
+		// Number of total channels in the base image
+		int32 BASE_CHANNEL_STRIDE>
+
+	inline void BufferLayerColourFormatInPlace(Image* pBase, vec3<float> col,
+		uint32 BaseChannelOffset,
+		bool bOnlyFirstLOD)
 	{
-        check( pResult->GetFormat() == pBase->GetFormat() );
-        check( pResult->GetSizeX() == pBase->GetSizeX() );
-        check( pResult->GetSizeY() == pBase->GetSizeY() );
-        check( pResult->GetLODCount() == pBase->GetLODCount() );
-        check( pBase->GetSizeX() == pMask->GetSizeX() );
-		check( pBase->GetSizeY() == pMask->GetSizeY() );
-		check( pMask->GetFormat() == EImageFormat::IF_L_UBYTE
-						||
-						pMask->GetFormat() == EImageFormat::IF_L_UBYTE_RLE );
+		uint8* pBaseBuf = pBase->GetData() + BaseChannelOffset;
 
-        bool valid =
-                ( pResult->GetFormat() == pBase->GetFormat() )
-                &&
-                ( pResult->GetSizeX() == pBase->GetSizeX() )
-                &&
-                ( pResult->GetSizeY() == pBase->GetSizeY() )
-                &&
-                ( pResult->GetLODCount() == pBase->GetLODCount() )
-                &&
-                ( pBase->GetSizeX() == pMask->GetSizeX() )
-                &&
-                ( pBase->GetSizeY() == pMask->GetSizeY() )
-                &&
-                ( pMask->GetFormat() == EImageFormat::IF_L_UBYTE
-                                        ||
-                                        pMask->GetFormat() == EImageFormat::IF_L_UBYTE_RLE );
-        if (valid)
-        {
-            BufferLayerColour<BLEND_FUNC_MASKED, BLEND_FUNC, CLAMP>( pResult->GetData(), pBase, pMask, col );
-        }
+		uint32 top[4];
+		top[0] = (uint32)(255 * col[0]);
+		top[1] = (uint32)(255 * col[1]);
+		top[2] = (uint32)(255 * col[2]);
+		top[3] = (uint32)(255 * col[0]); // used for 1 channel blends to alpha
+
+		// Generic implementation
+		int32 PixelCount = 0;
+		if (bOnlyFirstLOD)
+		{
+			PixelCount = pBase->GetSizeX() * pBase->GetSizeY();
+		}
+		else
+		{
+			PixelCount = pBase->CalculatePixelCount();
+		}
+
+		ParallelFor(PixelCount,
+			[
+				pBaseBuf, top, BaseChannelOffset
+			] (uint32 i)
+			{
+				for (int c = 0; c < CHANNELS_TO_BLEND; ++c)
+				{
+					uint32 base = pBaseBuf[BASE_CHANNEL_STRIDE * i + c];
+					uint32 blended = top[c + BaseChannelOffset];
+					uint32 result = BLEND_FUNC(base, blended);
+					if (CLAMP)
+					{
+						pBaseBuf[BASE_CHANNEL_STRIDE * i + c] = (uint8)FMath::Min(255u, result);
+					}
+					else
+					{
+						pBaseBuf[BASE_CHANNEL_STRIDE * i + c] = (uint8)result;
+					}
+				}
+			});
 	}
 
 
 	//---------------------------------------------------------------------------------------------
-	//! Apply a blending function to an image with a colour source on the base image itself
+	template< unsigned (*BLEND_FUNC)(unsigned, unsigned), bool CLAMP, uint32 CHANNEL_COUNT >
+	inline void BufferLayerColourInPlace(Image* pBase, vec3<float> col, bool bOnlyOneMip, uint32 BaseOffset)
+	{
+		EImageFormat baseFormat = pBase->GetFormat();
+
+		if (baseFormat == EImageFormat::IF_RGB_UBYTE)
+		{
+			check(BaseOffset + CHANNEL_COUNT <= 3);
+			BufferLayerColourFormatInPlace< BLEND_FUNC, CLAMP, CHANNEL_COUNT, 3 >(pBase, col, BaseOffset, bOnlyOneMip);
+		}
+		else if (baseFormat == EImageFormat::IF_RGBA_UBYTE)
+		{
+			check(BaseOffset + CHANNEL_COUNT <= 4);
+			BufferLayerColourFormatInPlace< BLEND_FUNC, CLAMP, CHANNEL_COUNT, 4 >(pBase, col, BaseOffset, bOnlyOneMip);
+		}
+		else if (baseFormat == EImageFormat::IF_BGRA_UBYTE)
+		{
+			float temp = col[0];
+			col[0] = col[2];
+			col[2] = temp;
+			BufferLayerColourFormatInPlace< BLEND_FUNC, CLAMP, CHANNEL_COUNT, 4>(pBase, col, BaseOffset, bOnlyOneMip);
+		}
+
+		else if (baseFormat == EImageFormat::IF_L_UBYTE)
+		{
+			check(BaseOffset + CHANNEL_COUNT <= 1);
+			BufferLayerColourFormatInPlace< BLEND_FUNC, CLAMP, CHANNEL_COUNT, 1 >(pBase, col, BaseOffset, bOnlyOneMip);
+		}
+		else
+		{
+			checkf(false, TEXT("Unsupported format."));
+		}
+	}
+
+
 	//---------------------------------------------------------------------------------------------
-	template< unsigned (*BLEND_FUNC_MASKED)(unsigned,unsigned,unsigned),
-			  unsigned (*BLEND_FUNC)(unsigned,unsigned),
-			  bool CLAMP >
-	inline void ImageLayerColourOnBase( Image* pBase, const Image* pMask, vec3<float> col )
+	template< 
+		unsigned (*BLEND_FUNC_MASKED)(unsigned, unsigned, unsigned),
+		unsigned (*BLEND_FUNC)(unsigned, unsigned), 
+		bool CLAMP, 
+		int32 CHANNEL_COUNT >
+	inline void BufferLayerColourInPlace(Image* pBase, const Image* Mask, vec3<float> col, bool bOnlyOneMip, uint32 BaseOffset)
 	{
-		check( pBase->GetSizeX() == pMask->GetSizeX() );
-		check( pBase->GetSizeY() == pMask->GetSizeY() );
-		check( pMask->GetFormat() == EImageFormat::IF_L_UBYTE
-						||
-						pMask->GetFormat() == EImageFormat::IF_L_UBYTE_RLE );
+		check(pBase->GetSizeX() == Mask->GetSizeX());
+		check(pBase->GetSizeY() == Mask->GetSizeY());
 
-		BufferLayerColour<BLEND_FUNC_MASKED,BLEND_FUNC,CLAMP>( pBase->GetData(), pBase, pMask, col );
+		EImageFormat baseFormat = pBase->GetFormat();
+
+		if (baseFormat == EImageFormat::IF_RGB_UBYTE)
+		{
+			check(BaseOffset + CHANNEL_COUNT <= 3);
+			BufferLayerColourFormat< BLEND_FUNC_MASKED, BLEND_FUNC, CLAMP, CHANNEL_COUNT, 3 >(pBase->GetData(), pBase, Mask, col, BaseOffset, bOnlyOneMip);
+		}
+		else if (baseFormat == EImageFormat::IF_RGBA_UBYTE)
+		{
+			check(BaseOffset + CHANNEL_COUNT <= 4);
+			BufferLayerColourFormat< BLEND_FUNC_MASKED, BLEND_FUNC, CLAMP, CHANNEL_COUNT, 4 >(pBase->GetData(), pBase, Mask, col, BaseOffset, bOnlyOneMip);
+		}
+		else if (baseFormat == EImageFormat::IF_BGRA_UBYTE)
+		{
+			check(BaseOffset + CHANNEL_COUNT <= 4);
+			float temp = col[0];
+			col[0] = col[2];
+			col[2] = temp;
+			BufferLayerColourFormat< BLEND_FUNC_MASKED, BLEND_FUNC, CLAMP, CHANNEL_COUNT, 4 >(pBase->GetData(), pBase, Mask, col, BaseOffset, bOnlyOneMip);
+		}
+		else if (baseFormat == EImageFormat::IF_L_UBYTE)
+		{
+			check(BaseOffset + CHANNEL_COUNT <= 1);
+			BufferLayerColourFormat< BLEND_FUNC_MASKED, BLEND_FUNC, CLAMP, CHANNEL_COUNT, 1 >(pBase->GetData(), pBase, Mask, col, BaseOffset, bOnlyOneMip);
+		}
+		else
+		{
+			checkf(false, TEXT("Unsupported format."));
+		}
 	}
 
 
 	//---------------------------------------------------------------------------------------------
 	//! Apply a blending function to an image with another image as blending layer
 	//---------------------------------------------------------------------------------------------
-	template< unsigned (*BLEND_FUNC)(unsigned,unsigned), bool CLAMP, int NC >
-    inline void BufferLayer(uint8* pDestBuf, const Image* pBase, const Image* pBlended,
-                            bool applyToAlpha, bool bOnlyFirstLOD )
+	template< unsigned (*BLEND_FUNC)(unsigned,unsigned), bool CLAMP,
+		// Number of total channels to actually process
+		int32 CHANNELS_TO_BLEND,
+		// Number of total channels in the base image
+		int32 BASE_CHANNEL_STRIDE,
+		// Number of total channels in the blend image
+		int32 BLENDED_CHANNEL_STRIDE>
+
+    inline void BufferLayerFormatInPlace(Image* pBase, const Image* pBlended,
+		uint32 BaseChannelOffset,
+		uint32 BlendedChannelOffset,
+		bool bOnlyFirstLOD )
 	{
-		check( pBase->GetSizeX() == pBlended->GetSizeX() );
-		check( pBase->GetSizeY() == pBlended->GetSizeY() );
-		check( pBase->GetFormat() == pBlended->GetFormat() );
+		check(pBase->GetSizeX() == pBlended->GetSizeX());
+		check(pBase->GetSizeY() == pBlended->GetSizeY());
+		check(pBase->GetFormat() == pBlended->GetFormat());
+		check(bOnlyFirstLOD || pBase->GetLODCount() <= pBlended->GetLODCount());
 
-        const uint8* pBaseBuf = pBase->GetData();
-        const uint8* pBlendedBuf = pBlended->GetData();
+        uint8* pBaseBuf = pBase->GetData() + BaseChannelOffset;
+        const uint8* pBlendedBuf = pBlended->GetData() + BlendedChannelOffset;
 
 		// Generic implementation
 		int32 PixelCount = 0;
@@ -497,45 +613,85 @@ namespace mu
 
 		ParallelFor(PixelCount,
 			[
-				pBaseBuf, pBlendedBuf, pDestBuf, applyToAlpha
+				pBaseBuf, pBlendedBuf
 			] (uint32 i)
 			{
-				if (applyToAlpha && NC > 3)
+				for (int c = 0; c < CHANNELS_TO_BLEND; ++c)
 				{
-					for (int c = 0; c < 3; ++c)
+					unsigned base = pBaseBuf[BASE_CHANNEL_STRIDE * i + c];
+					unsigned blended = pBlendedBuf[BLENDED_CHANNEL_STRIDE * i + c];
+					unsigned result = BLEND_FUNC(base, blended);
+					if (CLAMP)
 					{
-						unsigned base = pBaseBuf[NC * i + c];
-						unsigned blended = pBlendedBuf[NC * i + c];
-						unsigned result = BLEND_FUNC(base, blended);
-						if (CLAMP)
-						{
-							pDestBuf[NC * i + c] = (uint8)FMath::Min(255u, result);
-						}
-						else
-						{
-							pDestBuf[NC * i + c] = (uint8)result;
-						}
+						pBaseBuf[BASE_CHANNEL_STRIDE * i + c] = (uint8)FMath::Min(255u, result);
 					}
+					else
+					{
+						pBaseBuf[BASE_CHANNEL_STRIDE * i + c] = (uint8)result;
+					}
+				}				
+			});
+	}
+	
+	//---------------------------------------------------------------------------------------------
+	//! Apply a blending function to an image with another image as blending layer
+	//---------------------------------------------------------------------------------------------
+	template< unsigned (*BLEND_FUNC)(unsigned, unsigned), bool CLAMP,
+		// Number of total channels to actually process
+		int32 CHANNELS_TO_BLEND,
+		// Number of total channels in the base image
+		int32 BASE_CHANNEL_STRIDE,
+		// Number of total channels in the blend image
+		int32 BLENDED_CHANNEL_STRIDE>
+
+	inline void BufferLayerFormat(uint8* pDestBuf, const Image* pBase, const Image* pBlended,
+		bool bOnlyFirstLOD)
+	{
+		check(pBase->GetSizeX() == pBlended->GetSizeX());
+		check(pBase->GetSizeY() == pBlended->GetSizeY());
+		check(pBase->GetFormat() == pBlended->GetFormat());
+		check(bOnlyFirstLOD || pBase->GetLODCount() <= pBlended->GetLODCount());
 
-					pDestBuf[NC * i + 3] = pBaseBuf[NC * i + 3];
-				}
-				else
+		const uint8* pBaseBuf = pBase->GetData();
+		const uint8* pBlendedBuf = pBlended->GetData();
+
+		// Generic implementation
+		int32 PixelCount = 0;
+		if (bOnlyFirstLOD)
+		{
+			PixelCount = pBase->GetSizeX() * pBase->GetSizeY();
+		}
+		else
+		{
+			PixelCount = pBase->CalculatePixelCount();
+		}
+
+		ParallelFor(PixelCount,
+			[
+				pBaseBuf, pBlendedBuf, pDestBuf
+			] (uint32 i)
+			{
+				for (int c = 0; c < CHANNELS_TO_BLEND; ++c)
 				{
-					for (int c = 0; c < NC; ++c)
+					unsigned base = pBaseBuf[BASE_CHANNEL_STRIDE * i + c];
+					unsigned blended = pBlendedBuf[BLENDED_CHANNEL_STRIDE * i + c];
+					unsigned result = BLEND_FUNC(base, blended);
+					if (CLAMP)
 					{
-						unsigned base = pBaseBuf[NC * i + c];
-						unsigned blended = pBlendedBuf[NC * i + c];
-						unsigned result = BLEND_FUNC(base, blended);
-						if (CLAMP)
-						{
-							pDestBuf[NC * i + c] = (uint8)FMath::Min(255u, result);
-						}
-						else
-						{
-							pDestBuf[NC * i + c] = (uint8)result;
-						}
+						pDestBuf[BASE_CHANNEL_STRIDE * i + c] = (uint8)FMath::Min(255u, result);
+					}
+					else
+					{
+						pDestBuf[BASE_CHANNEL_STRIDE * i + c] = (uint8)result;
 					}
 				}
+
+				// Copy the unblended channels
+				// \TODO: unnecessary when doing it in-place?
+				for (uint32 c = 0; c < BASE_CHANNEL_STRIDE - CHANNELS_TO_BLEND; ++c)
+				{
+					pDestBuf[BASE_CHANNEL_STRIDE * i + CHANNELS_TO_BLEND + c] = pBaseBuf[BASE_CHANNEL_STRIDE * i + CHANNELS_TO_BLEND + c];
+				}
 			});
 	}
 
@@ -543,104 +699,120 @@ namespace mu
 	//---------------------------------------------------------------------------------------------
 	//! Apply a blending function to an image with another image as blending layer
 	//---------------------------------------------------------------------------------------------
-	template< unsigned (*BLEND_FUNC)(unsigned,unsigned), bool CLAMP >
-    inline void ImageLayer( Image* pResult, const Image* pBase, const Image* pBlended,
-                            bool applyToAlpha )
+	template< unsigned (*BLEND_FUNC)(unsigned, unsigned), bool CLAMP >
+	inline void BufferLayer(Image* pResult, const Image* pBase, const Image* pBlended, bool bApplyToAlpha, bool bOnlyOneMip)
 	{
-        check( pResult->GetFormat() == pBase->GetFormat() );
-        check( pResult->GetSizeX() == pBase->GetSizeX() );
-        check( pResult->GetSizeY() == pBase->GetSizeY() );
-        check( pResult->GetLODCount() == pBase->GetLODCount() );
-        check( pBase->GetSizeX() == pBlended->GetSizeX() );
-		check( pBase->GetSizeY() == pBlended->GetSizeY() );
-		check( pBase->GetFormat() == pBlended->GetFormat() );
-        check( pResult->GetLODCount() <= pBlended->GetLODCount() );
+		check(pResult->GetFormat() == pBase->GetFormat());
+		check(pResult->GetSizeX() == pBase->GetSizeX());
+		check(pResult->GetSizeY() == pBase->GetSizeY());
+		check(bOnlyOneMip || pResult->GetLODCount() == pBase->GetLODCount());
+		check(pBase->GetSizeX() == pBlended->GetSizeX());
+		check(pBase->GetSizeY() == pBlended->GetSizeY());
+		check(pBase->GetFormat() == pBlended->GetFormat());
+		check(bOnlyOneMip || pResult->GetLODCount() <= pBlended->GetLODCount());
 
 		EImageFormat baseFormat = pBase->GetFormat();
+		uint8* pDestBuf = pResult->GetData();
 
-        uint8* pDestBuf = pResult->GetData();
-
-		if ( baseFormat==EImageFormat::IF_RGB_UBYTE )
+		if (baseFormat == EImageFormat::IF_RGB_UBYTE)
 		{
-            BufferLayer< BLEND_FUNC, CLAMP, 3 > ( pDestBuf, pBase, pBlended, applyToAlpha, false );
+			BufferLayerFormat< BLEND_FUNC, CLAMP, 3, 3, 3 >(pDestBuf, pBase, pBlended, bOnlyOneMip);
 		}
-        else if ( baseFormat==EImageFormat::IF_RGBA_UBYTE || baseFormat==EImageFormat::IF_BGRA_UBYTE )
+		else if (baseFormat == EImageFormat::IF_RGBA_UBYTE || baseFormat == EImageFormat::IF_BGRA_UBYTE)
 		{
-            BufferLayer< BLEND_FUNC, CLAMP, 4 > ( pDestBuf, pBase, pBlended, applyToAlpha, false);
+			if (bApplyToAlpha)
+			{
+				BufferLayerFormat< BLEND_FUNC, CLAMP, 4, 4, 4 >(pDestBuf, pBase, pBlended, bOnlyOneMip);
+			}
+			else
+			{
+				BufferLayerFormat< BLEND_FUNC, CLAMP, 3, 4, 4 >(pDestBuf, pBase, pBlended, bOnlyOneMip);
+			}
 		}
-		else if ( baseFormat==EImageFormat::IF_L_UBYTE )
+		else if (baseFormat == EImageFormat::IF_L_UBYTE)
 		{
-            BufferLayer< BLEND_FUNC, CLAMP, 1 > ( pDestBuf, pBase, pBlended, applyToAlpha, false);
+			BufferLayerFormat< BLEND_FUNC, CLAMP, 1, 1, 1 >(pDestBuf, pBase, pBlended, bOnlyOneMip);
 		}
 		else
 		{
-			checkf( false, TEXT("Unsupported format.") );
+			checkf(false, TEXT("Unsupported format."));
 		}
 	}
 
 
 	//---------------------------------------------------------------------------------------------
-	//! Apply a blending function to an image with another image as blending layer on the base image
-	//---------------------------------------------------------------------------------------------
-	template< unsigned (*BLEND_FUNC)(unsigned,unsigned), bool CLAMP >
-    inline void ImageLayerOnBase( Image* pBase, const Image* pBlended,
-                                  bool applyToAlpha, bool bOnlyFirstLOD=false )
+	template< unsigned (*BLEND_FUNC)(unsigned, unsigned), bool CLAMP, int32 CHANNEL_COUNT >
+	inline void BufferLayerInPlace(Image* pBase, const Image* pBlended, bool bOnlyOneMip, uint32 BaseOffset, uint32 BlendedOffset)
 	{
-		check( pBase->GetSizeX() == pBlended->GetSizeX() );
-		check( pBase->GetSizeY() == pBlended->GetSizeY() );
-		check( pBase->GetFormat() == pBlended->GetFormat() );
+		check(pBase->GetSizeX() == pBlended->GetSizeX());
+		check(pBase->GetSizeY() == pBlended->GetSizeY());
+		check(pBase->GetFormat() == pBlended->GetFormat());
 
 		EImageFormat baseFormat = pBase->GetFormat();
-        uint8* pDestBuf = pBase->GetData();
 
-		if ( baseFormat==EImageFormat::IF_RGB_UBYTE )
+		if (baseFormat == EImageFormat::IF_RGB_UBYTE)
 		{
-            BufferLayer< BLEND_FUNC, CLAMP, 3 > ( pDestBuf, pBase, pBlended, applyToAlpha, bOnlyFirstLOD);
+			check(BaseOffset+CHANNEL_COUNT <= 3);
+			check(BlendedOffset+CHANNEL_COUNT <= 3);
+			BufferLayerFormatInPlace< BLEND_FUNC, CLAMP, CHANNEL_COUNT, 3, 3 >(pBase, pBlended, BaseOffset, BlendedOffset, bOnlyOneMip);
 		}
-        else if ( baseFormat==EImageFormat::IF_RGBA_UBYTE || baseFormat==EImageFormat::IF_BGRA_UBYTE )
+		else if (baseFormat == EImageFormat::IF_RGBA_UBYTE || baseFormat == EImageFormat::IF_BGRA_UBYTE)
 		{
-            BufferLayer< BLEND_FUNC, CLAMP, 4 > ( pDestBuf, pBase, pBlended, applyToAlpha, bOnlyFirstLOD);
+			check(BaseOffset+CHANNEL_COUNT <= 4);
+			check(BlendedOffset+CHANNEL_COUNT <= 4);
+			BufferLayerFormatInPlace< BLEND_FUNC, CLAMP, CHANNEL_COUNT, 4, 4 >(pBase, pBlended, BaseOffset, BlendedOffset, bOnlyOneMip);
 		}
-		else if ( baseFormat==EImageFormat::IF_L_UBYTE )
+		else if (baseFormat == EImageFormat::IF_L_UBYTE)
 		{
-            BufferLayer< BLEND_FUNC, CLAMP, 1 > ( pDestBuf, pBase, pBlended, applyToAlpha, bOnlyFirstLOD);
+			check(BaseOffset+CHANNEL_COUNT <= 1);
+			BufferLayerFormatInPlace< BLEND_FUNC, CLAMP, CHANNEL_COUNT, 1, 1 >(pBase, pBlended, BaseOffset, BlendedOffset, bOnlyOneMip);
 		}
 		else
 		{
-			checkf( false, TEXT("Unsupported format.") );
+			checkf(false, TEXT("Unsupported format."));
 		}
 	}
 
 
 	//---------------------------------------------------------------------------------------------
-	template< unsigned (*BLEND_FUNC_MASKED)(unsigned,unsigned,unsigned),
-			  unsigned (*BLEND_FUNC)(unsigned,unsigned),
-			  bool CLAMP,
-			  int NC >
-    inline void BufferLayerFormat( uint8* pStartDestBuf,
-                                   const Image* pBase,
-                                   const Image* pMask,
-                                   const Image* pBlend,
-                                   bool applyToAlpha,
-								   bool bOnlyFirstLOD )
+	template< unsigned (*BLEND_FUNC_MASKED)(unsigned, unsigned, unsigned),
+		unsigned (*BLEND_FUNC)(unsigned, unsigned),
+		bool CLAMP,
+		// Number of total channels to actually process
+		int32 CHANNELS_TO_BLEND,
+		// Number of total channels in the base image
+		int32 BASE_CHANNEL_STRIDE,
+		// Number of total channels in the blend image
+		int32 BLENDED_CHANNEL_STRIDE>
+	inline void BufferLayerFormat(uint8* pStartDestBuf,
+		const Image* pBase,
+		const Image* pMask,
+		const Image* pBlend,
+		uint32 BaseChannelOffset,
+		uint32 BlendedChannelOffset,
+		bool bOnlyFirstLOD)
 	{
-        check(pBase->GetSizeX()==pMask->GetSizeX() && pBase->GetSizeY()==pMask->GetSizeY());
-        check(pBase->GetSizeX()==pBlend->GetSizeX() && pBase->GetSizeY()==pBlend->GetSizeY());
-        check(bOnlyFirstLOD || pBase->GetLODCount() <= pMask->GetLODCount() );
-        check(bOnlyFirstLOD || pBase->GetLODCount() <= pBlend->GetLODCount() );
+		check(pBase->GetSizeX() == pMask->GetSizeX() && pBase->GetSizeY() == pMask->GetSizeY());
+		check(pBase->GetSizeX() == pBlend->GetSizeX() && pBase->GetSizeY() == pBlend->GetSizeY());
+		check(bOnlyFirstLOD || pBase->GetLODCount() <= pMask->GetLODCount());
+		check(bOnlyFirstLOD || pBase->GetLODCount() <= pBlend->GetLODCount());
+		check(pBase->GetFormat() == pBlend->GetFormat());
+		check(pMask->GetFormat() == EImageFormat::IF_L_UBYTE
+			||
+			pMask->GetFormat() == EImageFormat::IF_L_UBYTE_RLE);
 
 		uint8* pDestBuf = pStartDestBuf;
-        const uint8* pMaskBuf = pMask->GetData();
-        const uint8* pBaseBuf = pBase->GetData();
-        const uint8* pBlendedBuf = pBlend->GetData();
+		const uint8* pMaskBuf = pMask->GetData();
+		const uint8* pBaseBuf = pBase->GetData() + BaseChannelOffset;
+		const uint8* pBlendedBuf = pBlend->GetData() + BlendedChannelOffset;
 
 		EImageFormat maskFormat = pMask->GetFormat();
-        bool isUncompressed = ( maskFormat == EImageFormat::IF_L_UBYTE );
+		bool bIsMaskUncompressed = (maskFormat == EImageFormat::IF_L_UBYTE);
 
 		// The base determines the number of lods to process.
 		int32 LODCount = bOnlyFirstLOD ? 1 : pBase->GetLODCount();
 
-        if ( isUncompressed )
+		if (bIsMaskUncompressed)
 		{
 			int32 PixelCount = 0;
 			if (bOnlyFirstLOD)
@@ -652,66 +824,42 @@ namespace mu
 				PixelCount = pBase->CalculatePixelCount();
 			}
 
-            if (!applyToAlpha && NC>3)
-            {
-				ParallelFor(PixelCount,
-					[
-						pBaseBuf, pBlendedBuf, pMaskBuf, pDestBuf
-					] (uint32 i)
-					{
-						unsigned mask = pMaskBuf[i];
-						constexpr int32 NumColorChannels = FMath::Min(NC, 3);
-						for (int c = 0; c <NumColorChannels; ++c)
-						{
-							unsigned base = pBaseBuf[NC * i + c];
-							unsigned blended = pBlendedBuf[NC * i + c];
-							unsigned result = BLEND_FUNC_MASKED(base, blended, mask);
-							if (CLAMP)
-							{
-								pDestBuf[NC * i + c] = (uint8)FMath::Min(255u, result);
-							}
-							else
-							{
-								pDestBuf[NC * i + c] = (uint8)result;
-							}
-						}
-
-						pDestBuf[NC * i + 3] = pBaseBuf[NC * i + 3];
-					});
-            }
-            else
-            {
-				ParallelFor(PixelCount,
-					[
-						pBaseBuf, pBlendedBuf, pMaskBuf, pDestBuf
-					] (uint32 i)
-					{
-						unsigned mask = pMaskBuf[i];
-						for (int c = 0; c < NC; ++c)
-						{
-							unsigned base = pBaseBuf[NC * i + c];
-							unsigned blended = pBlendedBuf[NC * i + c];
-							unsigned result = BLEND_FUNC_MASKED(base, blended, mask);
-							if (CLAMP)
-							{
-								pDestBuf[NC * i + c] = (uint8)FMath::Min(255u, result);
-							}
-							else
-							{
-								pDestBuf[NC * i + c] = (uint8)result;
-							}
-						}
-					});
-            }
+			ParallelFor(PixelCount,
+				[
+					pBaseBuf, pBlendedBuf, pMaskBuf, pDestBuf
+				] (uint32 i)
+				{
+					unsigned mask = pMaskBuf[i];
+			for (int c = 0; c < CHANNELS_TO_BLEND; ++c)
+			{
+				uint32 base = pBaseBuf[BASE_CHANNEL_STRIDE * i + c];
+				uint32 blended = pBlendedBuf[BLENDED_CHANNEL_STRIDE * i + c];
+				uint32 result = BLEND_FUNC_MASKED(base, blended, mask);
+				if (CLAMP)
+				{
+					pDestBuf[BASE_CHANNEL_STRIDE * i + c] = (uint8)FMath::Min(255u, result);
+				}
+				else
+				{
+					pDestBuf[BASE_CHANNEL_STRIDE * i + c] = (uint8)result;
+				}
+			}
+			// Copy the unblended channels
+			// \TODO: unnecessary when doing it in-place?
+			for (int c = 0; c < BASE_CHANNEL_STRIDE - CHANNELS_TO_BLEND; ++c)
+			{
+				pDestBuf[BASE_CHANNEL_STRIDE * i + CHANNELS_TO_BLEND + c] = pBaseBuf[BASE_CHANNEL_STRIDE * i + CHANNELS_TO_BLEND + c];
+			}
+				});
 		}
-        else if ( maskFormat==EImageFormat::IF_L_UBYTE_RLE )
+		else if (maskFormat == EImageFormat::IF_L_UBYTE_RLE)
 		{
 			int rows = pBase->GetSizeY();
 			int width = pBase->GetSizeX();
 
-            for (int lod=0;lod<LODCount; ++lod)
-            {
-                pMaskBuf += 4+rows*sizeof(uint32);
+			for (int lod = 0; lod < LODCount; ++lod)
+			{
+				pMaskBuf += 4 + rows * sizeof(uint32);
 
 				// \todo: See how to handle mask buf here
 				//ParallelFor(rows,
@@ -719,11 +867,9 @@ namespace mu
 				//		pBaseBuf, pBlendedBuf, pMaskBuf, pDestBuf, width, applyToAlpha
 				//	] (uint32 r)
 				for (int r = 0; r < rows; ++r)
-				{						
-					//UE_LOG(LogMutableCore, Warning, "row: %d", r);
-
-					const uint8* pDestRowEnd = pDestBuf + width*NC;
-					while ( pDestBuf!=pDestRowEnd )
+				{
+					const uint8* pDestRowEnd = pDestBuf + width * BASE_CHANNEL_STRIDE;
+					while (pDestBuf != pDestRowEnd)
 					{
 						// Decode header
 						uint16 equal = *(const uint16*)pMaskBuf;
@@ -735,168 +881,188 @@ namespace mu
 						uint8 equalPixel = *pMaskBuf;
 						++pMaskBuf;
 
-						//UE_LOG(LogMutableCore, Warning, "    block: %d %d", (int)equal, (int)different );
-
 						// Equal pixels
-						check(pDestBuf + NC * equal <= pStartDestBuf + pBase->GetDataSize());
-						if ( equalPixel==255 )
+						check(pDestBuf + BASE_CHANNEL_STRIDE * equal <= pStartDestBuf + pBase->GetDataSize());
+						if (equalPixel == 255)
 						{
-							for ( int i=0; i<equal; ++i )
+							for (int i = 0; i < equal; ++i)
 							{
-								if (!applyToAlpha && NC>3)
+								for (int c = 0; c < CHANNELS_TO_BLEND; ++c)
 								{
-									for ( int c=0; c<3; ++c )
+									uint32 base = pBaseBuf[BASE_CHANNEL_STRIDE * i + c];
+									uint32 blended = pBlendedBuf[BLENDED_CHANNEL_STRIDE * i + c];
+									uint32 result = BLEND_FUNC(base, blended);
+									if (CLAMP)
 									{
-										unsigned base = pBaseBuf[NC*i+c];
-										unsigned blended = pBlendedBuf[NC*i+c];
-										unsigned result = BLEND_FUNC( base, blended );
-										if ( CLAMP )
-										{
-											pDestBuf[NC*i+c] = (uint8)FMath::Min( 255u, result );
-										}
-										else
-										{
-											pDestBuf[NC*i+c] = (uint8)result;
-										}
+										pDestBuf[BASE_CHANNEL_STRIDE * i + c] = (uint8)FMath::Min(255u, result);
 									}
-
-									pDestBuf[NC*i+3] = (uint8)pBaseBuf[NC*i+3];
-								}
-								else
-								{
-									for ( int c=0; c<NC; ++c )
+									else
 									{
-										unsigned base = pBaseBuf[NC*i+c];
-										unsigned blended = pBlendedBuf[NC*i+c];
-										unsigned result = BLEND_FUNC( base, blended );
-										if ( CLAMP )
-										{
-											pDestBuf[NC*i+c] = (uint8)FMath::Min( 255u, result );
-										}
-										else
-										{
-											pDestBuf[NC*i+c] = (uint8)result;
-										}
+										pDestBuf[BASE_CHANNEL_STRIDE * i + c] = (uint8)result;
 									}
 								}
+
+								// Copy the unblended channels
+								// \TODO: unnecessary when doing it in-place?
+								for (int c = 0; c < BASE_CHANNEL_STRIDE - CHANNELS_TO_BLEND; ++c)
+								{
+									pDestBuf[BASE_CHANNEL_STRIDE * i + CHANNELS_TO_BLEND + c] = pBaseBuf[BASE_CHANNEL_STRIDE * i + CHANNELS_TO_BLEND + c];
+								}
 							}
 						}
-						else if ( equalPixel>0 )
+						else if (equalPixel > 0)
 						{
-							for ( int i=0; i<equal; ++i )
+							for (int i = 0; i < equal; ++i)
 							{
-								if (!applyToAlpha && NC>3)
+								for (int c = 0; c < CHANNELS_TO_BLEND; ++c)
 								{
-									for ( int c=0; c<3; ++c )
+									unsigned base = pBaseBuf[BASE_CHANNEL_STRIDE * i + c];
+									unsigned blended = pBlendedBuf[BLENDED_CHANNEL_STRIDE * i + c];
+									unsigned result = BLEND_FUNC_MASKED(base, blended, equalPixel);
+									if (CLAMP)
 									{
-										unsigned base = pBaseBuf[NC*i+c];
-										unsigned blended = pBlendedBuf[NC*i+c];
-										unsigned result = BLEND_FUNC_MASKED( base, blended, equalPixel );
-										if ( CLAMP )
-										{
-											pDestBuf[NC*i+c] = (uint8)FMath::Min( 255u, result );
-										}
-										else
-										{
-											pDestBuf[NC*i+c] = (uint8)result;
-										}
+										pDestBuf[BASE_CHANNEL_STRIDE * i + c] = (uint8)FMath::Min(255u, result);
 									}
-
-									pDestBuf[NC*i+3] = (uint8)pBaseBuf[NC*i+3];
-								}
-								else
-								{
-									for ( int c=0; c<NC; ++c )
+									else
 									{
-										unsigned base = pBaseBuf[NC*i+c];
-										unsigned blended = pBlendedBuf[NC*i+c];
-										unsigned result = BLEND_FUNC_MASKED( base, blended, equalPixel );
-										if ( CLAMP )
-										{
-											pDestBuf[NC*i+c] = (uint8)FMath::Min( 255u, result );
-										}
-										else
-										{
-											pDestBuf[NC*i+c] = (uint8)result;
-										}
+										pDestBuf[BASE_CHANNEL_STRIDE * i + c] = (uint8)result;
 									}
 								}
+
+								// Copy the unblended channels
+								// \TODO: unnecessary when doing it in-place?
+								for (int c = 0; c < BASE_CHANNEL_STRIDE - CHANNELS_TO_BLEND; ++c)
+								{
+									pDestBuf[BASE_CHANNEL_STRIDE * i + CHANNELS_TO_BLEND + c] = pBaseBuf[BASE_CHANNEL_STRIDE * i + CHANNELS_TO_BLEND + c];
+								}
 							}
 						}
 						else
 						{
 							// It could happen if xxxxxOnBase
-							if (pDestBuf!=pBaseBuf)
+							if (pDestBuf != pBaseBuf)
 							{
-								FMemory::Memmove( pDestBuf, pBaseBuf, NC*equal );
+								FMemory::Memmove(pDestBuf, pBaseBuf, BASE_CHANNEL_STRIDE * equal);
 							}
 						}
-						pDestBuf += NC*equal;
-						pBaseBuf += NC*equal;
-						pBlendedBuf += NC*equal;
+						pDestBuf += BASE_CHANNEL_STRIDE * equal;
+						pBaseBuf += BASE_CHANNEL_STRIDE * equal;
+						pBlendedBuf += BLENDED_CHANNEL_STRIDE * equal;
 
 						// Different pixels
-						check(pDestBuf + NC * different <= pStartDestBuf + pBase->GetDataSize());
-						for ( int i=0; i<different; ++i )
+						check(pDestBuf + BASE_CHANNEL_STRIDE * different <= pStartDestBuf + pBase->GetDataSize());
+						for (int i = 0; i < different; ++i)
 						{
-							if (!applyToAlpha && NC>3)
+							for (int c = 0; c < CHANNELS_TO_BLEND; ++c)
 							{
-								for ( int c=0; c<3; ++c )
+								unsigned mask = pMaskBuf[i];
+								unsigned base = pBaseBuf[BASE_CHANNEL_STRIDE * i + c];
+								unsigned blended = pBlendedBuf[BLENDED_CHANNEL_STRIDE * i + c];
+								unsigned result = BLEND_FUNC_MASKED(base, blended, mask);
+								if (CLAMP)
 								{
-									unsigned mask = pMaskBuf[i];
-									unsigned base = pBaseBuf[NC*i+c];
-									unsigned blended = pBlendedBuf[NC*i+c];
-									unsigned result = BLEND_FUNC_MASKED( base, blended, mask );
-									if ( CLAMP )
-									{
-										pDestBuf[NC*i+c] = (uint8)FMath::Min( 255u, result );
-									}
-									else
-									{
-										pDestBuf[NC*i+c] = (uint8)result;
-									}
+									pDestBuf[BASE_CHANNEL_STRIDE * i + c] = (uint8)FMath::Min(255u, result);
 								}
-
-								pDestBuf[NC*i+3] = (uint8)pBaseBuf[NC*i+3];
-							}
-							else
-							{
-								for ( int c=0; c<NC; ++c )
+								else
 								{
-									unsigned mask = pMaskBuf[i];
-									unsigned base = pBaseBuf[NC*i+c];
-									unsigned blended = pBlendedBuf[NC*i+c];
-									unsigned result = BLEND_FUNC_MASKED( base, blended, mask );
-									if ( CLAMP )
-									{
-										pDestBuf[NC*i+c] = (uint8)FMath::Min( 255u, result );
-									}
-									else
-									{
-										pDestBuf[NC*i+c] = (uint8)result;
-									}
+									pDestBuf[BASE_CHANNEL_STRIDE * i + c] = (uint8)result;
 								}
 							}
+
+							// Copy the unblended channels
+							// \TODO: unnecessary when doing it in-place?
+							for (int c = 0; c < BASE_CHANNEL_STRIDE - CHANNELS_TO_BLEND; ++c)
+							{
+								pDestBuf[BASE_CHANNEL_STRIDE * i + CHANNELS_TO_BLEND + c] = pBaseBuf[BASE_CHANNEL_STRIDE * i + CHANNELS_TO_BLEND + c];
+							}
 						}
 
-						pDestBuf += NC*different;
-						pBaseBuf += NC*different;
-						pBlendedBuf += NC*different;
+						pDestBuf += BASE_CHANNEL_STRIDE * different;
+						pBaseBuf += BASE_CHANNEL_STRIDE * different;
+						pBlendedBuf += BLENDED_CHANNEL_STRIDE * different;
 						pMaskBuf += different;
 					}
 				}
 
-                rows = FMath::DivideAndRoundUp(rows,2);
-                width = FMath::DivideAndRoundUp(width,2);
-            }
+				rows = FMath::DivideAndRoundUp(rows, 2);
+				width = FMath::DivideAndRoundUp(width, 2);
+			}
 		}
 		else
 		{
-			checkf( false, TEXT("Unsupported mask format.") );
+			checkf(false, TEXT("Unsupported mask format."));
 		}
 	}
 
 
+	//---------------------------------------------------------------------------------------------
+	template< unsigned (*BLEND_FUNC_MASKED)(unsigned, unsigned, unsigned),
+		unsigned (*BLEND_FUNC)(unsigned, unsigned),
+		bool CLAMP,
+		// Number of total channels to actually process
+		int32 CHANNELS_TO_BLEND,
+		// Number of total channels in the base image
+		int32 BASE_CHANNEL_STRIDE,
+		// Number of total channels in the blend image
+		int32 BLENDED_CHANNEL_STRIDE>
+	inline void BufferLayerFormatEmbeddedMask(uint8* pStartDestBuf,
+		const Image* pBase,
+		const Image* pBlend,
+		uint32 BaseChannelOffset,
+		bool bOnlyFirstLOD)
+	{
+		check(pBase->GetSizeX() == pBlend->GetSizeX() && pBase->GetSizeY() == pBlend->GetSizeY());
+		check(bOnlyFirstLOD || pBase->GetLODCount() <= pBlend->GetLODCount());
+		check(pBlend->GetFormat()==EImageFormat::IF_RGBA_UBYTE);
+
+		uint8* pDestBuf = pStartDestBuf;
+		const uint8* pBaseBuf = pBase->GetData() + BaseChannelOffset;
+		const uint8* pBlendedBuf = pBlend->GetData() + BaseChannelOffset;
+		const uint8* pMaskBuf = pBlend->GetData() + 3;
+
+		// The base determines the number of lods to process.
+		int32 LODCount = bOnlyFirstLOD ? 1 : pBase->GetLODCount();
+
+		int32 PixelCount = 0;
+		if (bOnlyFirstLOD)
+		{
+			PixelCount = pBase->GetSizeX() * pBase->GetSizeY();
+		}
+		else
+		{
+			PixelCount = pBase->CalculatePixelCount();
+		}
+
+		ParallelFor(PixelCount,
+			[
+				pBaseBuf, pBlendedBuf, pDestBuf, pMaskBuf
+			] (uint32 i)
+			{
+				uint32 mask = pMaskBuf[BLENDED_CHANNEL_STRIDE * i];
+				for (int c = 0; c < CHANNELS_TO_BLEND; ++c)
+				{
+					uint32 base = pBaseBuf[BASE_CHANNEL_STRIDE * i + c];
+					uint32 blended = pBlendedBuf[BLENDED_CHANNEL_STRIDE * i + c];
+					uint32 result = BLEND_FUNC_MASKED(base, blended, mask);
+					if (CLAMP)
+					{
+						pDestBuf[BASE_CHANNEL_STRIDE * i + c] = (uint8)FMath::Min(255u, result);
+					}
+					else
+					{
+						pDestBuf[BASE_CHANNEL_STRIDE * i + c] = (uint8)result;
+					}
+				}
+				// Copy the unblended channels
+				// \TODO: unnecessary when doing it in-place?
+				for (int c = 0; c < BASE_CHANNEL_STRIDE - CHANNELS_TO_BLEND; ++c)
+				{
+					pDestBuf[BASE_CHANNEL_STRIDE * i + CHANNELS_TO_BLEND + c] = pBaseBuf[BASE_CHANNEL_STRIDE * i + CHANNELS_TO_BLEND + c];
+				}
+			});
+	}
+
+
 	//---------------------------------------------------------------------------------------------
 	template< unsigned (*BLEND_FUNC)(unsigned,unsigned),
 			  bool CLAMP,
@@ -1005,23 +1171,31 @@ namespace mu
 							 const Image* pBase,
 							 const Image* pMask,
                              const Image* pBlend,
-                             bool applyToAlpha,
+                             bool bApplyToAlpha,
 							 bool bOnlyFirstLOD )
 	{
 		if ( pBase->GetFormat()==EImageFormat::IF_RGB_UBYTE )
 		{
-            BufferLayerFormat< BLEND_FUNC_MASKED, BLEND_FUNC, CLAMP, 3 >
-                    ( pDestBuf, pBase, pMask, pBlend, applyToAlpha, bOnlyFirstLOD) ;
+            BufferLayerFormat< BLEND_FUNC_MASKED, BLEND_FUNC, CLAMP, 3, 3, 3 >
+                    ( pDestBuf, pBase, pMask, pBlend, 0, 0, bOnlyFirstLOD) ;
 		}
         else if ( pBase->GetFormat()==EImageFormat::IF_RGBA_UBYTE || pBase->GetFormat()==EImageFormat::IF_BGRA_UBYTE )
 		{
-            BufferLayerFormat< BLEND_FUNC_MASKED, BLEND_FUNC, CLAMP, 4 >
-                    ( pDestBuf, pBase, pMask, pBlend, applyToAlpha, bOnlyFirstLOD) ;
+			if (bApplyToAlpha)
+			{
+				BufferLayerFormat< BLEND_FUNC_MASKED, BLEND_FUNC, CLAMP, 4, 4, 4 >
+					(pDestBuf, pBase, pMask, pBlend, 0, 0, bOnlyFirstLOD);
+			}
+			else
+			{
+				BufferLayerFormat< BLEND_FUNC_MASKED, BLEND_FUNC, CLAMP, 3, 4, 4 >
+					(pDestBuf, pBase, pMask, pBlend, 0, 0, bOnlyFirstLOD);
+			}
 		}
 		else if ( pBase->GetFormat()==EImageFormat::IF_L_UBYTE )
 		{
-            BufferLayerFormat< BLEND_FUNC_MASKED, BLEND_FUNC, CLAMP, 1 >
-                    ( pDestBuf, pBase, pMask, pBlend, applyToAlpha, bOnlyFirstLOD) ;
+            BufferLayerFormat< BLEND_FUNC_MASKED, BLEND_FUNC, CLAMP, 1, 1, 1 >
+                    ( pDestBuf, pBase, pMask, pBlend, 0, 0, bOnlyFirstLOD) ;
 		}
         else
 		{
@@ -1030,6 +1204,113 @@ namespace mu
 	}
 
 
+	//---------------------------------------------------------------------------------------------
+	template< unsigned (*BLEND_FUNC_MASKED)(unsigned, unsigned, unsigned),
+		unsigned (*BLEND_FUNC)(unsigned, unsigned),
+		bool CLAMP >
+	inline void BufferLayerEmbeddedMask(uint8* pDestBuf,
+		const Image* pBase,
+		const Image* pBlend,
+		bool bApplyToAlpha,
+		bool bOnlyFirstLOD)
+	{
+		if (pBase->GetFormat() == EImageFormat::IF_RGB_UBYTE)
+		{
+			BufferLayerFormatEmbeddedMask< BLEND_FUNC_MASKED, BLEND_FUNC, CLAMP, 3, 3, 3 >
+				(pDestBuf, pBase, pBlend, 0, bOnlyFirstLOD);
+		}
+		else if (pBase->GetFormat() == EImageFormat::IF_RGBA_UBYTE || pBase->GetFormat() == EImageFormat::IF_BGRA_UBYTE)
+		{
+			if (bApplyToAlpha)
+			{
+				BufferLayerFormatEmbeddedMask< BLEND_FUNC_MASKED, BLEND_FUNC, CLAMP, 4, 4, 4 >
+					(pDestBuf, pBase, pBlend, 0, bOnlyFirstLOD);
+			}
+			else
+			{
+				BufferLayerFormatEmbeddedMask< BLEND_FUNC_MASKED, BLEND_FUNC, CLAMP, 3, 4, 4 >
+					(pDestBuf, pBase, pBlend, 0, bOnlyFirstLOD);
+			}
+		}
+		else
+		{
+			checkf(false, TEXT("Unsupported format."));
+		}
+	}
+
+
+	//---------------------------------------------------------------------------------------------
+	template< unsigned (*RGB_FUNC_MASKED)(unsigned, unsigned, unsigned),
+		unsigned (*A_FUNC)(unsigned, unsigned),
+		bool CLAMP >
+	inline void BufferLayerComposite(
+		Image* pBase,
+		const Image* pBlend,
+		bool bOnlyFirstLOD)
+	{
+		check(pBase->GetFormat() == EImageFormat::IF_RGBA_UBYTE);
+		check(pBlend->GetFormat() == EImageFormat::IF_RGBA_UBYTE);
+		check(pBase->GetSizeX() == pBlend->GetSizeX() && pBase->GetSizeY() == pBlend->GetSizeY());
+		check(bOnlyFirstLOD || pBase->GetLODCount() <= pBlend->GetLODCount());
+
+		uint8* pBaseBuf = pBase->GetData();
+		const uint8* pBlendedBuf = pBlend->GetData();
+
+		// The base determines the number of lods to process.
+		int32 LODCount = bOnlyFirstLOD ? 1 : pBase->GetLODCount();
+
+		int32 PixelCount = 0;
+		if (bOnlyFirstLOD)
+		{
+			PixelCount = pBase->GetSizeX() * pBase->GetSizeY();
+		}
+		else
+		{
+			PixelCount = pBase->CalculatePixelCount();
+		}
+
+		ParallelFor(PixelCount,
+			[
+				pBaseBuf, pBlendedBuf
+			] (uint32 i)
+			{
+				// TODO: Optimize this (SIMD?)
+				uint32 mask = pBlendedBuf[4 * i + 3];
+
+				// RGB
+				for (int c = 0; c < 3; ++c)
+				{
+					uint32 base = pBaseBuf[4 * i + c];
+					uint32 blended = pBlendedBuf[4 * i + c];
+					uint32 result = RGB_FUNC_MASKED(base, blended, mask);
+					if (CLAMP)
+					{
+						pBaseBuf[4 * i + c] = (uint8)FMath::Min(255u, result);
+					}
+					else
+					{
+						pBaseBuf[4 * i + c] = (uint8)result;
+					}
+				}
+
+				// A
+				{
+					uint32 base = pBaseBuf[4 * i + 3];
+					uint32 blended = pBlendedBuf[4 * i + 3];
+					uint32 result = A_FUNC(base, blended);
+					if (CLAMP)
+					{
+						pBaseBuf[4 * i + 3] = (uint8)FMath::Min(255u, result);
+					}
+					else
+					{
+						pBaseBuf[4 * i + 3] = (uint8)result;
+					}
+				}
+			});
+	}
+
+
 	//---------------------------------------------------------------------------------------------
 	template< unsigned (*BLEND_FUNC)(unsigned,unsigned),
 			  bool CLAMP >
@@ -1057,64 +1338,6 @@ namespace mu
 	}
 
 
-	//---------------------------------------------------------------------------------------------
-	//! Apply a blending function to an image with another image as blending layer with a mask
-	//---------------------------------------------------------------------------------------------
-	template< unsigned (*BLEND_FUNC_MASKED)(unsigned,unsigned,unsigned),
-			  unsigned (*BLEND_FUNC)(unsigned,unsigned),
-			  bool CLAMP >
-    inline void ImageLayer( Image* pResult,
-                            const Image* pBase,
-                            const Image* pMask,
-                            const Image* pBlended,
-                            bool applyToAlpha )
-	{
-        check( pResult->GetFormat() == pBase->GetFormat() );
-        check( pResult->GetSizeX() == pBase->GetSizeX() );
-        check( pResult->GetSizeY() == pBase->GetSizeY() );
-        check( pResult->GetLODCount() == pBase->GetLODCount() );
-		check( pBase->GetSizeX() == pMask->GetSizeX() );
-		check( pBase->GetSizeY() == pMask->GetSizeY() );
-		check( pBase->GetSizeX() == pBlended->GetSizeX() );
-		check( pBase->GetSizeY() == pBlended->GetSizeY() );
-		check( pBase->GetFormat() == pBlended->GetFormat() );
-		check( pMask->GetFormat() == EImageFormat::IF_L_UBYTE
-						||
-						pMask->GetFormat() == EImageFormat::IF_L_UBYTE_RLE );
-        check( pBase->GetLODCount() <= pMask->GetLODCount() );
-        check( pBase->GetLODCount() <= pBlended->GetLODCount() );
-
-        BufferLayer<BLEND_FUNC_MASKED,BLEND_FUNC,CLAMP>( pResult->GetData(), pBase, pMask,
-                                                         pBlended, applyToAlpha, false);
-	}
-
-
-	//---------------------------------------------------------------------------------------------
-	//! Apply a blending function to an image with another image as blending layer, on the base
-	//! image itself
-	//---------------------------------------------------------------------------------------------
-	template< unsigned (*BLEND_FUNC_MASKED)(unsigned,unsigned,unsigned),
-			  unsigned (*BLEND_FUNC)(unsigned,unsigned),
-			  bool CLAMP >
-    inline void ImageLayerOnBase( Image* pBase, const Image* pMask, const Image* pBlended,
-                                  bool applyToAlpha, bool bOnlyFirstLOD )
-	{
-		check( pBase->GetSizeX() == pMask->GetSizeX() );
-		check( pBase->GetSizeY() == pMask->GetSizeY() );
-		check( pBase->GetSizeX() == pBlended->GetSizeX() );
-		check( pBase->GetSizeY() == pBlended->GetSizeY() );
-		check( pBase->GetFormat() == pBlended->GetFormat() );
-		check( pMask->GetFormat() == EImageFormat::IF_L_UBYTE
-						||
-						pMask->GetFormat() == EImageFormat::IF_L_UBYTE_RLE );
-        check(bOnlyFirstLOD || pBase->GetLODCount() <= pMask->GetLODCount() );
-        check(bOnlyFirstLOD || pBase->GetLODCount() <= pBlended->GetLODCount() );
-
-        BufferLayer<BLEND_FUNC_MASKED,BLEND_FUNC,CLAMP>( pBase->GetData(), pBase, pMask,
-                                                         pBlended, applyToAlpha, bOnlyFirstLOD);
-	}
-
-
 	//---------------------------------------------------------------------------------------------
 	//! Apply a blending function to an image with another image as blending layer, on a subrect of
 	//! the base image.
@@ -1150,6 +1373,7 @@ namespace mu
 
 
 	//---------------------------------------------------------------------------------------------
+	// \TODO: Unused?
 	template< int NC >
     inline void BufferTableColourFormat( uint8* pStartDestBuf,
 										 const Image* pBase,
@@ -1314,6 +1538,7 @@ namespace mu
 
 
 	//---------------------------------------------------------------------------------------------
+	// \TODO: Unused?
     inline void ImageTable( Image* pResult,
                                 const Image* pBase,
 								const Image* pMask,
diff --git a/Engine/Plugins/Experimental/Mutable/Source/MutableRuntime/Private/MuR/OpImageMultiply.h b/Engine/Plugins/Experimental/Mutable/Source/MutableRuntime/Private/MuR/OpImageMultiply.h
deleted file mode 100644
diff --git a/Engine/Plugins/Experimental/Mutable/Source/MutableRuntime/Private/MuR/OpImageOverlay.h b/Engine/Plugins/Experimental/Mutable/Source/MutableRuntime/Private/MuR/OpImageOverlay.h
deleted file mode 100644
diff --git a/Engine/Plugins/Experimental/Mutable/Source/MutableRuntime/Private/MuR/OpImageProject.cpp b/Engine/Plugins/Experimental/Mutable/Source/MutableRuntime/Private/MuR/OpImageProject.cpp
--- a/Engine/Plugins/Experimental/Mutable/Source/MutableRuntime/Private/MuR/OpImageProject.cpp
+++ b/Engine/Plugins/Experimental/Mutable/Source/MutableRuntime/Private/MuR/OpImageProject.cpp
@@ -1792,7 +1792,7 @@ void MeshProject_Optimised_Wrapping( const Mesh* pMesh,
 	projectedPositions.SetNum(vertexCount);
 
     // Iterate the faces and trace a ray to find the origin face of the projection
-    const float maxDist = 100000.f;
+    const float maxDist = 100000.f; // TODO: This should be proportional to the mesh bounding box size.
     float min_t = maxDist;
     float rayLength = maxDist;
     int intersectedFace = -1;
diff --git a/Engine/Plugins/Experimental/Mutable/Source/MutableRuntime/Private/MuR/OpImageScreen.h b/Engine/Plugins/Experimental/Mutable/Source/MutableRuntime/Private/MuR/OpImageScreen.h
deleted file mode 100644
diff --git a/Engine/Plugins/Experimental/Mutable/Source/MutableRuntime/Private/MuR/OpImageSoftLight.h b/Engine/Plugins/Experimental/Mutable/Source/MutableRuntime/Private/MuR/OpImageSoftLight.h
deleted file mode 100644
diff --git a/Engine/Plugins/Experimental/Mutable/Source/MutableRuntime/Private/MuR/Operations.cpp b/Engine/Plugins/Experimental/Mutable/Source/MutableRuntime/Private/MuR/Operations.cpp
--- a/Engine/Plugins/Experimental/Mutable/Source/MutableRuntime/Private/MuR/Operations.cpp
+++ b/Engine/Plugins/Experimental/Mutable/Source/MutableRuntime/Private/MuR/Operations.cpp
@@ -304,13 +304,6 @@ namespace mu
             f(&op.args.ImageInterpolate3.target2 );
             break;
 
-        case OP_TYPE::IM_SWIZZLE:
-            for (int t=0;t<MUTABLE_OP_MAX_SWIZZLE_CHANNELS;++t)
-            {
-                f(&op.args.ImageSwizzle.sources[t] );
-            }
-            break;
-
         case OP_TYPE::IM_SATURATE:
             f(&op.args.ImageSaturate.base );
             f(&op.args.ImageSaturate.factor );
diff --git a/Engine/Plugins/Experimental/Mutable/Source/MutableRuntime/Private/MuR/Operations.h b/Engine/Plugins/Experimental/Mutable/Source/MutableRuntime/Private/MuR/Operations.h
--- a/Engine/Plugins/Experimental/Mutable/Source/MutableRuntime/Private/MuR/Operations.h
+++ b/Engine/Plugins/Experimental/Mutable/Source/MutableRuntime/Private/MuR/Operations.h
@@ -546,16 +546,19 @@ namespace mu
             ADDRESS blended;
             ADDRESS rangeSize;
             uint16 rangeId;
-            uint16 blendType; // One of BLEND_TYPE
-        };
+			uint8 blendType;		// One of BLEND_TYPE
+			uint8 blendTypeAlpha;	// One of BLEND_TYPE
+			uint8 bUseMaskFromBlended;	
+		};
 
         struct ImageLayerColourArgs
         {
             ADDRESS base;
             ADDRESS mask;
             ADDRESS colour;
-			uint8 blendType;	// One of BLEND_TYPE
-        };
+			uint8 blendType;		// One of BLEND_TYPE
+			uint8 blendTypeAlpha;	// One of BLEND_TYPE
+		};
 
         struct ImagePixelFormatArgs
         {
@@ -1084,7 +1087,6 @@ namespace mu
             ImageInterpolate3Args ImageInterpolate3;
             ImageSaturateArgs ImageSaturate;
             ImageLuminanceArgs ImageLuminance;
-            ImageSwizzleArgs ImageSwizzle;
             ImageSelectColourArgs ImageSelectColour;
             ImageColourMapArgs ImageColourMap;
             ImageGradientArgs ImageGradient;
diff --git a/Engine/Plugins/Experimental/Mutable/Source/MutableRuntime/Public/MuR/Image.h b/Engine/Plugins/Experimental/Mutable/Source/MutableRuntime/Public/MuR/Image.h
--- a/Engine/Plugins/Experimental/Mutable/Source/MutableRuntime/Public/MuR/Image.h
+++ b/Engine/Plugins/Experimental/Mutable/Source/MutableRuntime/Public/MuR/Image.h
@@ -107,7 +107,8 @@ namespace mu
 	//! List of supported modes in generic image layering operations.
 	enum class EBlendType
 	{
-		BT_SOFTLIGHT = 0,
+		BT_NONE = 0,
+		BT_SOFTLIGHT,
 		BT_HARDLIGHT,
 		BT_BURN,
 		BT_DODGE,
@@ -115,7 +116,7 @@ namespace mu
 		BT_OVERLAY,
 		BT_BLEND,
 		BT_MULTIPLY,
-		BT_ALPHA_OVERLAY, // TODO: This name is not descriptive.
+		BT_LIGHTEN,				// Increase the channel value by a given proportion of what is missing from white 
 		BT_NORMAL_COMBINE,
 		_BT_COUNT
 	};
diff --git a/Engine/Plugins/Experimental/Mutable/Source/MutableRuntime/Public/MuR/Model.h b/Engine/Plugins/Experimental/Mutable/Source/MutableRuntime/Public/MuR/Model.h
--- a/Engine/Plugins/Experimental/Mutable/Source/MutableRuntime/Public/MuR/Model.h
+++ b/Engine/Plugins/Experimental/Mutable/Source/MutableRuntime/Public/MuR/Model.h
@@ -12,7 +12,7 @@
 //! data structures. Compiled models are not necessarily compatible when the runtime is updated,
 //! so this version number can be used externally to verify this. It is not used internally, and
 //! serializing models from different versions than this runtime will probably result in a crash.
-#define MUTABLE_COMPILED_MODEL_CODE_VERSION		uint32( 50 )
+#define MUTABLE_COMPILED_MODEL_CODE_VERSION		uint32( 51 )
 #define MUTABLE_PARAMETERS_VERSION              uint32( 1 )
 
 
diff --git a/Engine/Plugins/Experimental/Mutable/Source/MutableTools/Private/MuT/AST.cpp b/Engine/Plugins/Experimental/Mutable/Source/MutableTools/Private/MuT/AST.cpp
--- a/Engine/Plugins/Experimental/Mutable/Source/MutableTools/Private/MuT/AST.cpp
+++ b/Engine/Plugins/Experimental/Mutable/Source/MutableTools/Private/MuT/AST.cpp
@@ -1205,12 +1205,6 @@ FImageDesc ASTOpFixed::GetImageDesc( bool returnBestOption, FGetImageDescContext
         }
         break;
 
-    case OP_TYPE::IM_SWIZZLE:
-        res = GetImageDesc( op.args.ImageSwizzle.sources[0], returnBestOption, context );
-        res.m_format = op.args.ImageSwizzle.format;
-        check( res.m_format != EImageFormat::IF_NONE );
-        break;
-
     case OP_TYPE::IM_SELECTCOLOUR:
         res = GetImageDesc( op.args.ImageSelectColour.base, returnBestOption, context );
         res.m_format = EImageFormat::IF_L_UBYTE;
@@ -1309,12 +1303,6 @@ void ASTOpFixed::GetLayoutBlockSize( int* pBlockX, int* pBlockY )
 		break;
 	}
 
-    case OP_TYPE::IM_SWIZZLE:
-    {
-        GetLayoutBlockSize( op.args.ImageSwizzle.sources[0], pBlockX, pBlockY );
-        break;
-    }
-
     case OP_TYPE::IM_LAYER:
     {
         GetLayoutBlockSize(op.args.ImageLayer.base, pBlockX, pBlockY);
@@ -1754,13 +1742,6 @@ mu::Ptr<ImageSizeExpression> ASTOpFixed::GetImageSizeExpression() const
         }
         break;
 
-    case OP_TYPE::IM_SWIZZLE:
-        if ( children[op.args.ImageSwizzle.sources[0]] )
-        {
-            pRes = children[op.args.ImageSwizzle.sources[0]].child()->GetImageSizeExpression();
-        }
-        break;
-
     case OP_TYPE::IM_GRADIENT:
         pRes->type = ImageSizeExpression::ISET_CONSTANT;
         pRes->size[0] = op.args.ImageGradient.size[0];
diff --git a/Engine/Plugins/Experimental/Mutable/Source/MutableTools/Private/MuT/AST.h b/Engine/Plugins/Experimental/Mutable/Source/MutableTools/Private/MuT/AST.h
--- a/Engine/Plugins/Experimental/Mutable/Source/MutableTools/Private/MuT/AST.h
+++ b/Engine/Plugins/Experimental/Mutable/Source/MutableTools/Private/MuT/AST.h
@@ -1052,12 +1052,9 @@ namespace mu
             }
         }
 
-        //! \todo move to ast swizzle op class when created
-        Ptr<ASTOp> OptimiseSwizzle( const FModelOptimizationOptions& ) const;
     };
 
 
-
     //---------------------------------------------------------------------------------------------
     //---------------------------------------------------------------------------------------------
     //---------------------------------------------------------------------------------------------
diff --git a/Engine/Plugins/Experimental/Mutable/Source/MutableTools/Private/MuT/ASTOpImageMultiLayer.cpp b/Engine/Plugins/Experimental/Mutable/Source/MutableTools/Private/MuT/ASTOpImageMultiLayer.cpp
--- a/Engine/Plugins/Experimental/Mutable/Source/MutableTools/Private/MuT/ASTOpImageMultiLayer.cpp
+++ b/Engine/Plugins/Experimental/Mutable/Source/MutableTools/Private/MuT/ASTOpImageMultiLayer.cpp
@@ -23,7 +23,6 @@ namespace mu
 		, blend(this)
 		, mask(this)
 		, range(this, nullptr, string(), string())
-		, blendType(EBlendType::BT_BLEND)
 	{
 	}
 
@@ -37,15 +36,17 @@ namespace mu
 
 
 	//-------------------------------------------------------------------------------------------------
-	bool ASTOpImageMultiLayer::IsEqual(const ASTOp& otherUntyped) const
+	bool ASTOpImageMultiLayer::IsEqual(const ASTOp& InOtherUntyped) const
 	{
-		if (auto other = dynamic_cast<const ASTOpImageMultiLayer*>(&otherUntyped))
+		if (const ASTOpImageMultiLayer* Other = dynamic_cast<const ASTOpImageMultiLayer*>(&InOtherUntyped))
 		{
-			return base == other->base &&
-				blend == other->blend &&
-				mask == other->mask &&
-				range == other->range &&
-				blendType == other->blendType;
+			return base == Other->base &&
+				blend == Other->blend &&
+				mask == Other->mask &&
+				range == Other->range &&
+				blendType == Other->blendType &&
+				blendTypeAlpha == Other->blendTypeAlpha &&
+				bUseMaskFromBlended == Other->bUseMaskFromBlended;
 		}
 		return false;
 	}
@@ -73,6 +74,8 @@ namespace mu
 		n->range.rangeUID = range.rangeUID;
 		n->range.rangeSize = mapChild(range.rangeSize.child());
 		n->blendType = blendType;
+		n->blendTypeAlpha = blendTypeAlpha;
+		n->bUseMaskFromBlended = bUseMaskFromBlended;
 		return n;
 	}
 
@@ -96,7 +99,9 @@ namespace mu
 			OP::ImageMultiLayerArgs args;
 			memset(&args, 0, sizeof(args));
 
-			args.blendType = (uint16)blendType;
+			args.blendType = (uint8)blendType;
+			args.blendTypeAlpha = (uint8)blendTypeAlpha;
+			args.bUseMaskFromBlended = bUseMaskFromBlended;
 
 			if (base) args.base = base->linkedAddress;
 			if (blend) args.blended = blend->linkedAddress;
diff --git a/Engine/Plugins/Experimental/Mutable/Source/MutableTools/Private/MuT/ASTOpImageMultiLayer.h b/Engine/Plugins/Experimental/Mutable/Source/MutableTools/Private/MuT/ASTOpImageMultiLayer.h
--- a/Engine/Plugins/Experimental/Mutable/Source/MutableTools/Private/MuT/ASTOpImageMultiLayer.h
+++ b/Engine/Plugins/Experimental/Mutable/Source/MutableTools/Private/MuT/ASTOpImageMultiLayer.h
@@ -25,7 +25,15 @@ struct FProgram;
 		ASTChild blend;
 		ASTChild mask;
 		FRangeData range;
-		EBlendType blendType;
+
+		/** Blend type used for the colour channels. */
+		EBlendType blendType = EBlendType::BT_NONE;
+
+		/** Blend type used for the alpha channel if any. */
+		EBlendType blendTypeAlpha = EBlendType::BT_NONE;
+
+		/** If true, use the alpha channel of the blended image as mask. Mask should be null.*/
+		bool bUseMaskFromBlended = false;
 
 	public:
 
diff --git a/Engine/Plugins/Experimental/Mutable/Source/MutableTools/Private/MuT/ASTOpImageSwizzle.cpp b/Engine/Plugins/Experimental/Mutable/Source/MutableTools/Private/MuT/ASTOpImageSwizzle.cpp
new file mode 100644
--- /dev/null
+++ b/Engine/Plugins/Experimental/Mutable/Source/MutableTools/Private/MuT/ASTOpImageSwizzle.cpp
@@ -0,0 +1,1021 @@
+// Copyright Epic Games, Inc. All Rights Reserved.
+
+#include "MuT/ASTOpImageSwizzle.h"
+
+#include "Containers/Map.h"
+#include "HAL/UnrealMemory.h"
+#include "Math/UnrealMathSSE.h"
+#include "Misc/AssertionMacros.h"
+#include "MuR/ImagePrivate.h"
+#include "MuR/MemoryPrivate.h"
+#include "MuR/ModelPrivate.h"
+#include "MuR/MutableMath.h"
+#include "MuR/RefCounted.h"
+#include "MuR/Types.h"
+#include "MuT/ASTOpConditional.h"
+#include "MuT/ASTOpImageCompose.h"
+#include "MuT/ASTOpImagePatch.h"
+#include "MuT/ASTOpImagePixelFormat.h"
+#include "MuT/ASTOpImageMultiLayer.h"
+#include "MuT/ASTOpSwitch.h"
+#include "MuT/StreamsPrivate.h"
+
+#include <memory>
+#include <utility>
+
+
+namespace mu
+{
+
+	ASTOpImageSwizzle::ASTOpImageSwizzle()
+		: Sources{ ASTChild(this),ASTChild(this),ASTChild(this),ASTChild(this) }
+	{
+	}
+
+
+	ASTOpImageSwizzle::~ASTOpImageSwizzle()
+	{
+		// Explicit call needed to avoid recursive destruction
+		ASTOp::RemoveChildren();
+	}
+
+
+	bool ASTOpImageSwizzle::IsEqual(const ASTOp& otherUntyped) const
+	{
+		if (const ASTOpImageSwizzle* Other = dynamic_cast<const ASTOpImageSwizzle*>(&otherUntyped))
+		{
+			for (int32 i = 0; i<MUTABLE_OP_MAX_SWIZZLE_CHANNELS; ++i)
+			{
+				if (!(Sources[i] == Other->Sources[i] && SourceChannels[i] == Other->SourceChannels[i]))
+				{
+					return false;
+				}
+			}
+			return Format == Other->Format;
+		}
+		return false;
+	}
+
+
+	uint64 ASTOpImageSwizzle::Hash() const
+	{
+		uint64 res = std::hash<void*>()(Sources[0].child().get());
+		hash_combine(res, std::hash<void*>()(Sources[1].child().get()));
+		hash_combine(res, std::hash<void*>()(Sources[2].child().get()));
+		hash_combine(res, std::hash<void*>()(Sources[3].child().get()));
+		hash_combine(res, std::hash<uint8>()(SourceChannels[0]));
+		hash_combine(res, std::hash<uint8>()(SourceChannels[1]));
+		hash_combine(res, std::hash<uint8>()(SourceChannels[2]));
+		hash_combine(res, std::hash<uint8>()(SourceChannels[3]));
+		hash_combine(res, Format);
+		return res;
+	}
+
+
+	mu::Ptr<ASTOp> ASTOpImageSwizzle::Clone(MapChildFuncRef mapChild) const
+	{
+		mu::Ptr<ASTOpImageSwizzle> n = new ASTOpImageSwizzle();
+		for (int32 i = 0; i < MUTABLE_OP_MAX_SWIZZLE_CHANNELS; ++i)
+		{
+			n->Sources[i] = mapChild(Sources[i].child());
+			n->SourceChannels[i] = SourceChannels[i];
+		}
+		n->Format = Format;
+		return n;
+	}
+
+
+	void ASTOpImageSwizzle::ForEachChild(const TFunctionRef<void(ASTChild&)> f)
+	{
+		for (int32 i = 0; i < MUTABLE_OP_MAX_SWIZZLE_CHANNELS; ++i)
+		{
+			f(Sources[i]);
+		}
+	}
+
+
+	void ASTOpImageSwizzle::Link(FProgram& program, const FLinkerOptions*)
+	{
+		// Already linked?
+		if (!linkedAddress)
+		{
+			OP::ImageSwizzleArgs args;
+			FMemory::Memzero(&args, sizeof(args));
+
+			args.format = Format;
+			for (int32 i = 0; i < MUTABLE_OP_MAX_SWIZZLE_CHANNELS; ++i)
+			{
+				if (Sources[i]) args.sources[i] = Sources[i]->linkedAddress;
+				args.sourceChannels[i] = SourceChannels[i];
+			}			
+
+			linkedAddress = (OP::ADDRESS)program.m_opAddress.Num();
+			program.m_opAddress.Add((uint32_t)program.m_byteCode.Num());
+			AppendCode(program.m_byteCode, GetOpType());
+			AppendCode(program.m_byteCode, args);
+		}
+	}
+
+
+	mu::Ptr<ASTOp> ASTOpImageSwizzle::OptimiseSemantic(const FModelOptimizationOptions&) const
+	{
+		// If children channels are also swizzle ops, recurse them
+		Ptr<ASTOpImageSwizzle> sat;
+
+		for (int c = 0; c < MUTABLE_OP_MAX_SWIZZLE_CHANNELS; ++c)
+		{
+			Ptr<ASTOp> candidate = Sources[c].child();
+			if (!candidate)
+			{
+				continue;
+			}
+
+			switch (candidate->GetOpType())
+			{
+			// Swizzle
+			case OP_TYPE::IM_SWIZZLE:
+			{
+				if (!sat)
+				{
+					sat = mu::Clone<ASTOpImageSwizzle>(this);
+				}
+				const ASTOpImageSwizzle* typedCandidate = dynamic_cast<const ASTOpImageSwizzle*>(candidate.get());
+				int candidateChannel = SourceChannels[c];
+
+				sat->Sources[c] = typedCandidate->Sources[candidateChannel].child();
+				sat->SourceChannels[c] = typedCandidate->SourceChannels[candidateChannel];
+
+				break;
+			}
+
+			// Format
+			case OP_TYPE::IM_PIXELFORMAT:
+			{
+				// We can remove the format if its source is already an uncompressed format
+				ASTOpImagePixelFormat* typedCandidate = dynamic_cast<ASTOpImagePixelFormat*>(candidate.get());
+				Ptr<ASTOp> formatSource = typedCandidate->Source.child();
+
+				if (formatSource)
+				{
+					FImageDesc desc = formatSource->GetImageDesc();
+					if (desc.m_format != EImageFormat::IF_NONE && !IsCompressedFormat(desc.m_format))
+					{
+						if (!sat)
+						{
+							sat = mu::Clone<ASTOpImageSwizzle>(this);
+						}
+						sat->Sources[c] = formatSource;
+					}
+				}
+
+				break;
+			}
+
+			default:
+				break;
+			}
+
+		}
+
+		return sat;
+	}
+
+	namespace
+	{
+		//---------------------------------------------------------------------------------------------
+		//! Set al the non-null sources of an image swizzle operation to the given value
+		//---------------------------------------------------------------------------------------------
+		void ReplaceAllSources(Ptr<ASTOpImageSwizzle>& op, Ptr<ASTOp>& value)
+		{
+			check(op->GetOpType() == OP_TYPE::IM_SWIZZLE);
+			for (int c = 0; c < MUTABLE_OP_MAX_SWIZZLE_CHANNELS; ++c)
+			{
+				if (op->Sources[c])
+				{
+					op->Sources[c] = value;
+				}
+			}
+		}
+	}
+
+	mu::Ptr<ASTOp> ASTOpImageSwizzle::OptimiseSink(const FModelOptimizationOptions& options, FOptimizeSinkContext& context) const
+	{
+		MUTABLE_CPUPROFILER_SCOPE(OptimiseSwizzleAST);
+
+		//! Basic optimisation first
+		Ptr<ASTOp> at = OptimiseSemantic(options);
+		if (at)
+		{
+			return at;
+		}
+
+		// If all sources are the same, we can sink the instruction
+		bool bAllChannelsAreTheSame = true;
+		bool bAllChannelsAreTheSameType = true;
+		Ptr<ASTOp> channelSourceAt;
+		for (int c = 0; c < MUTABLE_OP_MAX_SWIZZLE_CHANNELS; ++c)
+		{
+			Ptr<ASTOp> candidate = Sources[c].child();
+			if (candidate)
+			{
+				if (!channelSourceAt)
+				{
+					channelSourceAt = candidate;
+				}
+				else
+				{
+					bAllChannelsAreTheSame = bAllChannelsAreTheSame && (channelSourceAt == candidate);
+					bAllChannelsAreTheSameType = bAllChannelsAreTheSameType && (channelSourceAt->GetOpType() == candidate->GetOpType());
+				}
+			}
+		}
+
+		OP_TYPE sourceType = channelSourceAt->GetOpType();
+
+		if (bAllChannelsAreTheSame && channelSourceAt)
+		{
+			// The instruction can be sunk
+			switch (sourceType)
+			{
+
+			case OP_TYPE::IM_PLAINCOLOUR:
+			{
+				auto NewPlain = mu::Clone<ASTOpFixed>(channelSourceAt);
+				Ptr<ASTOpFixed> NewSwizzle = new ASTOpFixed;
+				NewSwizzle->op.type = OP_TYPE::CO_SWIZZLE;
+				for (int i = 0; i < MUTABLE_OP_MAX_SWIZZLE_CHANNELS; ++i)
+				{
+					NewSwizzle->SetChild(NewSwizzle->op.args.ColourSwizzle.sources[i], NewPlain->children[NewPlain->op.args.ImagePlainColour.colour]);
+					NewSwizzle->op.args.ColourSwizzle.sourceChannels[i] = SourceChannels[i];
+				}
+				NewPlain->SetChild(NewPlain->op.args.ImagePlainColour.colour, NewSwizzle);
+				NewPlain->op.args.ImagePlainColour.format = Format;
+				at = NewPlain;
+				break;
+			}
+
+			case OP_TYPE::IM_SWITCH:
+			{
+				// Move the swizzle down all the paths
+				Ptr<ASTOpSwitch> nop = mu::Clone<ASTOpSwitch>(channelSourceAt);
+
+				if (nop->def)
+				{
+					Ptr<ASTOpImageSwizzle> defOp = mu::Clone<ASTOpImageSwizzle>(this);
+					ReplaceAllSources(defOp, nop->def.child());
+					nop->def = defOp;
+				}
+
+				for (int32 v = 0; v < nop->cases.Num(); ++v)
+				{
+					if (nop->cases[v].branch)
+					{
+						Ptr<ASTOpImageSwizzle> bOp = mu::Clone<ASTOpImageSwizzle>(this);
+						ReplaceAllSources(bOp, nop->cases[v].branch.child());
+						nop->cases[v].branch = bOp;
+					}
+				}
+
+				at = nop;
+				break;
+			}
+
+			case OP_TYPE::IM_CONDITIONAL:
+			{
+				// We move the swizzle down the two paths
+				Ptr<ASTOpConditional> nop = mu::Clone<ASTOpConditional>(channelSourceAt);
+
+				Ptr<ASTOpImageSwizzle> aOp = mu::Clone<ASTOpImageSwizzle>(this);
+				ReplaceAllSources(aOp, nop->yes.child());
+				nop->yes = aOp;
+
+				Ptr<ASTOpImageSwizzle> bOp = mu::Clone<ASTOpImageSwizzle>(this);
+				ReplaceAllSources(bOp, nop->no.child());
+				nop->no = bOp;
+
+				at = nop;
+				break;
+			}
+
+			case OP_TYPE::IM_LAYER:
+			{
+				// We move the swizzle down the two paths
+				Ptr<ASTOpFixed> nop = mu::Clone<ASTOpFixed>(channelSourceAt);
+
+				Ptr<ASTOpImageSwizzle> aOp = mu::Clone<ASTOpImageSwizzle>(this);
+				ReplaceAllSources(aOp, nop->children[nop->op.args.ImageLayer.base].child());
+				nop->SetChild(nop->op.args.ImageLayer.base, aOp);
+
+				Ptr<ASTOpImageSwizzle> bOp = mu::Clone<ASTOpImageSwizzle>(this);
+				ReplaceAllSources(bOp, nop->children[nop->op.args.ImageLayer.blended].child());
+				nop->SetChild(nop->op.args.ImageLayer.blended, bOp);
+
+				at = nop;
+				break;
+			}
+
+			case OP_TYPE::IM_LAYERCOLOUR:
+			{
+				// We move the swizzle down the base path
+				Ptr<ASTOpFixed> nop = mu::Clone<ASTOpFixed>(channelSourceAt);
+
+				Ptr<ASTOpImageSwizzle> NewSwizzle = mu::Clone<ASTOpImageSwizzle>(this);
+				ReplaceAllSources(NewSwizzle, nop->children[nop->op.args.ImageLayerColour.base].child());
+				nop->SetChild(nop->op.args.ImageLayerColour.base, NewSwizzle);
+
+				// We need to swizzle the colour too
+				Ptr<ASTOpFixed> cOp = new ASTOpFixed;
+				cOp->op.type = OP_TYPE::CO_SWIZZLE;
+				for (int i = 0; i < MUTABLE_OP_MAX_SWIZZLE_CHANNELS; ++i)
+				{
+					cOp->SetChild(cOp->op.args.ColourSwizzle.sources[i], nop->children[nop->op.args.ImageLayerColour.colour]);
+					cOp->op.args.ColourSwizzle.sourceChannels[i] = SourceChannels[i];
+				}
+				nop->SetChild(nop->op.args.ImageLayerColour.colour, cOp);
+
+				at = nop;
+				break;
+			}
+
+			case OP_TYPE::IM_DISPLACE:
+			{
+				Ptr<ASTOpFixed> NewDisplace = mu::Clone<ASTOpFixed>(channelSourceAt);
+				Ptr<ASTOpImageSwizzle> NewSwizzle = mu::Clone<ASTOpImageSwizzle>(this);
+				ReplaceAllSources(NewSwizzle, NewDisplace->children[NewDisplace->op.args.ImageDisplace.source].child());
+				NewDisplace->SetChild(NewDisplace->op.args.ImageDisplace.source, NewSwizzle);
+				at = NewDisplace;
+				break;
+			}
+
+			case OP_TYPE::IM_RASTERMESH:
+			{
+				Ptr<ASTOpFixed> NewRaster = mu::Clone<ASTOpFixed>(channelSourceAt);
+				Ptr<ASTOpImageSwizzle> NewSwizzle = mu::Clone<ASTOpImageSwizzle>(this);
+				ReplaceAllSources(NewSwizzle, NewRaster->children[NewRaster->op.args.ImageRasterMesh.image].child());
+				NewRaster->SetChild(NewRaster->op.args.ImageRasterMesh.image, NewSwizzle);
+				at = NewRaster;
+				break;
+			}
+
+			default:
+				bAllChannelsAreTheSame = false;
+				break;
+			}
+		}
+
+		if (!bAllChannelsAreTheSame && bAllChannelsAreTheSameType)
+		{
+			// Maybe we can still sink the instruction in some cases
+
+			// If we have RGB being the same IM_MULTILAYER, and alpha a compatible IM_MULTILAYER we can optimize with
+			// a special multilayer blend mode. This happens often because of higher level group projector nodes.
+			if (!at
+				&&
+				Sources[0] == Sources[1] && Sources[0] == Sources[2]
+				&&
+				Sources[0] && Sources[0]->GetOpType() == OP_TYPE::IM_MULTILAYER
+				&&
+				Sources[3] && Sources[3]->GetOpType() == OP_TYPE::IM_MULTILAYER
+				&&
+				SourceChannels[0] == 0 && SourceChannels[1] == 1 && SourceChannels[2] == 2 && SourceChannels[3] == 0
+				)
+			{
+				const ASTOpImageMultiLayer* ColorMultiLayer = dynamic_cast<const ASTOpImageMultiLayer*>(Sources[0].child().get());
+				check(ColorMultiLayer);
+				const ASTOpImageMultiLayer* AlphaMultiLayer = dynamic_cast<const ASTOpImageMultiLayer*>(Sources[3].child().get());
+				check(AlphaMultiLayer);
+
+				bool bIsSpecialMultiLayer = !AlphaMultiLayer->mask
+					&&
+					ColorMultiLayer->range == AlphaMultiLayer->range;
+
+				if (bIsSpecialMultiLayer)
+				{
+					// We can combine the 2 multilayers into the composite blend+lighten mode
+
+					Ptr<ASTOpImageSwizzle> NewBase = mu::Clone<ASTOpImageSwizzle>(this);
+					NewBase->Sources[0] = ColorMultiLayer->base.child();
+					NewBase->Sources[1] = ColorMultiLayer->base.child();
+					NewBase->Sources[2] = ColorMultiLayer->base.child();
+					NewBase->Sources[3] = AlphaMultiLayer->base.child();
+
+					Ptr<ASTOpImageSwizzle> NewBlended = mu::Clone<ASTOpImageSwizzle>(this);
+					NewBlended->Sources[0] = ColorMultiLayer->blend.child();
+					NewBlended->Sources[1] = ColorMultiLayer->blend.child();
+					NewBlended->Sources[2] = ColorMultiLayer->blend.child();
+					NewBlended->Sources[3] = AlphaMultiLayer->blend.child();
+
+					Ptr<ASTOpImageMultiLayer> NewMultiLayer = mu::Clone<ASTOpImageMultiLayer>(ColorMultiLayer);
+					NewMultiLayer->blendTypeAlpha = AlphaMultiLayer->blendType;
+					NewMultiLayer->base = NewBase;
+					NewMultiLayer->blend = NewBlended;
+
+					if ( NewMultiLayer->mask.child() == AlphaMultiLayer->blend.child()
+						&&
+						NewBlended->Format==EImageFormat::IF_RGBA_UBYTE
+						)
+					{
+						// Additional optimization is possible here.
+						NewMultiLayer->bUseMaskFromBlended = true;
+						NewMultiLayer->mask = nullptr;
+					}
+
+					at = NewMultiLayer;
+				}
+			}
+
+
+			// If the channels are compatible switches, we can still sink the swizzle.
+			if (!at && sourceType == OP_TYPE::IM_SWITCH)
+			{
+				const ASTOpSwitch* FirstSwitch = dynamic_cast<const ASTOpSwitch*>(Sources[0].child().get());
+				check(FirstSwitch);
+
+				bool bAreAllSwitchesCompatible = true;
+				for (int32 c = 1; c < MUTABLE_OP_MAX_SWIZZLE_CHANNELS; ++c)
+				{
+					if (Sources[c])
+					{
+						const ASTOpSwitch* Typed = dynamic_cast<const ASTOpSwitch*>(Sources[c].child().get());
+						check(Typed);
+						if (!Typed->IsCompatibleWith(FirstSwitch))
+						{
+							bAreAllSwitchesCompatible = false;
+							break;
+						}
+					}
+				}
+
+				if (bAreAllSwitchesCompatible)
+				{
+					// Move the swizzle down all the paths
+					Ptr<ASTOpSwitch> nop = mu::Clone<ASTOpSwitch>(channelSourceAt);
+
+					if (nop->def)
+					{
+						Ptr<ASTOpImageSwizzle> defOp = mu::Clone<ASTOpImageSwizzle>(this);
+						for (int c = 0; c < MUTABLE_OP_MAX_SWIZZLE_CHANNELS; ++c)
+						{
+							const ASTOpSwitch* ChannelSwitch = dynamic_cast<const ASTOpSwitch*>(Sources[c].child().get());
+							if (ChannelSwitch)
+							{
+								defOp->Sources[c] = ChannelSwitch->def.child();
+							}
+						}
+						nop->def = defOp;
+					}
+
+					for (int32 v = 0; v < nop->cases.Num(); ++v)
+					{
+						if (nop->cases[v].branch)
+						{
+							Ptr<ASTOpImageSwizzle> branchOp = mu::Clone<ASTOpImageSwizzle>(this);
+							for (int c = 0; c < MUTABLE_OP_MAX_SWIZZLE_CHANNELS; ++c)
+							{
+								const ASTOpSwitch* ChannelSwitch = dynamic_cast<const ASTOpSwitch*>(Sources[c].child().get());
+								if (ChannelSwitch)
+								{
+									branchOp->Sources[c] = ChannelSwitch->cases[v].branch.child();
+								}
+							}
+							nop->cases[v].branch = branchOp;
+						}
+					}
+
+					at = nop;
+				}
+			}
+
+			// Swizzle down a displace.
+			if (!at && sourceType == OP_TYPE::IM_DISPLACE)
+			{
+				const ASTOpFixed* FirstDisplace = dynamic_cast<const ASTOpFixed*>(Sources[0].child().get());
+				check(FirstDisplace);
+
+				bool bAreAllDisplacesCompatible = true;
+				for (int32 c = 1; c < MUTABLE_OP_MAX_SWIZZLE_CHANNELS; ++c)
+				{
+					if (Sources[c])
+					{
+						const ASTOpFixed* Typed = dynamic_cast<const ASTOpFixed*>(Sources[c].child().get());
+						check(Typed);
+						if (FirstDisplace->op.args.ImageDisplace.displacementMap != Typed->op.args.ImageDisplace.displacementMap)
+						{
+							bAreAllDisplacesCompatible = false;
+							break;
+						}
+					}
+				}
+
+				if (bAreAllDisplacesCompatible)
+				{
+					// Move the swizzle down all the paths
+					Ptr<ASTOpFixed> NewDisplace = mu::Clone<ASTOpFixed>(FirstDisplace);
+
+					Ptr<ASTOpImageSwizzle> SourceOp = mu::Clone<ASTOpImageSwizzle>(this);
+					for (int c = 0; c < MUTABLE_OP_MAX_SWIZZLE_CHANNELS; ++c)
+					{
+						const ASTOpFixed* ChannelDisplace = dynamic_cast<const ASTOpFixed*>(Sources[c].child().get());
+						if (ChannelDisplace)
+						{
+							SourceOp->Sources[c] = ChannelDisplace->children[ChannelDisplace->op.args.ImageDisplace.source].child();
+						}
+					}
+
+					NewDisplace->SetChild(NewDisplace->op.args.ImageDisplace.source, SourceOp);
+
+					at = NewDisplace;
+				}
+
+			}
+
+			// Swizzle down a raster mesh.
+			if (!at && sourceType == OP_TYPE::IM_RASTERMESH)
+			{
+				const ASTOpFixed* FirstRasterMesh = dynamic_cast<const ASTOpFixed*>(Sources[0].child().get());
+				check(FirstRasterMesh);
+
+				bool bAreAllFirstRasterMeshesCompatible = true;
+				for (int32 c = 1; c < MUTABLE_OP_MAX_SWIZZLE_CHANNELS; ++c)
+				{
+					if (Sources[c])
+					{
+						const ASTOpFixed* Typed = dynamic_cast<const ASTOpFixed*>(Sources[c].child().get());
+						check(Typed);
+						// Compare all args but the source image
+						OP::ImageRasterMeshArgs ArgCopy = FirstRasterMesh->op.args.ImageRasterMesh;
+						ArgCopy.image = Typed->op.args.ImageRasterMesh.image;
+
+						if (FMemory::Memcmp(&ArgCopy, &Typed->op.args.ImageRasterMesh, sizeof(OP::ImageRasterMeshArgs)) != 0)
+						{
+							bAreAllFirstRasterMeshesCompatible = false;
+							break;
+						}
+					}
+				}
+
+				if (bAreAllFirstRasterMeshesCompatible)
+				{
+					// Move the swizzle down all the paths
+					Ptr<ASTOpFixed> NewRaster = mu::Clone<ASTOpFixed>(FirstRasterMesh);
+
+					Ptr<ASTOpImageSwizzle> NewSwizzle = mu::Clone<ASTOpImageSwizzle>(this);
+					for (int c = 0; c < MUTABLE_OP_MAX_SWIZZLE_CHANNELS; ++c)
+					{
+						const ASTOpFixed* ChannelRaster = dynamic_cast<const ASTOpFixed*>(Sources[c].child().get());
+						if (ChannelRaster)
+						{
+							NewSwizzle->Sources[c] = ChannelRaster->children[ChannelRaster->op.args.ImageRasterMesh.image].child();
+						}
+					}
+
+					NewRaster->SetChild(NewRaster->op.args.ImageRasterMesh.image, NewSwizzle);
+
+					at = NewRaster;
+				}
+
+			}
+		}
+
+		// Swizzle of RGB from a source + A from a layer colour
+		// This can be optimize to apply the layer colour on-base directly to the alpha channel to skip the swizzle
+		if ( !at 
+			&&
+			Sources[0] && Sources[0]==Sources[1] && Sources[0]==Sources[2]
+			&&
+			Sources[3] && Sources[3]->GetOpType()==OP_TYPE::IM_LAYERCOLOUR )
+		{
+			// Move the swizzle down all the paths
+			Ptr<ASTOpFixed> NewLayerColour = mu::Clone<ASTOpFixed>(Sources[3].child());
+
+			Ptr<ASTOpImageSwizzle> NewSwizzle = mu::Clone<ASTOpImageSwizzle>(this);
+			NewSwizzle->Sources[3] = NewLayerColour->children[NewLayerColour->op.args.ImageLayerColour.base].child();
+
+			NewLayerColour->op.args.ImageLayerColour.blendTypeAlpha = NewLayerColour->op.args.ImageLayerColour.blendType;
+			NewLayerColour->op.args.ImageLayerColour.blendType = uint8(EBlendType::BT_NONE);
+			NewLayerColour->SetChild(NewLayerColour->op.args.ImageLayerColour.base, NewSwizzle);
+
+			at = NewLayerColour;
+
+		}
+
+
+//            OP::ADDRESS sourceAt = program.m_code[at].args.ImageSwizzle.sources[0];
+//            OP_TYPE sourceType = (OP_TYPE)program.m_code[sourceAt].type;
+//            switch ( sourceType )
+//            {
+//            case OP_TYPE::IM_COMPOSE:
+//            {
+//                bool canSink = true;
+//                OP::ADDRESS layout = program.m_code[sourceAt].args.ImageCompose.layout;
+//                uint32_t block = program.m_code[sourceAt].args.ImageCompose.blockIndex;
+//                for ( int c=1; c<MUTABLE_OP_MAX_SWIZZLE_CHANNELS; ++c )
+//                {
+//                    OP::ADDRESS candidate = program.m_code[at].args.ImageSwizzle.sources[c];
+//                    if ( candidate )
+//                    {
+//                        if ( program.m_code[candidate].type==sourceType )
+//                        {
+//                            OP::ADDRESS clayout =
+//                                    program.m_code[candidate].args.ImageCompose.layout;
+//                            uint32_t cblock =
+//                                    program.m_code[candidate].args.ImageCompose.blockIndex;
+//                            canSink = canSink
+//                                    && ( layout == clayout )
+//                                    && ( block == cblock );
+//                        }
+//                        else
+//                        {
+//                            canSink = false;
+//                        }
+//                    }
+//                }
+
+//                if ( canSink )
+//                {
+//                    m_modified = true;
+
+//                    OP op = program.m_code[sourceAt];
+
+//                    OP swizzleBaseOp = program.m_code[at];
+//                    OP swizzleBlockOp = program.m_code[at];
+
+//                    for ( int c=0; c<MUTABLE_OP_MAX_SWIZZLE_CHANNELS; ++c )
+//                    {
+//                        OP::ADDRESS source = program.m_code[at].args.ImageSwizzle.sources[c];
+//                        if ( source )
+//                        {
+//                            swizzleBaseOp.args.ImageSwizzle.sources[c] =
+//                                    program.m_code[ source ].args.ImageCompose.base;
+
+//                            swizzleBlockOp.args.ImageSwizzle.sources[c] =
+//                                    program.m_code[ source ].args.ImageCompose.blockImage;
+//                        }
+//                    }
+//                    op.args.ImageCompose.base = program.AddOp( swizzleBaseOp );
+//                    op.args.ImageCompose.blockImage = program.AddOp( swizzleBlockOp );
+
+//                    at = program.AddOp( op );
+//                }
+//                break;
+//            }
+
+//            case OP_TYPE::IM_CONDITIONAL:
+//            {
+//                // We can always sink into conditionals
+//                m_modified = true;
+
+//                OP op = program.m_code[sourceAt];
+
+//                OP swizzleYesOp = program.m_code[at];
+//                OP swizzleNoOp = program.m_code[at];
+//                OP::ADDRESS source0 = 0;
+//                for ( int c=0; c<MUTABLE_OP_MAX_SWIZZLE_CHANNELS; ++c )
+//                {
+//                    OP::ADDRESS source = program.m_code[at].args.ImageSwizzle.sources[c];
+//                    if ( c==0 )
+//                    {
+//                        source0 = source;
+//                    }
+
+//                    if ( source
+//                        &&
+//                        ( (c==0) || ( program.m_code[source0].args.Conditional.condition == program.m_code[source].args.Conditional.condition ) )
+//                        )
+//                    {
+//                        swizzleYesOp.args.ImageSwizzle.sources[c] =
+//                                program.m_code[ source ].args.Conditional.yes;
+
+//                        swizzleNoOp.args.ImageSwizzle.sources[c] =
+//                                program.m_code[ source ].args.Conditional.no;
+//                    }
+//                }
+//                op.args.Conditional.yes = program.AddOp( swizzleYesOp );
+//                op.args.Conditional.no = program.AddOp( swizzleNoOp );
+
+//                at = program.AddOp( op );
+//                break;
+//            }
+
+
+//            case OP_TYPE::IM_SWITCH:
+//            {
+//                bool canSink = true;
+//                OP::ADDRESS variable = program.m_code[sourceAt].args.Switch.variable;
+
+//                // If at least 3 channels can be combined, do it. We may be duplicating some data
+//                // But it greatly optimises speed in some cases like the constant alpha masks of
+//                // the face colour interpolation.
+//                // TODO: Make this dependent on compilation options.
+//                //for ( int c=1; c<MUTABLE_OP_MAX_SWIZZLE_CHANNELS; ++c )
+//                for ( int c=1; c<3; ++c )
+//                {
+//                    OP::ADDRESS candidate = program.m_code[at].args.ImageSwizzle.sources[c];
+//                    if ( candidate )
+//                    {
+//                        if ( program.m_code[candidate].type==sourceType )
+//                        {
+//                            OP::ADDRESS cvariable =
+//                                    program.m_code[candidate].args.Switch.variable;
+//                            canSink = canSink && ( variable == cvariable );
+//                        }
+//                        else
+//                        {
+//                            canSink = false;
+//                        }
+//                    }
+//                }
+
+//                if ( canSink )
+//                {
+//                    // TODO since data-switch
+////                    m_modified = true;
+
+////                    OP op = program.m_code[sourceAt];
+
+////                    if ( op.args.Switch.def )
+////                    {
+////                        OP swizzleDefOp = program.m_code[at];
+////                        for ( int c=0; c<MUTABLE_OP_MAX_SWIZZLE_CHANNELS; ++c )
+////                        {
+////                            OP::ADDRESS source = program.m_code[at].args.ImageSwizzle.sources[c];
+////                            if ( source )
+////                            {
+////                                if ( program.m_code[source].type==sourceType )
+////                                {
+////                                    swizzleDefOp.args.ImageSwizzle.sources[c] =
+////                                            program.m_code[ source ].args.Switch.def;
+////                                }
+////                                else
+////                                {
+////                                    swizzleDefOp.args.ImageSwizzle.sources[c] = source;
+////                                }
+////                            }
+////                        }
+////                        op.args.Switch.def = program.AddOp( swizzleDefOp );
+////                    }
+
+////                    for ( int o=0; o<MUTABLE_OP_MAX_SWITCH_OPTIONS; ++o )
+////                    {
+////                        if ( op.args.Switch.values[o] )
+////                        {
+////                            OP swizzleOptOp = program.m_code[at];
+////                            for ( int c=0; c<MUTABLE_OP_MAX_SWIZZLE_CHANNELS; ++c )
+////                            {
+////                                OP::ADDRESS source = program.m_code[at].args.ImageSwizzle.sources[c];
+////                                if ( source )
+////                                {
+////                                    if ( program.m_code[source].type==sourceType )
+////                                    {
+////                                        swizzleOptOp.args.ImageSwizzle.sources[c] =
+////                                                program.m_code[ source ].args.Switch.values[o];
+////                                    }
+////                                    else
+////                                    {
+////                                        swizzleOptOp.args.ImageSwizzle.sources[c] = source;
+////                                    }
+////                                }
+////                            }
+////                            op.args.Switch.values[o] = program.AddOp( swizzleOptOp );
+////                        }
+////                    }
+
+////                    at = program.AddOp( op );
+//                }
+//                break;
+//            }
+
+
+//            case OP_TYPE::IM_INTERPOLATE:
+//            {
+//                bool canSink = true;
+//                OP::ADDRESS factor = program.m_code[sourceAt].args.ImageInterpolate.factor;
+
+//                // It is worth to sink 3 channels
+//                // TODO: Make it depend on compilation parameters
+//                //for ( int c=1; c<MUTABLE_OP_MAX_SWIZZLE_CHANNELS; ++c )
+//                for ( int c=1; c<3; ++c )
+//                {
+//                    OP::ADDRESS candidate = program.m_code[at].args.ImageSwizzle.sources[c];
+//                    if ( candidate )
+//                    {
+//                        if ( program.m_code[candidate].type==sourceType )
+//                        {
+//                            OP::ADDRESS cfactor =
+//                                    program.m_code[candidate].args.ImageInterpolate.factor;
+//                            canSink = canSink && ( factor == cfactor );
+//                        }
+//                        else
+//                        {
+//                            canSink = false;
+//                        }
+//                    }
+//                }
+
+//                if ( canSink )
+//                {
+//                    m_modified = true;
+
+//                    OP op = program.m_code[sourceAt];
+
+//                    for ( int o=0; o<MUTABLE_OP_MAX_INTERPOLATE_COUNT; ++o )
+//                    {
+//                        if ( op.args.ImageInterpolate.targets[o] )
+//                        {
+//                            OP swizzleOptOp = program.m_code[at];
+//                            for ( int c=0; c<MUTABLE_OP_MAX_SWIZZLE_CHANNELS; ++c )
+//                            {
+//                                OP::ADDRESS source = program.m_code[at].args.ImageSwizzle.sources[c];
+
+//                                if ( source
+//                                     && program.m_code[source].type==OP_TYPE::IM_INTERPOLATE
+//                                     && program.m_code[source].args.ImageInterpolate.factor==factor
+//                                     )
+//                                {
+//                                    swizzleOptOp.args.ImageSwizzle.sources[c] =
+//                                            program.m_code[ source ].args.ImageInterpolate.targets[o];
+//                                }
+//                            }
+//                            op.args.ImageInterpolate.targets[o] = program.AddOp( swizzleOptOp );
+//                        }
+//                    }
+
+//                    at = program.AddOp( op );
+//                }
+//                break;
+//            }
+
+
+//            case OP_TYPE::IM_INTERPOLATE3:
+//            {
+//                // We can sink if all channels are interpolates for the same factors
+//                OP::ADDRESS sourceAt0 = program.m_code[at].args.ImageSwizzle.sources[0];
+//                OP::ADDRESS factor10 = program.m_code[sourceAt0].args.ImageInterpolate3.factor1;
+//                OP::ADDRESS factor20 = program.m_code[sourceAt0].args.ImageInterpolate3.factor2;
+
+//                OP::ADDRESS sourceAt1 = program.m_code[at].args.ImageSwizzle.sources[1];
+//                OP::ADDRESS factor11 = program.m_code[sourceAt1].args.ImageInterpolate3.factor1;
+//                OP::ADDRESS factor21 = program.m_code[sourceAt1].args.ImageInterpolate3.factor2;
+
+//                OP::ADDRESS sourceAt2 = program.m_code[at].args.ImageSwizzle.sources[2];
+//                OP::ADDRESS factor12 = program.m_code[sourceAt2].args.ImageInterpolate3.factor1;
+//                OP::ADDRESS factor22 = program.m_code[sourceAt2].args.ImageInterpolate3.factor2;
+
+//                bool canSink = true;
+//                canSink &= program.m_code[sourceAt0].type == program.m_code[sourceAt1].type;
+//                canSink &= program.m_code[sourceAt0].type == program.m_code[sourceAt2].type;
+//                canSink &= factor10 == factor11;
+//                canSink &= factor10 == factor12;
+//                canSink &= factor20 == factor21;
+//                canSink &= factor20 == factor22;
+
+//                if ( canSink )
+//                {
+//                    m_modified = true;
+
+//                    OP op = program.m_code[sourceAt0];
+
+//                    // Target 0
+//                    {
+//                        OP swizzleOptOp = program.m_code[at];
+//                        for ( int c=0; c<MUTABLE_OP_MAX_SWIZZLE_CHANNELS; ++c )
+//                        {
+//                            OP::ADDRESS source = program.m_code[at].args.ImageSwizzle.sources[c];
+//                            if ( source
+//                                 && program.m_code[source].type==OP_TYPE::IM_INTERPOLATE3
+//                                 && program.m_code[source].args.ImageInterpolate3.factor1==factor10
+//                                 && program.m_code[source].args.ImageInterpolate3.factor2==factor20
+//                                 )
+//                            {
+//                                swizzleOptOp.args.ImageSwizzle.sources[c] =
+//                                        program.m_code[ source ].args.ImageInterpolate3.target0;
+//                            }
+//                        }
+//                        op.args.ImageInterpolate3.target0 = program.AddOp( swizzleOptOp );
+//                    }
+
+//                    // Target 1
+//                    {
+//                        OP swizzleOptOp = program.m_code[at];
+//                        for ( int c=0; c<MUTABLE_OP_MAX_SWIZZLE_CHANNELS; ++c )
+//                        {
+//                            OP::ADDRESS source = program.m_code[at].args.ImageSwizzle.sources[c];
+//                            if ( source
+//                                 && program.m_code[source].type==OP_TYPE::IM_INTERPOLATE3
+//                                 && program.m_code[source].args.ImageInterpolate3.factor1==factor10
+//                                 && program.m_code[source].args.ImageInterpolate3.factor2==factor20
+//                                 )
+//                            {
+//                                swizzleOptOp.args.ImageSwizzle.sources[c] =
+//                                        program.m_code[ source ].args.ImageInterpolate3.target1;
+//                            }
+//                        }
+//                        op.args.ImageInterpolate3.target1 = program.AddOp( swizzleOptOp );
+//                    }
+
+//                    // Target 2
+//                    {
+//                        OP swizzleOptOp = program.m_code[at];
+//                        for ( int c=0; c<MUTABLE_OP_MAX_SWIZZLE_CHANNELS; ++c )
+//                        {
+//                            OP::ADDRESS source = program.m_code[at].args.ImageSwizzle.sources[c];
+//                            if ( source
+//                                 && program.m_code[source].type==OP_TYPE::IM_INTERPOLATE3
+//                                 && program.m_code[source].args.ImageInterpolate3.factor1==factor10
+//                                 && program.m_code[source].args.ImageInterpolate3.factor2==factor20
+//                                 )
+//                            {
+//                                swizzleOptOp.args.ImageSwizzle.sources[c] =
+//                                        program.m_code[ source ].args.ImageInterpolate3.target2;
+//                            }
+//                        }
+//                        op.args.ImageInterpolate3.target2 = program.AddOp( swizzleOptOp );
+//                    }
+
+//                    at = program.AddOp( op );
+//                }
+//                break;
+//            }
+
+//            default:
+//                break;
+//            }
+//		}
+
+		return at;
+	}
+
+
+
+	//!
+	FImageDesc ASTOpImageSwizzle::GetImageDesc(bool returnBestOption, FGetImageDescContext* context) const
+	{
+		FImageDesc res;
+
+		// Local context in case it is necessary
+		FGetImageDescContext localContext;
+		if (!context)
+		{
+			context = &localContext;
+		}
+		else
+		{
+			// Cached result?
+			FImageDesc* PtrValue = context->m_results.Find(this);
+			if (PtrValue)
+			{
+				return *PtrValue;
+			}
+		}
+
+		if (Sources[0].child())
+		{
+			res = Sources[0]->GetImageDesc(returnBestOption, context);
+			res.m_format = Format;
+			check(res.m_format != EImageFormat::IF_NONE);
+		}
+
+		// Cache the result
+		if (context)
+		{
+			context->m_results.Add(this, res);
+		}
+
+		return res;
+	}
+
+
+	void ASTOpImageSwizzle::GetLayoutBlockSize(int* pBlockX, int* pBlockY)
+	{
+		if (Sources[0].child())
+		{
+			// Assume the block size of the biggest mip
+			Sources[0].child()->GetLayoutBlockSize(pBlockX, pBlockY);
+		}
+	}
+
+
+	bool ASTOpImageSwizzle::IsImagePlainConstant(FVector4f& colour) const
+	{
+		// TODO: Maybe something could be done here.
+		return false;
+	}
+
+
+	mu::Ptr<ImageSizeExpression> ASTOpImageSwizzle::GetImageSizeExpression() const
+	{
+		mu::Ptr<ImageSizeExpression> pRes;
+
+		if (Sources[0].child())
+		{
+			pRes = Sources[0].child()->GetImageSizeExpression();
+		}
+		else
+		{
+			pRes = new ImageSizeExpression;
+		}
+
+		return pRes;
+	}
+
+}
\ No newline at end of file
diff --git a/Engine/Plugins/Experimental/Mutable/Source/MutableTools/Private/MuT/ASTOpImageSwizzle.h b/Engine/Plugins/Experimental/Mutable/Source/MutableTools/Private/MuT/ASTOpImageSwizzle.h
new file mode 100644
--- /dev/null
+++ b/Engine/Plugins/Experimental/Mutable/Source/MutableTools/Private/MuT/ASTOpImageSwizzle.h
@@ -0,0 +1,48 @@
+// Copyright Epic Games, Inc. All Rights Reserved.
+
+#pragma once
+
+#include "HAL/PlatformMath.h"
+#include "MuR/Image.h"
+#include "MuR/Operations.h"
+#include "MuR/Ptr.h"
+#include "MuT/AST.h"
+
+
+namespace mu
+{
+	struct FProgram;
+
+	class ASTOpImageSwizzle : public ASTOp
+	{
+	public:
+
+		ASTChild Sources[MUTABLE_OP_MAX_SWIZZLE_CHANNELS];
+
+		uint8 SourceChannels[MUTABLE_OP_MAX_SWIZZLE_CHANNELS] = { 0,0,0,0 };
+
+		EImageFormat Format = EImageFormat::IF_NONE;
+
+	public:
+
+		ASTOpImageSwizzle();
+		ASTOpImageSwizzle(const ASTOpImageSwizzle&) = delete;
+		~ASTOpImageSwizzle();
+
+		OP_TYPE GetOpType() const override { return OP_TYPE::IM_SWIZZLE; }
+		uint64 Hash() const override;
+		bool IsEqual(const ASTOp& otherUntyped) const override;
+		Ptr<ASTOp> Clone(MapChildFuncRef mapChild) const override;
+		void ForEachChild(const TFunctionRef<void(ASTChild&)>) override;
+		void Link(FProgram& program, const FLinkerOptions* Options) override;
+		Ptr<ASTOp> OptimiseSemantic(const FModelOptimizationOptions&) const override;
+		Ptr<ASTOp> OptimiseSink(const FModelOptimizationOptions& options, FOptimizeSinkContext& context) const override;
+		FImageDesc GetImageDesc(bool returnBestOption, FGetImageDescContext* context) const override;
+		void GetLayoutBlockSize(int* pBlockX, int* pBlockY) override;
+		bool IsImagePlainConstant(FVector4f& colour) const override;
+		Ptr<ImageSizeExpression> GetImageSizeExpression() const override;
+
+	};
+
+}
+
diff --git a/Engine/Plugins/Experimental/Mutable/Source/MutableTools/Private/MuT/CodeGenerator_Image.cpp b/Engine/Plugins/Experimental/Mutable/Source/MutableTools/Private/MuT/CodeGenerator_Image.cpp
--- a/Engine/Plugins/Experimental/Mutable/Source/MutableTools/Private/MuT/CodeGenerator_Image.cpp
+++ b/Engine/Plugins/Experimental/Mutable/Source/MutableTools/Private/MuT/CodeGenerator_Image.cpp
@@ -27,6 +27,7 @@
 #include "MuT/ASTOpImagePixelFormat.h"
 #include "MuT/ASTOpImageTransform.h"
 #include "MuT/ASTOpImageMakeGrowMap.h"
+#include "MuT/ASTOpImageSwizzle.h"
 #include "MuT/ASTOpMeshExtractLayoutBlocks.h"
 #include "MuT/ASTOpMeshFormat.h"
 #include "MuT/ASTOpParameter.h"
@@ -762,8 +763,7 @@ namespace mu
 
         // This node always produces a swizzle operation and sometimes it may produce a pixelformat
 		// operation to compress the result
-        Ptr<ASTOpFixed> op = new ASTOpFixed();
-        op->op.type = OP_TYPE::IM_SWIZZLE;
+        Ptr<ASTOpImageSwizzle> op = new ASTOpImageSwizzle();
 
 		// Format
 		EImageFormat compressedFormat = EImageFormat::IF_NONE;
@@ -773,7 +773,7 @@ namespace mu
         case EImageFormat::IF_BC1:
         case EImageFormat::IF_ASTC_4x4_RGB_LDR:
             compressedFormat = node.m_format;
-            op->op.args.ImageSwizzle.format = node.m_sources[3] ? EImageFormat::IF_RGBA_UBYTE : EImageFormat::IF_RGB_UBYTE;
+            op->Format = node.m_sources[3] ? EImageFormat::IF_RGBA_UBYTE : EImageFormat::IF_RGB_UBYTE;
 			break;
 
 		case EImageFormat::IF_BC2:
@@ -782,23 +782,23 @@ namespace mu
         case EImageFormat::IF_BC7:
         case EImageFormat::IF_ASTC_4x4_RGBA_LDR:
             compressedFormat = node.m_format;
-             op->op.args.ImageSwizzle.format = EImageFormat::IF_RGBA_UBYTE;
+             op->Format = EImageFormat::IF_RGBA_UBYTE;
 			break;
 
 		case EImageFormat::IF_BC4:
 			compressedFormat = node.m_format;
-            op->op.args.ImageSwizzle.format = EImageFormat::IF_L_UBYTE;
+            op->Format = EImageFormat::IF_L_UBYTE;
 			break;
 
 		case EImageFormat::IF_BC5:
         case EImageFormat::IF_ASTC_4x4_RG_LDR:
             compressedFormat = node.m_format;
 			// TODO: Should be RG
-            op->op.args.ImageSwizzle.format = EImageFormat::IF_RGB_UBYTE;
+            op->Format = EImageFormat::IF_RGB_UBYTE;
 			break;
 
 		default:
-            op->op.args.ImageSwizzle.format = node.m_format;
+            op->Format = node.m_format;
 			break;
 
 		}
@@ -809,8 +809,8 @@ namespace mu
 		check(node.m_sources.Num() == node.m_sourceChannels.Num());
 
 		// First source, for reference in the size
-        Ptr<ASTOp> first = 0;
-		for (std::size_t t = 0; t< node.m_sources.Num(); ++t)
+        Ptr<ASTOp> first;
+		for (int32 t = 0; t<node.m_sources.Num(); ++t)
 		{
 			if (Node* pChannel = node.m_sources[t].get())
 			{
@@ -820,8 +820,7 @@ namespace mu
 				if (!source)
 				{
 					// TODO: Warn?
-					source = GenerateMissingImageCode("Swizzle channel", EImageFormat::IF_L_UBYTE,
-						node.m_errorContext);
+					source = GenerateMissingImageCode("Swizzle channel", EImageFormat::IF_L_UBYTE, node.m_errorContext);
 				}
 
                 Ptr<ASTOp> sizedSource;
@@ -839,17 +838,16 @@ namespace mu
 					sizedSource = source;
 				}
 
-                op->SetChild( op->op.args.ImageSwizzle.sources[t], sizedSource);
-                op->op.args.ImageSwizzle.sourceChannels[t] = (uint8_t)node.m_sourceChannels[t];
+                op->Sources[t] = sizedSource;
+                op->SourceChannels[t] = (uint8)node.m_sourceChannels[t];
 			}
 		}
 
 		// At least one source is required
-        if (!op->op.args.ImageSwizzle.sources[0])
+        if (!op->Sources[0])
 		{
-            Ptr<ASTOp> source = GenerateMissingImageCode("First swizzle image", EImageFormat::IF_RGBA_UBYTE,
-				node.m_errorContext);
-            op->SetChild( op->op.args.ImageSwizzle.sources[0], source);
+            Ptr<ASTOp> source = GenerateMissingImageCode("First swizzle image", EImageFormat::IF_RGBA_UBYTE, node.m_errorContext);
+            op->Sources[0] = source;
 		}
 
         Ptr<ASTOp> resultOp = op;
diff --git a/Engine/Plugins/Experimental/Mutable/Source/MutableTools/Private/MuT/CodeOptimiser_Semantic.cpp b/Engine/Plugins/Experimental/Mutable/Source/MutableTools/Private/MuT/CodeOptimiser_Semantic.cpp
--- a/Engine/Plugins/Experimental/Mutable/Source/MutableTools/Private/MuT/CodeOptimiser_Semantic.cpp
+++ b/Engine/Plugins/Experimental/Mutable/Source/MutableTools/Private/MuT/CodeOptimiser_Semantic.cpp
@@ -20,6 +20,7 @@
 #include "MuT/ASTOpImagePatch.h"
 #include "MuT/ASTOpImagePixelFormat.h"
 #include "MuT/ASTOpImageMakeGrowMap.h"
+#include "MuT/ASTOpImageSwizzle.h"
 #include "MuT/ASTOpMeshMorph.h"
 #include "MuT/ASTOpSwitch.h"
 #include "MuT/CodeOptimiser.h"
@@ -31,21 +32,6 @@
 namespace mu
 {
 
-    //---------------------------------------------------------------------------------------------
-    //! Set al the non-null sources of an image swizzle operation to the given value
-    //---------------------------------------------------------------------------------------------
-    void ReplaceAllSources( Ptr<ASTOpFixed>& op, Ptr<ASTOp>& value )
-    {
-        check( op->GetOpType() == OP_TYPE::IM_SWIZZLE );
-        for ( int c=0; c<MUTABLE_OP_MAX_SWIZZLE_CHANNELS; ++c )
-        {
-            if ( op->op.args.ImageSwizzle.sources[c] )
-            {
-                op->SetChild( op->op.args.ImageSwizzle.sources[c], value);
-            }
-        }
-    }
-
 
     //---------------------------------------------------------------------------------------------
     //! Get an uncompressed version of the image.
@@ -294,70 +280,11 @@ namespace mu
 
         //-------------------------------------------------------------------------------------
         case OP_TYPE::IM_PIXELFORMAT:
-        {
-            // All in the sink optimiser?
-            break;
-        }
-
-
         case OP_TYPE::IM_SWIZZLE:
         {
-            // If children channels are also swizzle ops, recurse them
-            Ptr<ASTOpFixed> sat;
-
-            for ( int c=0; c<MUTABLE_OP_MAX_SWIZZLE_CHANNELS; ++c )
-            {
-                Ptr<ASTOp> candidate = children[op.args.ImageSwizzle.sources[c]].child();
-				if (!candidate)
-				{
-					continue;
-				}
-
-                // Swizzle
-				if (candidate->GetOpType()==OP_TYPE::IM_SWIZZLE)
-                {
-                    if (!sat)
-                    {
-                        sat = mu::Clone<ASTOpFixed>(this);
-                    }
-                    ASTOpFixed* typedCandidate = dynamic_cast<ASTOpFixed*>(candidate.get());
-                    int candidateChannel = op.args.ImageSwizzle.sourceChannels[c];
-
-                    sat->SetChild( sat->op.args.ImageSwizzle.sources[c],
-                                    typedCandidate->children[typedCandidate->op.args.ImageSwizzle.sources[candidateChannel]] );
-                    sat->op.args.ImageSwizzle.sourceChannels[c] =  typedCandidate->op.args.ImageSwizzle.sourceChannels[candidateChannel];
-                }
-
-                // Format
-                if (candidate->GetOpType()==OP_TYPE::IM_PIXELFORMAT)
-                {
-                    // We can remove the format if its source is already an uncompressed format
-                    ASTOpImagePixelFormat* typedCandidate = dynamic_cast<ASTOpImagePixelFormat*>(candidate.get());
-                    Ptr<ASTOp> formatSource = typedCandidate->Source.child();
-
-                    if (formatSource)
-                    {
-                        auto desc = formatSource->GetImageDesc();
-                        if ( desc.m_format!= EImageFormat::IF_NONE && !IsCompressedFormat( desc.m_format ) )
-                        {
-                            if (!sat)
-                            {
-                                sat = mu::Clone<ASTOpFixed>(this);
-                            }
-                            sat->SetChild( sat->op.args.ImageSwizzle.sources[c], formatSource );
-                        }
-                    }
-                }
-
-            }
-
-            if (sat)
-            {
-                at = sat;
-            }
-
-            break;
-        }
+			check(false); // moved to its own operation class
+			break;
+		}
 
 
         //-------------------------------------------------------------------------------------
@@ -478,525 +405,6 @@ namespace mu
     }
 
 
-    //---------------------------------------------------------------------------------------------
-    Ptr<ASTOp> ASTOpFixed::OptimiseSwizzle( const FModelOptimizationOptions& options ) const
-    {
-        MUTABLE_CPUPROFILER_SCOPE(OptimiseSwizzleAST);
-
-        //! Basic optimisation first
-        Ptr<ASTOp> at = OptimiseSemantic(options);
-        if (at)
-        {
-            if (at->GetOpType()==OP_TYPE::IM_SWIZZLE)
-            {
-                at = dynamic_cast<ASTOpFixed*>(at.get())->OptimiseSwizzle( options );
-            }
-            return at;
-        }
-
-        // If all sources are the same, we can sink the instruction
-        bool same = true;
-		int RepeatedChild = 0;
-        Ptr<ASTOp> channelSourceAt;
-        for ( int c=0; c<MUTABLE_OP_MAX_SWIZZLE_CHANNELS; ++c )
-        {
-            auto candidate = children[op.args.ImageSwizzle.sources[c]].child();
-            if ( candidate )
-            {
-                if (!channelSourceAt)
-                {
-                    channelSourceAt=candidate;
-                }
-                else
-                {
-                    same = channelSourceAt==candidate;
-					++RepeatedChild;
-				}
-            }
-        }
-
-        if (same && channelSourceAt)
-        {
-            // The instruction can be sunk
-            OP_TYPE sourceType = channelSourceAt->GetOpType();
-            switch ( sourceType )
-            {
-
-			case OP_TYPE::IM_PLAINCOLOUR:
-			{
-				auto NewPlain = mu::Clone<ASTOpFixed>(channelSourceAt);
-				Ptr<ASTOpFixed> NewSwizzle = new ASTOpFixed;
-				NewSwizzle->op.type = OP_TYPE::CO_SWIZZLE;
-				for (int i = 0; i < MUTABLE_OP_MAX_SWIZZLE_CHANNELS; ++i)
-				{
-					NewSwizzle->SetChild(NewSwizzle->op.args.ColourSwizzle.sources[i], NewPlain->children[NewPlain->op.args.ImagePlainColour.colour]);
-					NewSwizzle->op.args.ColourSwizzle.sourceChannels[i] = op.args.ImageSwizzle.sourceChannels[i];
-				}
-				NewPlain->SetChild(NewPlain->op.args.ImagePlainColour.colour, NewSwizzle);
-				NewPlain->op.args.ImagePlainColour.format = this->op.args.ImageSwizzle.format;
-				at = NewPlain;
-				break;
-			}
-
-            case OP_TYPE::IM_SWITCH:
-            {
-                // Move the swizzle down all the paths
-                auto nop = mu::Clone<ASTOpSwitch>(channelSourceAt);
-
-                if (nop->def)
-                {
-                    auto defOp = mu::Clone<ASTOpFixed>(this);
-                    ReplaceAllSources( defOp, nop->def.child() );
-                    nop->def = defOp;
-                }
-
-                // We need to copy the options because we change them
-                for ( int32 v=0; v<nop->cases.Num(); ++v )
-                {
-                    if ( nop->cases[v].branch )
-                    {
-                        auto bOp = mu::Clone<ASTOpFixed>(this);
-                        ReplaceAllSources( bOp, nop->cases[v].branch.child() );
-                        nop->cases[v].branch = bOp;
-                    }
-                }
-
-                at = nop;
-                break;
-            }
-
-            case OP_TYPE::IM_CONDITIONAL:
-            {
-                // We move the swizzle down the two paths
-                auto nop = mu::Clone<ASTOpConditional>(channelSourceAt);
-
-                auto aOp = mu::Clone<ASTOpFixed>(this);
-                ReplaceAllSources( aOp, nop->yes.child() );
-                nop->yes = aOp;
-
-                auto bOp = mu::Clone<ASTOpFixed>(this);
-                ReplaceAllSources( bOp, nop->no.child() );
-                nop->no = bOp;
-
-                at = nop;
-                break;
-            }
-
-            case OP_TYPE::IM_LAYER:
-            {
-                // We move the swizzle down the two paths
-                auto nop = mu::Clone<ASTOpFixed>(channelSourceAt);
-
-                auto aOp = mu::Clone<ASTOpFixed>(this);
-                ReplaceAllSources( aOp, nop->children[nop->op.args.ImageLayer.base].child() );
-                nop->SetChild( nop->op.args.ImageLayer.base, aOp );
-
-                auto bOp = mu::Clone<ASTOpFixed>(this);
-                ReplaceAllSources( bOp, nop->children[nop->op.args.ImageLayer.blended].child() );
-                nop->SetChild( nop->op.args.ImageLayer.blended, bOp );
-
-                at = nop;
-                break;
-            }
-
-            case OP_TYPE::IM_LAYERCOLOUR:
-            {
-                // We move the swizzle down the base paths
-                auto nop = mu::Clone<ASTOpFixed>(channelSourceAt);
-
-                auto NewSwizzle = mu::Clone<ASTOpFixed>(this);
-                ReplaceAllSources(NewSwizzle, nop->children[nop->op.args.ImageLayerColour.base].child() );
-                nop->SetChild( nop->op.args.ImageLayerColour.base, NewSwizzle);
-
-                // We need to swizzle the colour too
-                Ptr<ASTOpFixed> cOp = new ASTOpFixed;
-                cOp->op.type = OP_TYPE::CO_SWIZZLE;
-                for ( int i=0; i<MUTABLE_OP_MAX_SWIZZLE_CHANNELS; ++i)
-                {
-                    cOp->SetChild( cOp->op.args.ColourSwizzle.sources[i], nop->children[nop->op.args.ImageLayerColour.colour] );
-                    cOp->op.args.ColourSwizzle.sourceChannels[i] = nop->op.args.ImageSwizzle.sourceChannels[i];
-                }
-                nop->SetChild( nop->op.args.ImageLayerColour.colour, cOp );
-
-                at = nop;
-                break;
-            }
-
-			case OP_TYPE::IM_DISPLACE:
-			{
-				auto NewDisplace = mu::Clone<ASTOpFixed>(channelSourceAt);
-				auto NewSwizzle = mu::Clone<ASTOpFixed>(this);
-				ReplaceAllSources(NewSwizzle, NewDisplace->children[NewDisplace->op.args.ImageDisplace.source].child());
-				NewDisplace->SetChild(NewDisplace->op.args.ImageDisplace.source, NewSwizzle);
-				at = NewDisplace;
-				break;
-			}
-
-			case OP_TYPE::IM_RASTERMESH:
-			{
-				auto NewRaster = mu::Clone<ASTOpFixed>(channelSourceAt);
-				auto NewSwizzle = mu::Clone<ASTOpFixed>(this);
-				ReplaceAllSources(NewSwizzle, NewRaster->children[NewRaster->op.args.ImageRasterMesh.image].child());
-				NewRaster->SetChild(NewRaster->op.args.ImageRasterMesh.image, NewSwizzle);
-				at = NewRaster;
-				break;
-			}
-
-            default:
-                same = false;
-                break;
-            }
-        }
-
-        if (!same)
-        {
-            // Maybe we can still sink the instruction in some cases
-//            OP::ADDRESS sourceAt = program.m_code[at].args.ImageSwizzle.sources[0];
-//            OP_TYPE sourceType = (OP_TYPE)program.m_code[sourceAt].type;
-//            switch ( sourceType )
-//            {
-//            case OP_TYPE::IM_COMPOSE:
-//            {
-//                bool canSink = true;
-//                OP::ADDRESS layout = program.m_code[sourceAt].args.ImageCompose.layout;
-//                uint32_t block = program.m_code[sourceAt].args.ImageCompose.blockIndex;
-//                for ( int c=1; c<MUTABLE_OP_MAX_SWIZZLE_CHANNELS; ++c )
-//                {
-//                    OP::ADDRESS candidate = program.m_code[at].args.ImageSwizzle.sources[c];
-//                    if ( candidate )
-//                    {
-//                        if ( program.m_code[candidate].type==sourceType )
-//                        {
-//                            OP::ADDRESS clayout =
-//                                    program.m_code[candidate].args.ImageCompose.layout;
-//                            uint32_t cblock =
-//                                    program.m_code[candidate].args.ImageCompose.blockIndex;
-//                            canSink = canSink
-//                                    && ( layout == clayout )
-//                                    && ( block == cblock );
-//                        }
-//                        else
-//                        {
-//                            canSink = false;
-//                        }
-//                    }
-//                }
-
-//                if ( canSink )
-//                {
-//                    m_modified = true;
-
-//                    OP op = program.m_code[sourceAt];
-
-//                    OP swizzleBaseOp = program.m_code[at];
-//                    OP swizzleBlockOp = program.m_code[at];
-
-//                    for ( int c=0; c<MUTABLE_OP_MAX_SWIZZLE_CHANNELS; ++c )
-//                    {
-//                        OP::ADDRESS source = program.m_code[at].args.ImageSwizzle.sources[c];
-//                        if ( source )
-//                        {
-//                            swizzleBaseOp.args.ImageSwizzle.sources[c] =
-//                                    program.m_code[ source ].args.ImageCompose.base;
-
-//                            swizzleBlockOp.args.ImageSwizzle.sources[c] =
-//                                    program.m_code[ source ].args.ImageCompose.blockImage;
-//                        }
-//                    }
-//                    op.args.ImageCompose.base = program.AddOp( swizzleBaseOp );
-//                    op.args.ImageCompose.blockImage = program.AddOp( swizzleBlockOp );
-
-//                    at = program.AddOp( op );
-//                }
-//                break;
-//            }
-
-//            case OP_TYPE::IM_CONDITIONAL:
-//            {
-//                // We can always sink into conditionals
-//                m_modified = true;
-
-//                OP op = program.m_code[sourceAt];
-
-//                OP swizzleYesOp = program.m_code[at];
-//                OP swizzleNoOp = program.m_code[at];
-//                OP::ADDRESS source0 = 0;
-//                for ( int c=0; c<MUTABLE_OP_MAX_SWIZZLE_CHANNELS; ++c )
-//                {
-//                    OP::ADDRESS source = program.m_code[at].args.ImageSwizzle.sources[c];
-//                    if ( c==0 )
-//                    {
-//                        source0 = source;
-//                    }
-
-//                    if ( source
-//                        &&
-//                        ( (c==0) || ( program.m_code[source0].args.Conditional.condition == program.m_code[source].args.Conditional.condition ) )
-//                        )
-//                    {
-//                        swizzleYesOp.args.ImageSwizzle.sources[c] =
-//                                program.m_code[ source ].args.Conditional.yes;
-
-//                        swizzleNoOp.args.ImageSwizzle.sources[c] =
-//                                program.m_code[ source ].args.Conditional.no;
-//                    }
-//                }
-//                op.args.Conditional.yes = program.AddOp( swizzleYesOp );
-//                op.args.Conditional.no = program.AddOp( swizzleNoOp );
-
-//                at = program.AddOp( op );
-//                break;
-//            }
-
-
-//            case OP_TYPE::IM_SWITCH:
-//            {
-//                bool canSink = true;
-//                OP::ADDRESS variable = program.m_code[sourceAt].args.Switch.variable;
-
-//                // If at least 3 channels can be combined, do it. We may be duplicating some data
-//                // But it greatly optimises speed in some cases like the constant alpha masks of
-//                // the face colour interpolation.
-//                // TODO: Make this dependent on compilation options.
-//                //for ( int c=1; c<MUTABLE_OP_MAX_SWIZZLE_CHANNELS; ++c )
-//                for ( int c=1; c<3; ++c )
-//                {
-//                    OP::ADDRESS candidate = program.m_code[at].args.ImageSwizzle.sources[c];
-//                    if ( candidate )
-//                    {
-//                        if ( program.m_code[candidate].type==sourceType )
-//                        {
-//                            OP::ADDRESS cvariable =
-//                                    program.m_code[candidate].args.Switch.variable;
-//                            canSink = canSink && ( variable == cvariable );
-//                        }
-//                        else
-//                        {
-//                            canSink = false;
-//                        }
-//                    }
-//                }
-
-//                if ( canSink )
-//                {
-//                    // TODO since data-switch
-////                    m_modified = true;
-
-////                    OP op = program.m_code[sourceAt];
-
-////                    if ( op.args.Switch.def )
-////                    {
-////                        OP swizzleDefOp = program.m_code[at];
-////                        for ( int c=0; c<MUTABLE_OP_MAX_SWIZZLE_CHANNELS; ++c )
-////                        {
-////                            OP::ADDRESS source = program.m_code[at].args.ImageSwizzle.sources[c];
-////                            if ( source )
-////                            {
-////                                if ( program.m_code[source].type==sourceType )
-////                                {
-////                                    swizzleDefOp.args.ImageSwizzle.sources[c] =
-////                                            program.m_code[ source ].args.Switch.def;
-////                                }
-////                                else
-////                                {
-////                                    swizzleDefOp.args.ImageSwizzle.sources[c] = source;
-////                                }
-////                            }
-////                        }
-////                        op.args.Switch.def = program.AddOp( swizzleDefOp );
-////                    }
-
-////                    for ( int o=0; o<MUTABLE_OP_MAX_SWITCH_OPTIONS; ++o )
-////                    {
-////                        if ( op.args.Switch.values[o] )
-////                        {
-////                            OP swizzleOptOp = program.m_code[at];
-////                            for ( int c=0; c<MUTABLE_OP_MAX_SWIZZLE_CHANNELS; ++c )
-////                            {
-////                                OP::ADDRESS source = program.m_code[at].args.ImageSwizzle.sources[c];
-////                                if ( source )
-////                                {
-////                                    if ( program.m_code[source].type==sourceType )
-////                                    {
-////                                        swizzleOptOp.args.ImageSwizzle.sources[c] =
-////                                                program.m_code[ source ].args.Switch.values[o];
-////                                    }
-////                                    else
-////                                    {
-////                                        swizzleOptOp.args.ImageSwizzle.sources[c] = source;
-////                                    }
-////                                }
-////                            }
-////                            op.args.Switch.values[o] = program.AddOp( swizzleOptOp );
-////                        }
-////                    }
-
-////                    at = program.AddOp( op );
-//                }
-//                break;
-//            }
-
-
-//            case OP_TYPE::IM_INTERPOLATE:
-//            {
-//                bool canSink = true;
-//                OP::ADDRESS factor = program.m_code[sourceAt].args.ImageInterpolate.factor;
-
-//                // It is worth to sink 3 channels
-//                // TODO: Make it depend on compilation parameters
-//                //for ( int c=1; c<MUTABLE_OP_MAX_SWIZZLE_CHANNELS; ++c )
-//                for ( int c=1; c<3; ++c )
-//                {
-//                    OP::ADDRESS candidate = program.m_code[at].args.ImageSwizzle.sources[c];
-//                    if ( candidate )
-//                    {
-//                        if ( program.m_code[candidate].type==sourceType )
-//                        {
-//                            OP::ADDRESS cfactor =
-//                                    program.m_code[candidate].args.ImageInterpolate.factor;
-//                            canSink = canSink && ( factor == cfactor );
-//                        }
-//                        else
-//                        {
-//                            canSink = false;
-//                        }
-//                    }
-//                }
-
-//                if ( canSink )
-//                {
-//                    m_modified = true;
-
-//                    OP op = program.m_code[sourceAt];
-
-//                    for ( int o=0; o<MUTABLE_OP_MAX_INTERPOLATE_COUNT; ++o )
-//                    {
-//                        if ( op.args.ImageInterpolate.targets[o] )
-//                        {
-//                            OP swizzleOptOp = program.m_code[at];
-//                            for ( int c=0; c<MUTABLE_OP_MAX_SWIZZLE_CHANNELS; ++c )
-//                            {
-//                                OP::ADDRESS source = program.m_code[at].args.ImageSwizzle.sources[c];
-
-//                                if ( source
-//                                     && program.m_code[source].type==OP_TYPE::IM_INTERPOLATE
-//                                     && program.m_code[source].args.ImageInterpolate.factor==factor
-//                                     )
-//                                {
-//                                    swizzleOptOp.args.ImageSwizzle.sources[c] =
-//                                            program.m_code[ source ].args.ImageInterpolate.targets[o];
-//                                }
-//                            }
-//                            op.args.ImageInterpolate.targets[o] = program.AddOp( swizzleOptOp );
-//                        }
-//                    }
-
-//                    at = program.AddOp( op );
-//                }
-//                break;
-//            }
-
-
-//            case OP_TYPE::IM_INTERPOLATE3:
-//            {
-//                // We can sink if all channels are interpolates for the same factors
-//                OP::ADDRESS sourceAt0 = program.m_code[at].args.ImageSwizzle.sources[0];
-//                OP::ADDRESS factor10 = program.m_code[sourceAt0].args.ImageInterpolate3.factor1;
-//                OP::ADDRESS factor20 = program.m_code[sourceAt0].args.ImageInterpolate3.factor2;
-
-//                OP::ADDRESS sourceAt1 = program.m_code[at].args.ImageSwizzle.sources[1];
-//                OP::ADDRESS factor11 = program.m_code[sourceAt1].args.ImageInterpolate3.factor1;
-//                OP::ADDRESS factor21 = program.m_code[sourceAt1].args.ImageInterpolate3.factor2;
-
-//                OP::ADDRESS sourceAt2 = program.m_code[at].args.ImageSwizzle.sources[2];
-//                OP::ADDRESS factor12 = program.m_code[sourceAt2].args.ImageInterpolate3.factor1;
-//                OP::ADDRESS factor22 = program.m_code[sourceAt2].args.ImageInterpolate3.factor2;
-
-//                bool canSink = true;
-//                canSink &= program.m_code[sourceAt0].type == program.m_code[sourceAt1].type;
-//                canSink &= program.m_code[sourceAt0].type == program.m_code[sourceAt2].type;
-//                canSink &= factor10 == factor11;
-//                canSink &= factor10 == factor12;
-//                canSink &= factor20 == factor21;
-//                canSink &= factor20 == factor22;
-
-//                if ( canSink )
-//                {
-//                    m_modified = true;
-
-//                    OP op = program.m_code[sourceAt0];
-
-//                    // Target 0
-//                    {
-//                        OP swizzleOptOp = program.m_code[at];
-//                        for ( int c=0; c<MUTABLE_OP_MAX_SWIZZLE_CHANNELS; ++c )
-//                        {
-//                            OP::ADDRESS source = program.m_code[at].args.ImageSwizzle.sources[c];
-//                            if ( source
-//                                 && program.m_code[source].type==OP_TYPE::IM_INTERPOLATE3
-//                                 && program.m_code[source].args.ImageInterpolate3.factor1==factor10
-//                                 && program.m_code[source].args.ImageInterpolate3.factor2==factor20
-//                                 )
-//                            {
-//                                swizzleOptOp.args.ImageSwizzle.sources[c] =
-//                                        program.m_code[ source ].args.ImageInterpolate3.target0;
-//                            }
-//                        }
-//                        op.args.ImageInterpolate3.target0 = program.AddOp( swizzleOptOp );
-//                    }
-
-//                    // Target 1
-//                    {
-//                        OP swizzleOptOp = program.m_code[at];
-//                        for ( int c=0; c<MUTABLE_OP_MAX_SWIZZLE_CHANNELS; ++c )
-//                        {
-//                            OP::ADDRESS source = program.m_code[at].args.ImageSwizzle.sources[c];
-//                            if ( source
-//                                 && program.m_code[source].type==OP_TYPE::IM_INTERPOLATE3
-//                                 && program.m_code[source].args.ImageInterpolate3.factor1==factor10
-//                                 && program.m_code[source].args.ImageInterpolate3.factor2==factor20
-//                                 )
-//                            {
-//                                swizzleOptOp.args.ImageSwizzle.sources[c] =
-//                                        program.m_code[ source ].args.ImageInterpolate3.target1;
-//                            }
-//                        }
-//                        op.args.ImageInterpolate3.target1 = program.AddOp( swizzleOptOp );
-//                    }
-
-//                    // Target 2
-//                    {
-//                        OP swizzleOptOp = program.m_code[at];
-//                        for ( int c=0; c<MUTABLE_OP_MAX_SWIZZLE_CHANNELS; ++c )
-//                        {
-//                            OP::ADDRESS source = program.m_code[at].args.ImageSwizzle.sources[c];
-//                            if ( source
-//                                 && program.m_code[source].type==OP_TYPE::IM_INTERPOLATE3
-//                                 && program.m_code[source].args.ImageInterpolate3.factor1==factor10
-//                                 && program.m_code[source].args.ImageInterpolate3.factor2==factor20
-//                                 )
-//                            {
-//                                swizzleOptOp.args.ImageSwizzle.sources[c] =
-//                                        program.m_code[ source ].args.ImageInterpolate3.target2;
-//                            }
-//                        }
-//                        op.args.ImageInterpolate3.target2 = program.AddOp( swizzleOptOp );
-//                    }
-
-//                    at = program.AddOp( op );
-//                }
-//                break;
-//            }
-
-//            default:
-//                break;
-//            }
-        }
-
-        return at;
-    }
-
-
     //---------------------------------------------------------------------------------------------
     //---------------------------------------------------------------------------------------------
     //---------------------------------------------------------------------------------------------
@@ -1381,8 +789,10 @@ namespace mu
         //-----------------------------------------------------------------------------------------
         //-----------------------------------------------------------------------------------------
         //-----------------------------------------------------------------------------------------
-        case OP_TYPE::LA_REMOVEBLOCKS:
+		case OP_TYPE::IM_SWIZZLE:
+		case OP_TYPE::LA_REMOVEBLOCKS:
         {
+			// Moved to their own operation class
 			check(false);
             break;
         }
@@ -1550,13 +960,6 @@ namespace mu
             break;
         }
 
-        //-------------------------------------------------------------------------------------
-        case OP_TYPE::IM_SWIZZLE:
-        {
-            at = OptimiseSwizzle( options );
-            break;
-        }
-
         //-------------------------------------------------------------------------------------
         case OP_TYPE::IM_LAYER:
         {
@@ -1678,6 +1081,128 @@ namespace mu
             break;
         }
 
+
+
+		//-------------------------------------------------------------------------------------
+        case OP_TYPE::IM_LAYERCOLOUR:
+		{
+			// Layer effects may be worth sinking down switches and conditionals, to be able
+			// to apply extra optimisations
+			auto baseAt = children[op.args.ImageLayerColour.base].child();
+			auto maskAt = children[op.args.ImageLayerColour.mask].child();
+
+			// Promote conditions from the base
+			OP_TYPE baseType = baseAt->GetOpType();
+			switch (baseType)
+			{
+				// Seems to cause operation explosion in optimizer in bandit model.
+				// moved to generic sink in the default.
+	//            case OP_TYPE::IM_CONDITIONAL:
+	//            {
+	//                m_modified = true;
+
+	//                OP op = program.m_code[baseAt];
+
+	//                OP aOp = program.m_code[at];
+	//                aOp.args.ImageLayerColour.base = program.m_code[baseAt].args.Conditional.yes;
+	//                op.args.Conditional.yes = program.AddOp( aOp );
+
+	//                OP bOp = program.m_code[at];
+	//                bOp.args.ImageLayerColour.base = program.m_code[baseAt].args.Conditional.no;
+	//                op.args.Conditional.no = program.AddOp( bOp );
+
+	//                at = program.AddOp( op );
+	//                break;
+	//            }
+
+			case OP_TYPE::IM_SWITCH:
+			{
+				// Warning:
+				// It seems to cause data explosion in optimizer in some models. Because
+				// all switch branches become unique constants
+
+                // See if the blended has an identical switch, to optimise it too
+				const ASTOpSwitch* baseSwitch = dynamic_cast<const ASTOpSwitch*>( baseAt.get() );
+
+                // Mask not supported yet
+				if (maskAt)
+				{
+					break;
+				}
+
+                // Move the layer operation down base paths
+                Ptr<ASTOpSwitch> nop = mu::Clone<ASTOpSwitch>(baseSwitch);
+
+                if (nop->def)
+                {
+					Ptr<ASTOpFixed> defOp = mu::Clone<ASTOpFixed>(this);
+                    defOp->SetChild( defOp->op.args.ImageLayerColour.base, nop->def );
+                    nop->def = defOp;
+                }
+
+                for ( int32 v=0; v<nop->cases.Num(); ++v )
+                {
+                    if ( nop->cases[v].branch )
+                    {
+                        Ptr<ASTOpFixed> bOp = mu::Clone<ASTOpFixed>(this);
+                        bOp->SetChild( bOp->op.args.ImageLayerColour.base, nop->cases[v].branch );
+                        nop->cases[v].branch = bOp;
+                    }
+                }
+
+                at = nop;
+				break;
+			}
+
+			case OP_TYPE::IM_DISPLACE:
+			{
+				// Mask not supported yet. If there is a mask it wouldn't be correct to sink
+				// unless the mask was a similar displace.
+				if (maskAt)
+				{
+					break;
+				}
+
+				Ptr<ASTOpFixed> NewDisplace = mu::Clone<ASTOpFixed>(baseAt);
+
+				Ptr<ASTOp> sourceOp = NewDisplace->children[NewDisplace->op.args.ImageDisplace.source].child();
+				Ptr<ASTOpFixed> NewSource = mu::Clone<ASTOpFixed>(this);
+				NewSource->SetChild(NewSource->op.args.ImageLayerColour.base, sourceOp);
+				NewDisplace->SetChild(NewDisplace->op.args.ImageDisplace.source, NewSource);
+
+				at = NewDisplace;
+				break;
+			}
+
+			case OP_TYPE::IM_RASTERMESH:
+			{
+				// Mask not supported yet. If there is a mask it wouldn't be correct to sink.				
+				if (maskAt)
+				{
+					break;
+				}
+
+				Ptr<ASTOpFixed> NewRaster = mu::Clone<ASTOpFixed>(baseAt);
+
+				Ptr<ASTOp> sourceOp = NewRaster->children[NewRaster->op.args.ImageRasterMesh.image].child();
+				Ptr<ASTOpFixed> NewSource = mu::Clone<ASTOpFixed>(this);
+				NewSource->SetChild(NewSource->op.args.ImageLayerColour.base, sourceOp);
+				NewRaster->SetChild(NewRaster->op.args.ImageRasterMesh.image, NewSource);
+
+				at = NewRaster;
+				break;
+			}
+
+
+			default:
+				break;
+
+			}
+
+			break;
+		}
+
+
         //-----------------------------------------------------------------------------------------
         //-----------------------------------------------------------------------------------------
         //-----------------------------------------------------------------------------------------
@@ -1923,9 +1448,9 @@ namespace mu
 			case OP_TYPE::IM_SWITCH:
 			{
 				// We move the mask creation down all the paths
-				auto newOp = mu::Clone<ASTOpSwitch>(at);
+				Ptr<ASTOpSwitch> newOp = mu::Clone<ASTOpSwitch>(at);
 				newOp->def = Visit(newOp->def.child(), currentSinkingOp);
-				for (auto& c : newOp->cases)
+				for (ASTOpSwitch::FCase& c : newOp->cases)
 				{
 					c.branch = Visit(c.branch.child(), currentSinkingOp);
 				}
@@ -1935,13 +1460,13 @@ namespace mu
 
 			case OP_TYPE::IM_SWIZZLE:
 			{
-				auto newOp = mu::Clone<ASTOpFixed>(at);
+				Ptr<ASTOpImageSwizzle> newOp = mu::Clone<ASTOpImageSwizzle>(at);
 				for (int s = 0; s < MUTABLE_OP_MAX_SWIZZLE_CHANNELS; ++s)
 				{
-					auto channelOp = newOp->children[newOp->op.args.ImageSwizzle.sources[s]].child();
+					Ptr<ASTOp> channelOp = newOp->Sources[s].child();
 					if (channelOp)
 					{
-						newOp->SetChild(newOp->op.args.ImageSwizzle.sources[s], Visit(channelOp, currentSinkingOp));
+						newOp->Sources[s] = Visit(channelOp, currentSinkingOp);
 					}
 				}
 				newAt = newOp;
diff --git a/Engine/Plugins/Experimental/Mutable/Source/MutableTools/Private/MuT/NodeImageLayer.cpp b/Engine/Plugins/Experimental/Mutable/Source/MutableTools/Private/MuT/NodeImageLayer.cpp
--- a/Engine/Plugins/Experimental/Mutable/Source/MutableTools/Private/MuT/NodeImageLayer.cpp
+++ b/Engine/Plugins/Experimental/Mutable/Source/MutableTools/Private/MuT/NodeImageLayer.cpp
@@ -25,21 +25,6 @@ namespace mu
 	MUTABLE_IMPLEMENT_NODE( NodeImageLayer, EType::Layer, Node, Node::EType::Image)
 
 
-	//---------------------------------------------------------------------------------------------
-	const char* NodeImageLayer::s_blendTypeName[] =
-	{
-		"SoftLight",
-		"HardLight",
-		"Burn",
-		"Dodge",
-		"Screen",
-		"Overlay",
-		"Blend",
-        "Multiply",
-        "AlphaOverlay"
-	};
-
-
 	//---------------------------------------------------------------------------------------------
 	// Node Interface
 	//---------------------------------------------------------------------------------------------
diff --git a/Engine/Plugins/Experimental/Mutable/Source/MutableTools/Private/MuT/NodeImageLayerColourPrivate.h b/Engine/Plugins/Experimental/Mutable/Source/MutableTools/Private/MuT/NodeImageLayerColourPrivate.h
--- a/Engine/Plugins/Experimental/Mutable/Source/MutableTools/Private/MuT/NodeImageLayerColourPrivate.h
+++ b/Engine/Plugins/Experimental/Mutable/Source/MutableTools/Private/MuT/NodeImageLayerColourPrivate.h
@@ -31,7 +31,7 @@ namespace mu
 		//!
 		void Serialise( OutputArchive& arch ) const
 		{
-            uint32_t ver = 0;
+            uint32_t ver = 1;
 			arch << ver;
 
 			arch << m_pBase;
@@ -45,12 +45,17 @@ namespace mu
 		{
             uint32_t ver;
 			arch >> ver;
-			check(ver==0);
+			check(ver<=1);
 
 			arch >> m_pBase;
 			arch >> m_pMask;
 			arch >> m_pColour;
 			arch >> m_type;
+			if (ver < 1)
+			{
+				m_type = EBlendType( uint32(m_type)+1 );
+			}
+
 		}
 	};
 
diff --git a/Engine/Plugins/Experimental/Mutable/Source/MutableTools/Private/MuT/NodeImageLayerPrivate.h b/Engine/Plugins/Experimental/Mutable/Source/MutableTools/Private/MuT/NodeImageLayerPrivate.h
--- a/Engine/Plugins/Experimental/Mutable/Source/MutableTools/Private/MuT/NodeImageLayerPrivate.h
+++ b/Engine/Plugins/Experimental/Mutable/Source/MutableTools/Private/MuT/NodeImageLayerPrivate.h
@@ -31,7 +31,7 @@ namespace mu
 		//!
 		void Serialise( OutputArchive& arch ) const
 		{
-            uint32_t ver = 0;
+            uint32_t ver = 1;
 			arch << ver;
 
 			arch << m_pBase;
@@ -45,7 +45,7 @@ namespace mu
 		{
             uint32_t ver;
 			arch >> ver;
-			check(ver==0);
+			check(ver<=1);
 
 			arch >> m_pBase;
 			arch >> m_pMask;
@@ -53,6 +53,10 @@ namespace mu
 
             uint32_t t;
             arch >> t;
+			if (ver < 1)
+			{
+				++t;
+			}
             m_type=(EBlendType)t;
 		}
 	};
diff --git a/Engine/Plugins/Experimental/Mutable/Source/MutableTools/Private/MuT/NodeImageMultiLayerPrivate.h b/Engine/Plugins/Experimental/Mutable/Source/MutableTools/Private/MuT/NodeImageMultiLayerPrivate.h
--- a/Engine/Plugins/Experimental/Mutable/Source/MutableTools/Private/MuT/NodeImageMultiLayerPrivate.h
+++ b/Engine/Plugins/Experimental/Mutable/Source/MutableTools/Private/MuT/NodeImageMultiLayerPrivate.h
@@ -32,7 +32,7 @@ namespace mu
 		//!
 		void Serialise( OutputArchive& arch ) const
 		{
-            uint32_t ver = 0;
+            uint32_t ver = 1;
 			arch << ver;
             arch << uint32_t(m_type);
 
@@ -47,10 +47,14 @@ namespace mu
 		{
             uint32_t ver;
 			arch >> ver;
-			check(ver==0);
+			check(ver<=1);
 
             uint32_t t;
             arch >> t;
+			if (ver < 1)
+			{
+				++t;
+			}
             m_type=(EBlendType)t;
 
             arch >> m_pBase;
diff --git a/Engine/Plugins/Experimental/Mutable/Source/MutableTools/Private/MuT/NodePatchImagePrivate.h b/Engine/Plugins/Experimental/Mutable/Source/MutableTools/Private/MuT/NodePatchImagePrivate.h
--- a/Engine/Plugins/Experimental/Mutable/Source/MutableTools/Private/MuT/NodePatchImagePrivate.h
+++ b/Engine/Plugins/Experimental/Mutable/Source/MutableTools/Private/MuT/NodePatchImagePrivate.h
@@ -36,7 +36,7 @@ namespace mu
         //!
         void Serialise( OutputArchive& arch ) const
         {
-            uint32_t ver = 2;
+            uint32_t ver = 3;
             arch << ver;
 
             arch << m_pImage;
@@ -51,7 +51,7 @@ namespace mu
         {
             uint32_t ver;
             arch >> ver;
-            check(ver==2);
+            check(ver>=2 && ver<=3);
 
             arch >> m_pImage;
             arch >> m_pMask;
@@ -59,6 +59,10 @@ namespace mu
 
 			uint32_t t;
             arch >> t;
+			if (ver <= 2)
+			{
+				++t;
+			}
             m_blendType=(EBlendType)t;
 
             arch >> m_applyToAlpha;
diff --git a/Engine/Plugins/Experimental/Mutable/Source/MutableTools/Private/MuT/TypeInfo.cpp b/Engine/Plugins/Experimental/Mutable/Source/MutableTools/Private/MuT/TypeInfo.cpp
--- a/Engine/Plugins/Experimental/Mutable/Source/MutableTools/Private/MuT/TypeInfo.cpp
+++ b/Engine/Plugins/Experimental/Mutable/Source/MutableTools/Private/MuT/TypeInfo.cpp
@@ -132,6 +132,7 @@ namespace mu
 
 	const char* TypeInfo::s_blendModeName[uint32(EBlendType::_BT_COUNT)] =
 	{
+		"None",
 		"SoftLight",
 		"HardLight",
 		"Burn",
@@ -140,7 +141,7 @@ namespace mu
 		"Overlay",
 		"Blend",
 		"Multiply",
-		"AlphaOverlay", // TODO: This name is not descriptive.
+		"Lighten", // TODO: This name is not descriptive.
 		"NormalCombine",
 	};
 
diff --git a/Engine/Plugins/Experimental/Mutable/Source/MutableTools/Public/MuT/NodeImageLayer.h b/Engine/Plugins/Experimental/Mutable/Source/MutableTools/Public/MuT/NodeImageLayer.h
--- a/Engine/Plugins/Experimental/Mutable/Source/MutableTools/Public/MuT/NodeImageLayer.h
+++ b/Engine/Plugins/Experimental/Mutable/Source/MutableTools/Public/MuT/NodeImageLayer.h
@@ -64,8 +64,6 @@ namespace mu
 		NodeImagePtr GetBlended() const;
 		void SetBlended( NodeImagePtr );
 
-		static const char* s_blendTypeName[int32(EBlendType::_BT_COUNT)];
-
 		//!
 		EBlendType GetBlendType() const;
 		void SetBlendType(EBlendType);
