#rnx

diff --git a/Engine/Plugins/Experimental/Mutable/Source/CustomizableObject/Classes/MuCO/CustomizableObjectSystem.h b/Engine/Plugins/Experimental/Mutable/Source/CustomizableObject/Classes/MuCO/CustomizableObjectSystem.h
--- a/Engine/Plugins/Experimental/Mutable/Source/CustomizableObject/Classes/MuCO/CustomizableObjectSystem.h
+++ b/Engine/Plugins/Experimental/Mutable/Source/CustomizableObject/Classes/MuCO/CustomizableObjectSystem.h
@@ -326,6 +326,7 @@ public:
 
 	// Give access to the internal object data.
 	FCustomizableObjectSystemPrivate* GetPrivate() { return Private.Get(); }
+	const FCustomizableObjectSystemPrivate* GetPrivate() const { return Private.Get(); }
 
 	FStreamableManager& GetStreamableManager() { return StreamableManager; }
 
@@ -365,6 +366,8 @@ public:
 	UFUNCTION(BlueprintCallable, Category = Status)
 	void SetReleaseMutableTexturesImmediately(bool bReleaseTextures);
 
+	bool IsMutableAnimInfoDebuggingEnabled() const;
+
 #if WITH_EDITOR
 	void RecompileCustomizableObjectAsync(const FAssetData& InAssetData, const UCustomizableObject* InObject);
 	
diff --git a/Engine/Plugins/Experimental/Mutable/Source/CustomizableObject/Private/MuCO/CustomizableInstancePrivateData.h b/Engine/Plugins/Experimental/Mutable/Source/CustomizableObject/Private/MuCO/CustomizableInstancePrivateData.h
--- a/Engine/Plugins/Experimental/Mutable/Source/CustomizableObject/Private/MuCO/CustomizableInstancePrivateData.h
+++ b/Engine/Plugins/Experimental/Mutable/Source/CustomizableObject/Private/MuCO/CustomizableInstancePrivateData.h
@@ -122,6 +122,11 @@ struct FCustomizableInstanceComponentData
 	UPROPERTY(Transient, Category = CustomizableObjectInstance, editfixedsize, VisibleAnywhere)
 	TMap<int32, TSoftClassPtr<UAnimInstance>> AnimSlotToBP;
 
+#if WITH_EDITORONLY_DATA
+	// Just used for mutable.EnableMutableAnimInfoDebugging command
+	TArray<FString> MeshPartPaths;
+#endif
+
 	/** Skeletons required by the current generated instance. Skeletons to be loaded and merged.*/
 	UPROPERTY(Transient)
 	FReferencedSkeletons Skeletons;
@@ -290,7 +295,6 @@ public:
 	// Maximum number of SkeletalMesh LODs to stream
 	uint8 NumMaxLODsToStream;
 	
-	
 	TMap<FString, FTextureCoverageQueryData> TextureCoverageQueries;
 
 	UPROPERTY(Transient)
diff --git a/Engine/Plugins/Experimental/Mutable/Source/CustomizableObject/Private/MuCO/CustomizableObjectInstance.cpp b/Engine/Plugins/Experimental/Mutable/Source/CustomizableObject/Private/MuCO/CustomizableObjectInstance.cpp
--- a/Engine/Plugins/Experimental/Mutable/Source/CustomizableObject/Private/MuCO/CustomizableObjectInstance.cpp
+++ b/Engine/Plugins/Experimental/Mutable/Source/CustomizableObject/Private/MuCO/CustomizableObjectInstance.cpp
@@ -289,6 +289,11 @@ void UCustomizableInstancePrivateData::PrepareForUpdate(const TSharedPtr<FMutabl
 			if (FCustomizableInstanceComponentData* ComponentData = GetComponentData(Component.Id))
 			{
 				ComponentData->AnimSlotToBP.Empty();
+
+#if WITH_EDITORONLY_DATA
+				ComponentData->MeshPartPaths.Empty();
+#endif
+
 				ComponentData->Skeletons.SkeletonsToLoad.Empty();
 				ComponentData->Skeletons.SkeletonsToMerge.Empty();
 				ComponentData->PhysicsAssets.PhysicsAssetToLoad.Empty();
@@ -4249,6 +4254,7 @@ bool UCustomizableInstancePrivateData::BuildSkeletalMeshRenderData(const TShared
 	return true;
 }
 
+
 FGraphEventRef UCustomizableInstancePrivateData::LoadAdditionalAssetsAsync(const TSharedPtr<FMutableOperationData>& OperationData, UCustomizableObjectInstance* Public, FStreamableManager &StreamableManager)
 {
 	UCustomizableObject* CustomizableObject = Public->GetCustomizableObject();
@@ -4269,7 +4275,6 @@ FGraphEventRef UCustomizableInstancePrivateData::LoadAdditionalAssetsAsync(const
 
 	GatheredAnimBPs.Empty();
 	AnimBPGameplayTags.Reset();
-
 	
 	for (const FInstanceUpdateData::FSurface& Surface : OperationData->InstanceUpdateData.Surfaces)
 	{
@@ -4458,6 +4463,12 @@ FGraphEventRef UCustomizableInstancePrivateData::LoadAdditionalAssetsAsync(const
 				{
 					AnimBPGameplayTags.AddTag(FGameplayTag::RequestGameplayTag(*Tag));
 				}
+#if WITH_EDITORONLY_DATA
+				else if (Tag.RemoveFromStart("__MeshPath:"))
+				{
+					ComponentData->MeshPartPaths.Add(Tag);
+				}
+#endif
 			}
 		}
 	}
diff --git a/Engine/Plugins/Experimental/Mutable/Source/CustomizableObject/Private/MuCO/CustomizableObjectSystem.cpp b/Engine/Plugins/Experimental/Mutable/Source/CustomizableObject/Private/MuCO/CustomizableObjectSystem.cpp
--- a/Engine/Plugins/Experimental/Mutable/Source/CustomizableObject/Private/MuCO/CustomizableObjectSystem.cpp
+++ b/Engine/Plugins/Experimental/Mutable/Source/CustomizableObject/Private/MuCO/CustomizableObjectSystem.cpp
@@ -482,10 +482,12 @@ void FCustomizableObjectSystemPrivate::CreatedTexture(UTexture2D* Texture)
 }
 
 
-int32 GVarEnableMutableAnimInfoDebugging = 0;
+int32 FCustomizableObjectSystemPrivate::EnableMutableAnimInfoDebugging = 0;
+
 static FAutoConsoleVariableRef CVarEnableMutableAnimInfoDebugging(
-	TEXT("mutable.EnableMutableAnimInfoDebugging"), GVarEnableMutableAnimInfoDebugging,
-	TEXT("If set to 1 or greater print on screen the animation info of the pawn's Customizable Object Instance. Anim BPs, slots and tags will be displayed."),
+	TEXT("mutable.EnableMutableAnimInfoDebugging"), FCustomizableObjectSystemPrivate::EnableMutableAnimInfoDebugging,
+	TEXT("If set to 1 or greater print on screen the animation info of the pawn's Customizable Object Instance. Anim BPs, slots and tags will be displayed."
+	"If the root Customizable Object is recompiled after this command is run, the used skeletal meshes will also be displayed."),
 	ECVF_Default);
 
 
@@ -619,7 +621,8 @@ void FCustomizableObjectSystemPrivate::UpdateStats()
 		SET_DWORD_STAT(STAT_MutableTextureGeneratedMemory, SizeGenerated / 1024.f);
 #endif
 
-		if (GVarEnableMutableAnimInfoDebugging > 0)
+#if WITH_EDITORONLY_DATA
+		if (FCustomizableObjectSystemPrivate::IsMutableAnimInfoDebuggingEnabled())
 		{
 			if (GEngine)
 			{
@@ -661,7 +664,26 @@ void FCustomizableObjectSystemPrivate::UpdateStats()
 						}
 
 						GEngine->AddOnScreenDebugMessage(MsgIndex++, .0f, FColor::Green, TEXT("Slots-AnimBP: "));
-						GEngine->AddOnScreenDebugMessage(MsgIndex++, .0f, FColor::Cyan, TEXT("Player Pawn Mutable Animation info "));
+
+						if (ComponentData)
+						{
+							if (ComponentData->MeshPartPaths.IsEmpty())
+							{
+								GEngine->AddOnScreenDebugMessage(MsgIndex++, .0f, FColor::Magenta,
+									TEXT("No meshes found. In order to see the meshes compile the pawn's root CustomizableObject after the 'mutable.EnableMutableAnimInfoDebugging 1' command has been run."));
+							}
+
+							for (const FString& MeshPath : ComponentData->MeshPartPaths)
+							{
+								GEngine->AddOnScreenDebugMessage(MsgIndex++, .0f, FColor::Magenta, MeshPath);
+							}
+						}
+
+						GEngine->AddOnScreenDebugMessage(MsgIndex++, .0f, FColor::Magenta, TEXT("Meshes: "));
+
+						GEngine->AddOnScreenDebugMessage(MsgIndex++, .0f, FColor::Cyan,
+							TEXT("Player Pawn Mutable Mesh/Animation info for component ") + FString::Printf(TEXT("%d"), 
+							CustomizableSkeletalComponent->ComponentIndex));
 					}
 				}
 
@@ -671,6 +693,7 @@ void FCustomizableObjectSystemPrivate::UpdateStats()
 				}
 			}
 		}
+#endif
 	}
 }
 
@@ -686,7 +709,6 @@ void FCustomizableObjectSystemPrivate::AddReferencedObjects(FReferenceCollector&
 
 int32 FCustomizableObjectSystemPrivate::EnableMutableProgressiveMipStreaming = 1;
 
-
 // Warning! If this is enabled, do not get references to the textures generated by Mutable! They are owned by Mutable and could become invalid at any moment
 static FAutoConsoleVariableRef CVarEnableMutableProgressiveMipStreaming(
 	TEXT("mutable.EnableMutableProgressiveMipStreaming"), FCustomizableObjectSystemPrivate::EnableMutableProgressiveMipStreaming,
@@ -2789,4 +2811,24 @@ void UCustomizableObjectSystem::ClearImageCache()
 {
 	check(GetPrivate()->ImageProvider);
 	GetPrivate()->ImageProvider->ClearCache();
-}
\ No newline at end of file
+}
+
+
+bool FCustomizableObjectSystemPrivate::IsMutableAnimInfoDebuggingEnabled() const
+{ 
+#if WITH_EDITORONLY_DATA
+	return EnableMutableAnimInfoDebugging > 0;
+#else
+	return false;
+#endif
+}
+
+
+bool UCustomizableObjectSystem::IsMutableAnimInfoDebuggingEnabled() const
+{
+#if WITH_EDITOR
+	return GetPrivate()->IsMutableAnimInfoDebuggingEnabled();
+#else
+	return false;
+#endif
+}
diff --git a/Engine/Plugins/Experimental/Mutable/Source/CustomizableObject/Private/MuCO/CustomizableObjectSystemPrivate.h b/Engine/Plugins/Experimental/Mutable/Source/CustomizableObject/Private/MuCO/CustomizableObjectSystemPrivate.h
--- a/Engine/Plugins/Experimental/Mutable/Source/CustomizableObject/Private/MuCO/CustomizableObjectSystemPrivate.h
+++ b/Engine/Plugins/Experimental/Mutable/Source/CustomizableObject/Private/MuCO/CustomizableObjectSystemPrivate.h
@@ -480,6 +480,7 @@ public:
 	TQueue<FMutableTask> PendingTasks;
 
 	static int32 EnableMutableProgressiveMipStreaming;
+	static int32 EnableMutableAnimInfoDebugging;
 
 	/** */
 	inline void AddGameThreadTask(const FMutableTask& Task)
@@ -746,6 +747,8 @@ public:
 	// Check and update the streaming memory limit. Only safe from game thread and when the mutable thread is idle.
 	void UpdateStreamingLimit();
 
+	bool IsMutableAnimInfoDebuggingEnabled() const;
+
 private:
 
 	/** Access to this must be protected with this. */
diff --git a/Engine/Plugins/Experimental/Mutable/Source/CustomizableObjectEditor/Private/MuCOE/GenerateMutableSource/GenerateMutableSourceMesh.cpp b/Engine/Plugins/Experimental/Mutable/Source/CustomizableObjectEditor/Private/MuCOE/GenerateMutableSource/GenerateMutableSourceMesh.cpp
--- a/Engine/Plugins/Experimental/Mutable/Source/CustomizableObjectEditor/Private/MuCOE/GenerateMutableSource/GenerateMutableSourceMesh.cpp
+++ b/Engine/Plugins/Experimental/Mutable/Source/CustomizableObjectEditor/Private/MuCOE/GenerateMutableSource/GenerateMutableSourceMesh.cpp
@@ -48,6 +48,7 @@
 #include "Misc/Guid.h"
 #include "Modules/ModuleManager.h"
 #include "MuCO/CustomizableObject.h"
+#include "MuCO/CustomizableObjectSystemPrivate.h"
 #include "MuCO/CustomizableObjectClothingTypes.h"
 #include "MuCO/CustomizableObjectInstance.h"
 #include "MuCO/CustomizableObjectUIData.h"
@@ -2772,6 +2773,14 @@ mu::NodeMeshPtr GenerateMutableSourceMesh(const UEdGraphPin * Pin,
 
 				AddSocketTagsToMesh(TypedNodeSkel->SkeletalMesh, MutableMesh, GenerationContext);
 
+				if (UCustomizableObjectSystem::GetInstance()->IsMutableAnimInfoDebuggingEnabled())
+				{
+					FString MeshPath;
+					TypedNodeSkel->SkeletalMesh->GetOuter()->GetPathName(nullptr, MeshPath);
+					FString MeshTag = FString("__MeshPath:") + MeshPath;
+					AddTagToMutableMeshUnique(*MutableMesh, MeshTag);
+				}
+
 				MeshData.bHasVertexColors = TypedNodeSkel->SkeletalMesh->GetHasVertexColors();
 				FSkeletalMeshModel* importModel = Helper_GetImportedModel(TypedNodeSkel->SkeletalMesh);
 				MeshData.NumTexCoordChannels = importModel->LODModels[LOD].NumTexCoords;
diff --git a/Engine/Plugins/Experimental/Mutable/Source/CustomizableObjectEditor/Private/MuCOE/GenerateMutableSource/GenerateMutableSourceTable.cpp b/Engine/Plugins/Experimental/Mutable/Source/CustomizableObjectEditor/Private/MuCOE/GenerateMutableSource/GenerateMutableSourceTable.cpp
--- a/Engine/Plugins/Experimental/Mutable/Source/CustomizableObjectEditor/Private/MuCOE/GenerateMutableSource/GenerateMutableSourceTable.cpp
+++ b/Engine/Plugins/Experimental/Mutable/Source/CustomizableObjectEditor/Private/MuCOE/GenerateMutableSource/GenerateMutableSourceTable.cpp
@@ -24,6 +24,7 @@
 #include "Materials/MaterialInstance.h"
 #include "Math/Color.h"
 #include "Misc/Guid.h"
+#include "MuCO/CustomizableObjectSystem.h"
 #include "MuCO/CustomizableObjectParameterTypeDefinitions.h"
 #include "MuCO/CustomizableObjectUIData.h"
 #include "MuCO/UnrealPortabilityHelpers.h"
@@ -260,6 +261,14 @@ void FillTableColumn(const UEdGraphPin* Pin, mu::TablePtr MutableTable, FString
 
 						AddSocketTagsToMesh(SkeletalMesh, MutableMesh, GenerationContext);
 
+						if (UCustomizableObjectSystem::GetInstance()->IsMutableAnimInfoDebuggingEnabled())
+						{
+							FString MeshPath;
+							SkeletalMesh->GetOuter()->GetPathName(nullptr, MeshPath);
+							FString MeshTag = FString("__MeshPath:") + MeshPath;
+							AddTagToMutableMeshUnique(*MutableMesh, MeshTag);
+						}
+
 						MutableTable->SetCell(CurrentColumn, RowIdx, MutableMesh.get());
 					}
 					else
