[Mutable] Rever changes to the Texture Layer Node blend mode.

The name change was producing errors on unserializing those nodes.

diff --git a/Engine/Plugins/Experimental/Mutable/Source/CustomizableObjectEditor/Classes/MuCOE/Nodes/CustomizableObjectNodeTextureLayer.h b/Engine/Plugins/Experimental/Mutable/Source/CustomizableObjectEditor/Classes/MuCOE/Nodes/CustomizableObjectNodeTextureLayer.h
--- a/Engine/Plugins/Experimental/Mutable/Source/CustomizableObjectEditor/Classes/MuCOE/Nodes/CustomizableObjectNodeTextureLayer.h
+++ b/Engine/Plugins/Experimental/Mutable/Source/CustomizableObjectEditor/Classes/MuCOE/Nodes/CustomizableObjectNodeTextureLayer.h
@@ -33,7 +33,7 @@ enum ECustomizableObjectTextureLayerEffect
 	COTLE_BURN					UMETA(DisplayName = "BURN"),
 	COTLE_SCREEN				UMETA(DisplayName = "SCREEN"),
 	COTLE_OVERLAY				UMETA(DisplayName = "OVERLAY"),
-	COTLE_LIGHTEN				UMETA(DisplayName = "LIGHTEN"),
+	COTLE_ALPHA_OVERLAY			UMETA(DisplayName = "LIGHTEN"),
 	COTLE_NORMAL_COMBINE		UMETA(DisplayName = "COMBINE")
 };
 
diff --git a/Engine/Plugins/Experimental/Mutable/Source/CustomizableObjectEditor/Private/MuCOE/GenerateMutableSource/GenerateMutableSourceImage.cpp b/Engine/Plugins/Experimental/Mutable/Source/CustomizableObjectEditor/Private/MuCOE/GenerateMutableSource/GenerateMutableSourceImage.cpp
--- a/Engine/Plugins/Experimental/Mutable/Source/CustomizableObjectEditor/Private/MuCOE/GenerateMutableSource/GenerateMutableSourceImage.cpp
+++ b/Engine/Plugins/Experimental/Mutable/Source/CustomizableObjectEditor/Private/MuCOE/GenerateMutableSource/GenerateMutableSourceImage.cpp
@@ -329,7 +329,7 @@ mu::NodeImagePtr GenerateMutableSourceImage(const UEdGraphPin* Pin, FMutableGrap
 				case COTLE_BURN: Type = mu::EBlendType::BT_BURN; break;
 				case COTLE_SCREEN: Type = mu::EBlendType::BT_SCREEN; break;
 				case COTLE_OVERLAY: Type = mu::EBlendType::BT_OVERLAY; break;
-				case COTLE_LIGHTEN: Type = mu::EBlendType::BT_LIGHTEN; break;
+				case COTLE_ALPHA_OVERLAY: Type = mu::EBlendType::BT_LIGHTEN; break;
 				case COTLE_NORMAL_COMBINE: Type = mu::EBlendType::BT_NORMAL_COMBINE; break;
 				default:
 					GenerationContext.Compiler->CompilerLog(LOCTEXT("UnsupportedImageEffect", "Texture layer effect not supported. Setting to 'Blend'."), Node);
