[mutable] Split multilayer image blends in separate steps per-layer to prevent memory usage peak.
Rename more structures according to coding standards.
Improved memory allocation behaviour for instance operations and mesh projection.
#rnx

diff --git a/Engine/Plugins/Experimental/Mutable/Source/MutableRuntime/Private/MuR/CodeRunner.cpp b/Engine/Plugins/Experimental/Mutable/Source/MutableRuntime/Private/MuR/CodeRunner.cpp
--- a/Engine/Plugins/Experimental/Mutable/Source/MutableRuntime/Private/MuR/CodeRunner.cpp
+++ b/Engine/Plugins/Experimental/Mutable/Source/MutableRuntime/Private/MuR/CodeRunner.cpp
@@ -83,21 +83,21 @@ namespace mu
 		const Model* pModel,
 		const Parameters* pParams,
 		OP::ADDRESS at,
-		uint32 lodMask, uint8 executionOptions, SCHEDULED_OP::EType Type )
+		uint32 lodMask, uint8 executionOptions, FScheduledOp::EType Type )
 		: m_pSettings(pSettings), m_pSystem(s), m_pModel(pModel), m_pParams(pParams), m_lodMask(lodMask)
 	{
 		ScheduledStagePerOp.resize(m_pModel->GetPrivate()->m_program.m_opAddress.Num());
 
 		// We will read this in the end, so make sure we keep it.
-		m_pSystem->m_memory->IncreaseHitCount(CACHE_ADDRESS(at, 0, executionOptions));
+		m_pSystem->m_memory->IncreaseHitCount(FCacheAddress(at, 0, executionOptions));
 
 		FProgram& program = m_pModel->GetPrivate()->m_program;
 		m_romPendingOps.SetNum(program.m_roms.Num());
 
 		// Push the root operation
-		SCHEDULED_OP rootOp;
-		rootOp.at = at;
-		rootOp.executionOptions = executionOptions;
+		FScheduledOp rootOp;
+		rootOp.At = at;
+		rootOp.ExecutionOptions = executionOptions;
 		rootOp.Type = Type;
 		AddOp(rootOp);
 	}
@@ -135,25 +135,25 @@ namespace mu
 
 
     //---------------------------------------------------------------------------------------------
-    void CodeRunner::RunCode_Conditional( SCHEDULED_OP& item,
+    void CodeRunner::RunCode_Conditional( FScheduledOp& item,
                                           const Model* pModel
                                           )
     {
         //MUTABLE_CPUPROFILER_SCOPE(RunCode_Conditional);
 
-		OP_TYPE type = pModel->GetPrivate()->m_program.GetOpType(item.at);
-        auto args = pModel->GetPrivate()->m_program.GetOpArgs<OP::ConditionalArgs>(item.at);
+		OP_TYPE type = pModel->GetPrivate()->m_program.GetOpType(item.At);
+        auto args = pModel->GetPrivate()->m_program.GetOpArgs<OP::ConditionalArgs>(item.At);
 
         // Conditionals have the following execution stages:
         // 0: we need to run the condition
         // 1: we need to run the branch
         // 2: we need to fetch the result and store it in this op
-        switch( item.stage )
+        switch( item.Stage )
         {
         case 0:
         {
-            AddOp( SCHEDULED_OP( item.at,item,1 ),
-                   SCHEDULED_OP( args.condition, item ) );
+            AddOp( FScheduledOp( item.At,item,1 ),
+                   FScheduledOp( args.condition, item ) );
             break;
         }
 
@@ -163,24 +163,24 @@ namespace mu
 
             // If there is no expression, we'll assume true.
             bool value = true;
-            value = GetMemory().GetBool( CACHE_ADDRESS(args.condition, item.executionIndex, item.executionOptions) );
+            value = GetMemory().GetBool( FCacheAddress(args.condition, item.ExecutionIndex, item.ExecutionOptions) );
 
             OP::ADDRESS resultAt = value ? args.yes : args.no;
 
             // Schedule the end of this instruction if necessary
-            AddOp( SCHEDULED_OP( item.at, item, 2, (uint32)value),
-				SCHEDULED_OP( resultAt, item) );
+            AddOp( FScheduledOp( item.At, item, 2, (uint32)value),
+				FScheduledOp( resultAt, item) );
 
             break;
         }
 
         case 2:
         {
-            OP::ADDRESS resultAt = item.customState ? args.yes : args.no;
+            OP::ADDRESS resultAt = item.CustomState ? args.yes : args.no;
 
             // Store the final result
-            CACHE_ADDRESS cat( item );
-            CACHE_ADDRESS rat( resultAt, item );
+            FCacheAddress cat( item );
+            FCacheAddress rat( resultAt, item );
             switch (GetOpDataType(type))
             {
             case DT_BOOL:       GetMemory().SetBool( cat, GetMemory().GetBool(rat) ); break;
@@ -208,13 +208,13 @@ namespace mu
 
 
 	//---------------------------------------------------------------------------------------------
-	void CodeRunner::RunCode_Switch(SCHEDULED_OP& item,
+	void CodeRunner::RunCode_Switch(FScheduledOp& item,
 		const Model* pModel
 	)
 	{
-		OP_TYPE type = pModel->GetPrivate()->m_program.GetOpType(item.at);
+		OP_TYPE type = pModel->GetPrivate()->m_program.GetOpType(item.At);
 
-		const uint8* data = pModel->GetPrivate()->m_program.GetOpArgsPointer(item.at);
+		const uint8* data = pModel->GetPrivate()->m_program.GetOpArgsPointer(item.At);
 
 		OP::ADDRESS VarAddress;
 		FMemory::Memcpy(&VarAddress, data, sizeof(OP::ADDRESS));
@@ -228,14 +228,14 @@ namespace mu
 		FMemory::Memcpy(&CaseCount, data, sizeof(uint32));
 		data += sizeof(uint32);
 
-		switch (item.stage)
+		switch (item.Stage)
 		{
 		case 0:
 		{
 			if (VarAddress)
 			{
-				AddOp(SCHEDULED_OP(item.at, item, 1),
-					SCHEDULED_OP(VarAddress, item));
+				AddOp(FScheduledOp(item.At, item, 1),
+					FScheduledOp(VarAddress, item));
 			}
 			else
 			{
@@ -262,7 +262,7 @@ namespace mu
 		case 1:
 		{
 			// Get the variable result
-			int var = GetMemory().GetInt(CACHE_ADDRESS(VarAddress, item));
+			int var = GetMemory().GetInt(FCacheAddress(VarAddress, item));
 
 			OP::ADDRESS valueAt = DefAddress;
 			for (uint32 C = 0; C < CaseCount; ++C)
@@ -283,19 +283,19 @@ namespace mu
 			}
 
             // Schedule the end of this instruction if necessary
-            AddOp( SCHEDULED_OP( item.at, item, 2, valueAt ),
-				   SCHEDULED_OP( valueAt, item ) );
+            AddOp( FScheduledOp( item.At, item, 2, valueAt ),
+				   FScheduledOp( valueAt, item ) );
 
             break;
         }
 
         case 2:
         {
-			OP::ADDRESS resultAt = OP::ADDRESS(item.customState);
+			OP::ADDRESS resultAt = OP::ADDRESS(item.CustomState);
 
             // Store the final result
-            CACHE_ADDRESS cat( item );
-            CACHE_ADDRESS rat( resultAt, item );
+            FCacheAddress cat( item );
+            FCacheAddress rat( resultAt, item );
             switch (GetOpDataType(type))
             {
             case DT_BOOL:       GetMemory().SetBool( cat, GetMemory().GetBool(rat) ); break;
@@ -323,32 +323,32 @@ namespace mu
 
 
     //---------------------------------------------------------------------------------------------
-    void CodeRunner::RunCode_Instance( SCHEDULED_OP& item,
+    void CodeRunner::RunCode_Instance( FScheduledOp& item,
                                   const Model* pModel,
                                   uint32 lodMask
                                   )
     {
         //MUTABLE_CPUPROFILER_SCOPE(RunCode_Instance);
 
-		OP_TYPE type = pModel->GetPrivate()->m_program.GetOpType(item.at);
+		OP_TYPE type = pModel->GetPrivate()->m_program.GetOpType(item.At);
         switch (type)
         {
 
         case OP_TYPE::IN_ADDVECTOR:
         {
-            auto args = pModel->GetPrivate()->m_program.GetOpArgs<OP::InstanceAddArgs>(item.at);
+            auto args = pModel->GetPrivate()->m_program.GetOpArgs<OP::InstanceAddArgs>(item.At);
 
-            switch (item.stage)
+            switch (item.Stage)
             {
             case 0:
-                    AddOp( SCHEDULED_OP( item.at, item, 1),
-                           SCHEDULED_OP( args.instance, item),
-                           SCHEDULED_OP( args.value, item) );
+                    AddOp( FScheduledOp( item.At, item, 1),
+                           FScheduledOp( args.instance, item),
+                           FScheduledOp( args.value, item) );
                 break;
 
             case 1:
             {
-                InstancePtrConst pBase = GetMemory().GetInstance( CACHE_ADDRESS(args.instance,item) );
+                InstancePtrConst pBase = GetMemory().GetInstance( FCacheAddress(args.instance,item) );
                 InstancePtr pResult;
                 if (!pBase)
                 {
@@ -361,7 +361,7 @@ namespace mu
 
                 if ( args.value )
                 {
-					FVector4f value = GetMemory().GetColour( CACHE_ADDRESS(args.value,item) );
+					FVector4f value = GetMemory().GetColour( FCacheAddress(args.value,item) );
 
                     OP::ADDRESS nameAd = args.name;
                     check(  nameAd < (uint32)pModel->GetPrivate()->m_program.m_constantStrings.Num() );
@@ -382,18 +382,18 @@ namespace mu
 
         case OP_TYPE::IN_ADDSCALAR:
         {
-            auto args = pModel->GetPrivate()->m_program.GetOpArgs<OP::InstanceAddArgs>(item.at);
-            switch (item.stage)
+            auto args = pModel->GetPrivate()->m_program.GetOpArgs<OP::InstanceAddArgs>(item.At);
+            switch (item.Stage)
             {
             case 0:
-                    AddOp( SCHEDULED_OP( item.at, item, 1),
-                           SCHEDULED_OP( args.instance, item),
-                           SCHEDULED_OP( args.value, item) );
+                    AddOp( FScheduledOp( item.At, item, 1),
+                           FScheduledOp( args.instance, item),
+                           FScheduledOp( args.value, item) );
                 break;
 
             case 1:
             {
-                InstancePtrConst pBase = GetMemory().GetInstance( CACHE_ADDRESS(args.instance,item) );
+                InstancePtrConst pBase = GetMemory().GetInstance( FCacheAddress(args.instance,item) );
                 InstancePtr pResult;
                 if (!pBase)
                 {
@@ -401,12 +401,12 @@ namespace mu
                 }
                 else
                 {
-                    pResult = pBase->Clone();
+                    pResult = CloneOrTakeOver<Instance>(pBase.get());
                 }
 
                 if ( args.value )
                 {
-                    float value = GetMemory().GetScalar( CACHE_ADDRESS(args.value,item) );
+                    float value = GetMemory().GetScalar( FCacheAddress(args.value,item) );
 
                     OP::ADDRESS nameAd = args.name;
                     check(  nameAd < (uint32)pModel->GetPrivate()->m_program.m_constantStrings.Num() );
@@ -427,18 +427,18 @@ namespace mu
 
         case OP_TYPE::IN_ADDSTRING:
         {
-            auto args = pModel->GetPrivate()->m_program.GetOpArgs<OP::InstanceAddArgs>( item.at );
-            switch ( item.stage )
+            auto args = pModel->GetPrivate()->m_program.GetOpArgs<OP::InstanceAddArgs>( item.At );
+            switch ( item.Stage )
             {
             case 0:
-                    AddOp( SCHEDULED_OP( item.at, item, 1 ), SCHEDULED_OP( args.instance, item ),
-                           SCHEDULED_OP( args.value, item ) );
+                    AddOp( FScheduledOp( item.At, item, 1 ), FScheduledOp( args.instance, item ),
+                           FScheduledOp( args.value, item ) );
                 break;
 
             case 1:
             {
                 InstancePtrConst pBase =
-                    GetMemory().GetInstance( CACHE_ADDRESS( args.instance, item ) );
+                    GetMemory().GetInstance( FCacheAddress( args.instance, item ) );
                 InstancePtr pResult;
                 if ( !pBase )
                 {
@@ -446,13 +446,13 @@ namespace mu
                 }
                 else
                 {
-                    pResult = pBase->Clone();
-                }
+					pResult = CloneOrTakeOver<Instance>(pBase.get());
+				}
 
                 if ( args.value )
                 {
                     Ptr<const String> value =
-                        GetMemory().GetString( CACHE_ADDRESS( args.value, item ) );
+                        GetMemory().GetString( FCacheAddress( args.value, item ) );
 
                     OP::ADDRESS nameAd = args.name;
                     check( nameAd < (uint32)pModel->GetPrivate()->m_program.m_constantStrings.Num() );
@@ -474,18 +474,18 @@ namespace mu
 
         case OP_TYPE::IN_ADDCOMPONENT:
         {
-            auto args = pModel->GetPrivate()->m_program.GetOpArgs<OP::InstanceAddArgs>(item.at);
-            switch (item.stage)
+            auto args = pModel->GetPrivate()->m_program.GetOpArgs<OP::InstanceAddArgs>(item.At);
+            switch (item.Stage)
             {
             case 0:
-                    AddOp( SCHEDULED_OP( item.at, item, 1),
-                           SCHEDULED_OP( args.instance, item),
-                           SCHEDULED_OP( args.value, item) );
+                    AddOp( FScheduledOp( item.At, item, 1),
+                           FScheduledOp( args.instance, item),
+                           FScheduledOp( args.value, item) );
                 break;
 
             case 1:
             {
-                InstancePtrConst pBase = GetMemory().GetInstance( CACHE_ADDRESS(args.instance,item) );
+                InstancePtrConst pBase = GetMemory().GetInstance( FCacheAddress(args.instance,item) );
                 InstancePtr pResult;
                 if (!pBase)
                 {
@@ -493,12 +493,12 @@ namespace mu
                 }
                 else
                 {
-                    pResult = pBase->Clone();
-                }
+					pResult = CloneOrTakeOver<Instance>(pBase.get());
+				}
 
                 if ( args.value )
                 {
-                    InstancePtrConst pComp = GetMemory().GetInstance( CACHE_ADDRESS(args.value,item) );
+                    InstancePtrConst pComp = GetMemory().GetInstance( FCacheAddress(args.value,item) );
 
                     int cindex = pResult->GetPrivate()->AddComponent( 0 );
 
@@ -533,28 +533,34 @@ namespace mu
 
         case OP_TYPE::IN_ADDSURFACE:
         {
-            auto args = pModel->GetPrivate()->m_program.GetOpArgs<OP::InstanceAddArgs>(item.at);
-            switch (item.stage)
+            auto args = pModel->GetPrivate()->m_program.GetOpArgs<OP::InstanceAddArgs>(item.At);
+            switch (item.Stage)
             {
             case 0:
-                    AddOp( SCHEDULED_OP( item.at, item, 1),
-                           SCHEDULED_OP( args.instance, item),
-                           SCHEDULED_OP( args.value, item) );
+                    AddOp( FScheduledOp( item.At, item, 1),
+                           FScheduledOp( args.instance, item),
+                           FScheduledOp( args.value, item) );
                 break;
 
             case 1:
             {
-                InstancePtrConst pBase = GetMemory().GetInstance( CACHE_ADDRESS(args.instance,item) );
+                InstancePtrConst pBase = GetMemory().GetInstance( FCacheAddress(args.instance,item) );
 
                 InstancePtr pResult;
-                if (pBase) pResult = pBase->Clone();
-                else pResult = new Instance();
+				if (pBase)
+				{
+					pResult = CloneOrTakeOver<Instance>(pBase.get());
+				}
+				else
+				{
+					pResult = new Instance();
+				}
 
                 // Empty surfaces are ok, they still need to be created, because they may contain
                 // additional information like internal or external IDs
                 //if ( args.value )
                 {
-                    InstancePtrConst pSurf = GetMemory().GetInstance( CACHE_ADDRESS(args.value,item) );
+                    InstancePtrConst pSurf = GetMemory().GetInstance( FCacheAddress(args.value,item) );
 
                     int sindex = pResult->GetPrivate()->AddSurface( 0, 0 );
 
@@ -597,26 +603,26 @@ namespace mu
 
         case OP_TYPE::IN_ADDLOD:
         {
-            auto args = pModel->GetPrivate()->m_program.GetOpArgs<OP::InstanceAddLODArgs>(item.at);
-            switch (item.stage)
+            auto args = pModel->GetPrivate()->m_program.GetOpArgs<OP::InstanceAddLODArgs>(item.At);
+            switch (item.Stage)
             {
             case 0:
             {                
-                    TArray<SCHEDULED_OP> deps;
-                    for ( int i=0; i<MUTABLE_OP_MAX_ADD_COUNT; ++i )
+                TArray<FScheduledOp> deps;
+                for ( int i=0; i<MUTABLE_OP_MAX_ADD_COUNT; ++i )
+                {
+                    if ( args.lod[i] )
                     {
-                        if ( args.lod[i] )
-                        {
-                            bool selectedLod = ( (1<<i) & lodMask ) != 0;
+                        bool selectedLod = ( (1<<i) & lodMask ) != 0;
 
-                            if ( selectedLod )
-                            {
-                                deps.Emplace(args.lod[i], item);
-                            }
+                        if ( selectedLod )
+                        {
+                            deps.Emplace(args.lod[i], item);
                         }
                     }
+                }
 
-                    AddOp( SCHEDULED_OP( item.at,item, 1), deps );
+                AddOp( FScheduledOp( item.At,item, 1), deps );
 
                 break;
             }
@@ -634,14 +640,14 @@ namespace mu
 
                         if ( selectedLod )
                         {
-                            InstancePtrConst pLOD = GetMemory().GetInstance( CACHE_ADDRESS(args.lod[i],item) );
+                            InstancePtrConst pLOD = GetMemory().GetInstance( FCacheAddress(args.lod[i],item) );
 
-                            int lindex = pResult->GetPrivate()->AddLOD();
+                            int LODIndex = pResult->GetPrivate()->AddLOD();
 
                             // In a degenerated case, the returned pLOD may not have an LOD inside
                             if ( pLOD && !pLOD->GetPrivate()->m_lods.IsEmpty() )
                             {
-                                pResult->GetPrivate()->m_lods[lindex] = pLOD->GetPrivate()->m_lods[0];
+                                pResult->GetPrivate()->m_lods[LODIndex] = pLOD->GetPrivate()->m_lods[0];
                             }
                         }
                         else
@@ -670,25 +676,25 @@ namespace mu
 
 
     //---------------------------------------------------------------------------------------------
-    void CodeRunner::RunCode_InstanceAddResource( SCHEDULED_OP& item,
+    void CodeRunner::RunCode_InstanceAddResource( FScheduledOp& item,
                                   const Model* pModel,
                                   const Parameters* pParams
                                   )
     {
 		MUTABLE_CPUPROFILER_SCOPE(RunCode_InstanceAddResource);
 
-		OP_TYPE type = pModel->GetPrivate()->m_program.GetOpType(item.at);
+		OP_TYPE type = pModel->GetPrivate()->m_program.GetOpType(item.At);
         switch (type)
         {
 
         case OP_TYPE::IN_ADDMESH:
         {
-            auto args = pModel->GetPrivate()->m_program.GetOpArgs<OP::InstanceAddArgs>(item.at);
-            switch (item.stage)
+            auto args = pModel->GetPrivate()->m_program.GetOpArgs<OP::InstanceAddArgs>(item.At);
+            switch (item.Stage)
             {
             case 0:
-                    AddOp( SCHEDULED_OP( item.at, item, 1),
-                           SCHEDULED_OP( args.instance, item) );
+                    AddOp( FScheduledOp( item.At, item, 1),
+                           FScheduledOp( args.instance, item) );
 
                 // We don't build the resources when building instance: just store ids for them.
                 //PushIfNotVisited(args.value, item);
@@ -696,7 +702,7 @@ namespace mu
 
             case 1:
             {
-                InstancePtrConst pBase = GetMemory().GetInstance( CACHE_ADDRESS(args.instance,item) );
+                InstancePtrConst pBase = GetMemory().GetInstance( FCacheAddress(args.instance,item) );
                 InstancePtr pResult;
                 if (!pBase)
                 {
@@ -704,8 +710,8 @@ namespace mu
                 }
                 else
                 {
-                    pResult = pBase->Clone();
-                }
+					pResult = CloneOrTakeOver<Instance>(pBase.get());
+				}
 
                 if ( args.value )
                 {
@@ -730,17 +736,17 @@ namespace mu
 
         case OP_TYPE::IN_ADDIMAGE:
         {
-            auto args = pModel->GetPrivate()->m_program.GetOpArgs<OP::InstanceAddArgs>(item.at);
-            switch (item.stage)
+            auto args = pModel->GetPrivate()->m_program.GetOpArgs<OP::InstanceAddArgs>(item.At);
+            switch (item.Stage)
             {
             case 0:
 				// We don't build the resources when building instance: just store ids for them.
-				AddOp( SCHEDULED_OP( item.at, item, 1), SCHEDULED_OP( args.instance, item) );
+				AddOp( FScheduledOp( item.At, item, 1), FScheduledOp( args.instance, item) );
                 break;
 
             case 1:
             {
-                InstancePtrConst pBase = GetMemory().GetInstance( CACHE_ADDRESS(args.instance,item) );
+                InstancePtrConst pBase = GetMemory().GetInstance( FCacheAddress(args.instance,item) );
                 InstancePtr pResult;
                 if (!pBase)
                 {
@@ -748,8 +754,8 @@ namespace mu
                 }
                 else
                 {
-                    pResult = pBase->Clone();
-                }
+					pResult = CloneOrTakeOver<Instance>(pBase.get());
+				}
 
                 if ( args.value )
                 {
@@ -780,17 +786,17 @@ namespace mu
 
 
     //---------------------------------------------------------------------------------------------
-    void CodeRunner::RunCode_ConstantResource( SCHEDULED_OP& item, const Model* pModel )
+    void CodeRunner::RunCode_ConstantResource( FScheduledOp& item, const Model* pModel )
     {
 		MUTABLE_CPUPROFILER_SCOPE(RunCode_Constant);
 
-		OP_TYPE type = pModel->GetPrivate()->m_program.GetOpType(item.at);
+		OP_TYPE type = pModel->GetPrivate()->m_program.GetOpType(item.At);
         switch (type)
         {
 
         case OP_TYPE::ME_CONSTANT:
         {
-			OP::MeshConstantArgs args = pModel->GetPrivate()->m_program.GetOpArgs<OP::MeshConstantArgs>(item.at);
+			OP::MeshConstantArgs args = pModel->GetPrivate()->m_program.GetOpArgs<OP::MeshConstantArgs>(item.At);
 
             OP::ADDRESS cat = args.value;
 
@@ -819,26 +825,26 @@ namespace mu
 			}
 
             GetMemory().SetMesh( item, pSource );
-			//UE_LOG(LogMutableCore, Log, TEXT("Set mesh constant %d."), item.at);
+			//UE_LOG(LogMutableCore, Log, TEXT("Set mesh constant %d."), item.At);
             break;
         }
 
         case OP_TYPE::IM_CONSTANT:
         {
-            auto args = pModel->GetPrivate()->m_program.GetOpArgs<OP::ResourceConstantArgs>(item.at);
-                OP::ADDRESS cat = args.value;
+			OP::ResourceConstantArgs args = pModel->GetPrivate()->m_program.GetOpArgs<OP::ResourceConstantArgs>(item.At);
+            OP::ADDRESS cat = args.value;
 
             const FProgram& program = pModel->GetPrivate()->m_program;
 
-			int32 MipsToSkip = item.executionOptions;
-            ImagePtrConst pSource;
+			int32 MipsToSkip = item.ExecutionOptions;
+            Ptr<const Image> pSource;
             program.GetConstant( cat, pSource, MipsToSkip);
 
 			// Assume the ROM has been loaded previously in a task generated at IssueOp
 			check(pSource);
 
             GetMemory().SetImage( item, pSource );
-			//UE_LOG(LogMutableCore, Log, TEXT("Set image constant %d."), item.at);
+			//UE_LOG(LogMutableCore, Log, TEXT("Set image constant %d."), item.At);
 			break;
         }
 
@@ -853,35 +859,35 @@ namespace mu
     }
 
     //---------------------------------------------------------------------------------------------
-    void CodeRunner::RunCode_Mesh( SCHEDULED_OP& item, const Model* pModel )
+    void CodeRunner::RunCode_Mesh( FScheduledOp& item, const Model* pModel )
     {
 		MUTABLE_CPUPROFILER_SCOPE(RunCode_Mesh);
 
-		OP_TYPE type = pModel->GetPrivate()->m_program.GetOpType(item.at);
+		OP_TYPE type = pModel->GetPrivate()->m_program.GetOpType(item.At);
         switch (type)
         {
 
         case OP_TYPE::ME_APPLYLAYOUT:
         {
-            auto args = pModel->GetPrivate()->m_program.GetOpArgs<OP::MeshApplyLayoutArgs>(item.at);
-            switch (item.stage)
+            auto args = pModel->GetPrivate()->m_program.GetOpArgs<OP::MeshApplyLayoutArgs>(item.At);
+            switch (item.Stage)
             {
             case 0:
-                    AddOp( SCHEDULED_OP( item.at, item, 1),
-                           SCHEDULED_OP( args.mesh, item),
-                           SCHEDULED_OP( args.layout, item) );
+                    AddOp( FScheduledOp( item.At, item, 1),
+                           FScheduledOp( args.mesh, item),
+                           FScheduledOp( args.layout, item) );
                 break;
 
             case 1:
             {
-                Ptr<const Mesh> pBase = GetMemory().GetMesh( CACHE_ADDRESS(args.mesh,item) );
+                Ptr<const Mesh> pBase = GetMemory().GetMesh( FCacheAddress(args.mesh,item) );
 
                 MeshPtr pApplied;
                 if (pBase)
                 {
                     pApplied = pBase->Clone();
 
-                    Ptr<const Layout> pLayout = GetMemory().GetLayout( CACHE_ADDRESS(args.layout,item) );
+                    Ptr<const Layout> pLayout = GetMemory().GetLayout( FCacheAddress(args.layout,item) );
                     int texCoordsSet = args.channel;
                     MeshApplyLayout( pApplied.get(), pLayout.get(), texCoordsSet );
                 }
@@ -899,7 +905,7 @@ namespace mu
 
         case OP_TYPE::ME_DIFFERENCE:
         {
-			const uint8* data = pModel->GetPrivate()->m_program.GetOpArgsPointer(item.at);
+			const uint8* data = pModel->GetPrivate()->m_program.GetOpArgsPointer(item.At);
 
 			OP::ADDRESS BaseAt = 0;
 			FMemory::Memcpy(&BaseAt, data, sizeof(OP::ADDRESS)); data += sizeof(OP::ADDRESS);
@@ -907,14 +913,14 @@ namespace mu
 			OP::ADDRESS TargetAt = 0;
 			FMemory::Memcpy(&TargetAt, data, sizeof(OP::ADDRESS)); data += sizeof(OP::ADDRESS);
 
-            switch (item.stage)
+            switch (item.Stage)
             {
             case 0:
                 if (BaseAt && TargetAt)
                 {
-                    AddOp( SCHEDULED_OP( item.at, item, 1),
-                            SCHEDULED_OP( BaseAt, item),
-                            SCHEDULED_OP( TargetAt, item) );
+                    AddOp( FScheduledOp( item.At, item, 1),
+                            FScheduledOp( BaseAt, item),
+                            FScheduledOp( TargetAt, item) );
                 }
                 else
                 {
@@ -924,8 +930,8 @@ namespace mu
 
             case 1:
             {
-                Ptr<const Mesh> pBase = GetMemory().GetMesh(CACHE_ADDRESS(BaseAt,item));
-                Ptr<const Mesh> pTarget = GetMemory().GetMesh(CACHE_ADDRESS(TargetAt,item));
+                Ptr<const Mesh> pBase = GetMemory().GetMesh(FCacheAddress(BaseAt,item));
+                Ptr<const Mesh> pTarget = GetMemory().GetMesh(FCacheAddress(TargetAt,item));
 
 				TArray<MESH_BUFFER_SEMANTIC, TInlineAllocator<8>> Semantics;
 				TArray<int32, TInlineAllocator<8>> SemanticIndices;
@@ -964,7 +970,7 @@ namespace mu
 
         case OP_TYPE::ME_MORPH2:
         {
-			const uint8* data = pModel->GetPrivate()->m_program.GetOpArgsPointer(item.at);
+			const uint8* data = pModel->GetPrivate()->m_program.GetOpArgsPointer(item.At);
 
 			OP::ADDRESS FactorAt = 0;
 			FMemory::Memcpy(&FactorAt, data, sizeof(OP::ADDRESS)); data += sizeof(OP::ADDRESS);
@@ -982,14 +988,14 @@ namespace mu
 				FMemory::Memcpy(&Targets[T], data, sizeof(OP::ADDRESS)); data += sizeof(OP::ADDRESS);
 			}
 
-			switch (item.stage)
+			switch (item.Stage)
             {
             case 0:
                 if (BaseAt)
                 {
-                    AddOp( SCHEDULED_OP(item.at, item, 1),
-                           SCHEDULED_OP(BaseAt, item),
-                           SCHEDULED_OP(FactorAt, item) );
+                    AddOp( FScheduledOp(item.At, item, 1),
+                           FScheduledOp(BaseAt, item),
+                           FScheduledOp(FactorAt, item) );
                 }
                 else
                 {
@@ -999,8 +1005,8 @@ namespace mu
 
             case 1:
             {
-                bool baseValid = GetMemory().IsValid( CACHE_ADDRESS(BaseAt,item) );
-                float factor = GetMemory().GetScalar( CACHE_ADDRESS(FactorAt,item) );
+                bool baseValid = GetMemory().IsValid( FCacheAddress(BaseAt,item) );
+                float factor = GetMemory().GetScalar( FCacheAddress(FactorAt,item) );
 
                 // Factor goes from -1 to 1 across all targets. [0 - 1] represents positive morphs, while [-1, 0) represent negative morphs.
 				factor = FMath::Clamp(factor, -1.0f, 1.0f); // Is the factor not in range [-1, 1], it will index a non existing morph.
@@ -1033,35 +1039,35 @@ namespace mu
                 if (baseValid)
                 {
                     FScheduledOpData HeapData;
-					HeapData.bifactor = bifactor;
-					HeapData.min = FMath::Clamp(min, 0, NumTargets - 1);
-					HeapData.max = FMath::Clamp(max, 0, NumTargets - 1);
+					HeapData.Interpolate.Bifactor = bifactor;
+					HeapData.Interpolate.Min = FMath::Clamp(min, 0, NumTargets - 1);
+					HeapData.Interpolate.Max = FMath::Clamp(max, 0, NumTargets - 1);
 					uint32 dataAddress = uint32(m_heapData.Add(HeapData));
 
                     // Just the first of the targets
 					if ( bifactor < UE_SMALL_NUMBER && bifactor > -UE_SMALL_NUMBER )
                     {                        
-                        AddOp( SCHEDULED_OP( item.at, item, 2, dataAddress),
-                                SCHEDULED_OP(BaseAt, item),
-                                SCHEDULED_OP(Targets[min], item) );
+                        AddOp( FScheduledOp( item.At, item, 2, dataAddress),
+                                FScheduledOp(BaseAt, item),
+                                FScheduledOp(Targets[min], item) );
                     }
                     // Just the second of the targets
 					else if ( bifactor > 1.0f - UE_SMALL_NUMBER || bifactor <= -1.0f + UE_SMALL_NUMBER )
                     {
                         check( max>0 );
 
-                        AddOp( SCHEDULED_OP( item.at, item, 2, dataAddress),
-                                SCHEDULED_OP(BaseAt, item),
-                                SCHEDULED_OP(Targets[max], item) );
+                        AddOp( FScheduledOp( item.At, item, 2, dataAddress),
+                                FScheduledOp(BaseAt, item),
+                                FScheduledOp(Targets[max], item) );
                     }
                     // Mix two targets on the base
                     else
                     {
                         // We will need the base again
-                        AddOp( SCHEDULED_OP( item.at, item, 2, dataAddress),
-                                SCHEDULED_OP(BaseAt, item),
-                                SCHEDULED_OP(Targets[min], item),
-                                SCHEDULED_OP(Targets[max], item) );
+                        AddOp( FScheduledOp( item.At, item, 2, dataAddress),
+                                FScheduledOp(BaseAt, item),
+                                FScheduledOp(Targets[min], item),
+                                FScheduledOp(Targets[max], item) );
                     }
                 }
 
@@ -1070,13 +1076,13 @@ namespace mu
 
             case 2:
             {
-                Ptr<const Mesh> pBase = GetMemory().GetMesh( CACHE_ADDRESS(BaseAt,item) );
+                Ptr<const Mesh> pBase = GetMemory().GetMesh( FCacheAddress(BaseAt,item) );
 
                 // Factor from 0 to 1 between the two targets
-                const FScheduledOpData& HeapData = m_heapData[ (size_t)item.customState ];
-                float bifactor = HeapData.bifactor;
-                int min = HeapData.min;
-                int max = HeapData.max;
+                const FScheduledOpData& HeapData = m_heapData[ (size_t)item.CustomState ];
+                float bifactor = HeapData.Interpolate.Bifactor;
+                int min = HeapData.Interpolate.Min;
+                int max = HeapData.Interpolate.Max;
 
                 MeshPtrConst pResult;
 
@@ -1086,7 +1092,7 @@ namespace mu
 					if ( bifactor < UE_SMALL_NUMBER && bifactor > -UE_SMALL_NUMBER )
                     {
                         // Base with one full morph
-                        Ptr<const Mesh> pMorph = GetMemory().GetMesh( CACHE_ADDRESS(Targets[min],item) );
+                        Ptr<const Mesh> pMorph = GetMemory().GetMesh( FCacheAddress(Targets[min],item) );
                         if (pMorph)
                         {
 							pResult = MeshMorph(pBase.get(), pMorph.get());
@@ -1096,7 +1102,7 @@ namespace mu
                     else if ( bifactor > 1.0f - UE_SMALL_NUMBER )
                     {
                         check( max>0 );
-                        Ptr<const Mesh> pMorph = GetMemory().GetMesh( CACHE_ADDRESS(Targets[max],item) );
+                        Ptr<const Mesh> pMorph = GetMemory().GetMesh( FCacheAddress(Targets[max],item) );
                         if (pMorph)
                         {
 							pResult = MeshMorph(pBase.get(), pMorph.get());
@@ -1106,7 +1112,7 @@ namespace mu
 					else if (bifactor <= -1.0f + UE_SMALL_NUMBER)
 					{
 						check( max > 0 );
-						Ptr<const Mesh> pMorph = GetMemory().GetMesh(CACHE_ADDRESS(Targets[max], item));
+						Ptr<const Mesh> pMorph = GetMemory().GetMesh(FCacheAddress(Targets[max], item));
 						if (pMorph)
 						{
 							pResult = MeshMorph(pBase.get(), pMorph.get(), bifactor);
@@ -1115,8 +1121,8 @@ namespace mu
                     // Mix two targets on the base
                     else
                     {
-                        Ptr<const Mesh> pMin = GetMemory().GetMesh( CACHE_ADDRESS(Targets[min],item) );
-                        Ptr<const Mesh> pMax = GetMemory().GetMesh( CACHE_ADDRESS(Targets[max],item) );
+                        Ptr<const Mesh> pMin = GetMemory().GetMesh( FCacheAddress(Targets[min],item) );
+                        Ptr<const Mesh> pMax = GetMemory().GetMesh( FCacheAddress(Targets[max],item) );
                         if (pMin && pMax)
                         {
                             pResult = MeshMorph2( pBase.get(), pMin.get(), pMax.get(), bifactor );
@@ -1142,19 +1148,19 @@ namespace mu
 
         case OP_TYPE::ME_MERGE:
         {
-            auto args = pModel->GetPrivate()->m_program.GetOpArgs<OP::MeshMergeArgs>(item.at);
-            switch (item.stage)
+            auto args = pModel->GetPrivate()->m_program.GetOpArgs<OP::MeshMergeArgs>(item.At);
+            switch (item.Stage)
             {
             case 0:
-                    AddOp( SCHEDULED_OP( item.at, item, 1),
-                           SCHEDULED_OP( args.base, item),
-                           SCHEDULED_OP( args.added, item) );
+                    AddOp( FScheduledOp( item.At, item, 1),
+                           FScheduledOp( args.base, item),
+                           FScheduledOp( args.added, item) );
                 break;
 
             case 1:
             {
-                MeshPtrConst pA = GetMemory().GetMesh( CACHE_ADDRESS(args.base,item) );
-                MeshPtrConst pB = GetMemory().GetMesh( CACHE_ADDRESS(args.added,item) );
+                MeshPtrConst pA = GetMemory().GetMesh( FCacheAddress(args.base,item) );
+                MeshPtrConst pB = GetMemory().GetMesh( FCacheAddress(args.added,item) );
 
                 MeshPtr pResult;
 
@@ -1205,15 +1211,15 @@ namespace mu
 
         case OP_TYPE::ME_INTERPOLATE:
         {
-            auto args = pModel->GetPrivate()->m_program.GetOpArgs<OP::MeshInterpolateArgs>(item.at);
-            switch (item.stage)
+            auto args = pModel->GetPrivate()->m_program.GetOpArgs<OP::MeshInterpolateArgs>(item.At);
+            switch (item.Stage)
             {
             case 0:
                 if ( args.base )
                 {
-                    AddOp( SCHEDULED_OP( item.at, item, 1),
-                           SCHEDULED_OP( args.base, item),
-                           SCHEDULED_OP( args.factor, item ) );
+                    AddOp( FScheduledOp( item.At, item, 1),
+                           FScheduledOp( args.base, item),
+                           FScheduledOp( args.factor, item ) );
                 }
                 else
                 {
@@ -1232,7 +1238,7 @@ namespace mu
                 }
 
                 // Factor from 0 to 1 across all targets
-                float factor = GetMemory().GetScalar( CACHE_ADDRESS(args.factor,item) );
+                float factor = GetMemory().GetScalar( FCacheAddress(args.factor,item) );
 
                 float delta = 1.0f/(count-1);
                 int min = (int)floorf( factor/delta );
@@ -1242,9 +1248,9 @@ namespace mu
                 float bifactor = factor/delta - min;
 
                 FScheduledOpData data;
-                data.bifactor = bifactor;
-				data.min = FMath::Clamp(min, 0, count - 1);
-				data.max = FMath::Clamp(max, 0, count - 1);
+                data.Interpolate.Bifactor = bifactor;
+				data.Interpolate.Min = FMath::Clamp(min, 0, count - 1);
+				data.Interpolate.Max = FMath::Clamp(max, 0, count - 1);
 				uint32 dataAddress = uint32(m_heapData.Num());
 
                 // Just the first of the targets
@@ -1253,7 +1259,7 @@ namespace mu
                     if( min==0 )
                     {
                         // Just the base
-                            Ptr<const Mesh> pBase = GetMemory().GetMesh( CACHE_ADDRESS(args.base,item) );
+                            Ptr<const Mesh> pBase = GetMemory().GetMesh( FCacheAddress(args.base,item) );
                             MeshPtr pResult;
                             if (pBase)
                             {
@@ -1265,36 +1271,36 @@ namespace mu
                     {
                         // Base with one full morph
                         m_heapData.Add(data);
-                            AddOp( SCHEDULED_OP( item.at, item, 2, dataAddress),
-                                   SCHEDULED_OP( args.base, item),
-                                   SCHEDULED_OP( args.targets[min-1], item) );
+                            AddOp( FScheduledOp( item.At, item, 2, dataAddress),
+                                   FScheduledOp( args.base, item),
+                                   FScheduledOp( args.targets[min-1], item) );
                         }
                     }
                 // Just the second of the targets
                 else if ( bifactor > 1.0f-UE_SMALL_NUMBER )
                 {
                     m_heapData.Add(data);
-                        AddOp( SCHEDULED_OP( item.at, item, 2, dataAddress),
-                               SCHEDULED_OP( args.base, item),
-                               SCHEDULED_OP( args.targets[max-1], item) );
+                        AddOp( FScheduledOp( item.At, item, 2, dataAddress),
+                               FScheduledOp( args.base, item),
+                               FScheduledOp( args.targets[max-1], item) );
                     }
                 // Mix the first target on the base
                 else if ( min==0 )
                 {
                     m_heapData.Add(data);
-                        AddOp( SCHEDULED_OP( item.at, item, 2, dataAddress),
-                               SCHEDULED_OP( args.base, item),
-                               SCHEDULED_OP( args.targets[0], item)
+                        AddOp( FScheduledOp( item.At, item, 2, dataAddress),
+                               FScheduledOp( args.base, item),
+                               FScheduledOp( args.targets[0], item)
                                );
                     }
                 // Mix two targets on the base
                 else
                 {
                     m_heapData.Add(data);
-                        AddOp( SCHEDULED_OP( item.at, item, 2, dataAddress),
-                               SCHEDULED_OP( args.base, item),
-                               SCHEDULED_OP( args.targets[min-1], item),
-                               SCHEDULED_OP( args.targets[max-1], item) );
+                        AddOp( FScheduledOp( item.At, item, 2, dataAddress),
+                               FScheduledOp( args.base, item),
+                               FScheduledOp( args.targets[min-1], item),
+                               FScheduledOp( args.targets[max-1], item) );
                     }
 
                 break;
@@ -1310,14 +1316,14 @@ namespace mu
                     count++;
                 }
 
-                const FScheduledOpData& data = m_heapData[ (size_t)item.customState ];
+                const FScheduledOpData& data = m_heapData[ (size_t)item.CustomState ];
 
                 // Factor from 0 to 1 between the two targets
-                float bifactor = data.bifactor;
-                int min = data.min;
-                int max = data.max;
+                float bifactor = data.Interpolate.Bifactor;
+                int min = data.Interpolate.Min;
+                int max = data.Interpolate.Max;
 
-                Ptr<const Mesh> pBase = GetMemory().GetMesh( CACHE_ADDRESS(args.base, item) );
+                Ptr<const Mesh> pBase = GetMemory().GetMesh( FCacheAddress(args.base, item) );
 
                 MeshPtr pResult;
 
@@ -1334,7 +1340,7 @@ namespace mu
                         else
                         {
                             // Base with one full morph
-                            Ptr<const Mesh> pMorph = GetMemory().GetMesh( CACHE_ADDRESS(args.targets[min-1],item) );
+                            Ptr<const Mesh> pMorph = GetMemory().GetMesh( FCacheAddress(args.targets[min-1],item) );
                             pResult = MeshMorph( pBase.get(), pMorph.get() );
                         }
                     }
@@ -1342,7 +1348,7 @@ namespace mu
                     else if ( bifactor > 1.0f-UE_SMALL_NUMBER )
                     {
                         check( max>0 );
-                        Ptr<const Mesh> pMorph = GetMemory().GetMesh( CACHE_ADDRESS(args.targets[max-1],item) );
+                        Ptr<const Mesh> pMorph = GetMemory().GetMesh( FCacheAddress(args.targets[max-1],item) );
 
                         if (pMorph)
                         {
@@ -1356,7 +1362,7 @@ namespace mu
                     // Mix the first target on the base
                     else if ( min==0 )
                     {
-                        Ptr<const Mesh> pMorph = GetMemory().GetMesh( CACHE_ADDRESS( args.targets[0], item ) );
+                        Ptr<const Mesh> pMorph = GetMemory().GetMesh( FCacheAddress( args.targets[0], item ) );
                         if (pMorph)
                         {
                             pResult = MeshMorph( pBase.get(), pMorph.get(), bifactor );
@@ -1370,8 +1376,8 @@ namespace mu
                     // Mix two targets on the base
                     else
                     {
-                        Ptr<const Mesh> pMin = GetMemory().GetMesh( CACHE_ADDRESS(args.targets[min-1],item) );
-                        Ptr<const Mesh> pMax = GetMemory().GetMesh( CACHE_ADDRESS(args.targets[max-1],item) );
+                        Ptr<const Mesh> pMin = GetMemory().GetMesh( FCacheAddress(args.targets[min-1],item) );
+                        Ptr<const Mesh> pMax = GetMemory().GetMesh( FCacheAddress(args.targets[max-1],item) );
 
                         if (pMin && pMax)
                         {
@@ -1397,19 +1403,19 @@ namespace mu
 
         case OP_TYPE::ME_MASKCLIPMESH:
         {
-            auto args = pModel->GetPrivate()->m_program.GetOpArgs<OP::MeshMaskClipMeshArgs>(item.at);
-            switch (item.stage)
+            auto args = pModel->GetPrivate()->m_program.GetOpArgs<OP::MeshMaskClipMeshArgs>(item.At);
+            switch (item.Stage)
             {
             case 0:
-                    AddOp( SCHEDULED_OP( item.at, item, 1),
-                           SCHEDULED_OP( args.source, item),
-                           SCHEDULED_OP( args.clip, item) );
+                    AddOp( FScheduledOp( item.At, item, 1),
+                           FScheduledOp( args.source, item),
+                           FScheduledOp( args.clip, item) );
                 break;
 
             case 1:
             {
-                Ptr<const Mesh> pSource = GetMemory().GetMesh( CACHE_ADDRESS(args.source,item) );
-                Ptr<const Mesh> pClip = GetMemory().GetMesh( CACHE_ADDRESS(args.clip,item) );
+                Ptr<const Mesh> pSource = GetMemory().GetMesh( FCacheAddress(args.source,item) );
+                Ptr<const Mesh> pClip = GetMemory().GetMesh( FCacheAddress(args.clip,item) );
 
                 // Only if both are valid.
                 MeshPtr pResult;
@@ -1431,19 +1437,19 @@ namespace mu
 
         case OP_TYPE::ME_MASKDIFF:
         {
-            auto args = pModel->GetPrivate()->m_program.GetOpArgs<OP::MeshMaskDiffArgs>(item.at);
-            switch (item.stage)
+            auto args = pModel->GetPrivate()->m_program.GetOpArgs<OP::MeshMaskDiffArgs>(item.At);
+            switch (item.Stage)
             {
             case 0:
-                    AddOp( SCHEDULED_OP( item.at, item, 1),
-                           SCHEDULED_OP( args.source, item),
-                           SCHEDULED_OP( args.fragment, item) );
+                    AddOp( FScheduledOp( item.At, item, 1),
+                           FScheduledOp( args.source, item),
+                           FScheduledOp( args.fragment, item) );
                 break;
 
             case 1:
             {
-                Ptr<const Mesh> pSource = GetMemory().GetMesh( CACHE_ADDRESS(args.source,item) );
-                Ptr<const Mesh> pClip = GetMemory().GetMesh( CACHE_ADDRESS(args.fragment,item) );
+                Ptr<const Mesh> pSource = GetMemory().GetMesh( FCacheAddress(args.source,item) );
+                Ptr<const Mesh> pClip = GetMemory().GetMesh( FCacheAddress(args.fragment,item) );
 
                 // Only if both are valid.
                 MeshPtr pResult;
@@ -1465,19 +1471,19 @@ namespace mu
 
         case OP_TYPE::ME_SUBTRACT:
         {
-            auto args = pModel->GetPrivate()->m_program.GetOpArgs<OP::MeshSubtractArgs>(item.at);
-            switch (item.stage)
+            auto args = pModel->GetPrivate()->m_program.GetOpArgs<OP::MeshSubtractArgs>(item.At);
+            switch (item.Stage)
             {
             case 0:
-                    AddOp( SCHEDULED_OP( item.at, item, 1),
-                           SCHEDULED_OP( args.a, item),
-                           SCHEDULED_OP( args.b, item) );
+                    AddOp( FScheduledOp( item.At, item, 1),
+                           FScheduledOp( args.a, item),
+                           FScheduledOp( args.b, item) );
                 break;
 
             case 1:
             {
-                Ptr<const Mesh> pA = GetMemory().GetMesh( CACHE_ADDRESS(args.a,item) );
-                Ptr<const Mesh> pB = GetMemory().GetMesh( CACHE_ADDRESS(args.b,item) );
+                Ptr<const Mesh> pA = GetMemory().GetMesh( FCacheAddress(args.a,item) );
+                Ptr<const Mesh> pB = GetMemory().GetMesh( FCacheAddress(args.b,item) );
 
                 MeshPtr pResult = MeshSubtract( pA.get(), pB.get() );
 
@@ -1494,15 +1500,15 @@ namespace mu
 
         case OP_TYPE::ME_FORMAT:
         {
-            auto args = pModel->GetPrivate()->m_program.GetOpArgs<OP::MeshFormatArgs>(item.at);
-            switch (item.stage)
+            auto args = pModel->GetPrivate()->m_program.GetOpArgs<OP::MeshFormatArgs>(item.At);
+            switch (item.Stage)
             {
             case 0:
                 if (args.source && args.format)
                 {
-                         AddOp( SCHEDULED_OP( item.at, item, 1),
-                                SCHEDULED_OP( args.source, item),
-                                SCHEDULED_OP( args.format, item));
+                         AddOp( FScheduledOp( item.At, item, 1),
+                                FScheduledOp( args.source, item),
+                                FScheduledOp( args.format, item));
                     }
                     else
                     {
@@ -1512,8 +1518,8 @@ namespace mu
 
             case 1:
             {
-                Ptr<const Mesh> pSource = GetMemory().GetMesh( CACHE_ADDRESS(args.source,item) );
-                Ptr<const Mesh> pFormat = GetMemory().GetMesh( CACHE_ADDRESS(args.format,item) );
+                Ptr<const Mesh> pSource = GetMemory().GetMesh( FCacheAddress(args.source,item) );
+                Ptr<const Mesh> pFormat = GetMemory().GetMesh( FCacheAddress(args.format,item) );
 
                 uint8 flags = args.buffers;
                 MeshPtr pResult;
@@ -1544,7 +1550,7 @@ namespace mu
 
         case OP_TYPE::ME_EXTRACTLAYOUTBLOCK:
         {
-            const uint8* data = pModel->GetPrivate()->m_program.GetOpArgsPointer(item.at);
+            const uint8* data = pModel->GetPrivate()->m_program.GetOpArgsPointer(item.At);
 
             mu::OP::ADDRESS source;
             memcpy( &source, data, sizeof(OP::ADDRESS) );
@@ -1558,13 +1564,13 @@ namespace mu
             memcpy( &blockCount, data, sizeof(uint16) );
             data += sizeof(uint16);
 
-            switch (item.stage)
+            switch (item.Stage)
             {
             case 0:
                 if (source)
                 {
-                        AddOp( SCHEDULED_OP( item.at, item, 1),
-                               SCHEDULED_OP( source, item) );
+                        AddOp( FScheduledOp( item.At, item, 1),
+                               FScheduledOp( source, item) );
                     }
                     else
                     {
@@ -1574,7 +1580,7 @@ namespace mu
 
             case 1:
             {
-                Ptr<const Mesh> pSource = GetMemory().GetMesh( CACHE_ADDRESS(source,item) );
+                Ptr<const Mesh> pSource = GetMemory().GetMesh( FCacheAddress(source,item) );
 
                 // Access with memcpy necessary for unaligned arm issues.
                 uint32 blocks[1024];
@@ -1599,14 +1605,14 @@ namespace mu
 
         case OP_TYPE::ME_EXTRACTFACEGROUP:
         {
-            auto args = pModel->GetPrivate()->m_program.GetOpArgs<OP::MeshExtractFaceGroupArgs>(item.at);
-            switch (item.stage)
+            auto args = pModel->GetPrivate()->m_program.GetOpArgs<OP::MeshExtractFaceGroupArgs>(item.At);
+            switch (item.Stage)
             {
             case 0:
                 if (args.source)
                 {
-                        AddOp( SCHEDULED_OP( item.at, item, 1),
-                               SCHEDULED_OP( args.source, item) );
+                        AddOp( FScheduledOp( item.At, item, 1),
+                               FScheduledOp( args.source, item) );
                     }
                     else
                     {
@@ -1616,7 +1622,7 @@ namespace mu
 
             case 1:
             {
-                Ptr<const Mesh> pSource = GetMemory().GetMesh( CACHE_ADDRESS(args.source,item) );
+                Ptr<const Mesh> pSource = GetMemory().GetMesh( FCacheAddress(args.source,item) );
 
                 MeshPtr pResult;
                 pResult = MeshExtractFaceGroup( pSource.get(),
@@ -1635,14 +1641,14 @@ namespace mu
 
         case OP_TYPE::ME_TRANSFORM:
         {
-            auto args = pModel->GetPrivate()->m_program.GetOpArgs<OP::MeshTransformArgs>(item.at);
-            switch (item.stage)
+            auto args = pModel->GetPrivate()->m_program.GetOpArgs<OP::MeshTransformArgs>(item.At);
+            switch (item.Stage)
             {
             case 0:
                 if (args.source)
                 {
-                        AddOp( SCHEDULED_OP( item.at, item, 1),
-                               SCHEDULED_OP( args.source, item) );
+                        AddOp( FScheduledOp( item.At, item, 1),
+                               FScheduledOp( args.source, item) );
                     }
                     else
                     {
@@ -1652,7 +1658,7 @@ namespace mu
 
             case 1:
             {
-                Ptr<const Mesh> pSource = GetMemory().GetMesh( CACHE_ADDRESS(args.source,item) );
+                Ptr<const Mesh> pSource = GetMemory().GetMesh( FCacheAddress(args.source,item) );
 
                 const mat4f& mat = pModel->GetPrivate()->m_program.
                     m_constantMatrices[args.matrix];
@@ -1672,14 +1678,14 @@ namespace mu
 
         case OP_TYPE::ME_CLIPMORPHPLANE:
         {
-            auto args = pModel->GetPrivate()->m_program.GetOpArgs<OP::MeshClipMorphPlaneArgs>(item.at);
-            switch (item.stage)
+            auto args = pModel->GetPrivate()->m_program.GetOpArgs<OP::MeshClipMorphPlaneArgs>(item.At);
+            switch (item.Stage)
             {
             case 0:
                 if (args.source)
                 {
-                        AddOp( SCHEDULED_OP( item.at, item, 1),
-                               SCHEDULED_OP( args.source, item) );
+                        AddOp( FScheduledOp( item.At, item, 1),
+                               FScheduledOp( args.source, item) );
                     }
                     else
                     {
@@ -1689,7 +1695,7 @@ namespace mu
 
             case 1:
             {
-                Ptr<const Mesh> pSource = GetMemory().GetMesh( CACHE_ADDRESS(args.source,item) );
+                Ptr<const Mesh> pSource = GetMemory().GetMesh( FCacheAddress(args.source,item) );
 
                 MeshPtr pResult;
 
@@ -1742,15 +1748,15 @@ namespace mu
 
         case OP_TYPE::ME_CLIPWITHMESH:
         {
-            auto args = pModel->GetPrivate()->m_program.GetOpArgs<OP::MeshClipWithMeshArgs>(item.at);
-            switch (item.stage)
+            auto args = pModel->GetPrivate()->m_program.GetOpArgs<OP::MeshClipWithMeshArgs>(item.At);
+            switch (item.Stage)
             {
             case 0:
                 if (args.source)
                 {
-                        AddOp( SCHEDULED_OP( item.at, item, 1),
-                               SCHEDULED_OP( args.source, item),
-                               SCHEDULED_OP( args.clipMesh, item) );
+                        AddOp( FScheduledOp( item.At, item, 1),
+                               FScheduledOp( args.source, item),
+                               FScheduledOp( args.clipMesh, item) );
                     }
                     else
                     {
@@ -1760,8 +1766,8 @@ namespace mu
 
             case 1:
             {
-                Ptr<const Mesh> pSource = GetMemory().GetMesh( CACHE_ADDRESS(args.source,item) );
-                Ptr<const Mesh> pClip = GetMemory().GetMesh( CACHE_ADDRESS(args.clipMesh,item) );
+                Ptr<const Mesh> pSource = GetMemory().GetMesh( FCacheAddress(args.source,item) );
+                Ptr<const Mesh> pClip = GetMemory().GetMesh( FCacheAddress(args.clipMesh,item) );
 
                 // Only if both are valid.
                 MeshPtr pResult;
@@ -1786,15 +1792,15 @@ namespace mu
         }
 		case OP_TYPE::ME_CLIPDEFORM:
 		{
-			OP::MeshClipDeformArgs args = pModel->GetPrivate()->m_program.GetOpArgs<OP::MeshClipDeformArgs>(item.at);
-			switch (item.stage)
+			OP::MeshClipDeformArgs args = pModel->GetPrivate()->m_program.GetOpArgs<OP::MeshClipDeformArgs>(item.At);
+			switch (item.Stage)
 			{
 			case 0:
 				if (args.mesh)
 				{
-						AddOp(SCHEDULED_OP(item.at, item, 1),
-							SCHEDULED_OP(args.mesh, item),
-							SCHEDULED_OP(args.clipShape, item));
+						AddOp(FScheduledOp(item.At, item, 1),
+							FScheduledOp(args.mesh, item),
+							FScheduledOp(args.clipShape, item));
 					}
 					else
 					{
@@ -1804,8 +1810,8 @@ namespace mu
 			
 			case 1:
 			{
-				Ptr<const Mesh> BaseMesh = GetMemory().GetMesh(CACHE_ADDRESS(args.mesh, item));
-				Ptr<const Mesh> ClipShape = GetMemory().GetMesh(CACHE_ADDRESS(args.clipShape, item));
+				Ptr<const Mesh> BaseMesh = GetMemory().GetMesh(FCacheAddress(args.mesh, item));
+				Ptr<const Mesh> ClipShape = GetMemory().GetMesh(FCacheAddress(args.clipShape, item));
 				Ptr<Mesh> pResult;
 
 					if ( BaseMesh && ClipShape )
@@ -1830,15 +1836,15 @@ namespace mu
 
         case OP_TYPE::ME_REMAPINDICES:
         {
-            auto args = pModel->GetPrivate()->m_program.GetOpArgs<OP::MeshRemapIndicesArgs>(item.at);
-            switch (item.stage)
+            auto args = pModel->GetPrivate()->m_program.GetOpArgs<OP::MeshRemapIndicesArgs>(item.At);
+            switch (item.Stage)
             {
             case 0:
                 if (args.source)
                 {
-                        AddOp( SCHEDULED_OP( item.at, item, 1),
-                               SCHEDULED_OP( args.source, item),
-                               SCHEDULED_OP( args.reference, item) );
+                        AddOp( FScheduledOp( item.At, item, 1),
+                               FScheduledOp( args.source, item),
+                               FScheduledOp( args.reference, item) );
                     }
                     else
                     {
@@ -1848,8 +1854,8 @@ namespace mu
 
             case 1:
             {
-                Ptr<const Mesh> pSource = GetMemory().GetMesh( CACHE_ADDRESS(args.source,item) );
-                Ptr<const Mesh> pReference = GetMemory().GetMesh( CACHE_ADDRESS(args.reference,item) );
+                Ptr<const Mesh> pSource = GetMemory().GetMesh( FCacheAddress(args.source,item) );
+                Ptr<const Mesh> pReference = GetMemory().GetMesh( FCacheAddress(args.reference,item) );
 
                 // Only if both are valid.
                 MeshPtr pResult;
@@ -1876,15 +1882,15 @@ namespace mu
 
         case OP_TYPE::ME_APPLYPOSE:
         {
-            auto args = pModel->GetPrivate()->m_program.GetOpArgs<OP::MeshApplyPoseArgs>(item.at);
-            switch (item.stage)
+            auto args = pModel->GetPrivate()->m_program.GetOpArgs<OP::MeshApplyPoseArgs>(item.At);
+            switch (item.Stage)
             {
             case 0:
                 if (args.base)
                 {
-                        AddOp( SCHEDULED_OP( item.at, item, 1),
-                               SCHEDULED_OP( args.base, item),
-                               SCHEDULED_OP( args.pose, item) );
+                        AddOp( FScheduledOp( item.At, item, 1),
+                               FScheduledOp( args.base, item),
+                               FScheduledOp( args.pose, item) );
                     }
                     else
                     {
@@ -1894,8 +1900,8 @@ namespace mu
 
             case 1:
             {
-                Ptr<const Mesh> pBase = GetMemory().GetMesh( CACHE_ADDRESS(args.base,item) );
-                Ptr<const Mesh> pPose = GetMemory().GetMesh( CACHE_ADDRESS(args.pose,item) );
+                Ptr<const Mesh> pBase = GetMemory().GetMesh( FCacheAddress(args.base,item) );
+                Ptr<const Mesh> pPose = GetMemory().GetMesh( FCacheAddress(args.pose,item) );
 
                 // Only if both are valid.
                 MeshPtr pResult;
@@ -1922,17 +1928,17 @@ namespace mu
 
 		case OP_TYPE::ME_GEOMETRYOPERATION:
 		{
-			auto args = pModel->GetPrivate()->m_program.GetOpArgs<OP::MeshGeometryOperationArgs>(item.at);
-			switch (item.stage)
+			auto args = pModel->GetPrivate()->m_program.GetOpArgs<OP::MeshGeometryOperationArgs>(item.At);
+			switch (item.Stage)
 			{
 			case 0:
 				if (args.meshA)
 				{
-						AddOp(SCHEDULED_OP(item.at, item, 1),
-							SCHEDULED_OP(args.meshA, item),
-							SCHEDULED_OP(args.meshB, item),
-							SCHEDULED_OP(args.scalarA, item),
-							SCHEDULED_OP(args.scalarB, item));
+						AddOp(FScheduledOp(item.At, item, 1),
+							FScheduledOp(args.meshA, item),
+							FScheduledOp(args.meshB, item),
+							FScheduledOp(args.scalarA, item),
+							FScheduledOp(args.scalarB, item));
 					}
 					else
 					{
@@ -1942,10 +1948,10 @@ namespace mu
 
 			case 1:
 			{
-				Ptr<const Mesh> MeshA = GetMemory().GetMesh(CACHE_ADDRESS(args.meshA, item));
-				Ptr<const Mesh> MeshB = GetMemory().GetMesh(CACHE_ADDRESS(args.meshB, item));
-				float ScalarA = GetMemory().GetScalar(CACHE_ADDRESS(args.scalarA, item));
-				float ScalarB = GetMemory().GetScalar(CACHE_ADDRESS(args.scalarB, item));
+				Ptr<const Mesh> MeshA = GetMemory().GetMesh(FCacheAddress(args.meshA, item));
+				Ptr<const Mesh> MeshB = GetMemory().GetMesh(FCacheAddress(args.meshB, item));
+				float ScalarA = GetMemory().GetScalar(FCacheAddress(args.scalarA, item));
+				float ScalarB = GetMemory().GetScalar(FCacheAddress(args.scalarB, item));
 
 				MeshPtr pResult = MeshGeometryOperation(MeshA.get(),MeshB.get(),ScalarA,ScalarB);
 
@@ -1963,17 +1969,17 @@ namespace mu
 
 		case OP_TYPE::ME_BINDSHAPE:
 		{
-			OP::MeshBindShapeArgs args = pModel->GetPrivate()->m_program.GetOpArgs<OP::MeshBindShapeArgs>(item.at);
-			const uint8* data = pModel->GetPrivate()->m_program.GetOpArgsPointer(item.at);
+			OP::MeshBindShapeArgs args = pModel->GetPrivate()->m_program.GetOpArgs<OP::MeshBindShapeArgs>(item.At);
+			const uint8* data = pModel->GetPrivate()->m_program.GetOpArgsPointer(item.At);
 
-			switch (item.stage)
+			switch (item.Stage)
 			{
 			case 0:
 				if (args.mesh)
 				{
-					AddOp(SCHEDULED_OP(item.at, item, 1),
-							SCHEDULED_OP(args.mesh, item),
-							SCHEDULED_OP(args.shape, item));
+					AddOp(FScheduledOp(item.At, item, 1),
+							FScheduledOp(args.mesh, item),
+							FScheduledOp(args.shape, item));
 				}
 				else
 				{
@@ -1983,8 +1989,8 @@ namespace mu
 
 			case 1:
 			{
-				Ptr<const Mesh> BaseMesh = GetMemory().GetMesh(CACHE_ADDRESS(args.mesh, item));
-				Ptr<const Mesh> Shape = GetMemory().GetMesh(CACHE_ADDRESS(args.shape, item));
+				Ptr<const Mesh> BaseMesh = GetMemory().GetMesh(FCacheAddress(args.mesh, item));
+				Ptr<const Mesh> Shape = GetMemory().GetMesh(FCacheAddress(args.shape, item));
 				
 				EShapeBindingMethod BindingMethod = static_cast<EShapeBindingMethod>(args.bindingMethod); 
 
@@ -2043,15 +2049,15 @@ namespace mu
 
 		case OP_TYPE::ME_APPLYSHAPE:
 		{
-			auto args = pModel->GetPrivate()->m_program.GetOpArgs<OP::MeshApplyShapeArgs>(item.at);
-			switch (item.stage)
+			auto args = pModel->GetPrivate()->m_program.GetOpArgs<OP::MeshApplyShapeArgs>(item.At);
+			switch (item.Stage)
 			{
 			case 0:
 				if (args.mesh)
 				{
-						AddOp(SCHEDULED_OP(item.at, item, 1),
-							SCHEDULED_OP(args.mesh, item),
-							SCHEDULED_OP(args.shape, item));
+						AddOp(FScheduledOp(item.At, item, 1),
+							FScheduledOp(args.mesh, item),
+							FScheduledOp(args.shape, item));
 				}
 				else
 				{
@@ -2061,8 +2067,8 @@ namespace mu
 
 			case 1:
 			{
-				Ptr<const Mesh> BaseMesh = GetMemory().GetMesh(CACHE_ADDRESS(args.mesh, item));
-				Ptr<const Mesh> Shape = GetMemory().GetMesh(CACHE_ADDRESS(args.shape, item));
+				Ptr<const Mesh> BaseMesh = GetMemory().GetMesh(FCacheAddress(args.mesh, item));
+				Ptr<const Mesh> Shape = GetMemory().GetMesh(FCacheAddress(args.shape, item));
 
 
 				Ptr<Mesh> pResult = MeshApplyShape(BaseMesh.get(), Shape.get(), static_cast<EMeshBindShapeFlags>(args.flags));
@@ -2080,16 +2086,16 @@ namespace mu
 
 		case OP_TYPE::ME_MORPHRESHAPE:
 		{
-			OP::MeshMorphReshapeArgs Args = pModel->GetPrivate()->m_program.GetOpArgs<OP::MeshMorphReshapeArgs>(item.at);
-			switch(item.stage)
+			OP::MeshMorphReshapeArgs Args = pModel->GetPrivate()->m_program.GetOpArgs<OP::MeshMorphReshapeArgs>(item.At);
+			switch(item.Stage)
 			{
 			case 0:
 			{
 				if (Args.Morph)
 				{
-					AddOp(SCHEDULED_OP(item.at, item, 1), 
-						  SCHEDULED_OP(Args.Morph, item), 
-						  SCHEDULED_OP(Args.Reshape, item));
+					AddOp(FScheduledOp(item.At, item, 1), 
+						  FScheduledOp(Args.Morph, item), 
+						  FScheduledOp(Args.Reshape, item));
 				}
 				else 
 				{
@@ -2099,8 +2105,8 @@ namespace mu
 			}
 			case 1:
 			{
-				Ptr<const Mesh> MorphedMesh = GetMemory().GetMesh(CACHE_ADDRESS(Args.Morph, item));
-				Ptr<const Mesh> ReshapeMesh = GetMemory().GetMesh(CACHE_ADDRESS(Args.Reshape, item));
+				Ptr<const Mesh> MorphedMesh = GetMemory().GetMesh(FCacheAddress(Args.Morph, item));
+				Ptr<const Mesh> ReshapeMesh = GetMemory().GetMesh(FCacheAddress(Args.Reshape, item));
 
 				if (ReshapeMesh)
 				{
@@ -2135,15 +2141,15 @@ namespace mu
 
         case OP_TYPE::ME_SETSKELETON:
         {
-            auto args = pModel->GetPrivate()->m_program.GetOpArgs<OP::MeshSetSkeletonArgs>(item.at);
-            switch (item.stage)
+            auto args = pModel->GetPrivate()->m_program.GetOpArgs<OP::MeshSetSkeletonArgs>(item.At);
+            switch (item.Stage)
             {
             case 0:
                 if (args.source)
                 {
-                        AddOp( SCHEDULED_OP( item.at, item, 1),
-                               SCHEDULED_OP( args.source, item),
-                               SCHEDULED_OP( args.skeleton, item) );
+                        AddOp( FScheduledOp( item.At, item, 1),
+                               FScheduledOp( args.source, item),
+                               FScheduledOp( args.skeleton, item) );
                     }
                     else
                     {
@@ -2153,8 +2159,8 @@ namespace mu
 
             case 1:
             {
-                Ptr<const Mesh> pSource = GetMemory().GetMesh( CACHE_ADDRESS(args.source,item) );
-                Ptr<const Mesh> pSkeleton = GetMemory().GetMesh( CACHE_ADDRESS(args.skeleton,item) );
+                Ptr<const Mesh> pSource = GetMemory().GetMesh( FCacheAddress(args.source,item) );
+                Ptr<const Mesh> pSkeleton = GetMemory().GetMesh( FCacheAddress(args.skeleton,item) );
 
                 // Only if both are valid.
                 MeshPtr pResult;
@@ -2202,12 +2208,12 @@ namespace mu
         {
             // Decode op
             // TODO: Partial decode for each stage
-            const uint8* data = pModel->GetPrivate()->m_program.GetOpArgsPointer(item.at);
+            const uint8* data = pModel->GetPrivate()->m_program.GetOpArgsPointer(item.At);
 
             OP::ADDRESS source;
             FMemory::Memcpy(&source,data,sizeof(OP::ADDRESS)); data += sizeof(OP::ADDRESS);
 
-            TArray<SCHEDULED_OP> conditions;
+            TArray<FScheduledOp> conditions;
 			TArray<OP::ADDRESS> masks;
 
             uint16 removes;
@@ -2226,13 +2232,13 @@ namespace mu
 
 
             // Schedule next stages
-            switch (item.stage)
+            switch (item.Stage)
             {
             case 0:
                 if (source)
                 {
                     // Request the conditions
-                    AddOp( SCHEDULED_OP( item.at, item, 1), conditions );
+                    AddOp( FScheduledOp( item.At, item, 1), conditions );
                 }
                 else
                 {
@@ -2244,15 +2250,15 @@ namespace mu
             {
                 // Request the source and the necessary masks
                 // \todo: store condition values in heap?
-                TArray<SCHEDULED_OP> deps;
+                TArray<FScheduledOp> deps;
                 deps.Emplace( source, item );
                 for( size_t r=0; source && r<conditions.Num(); ++r )
                 {
                     // If there is no expression, we'll assume true.
                     bool value = true;
-                    if (conditions[r].at)
+                    if (conditions[r].At)
                     {
-                        value = GetMemory().GetBool( CACHE_ADDRESS(conditions[r].at, item) );
+                        value = GetMemory().GetBool( FCacheAddress(conditions[r].At, item) );
                     }
 
                     if (value)
@@ -2263,7 +2269,7 @@ namespace mu
 
                 if (source)
                 {
-                        AddOp( SCHEDULED_OP( item.at, item, 2), deps );
+                        AddOp( FScheduledOp( item.At, item, 2), deps );
                     }
                 break;
             }
@@ -2271,7 +2277,7 @@ namespace mu
             case 2:
             {
                 // \todo: single remove operation with all masks?
-                Ptr<const Mesh> pSource = GetMemory().GetMesh( CACHE_ADDRESS(source,item) );
+                Ptr<const Mesh> pSource = GetMemory().GetMesh( FCacheAddress(source,item) );
 
                 MeshPtrConst pResult = pSource;
 
@@ -2279,14 +2285,14 @@ namespace mu
                 {
                     // If there is no expression, we'll assume true.
                     bool value = true;
-                    if (conditions[r].at)
+                    if (conditions[r].At)
                     {
-                        value = GetMemory().GetBool( CACHE_ADDRESS(conditions[r].at, item) );
+                        value = GetMemory().GetBool( FCacheAddress(conditions[r].At, item) );
                     }
 
                     if (value)
                     {
-                        Ptr<const Mesh> pMask = GetMemory().GetMesh( CACHE_ADDRESS(masks[r],item) );
+                        Ptr<const Mesh> pMask = GetMemory().GetMesh( FCacheAddress(masks[r],item) );
                         if (pMask)
                         {
                             pResult = MeshRemoveMask(pResult.get(), pMask.get());
@@ -2307,15 +2313,15 @@ namespace mu
 
         case OP_TYPE::ME_PROJECT:
         {
-            auto args = pModel->GetPrivate()->m_program.GetOpArgs<OP::MeshProjectArgs>(item.at);
-            switch (item.stage)
+            auto args = pModel->GetPrivate()->m_program.GetOpArgs<OP::MeshProjectArgs>(item.At);
+            switch (item.Stage)
             {
             case 0:
                 if (args.mesh)
                 {
-                        AddOp( SCHEDULED_OP( item.at, item, 1),
-                               SCHEDULED_OP( args.mesh, item),
-                               SCHEDULED_OP( args.projector, item));
+                        AddOp( FScheduledOp( item.At, item, 1),
+                               FScheduledOp( args.mesh, item),
+                               FScheduledOp( args.projector, item));
                 }
                 else
                 {
@@ -2325,8 +2331,8 @@ namespace mu
 
             case 1:
             {
-                Ptr<const Mesh> pMesh = GetMemory().GetMesh( CACHE_ADDRESS(args.mesh,item) );
-                Ptr<const Projector> pProjector = GetMemory().GetProjector( CACHE_ADDRESS(args.projector,item) );
+                Ptr<const Mesh> pMesh = GetMemory().GetMesh( FCacheAddress(args.mesh,item) );
+                Ptr<const Projector> pProjector = GetMemory().GetProjector( FCacheAddress(args.projector,item) );
 
                 // Only if both are valid.
                 MeshPtr pResult;
@@ -2369,20 +2375,20 @@ namespace mu
     }
 
     //---------------------------------------------------------------------------------------------
-    void CodeRunner::RunCode_Image( SCHEDULED_OP& item,
+    void CodeRunner::RunCode_Image( FScheduledOp& item,
                                     const Parameters* pParams,
                                     const Model* pModel
                                     )
     {
 		MUTABLE_CPUPROFILER_SCOPE(RunCode_Image);
 
-		OP_TYPE type = pModel->GetPrivate()->m_program.GetOpType(item.at);
+		OP_TYPE type = pModel->GetPrivate()->m_program.GetOpType(item.At);
 		switch (type)
         {
 
         case OP_TYPE::IM_PARAMETER:
         {
-			OP::ParameterArgs args = pModel->GetPrivate()->m_program.GetOpArgs<OP::ParameterArgs>(item.at);
+			OP::ParameterArgs args = pModel->GetPrivate()->m_program.GetOpArgs<OP::ParameterArgs>(item.At);
 			Ptr<RangeIndex> Index = BuildCurrentOpRangeIndex(item, pParams, pModel, args.variable);
 			EXTERNAL_IMAGE_ID id = pParams->GetImageValue(args.variable, Index);
 			ImagePtr pResult = LoadExternalImage(id);
@@ -2392,14 +2398,14 @@ namespace mu
 
         case OP_TYPE::IM_LAYERCOLOUR:
         {
-			OP::ImageLayerColourArgs args = pModel->GetPrivate()->m_program.GetOpArgs<OP::ImageLayerColourArgs>(item.at);
-            switch (item.stage)
+			OP::ImageLayerColourArgs args = pModel->GetPrivate()->m_program.GetOpArgs<OP::ImageLayerColourArgs>(item.At);
+            switch (item.Stage)
             {
             case 0:
-                    AddOp( SCHEDULED_OP( item.at, item, 1),
-                           SCHEDULED_OP( args.base, item),
-                           SCHEDULED_OP::FromOpAndOptions( args.colour, item, 0),
-                           SCHEDULED_OP( args.mask, item) );
+                    AddOp( FScheduledOp( item.At, item, 1),
+                           FScheduledOp( args.base, item),
+                           FScheduledOp::FromOpAndOptions( args.colour, item, 0),
+                           FScheduledOp( args.mask, item) );
                 break;
 
             case 1:
@@ -2416,14 +2422,14 @@ namespace mu
 
         case OP_TYPE::IM_LAYER:
         {
-			OP::ImageLayerArgs args = pModel->GetPrivate()->m_program.GetOpArgs<OP::ImageLayerArgs>(item.at);
-            switch (item.stage)
+			OP::ImageLayerArgs args = pModel->GetPrivate()->m_program.GetOpArgs<OP::ImageLayerArgs>(item.At);
+            switch (item.Stage)
             {
             case 0:
-                    AddOp( SCHEDULED_OP( item.at, item, 1),
-                           SCHEDULED_OP( args.base, item),
-                           SCHEDULED_OP( args.blended, item),
-                           SCHEDULED_OP( args.mask, item) );
+                    AddOp( FScheduledOp( item.At, item, 1),
+                           FScheduledOp( args.base, item),
+                           FScheduledOp( args.blended, item),
+                           FScheduledOp( args.mask, item) );
                 break;
 
             case 1:
@@ -2440,12 +2446,13 @@ namespace mu
 
         case OP_TYPE::IM_MULTILAYER:
         {
-			OP::ImageMultiLayerArgs args = pModel->GetPrivate()->m_program.GetOpArgs<OP::ImageMultiLayerArgs>(item.at);
-            switch (item.stage)
+			OP::ImageMultiLayerArgs args = pModel->GetPrivate()->m_program.GetOpArgs<OP::ImageMultiLayerArgs>(item.At);
+            switch (item.Stage)
             {
             case 0:
-                AddOp( SCHEDULED_OP( item.at, item, 1),
-                       SCHEDULED_OP( args.rangeSize, item ) );
+                AddOp( FScheduledOp( item.At, item, 1),
+                       FScheduledOp( args.rangeSize, item ),
+					   FScheduledOp(args.base, item));
                 break;
 
             case 1:
@@ -2454,7 +2461,7 @@ namespace mu
                 int32 Iterations = 0;
                 if (args.rangeSize)
                 {
-                    CACHE_ADDRESS RangeAddress(args.rangeSize,item);
+                    FCacheAddress RangeAddress(args.rangeSize,item);
 
                     // We support both integers and scalars here, which is not common.
                     // \todo: review if this is necessary or we can enforce it at compile time.
@@ -2469,97 +2476,100 @@ namespace mu
                     }
                 }
 
-                // \todo Check that we are not overwriting the index (it shouldn't be set when we reach here)
-                SCHEDULED_OP ItemCopy = item;
-                ExecutionIndex Index = GetMemory().GetRageIndex( item.executionIndex );
+				Ptr<const Image> Base = GetMemory().GetImage(FCacheAddress(args.base, item));
 
-				TArray<SCHEDULED_OP, TInlineAllocator<48>> Deps;
-				Deps.SetNumUninitialized(Iterations * 2 + 1);
-                int32 d=0;
-				Deps[d++] = SCHEDULED_OP( args.base, item );
-                for (int i=0; i< Iterations; ++i)
-                {
-					Index.SetFromModelRangeIndex( args.rangeId, i );
+				if (Iterations <= 0)
+				{
+					// There are no layers: return the base
+					GetMemory().SetImage(item, Base);
+				}
+				else
+				{
+					// Store the base
+					Ptr<Image> New = mu::CloneOrTakeOver(Base.get());
 
-					ItemCopy.executionIndex = GetMemory().GetRageIndexIndex(Index);
+					// This shouldn't happen in optimised models, but it could happen in editors, etc.
+					// \todo: raise a performance warning?
+					EImageFormat BaseFormat = GetUncompressedFormat(Base->GetFormat());
+					if (Base->GetFormat() != BaseFormat)
+					{
+						Base = ImagePixelFormat(m_pSettings->GetPrivate()->m_imageCompressionQuality, Base.get(), BaseFormat);
+					}
 
-					Deps[d++] = SCHEDULED_OP( args.blended, ItemCopy);
-					Deps[d++] = SCHEDULED_OP( args.mask, ItemCopy);
-                }
-                AddOp( SCHEDULED_OP( item.at, item, 2, Iterations), Deps);
+					FScheduledOpData Data;
+					Data.Resource = New;
+					Data.MultiLayer.Iterations = Iterations;
+					Data.MultiLayer.OriginalBaseFormat = Base->GetFormat();
+					Data.MultiLayer.bBlendOnlyOneMip = false;
+					int32 DataPos = m_heapData.Add(Data);
+
+					// Request the first layer
+					int32 CurrentIteration = 0;
+					FScheduledOp ItemCopy = item;
+					ExecutionIndex Index = GetMemory().GetRageIndex(item.ExecutionIndex);
+					Index.SetFromModelRangeIndex(args.rangeId, CurrentIteration);
+					ItemCopy.ExecutionIndex = GetMemory().GetRageIndexIndex(Index);
+					AddOp(FScheduledOp(item.At, item, 2, DataPos), FScheduledOp(args.base, item), FScheduledOp(args.blended, ItemCopy), FScheduledOp(args.mask, ItemCopy));
+				}
 
                 break;
             }
 
-            case 2:
+            default:
             {
-                MUTABLE_CPUPROFILER_SCOPE(ImageMultiLayer);
-
-				int32 Iterations = item.customState;
+                //MUTABLE_CPUPROFILER_SCOPE(ImageMultiLayer);
 
-                Ptr<const Image> pBase = GetMemory().GetImage( CACHE_ADDRESS(args.base,item) );
-
-                if (Iterations ==0)
-                {
-					// Early out
-					GetMemory().SetImage(item, pBase);
-					break;
-                }
+				FScheduledOpData& Data = m_heapData[item.CustomState];
 
-				Ptr<const Image> pResult;
+				int32 Iterations = Data.MultiLayer.Iterations;
+				int32 CurrentIteration = item.Stage - 2;
+				check(CurrentIteration >= 0 && CurrentIteration < 120);
 
-                // This shouldn't happen in optimised models, but it could happen in editors, etc.
-                // \todo: raise a performance warning?
-                EImageFormat baseFormat = GetUncompressedFormat(pBase->GetFormat());
-                if ( pBase->GetFormat()!=baseFormat )
-                {
-                    pBase = ImagePixelFormat( m_pSettings->GetPrivate()->m_imageCompressionQuality,
-                                              pBase.get(), baseFormat );
-                }
+				TRACE_CPUPROFILER_EVENT_SCOPE_TEXT(*FString::Printf(TEXT("Layer %d of %d"), CurrentIteration, Iterations));
 
-				Ptr<Image> pNew = mu::CloneOrTakeOver(pBase.get());
-				pResult = pNew;
+				// Process the current layer
 
-                SCHEDULED_OP itemCopy = item;
-                ExecutionIndex index = GetMemory().GetRageIndex( item.executionIndex );
+				Ptr<Image> Base = static_cast<Image*>(Data.Resource.get());
+ 
+                FScheduledOp itemCopy = item;
+                ExecutionIndex index = GetMemory().GetRageIndex( item.ExecutionIndex );
 				
-				// This becomes true if we need to update the mips of the resulting image
-				// This could happen in the base image has mips, but one of the blended one doesn't.
-				bool bBlendOnlyOneMip = false;
-
-                for (int i=0; i< Iterations; ++i)
                 {
-                    index.SetFromModelRangeIndex( args.rangeId, i );
-                    itemCopy.executionIndex = GetMemory().GetRageIndexIndex(index);
-					itemCopy.customState = 0;
+                    index.SetFromModelRangeIndex( args.rangeId, CurrentIteration);
+                    itemCopy.ExecutionIndex = GetMemory().GetRageIndexIndex(index);
+					itemCopy.CustomState = 0;
 
-                    Ptr<const Image> pBlended = GetMemory().GetImage( CACHE_ADDRESS(args.blended,itemCopy) );
+                    Ptr<const Image> pBlended = GetMemory().GetImage( FCacheAddress(args.blended,itemCopy) );
 
                     // This shouldn't happen in optimised models, but it could happen in editors, etc.
                     // \todo: raise a performance warning?
-                    if (pBlended && pBlended->GetFormat()!=baseFormat )
+                    if (pBlended && pBlended->GetFormat()!=Base->GetFormat() )
                     {
 						MUTABLE_CPUPROFILER_SCOPE(ImageResize_BlendedReformat);
-						pBlended = ImagePixelFormat( m_pSettings->GetPrivate()->m_imageCompressionQuality, pBlended.get(), baseFormat );
+						pBlended = ImagePixelFormat( m_pSettings->GetPrivate()->m_imageCompressionQuality, pBlended.get(), Base->GetFormat());
                     }
 
 					// TODO: This shouldn't happen, but be defensive.
-					FImageSize ResultSize = pBase->GetSize();
+					FImageSize ResultSize = Base->GetSize();
 					if (pBlended && pBlended->GetSize() != ResultSize)
 					{
 						MUTABLE_CPUPROFILER_SCOPE(ImageResize_BlendedFixForMultilayer);
 						pBlended = ImageResizeLinear(0, pBlended.get(), ResultSize);
 					}
 
-					if (pBlended->GetLODCount() < pBase->GetLODCount())
+					if (pBlended->GetLODCount() < Base->GetLODCount())
 					{
-						bBlendOnlyOneMip = true;
+						Data.MultiLayer.bBlendOnlyOneMip = true;
 					}
 
 					bool bApplyColorBlendToAlpha = false;
 
 					bool bDone = false;
 
+					// This becomes true if we need to update the mips of the resulting image
+					// This could happen in the base image has mips, but one of the blended one doesn't.
+					bool bBlendOnlyOneMip = Data.MultiLayer.bBlendOnlyOneMip;
+
 					if (!args.mask && args.bUseMaskFromBlended
 						&&
 						args.blendType == uint8(EBlendType::BT_BLEND)
@@ -2569,17 +2579,12 @@ namespace mu
 						// This is a frequent critical-path case because of multilayer projectors.
 						bDone = true;
 						
-						if (pBlended->GetLODCount() < pBase->GetLODCount())
-						{
-							bBlendOnlyOneMip = true;
-						}
-
-						BufferLayerComposite<BlendChannelMasked, LightenChannel, false>(pNew.get(), pBlended.get(), bBlendOnlyOneMip);
+						BufferLayerComposite<BlendChannelMasked, LightenChannel, false>(Base.get(), pBlended.get(), bBlendOnlyOneMip);
 					}
 
                     if (!bDone && args.mask)
                     {
-                        Ptr<const Image> pMask = GetMemory().GetImage( CACHE_ADDRESS(args.mask,itemCopy) );
+                        Ptr<const Image> pMask = GetMemory().GetImage( FCacheAddress(args.mask,itemCopy) );
 
 						// TODO: This shouldn't happen, but be defensive.
 						if (pMask && pMask->GetSize() != ResultSize)
@@ -2588,54 +2593,36 @@ namespace mu
 							pMask = ImageResizeLinear(0, pMask.get(), ResultSize);
 						}
 
-                        if (pMask->GetLODCount()<pBase->GetLODCount() || pBlended->GetLODCount() < pBase->GetLODCount())
-                        {
-							bBlendOnlyOneMip = true;
-						}
-
                         switch (EBlendType(args.blendType))
                         {
-						case EBlendType::BT_NORMAL_COMBINE: 
-							// \todo: optimise with ping pong buffers or blend-in-place
-							ImageNormalCombine(pNew.get(), pBase.get(), pMask.get(), pBlended.get(), bBlendOnlyOneMip);
-							pResult = pNew;
-							pBase = pNew;
-							pNew = new Image(pBase->GetSizeX(), pBase->GetSizeY(), pBase->GetLODCount(), pBase->GetFormat());
-							break;
-                        case EBlendType::BT_SOFTLIGHT: BufferLayer<SoftLightChannelMasked, SoftLightChannel, false>( pNew->GetData(), pNew.get(), pMask.get(), pBlended.get(), bApplyColorBlendToAlpha, bBlendOnlyOneMip); break;
-                        case EBlendType::BT_HARDLIGHT: BufferLayer<HardLightChannelMasked, HardLightChannel, false>(pNew->GetData(), pNew.get(), pMask.get(), pBlended.get(), bApplyColorBlendToAlpha, bBlendOnlyOneMip); break;
-                        case EBlendType::BT_BURN: BufferLayer<BurnChannelMasked, BurnChannel, false>(pNew->GetData(), pNew.get(), pMask.get(), pBlended.get(), bApplyColorBlendToAlpha, bBlendOnlyOneMip); break;
-                        case EBlendType::BT_DODGE: BufferLayer<DodgeChannelMasked, DodgeChannel, false>(pNew->GetData(), pNew.get(), pMask.get(), pBlended.get(), bApplyColorBlendToAlpha, bBlendOnlyOneMip); break;
-                        case EBlendType::BT_SCREEN: BufferLayer<ScreenChannelMasked, ScreenChannel, false>(pNew->GetData(), pNew.get(), pMask.get(), pBlended.get(), bApplyColorBlendToAlpha, bBlendOnlyOneMip); break;
-                        case EBlendType::BT_OVERLAY: BufferLayer<OverlayChannelMasked, OverlayChannel, false>(pNew->GetData(), pNew.get(), pMask.get(), pBlended.get(), bApplyColorBlendToAlpha, bBlendOnlyOneMip); break;
-                        case EBlendType::BT_LIGHTEN: BufferLayer<LightenChannelMasked, LightenChannel, false>(pNew->GetData(), pNew.get(), pMask.get(), pBlended.get(), bApplyColorBlendToAlpha, bBlendOnlyOneMip); break;
-                        case EBlendType::BT_MULTIPLY: BufferLayer<MultiplyChannelMasked, MultiplyChannel, false>(pNew->GetData(), pNew.get(), pMask.get(), pBlended.get(), bApplyColorBlendToAlpha, bBlendOnlyOneMip); break;
-                        case EBlendType::BT_BLEND: BufferLayer<BlendChannelMasked, BlendChannel, false>(pNew->GetData(), pNew.get(), pMask.get(), pBlended.get(), bApplyColorBlendToAlpha, bBlendOnlyOneMip); break;
+						case EBlendType::BT_NORMAL_COMBINE: check(false); break;
+                        case EBlendType::BT_SOFTLIGHT: BufferLayer<SoftLightChannelMasked, SoftLightChannel, false>( Base->GetData(), Base.get(), pMask.get(), pBlended.get(), bApplyColorBlendToAlpha, bBlendOnlyOneMip); break;
+                        case EBlendType::BT_HARDLIGHT: BufferLayer<HardLightChannelMasked, HardLightChannel, false>(Base->GetData(), Base.get(), pMask.get(), pBlended.get(), bApplyColorBlendToAlpha, bBlendOnlyOneMip); break;
+                        case EBlendType::BT_BURN: BufferLayer<BurnChannelMasked, BurnChannel, false>(Base->GetData(), Base.get(), pMask.get(), pBlended.get(), bApplyColorBlendToAlpha, bBlendOnlyOneMip); break;
+                        case EBlendType::BT_DODGE: BufferLayer<DodgeChannelMasked, DodgeChannel, false>(Base->GetData(), Base.get(), pMask.get(), pBlended.get(), bApplyColorBlendToAlpha, bBlendOnlyOneMip); break;
+                        case EBlendType::BT_SCREEN: BufferLayer<ScreenChannelMasked, ScreenChannel, false>(Base->GetData(), Base.get(), pMask.get(), pBlended.get(), bApplyColorBlendToAlpha, bBlendOnlyOneMip); break;
+                        case EBlendType::BT_OVERLAY: BufferLayer<OverlayChannelMasked, OverlayChannel, false>(Base->GetData(), Base.get(), pMask.get(), pBlended.get(), bApplyColorBlendToAlpha, bBlendOnlyOneMip); break;
+                        case EBlendType::BT_LIGHTEN: BufferLayer<LightenChannelMasked, LightenChannel, false>(Base->GetData(), Base.get(), pMask.get(), pBlended.get(), bApplyColorBlendToAlpha, bBlendOnlyOneMip); break;
+                        case EBlendType::BT_MULTIPLY: BufferLayer<MultiplyChannelMasked, MultiplyChannel, false>(Base->GetData(), Base.get(), pMask.get(), pBlended.get(), bApplyColorBlendToAlpha, bBlendOnlyOneMip); break;
+                        case EBlendType::BT_BLEND: BufferLayer<BlendChannelMasked, BlendChannel, false>(Base->GetData(), Base.get(), pMask.get(), pBlended.get(), bApplyColorBlendToAlpha, bBlendOnlyOneMip); break;
                         default: check(false);
                         }
 
                     }
 					else if (!bDone && args.bUseMaskFromBlended)
 					{
-						if (pBlended->GetLODCount() < pBase->GetLODCount())
-						{
-							bBlendOnlyOneMip = true;
-						}
-
 						switch (EBlendType(args.blendType))
 						{
-						case EBlendType::BT_NORMAL_COMBINE:
-							check(false);
-							break;
-						case EBlendType::BT_SOFTLIGHT: BufferLayerEmbeddedMask<SoftLightChannelMasked, SoftLightChannel, false>(pNew->GetData(), pNew.get(), pBlended.get(), bApplyColorBlendToAlpha, bBlendOnlyOneMip); break;
-						case EBlendType::BT_HARDLIGHT: BufferLayerEmbeddedMask<HardLightChannelMasked, HardLightChannel, false>(pNew->GetData(), pNew.get(), pBlended.get(), bApplyColorBlendToAlpha, bBlendOnlyOneMip); break;
-						case EBlendType::BT_BURN: BufferLayerEmbeddedMask<BurnChannelMasked, BurnChannel, false>(pNew->GetData(), pNew.get(), pBlended.get(), bApplyColorBlendToAlpha, bBlendOnlyOneMip); break;
-						case EBlendType::BT_DODGE: BufferLayerEmbeddedMask<DodgeChannelMasked, DodgeChannel, false>(pNew->GetData(), pNew.get(), pBlended.get(), bApplyColorBlendToAlpha, bBlendOnlyOneMip); break;
-						case EBlendType::BT_SCREEN: BufferLayerEmbeddedMask<ScreenChannelMasked, ScreenChannel, false>(pNew->GetData(), pNew.get(), pBlended.get(), bApplyColorBlendToAlpha, bBlendOnlyOneMip); break;
-						case EBlendType::BT_OVERLAY: BufferLayerEmbeddedMask<OverlayChannelMasked, OverlayChannel, false>(pNew->GetData(), pNew.get(), pBlended.get(), bApplyColorBlendToAlpha, bBlendOnlyOneMip); break;
-						case EBlendType::BT_LIGHTEN: BufferLayerEmbeddedMask<LightenChannelMasked, LightenChannel, false>(pNew->GetData(), pNew.get(), pBlended.get(), bApplyColorBlendToAlpha, bBlendOnlyOneMip); break;
-						case EBlendType::BT_MULTIPLY: BufferLayerEmbeddedMask<MultiplyChannelMasked, MultiplyChannel, false>(pNew->GetData(), pNew.get(), pBlended.get(), bApplyColorBlendToAlpha, bBlendOnlyOneMip); break;
-						case EBlendType::BT_BLEND: BufferLayerEmbeddedMask<BlendChannelMasked, BlendChannel, false>(pNew->GetData(), pNew.get(), pBlended.get(), bApplyColorBlendToAlpha, bBlendOnlyOneMip); break;
+						case EBlendType::BT_NORMAL_COMBINE: check(false); break;
+						case EBlendType::BT_SOFTLIGHT: BufferLayerEmbeddedMask<SoftLightChannelMasked, SoftLightChannel, false>(Base->GetData(), Base.get(), pBlended.get(), bApplyColorBlendToAlpha, bBlendOnlyOneMip); break;
+						case EBlendType::BT_HARDLIGHT: BufferLayerEmbeddedMask<HardLightChannelMasked, HardLightChannel, false>(Base->GetData(), Base.get(), pBlended.get(), bApplyColorBlendToAlpha, bBlendOnlyOneMip); break;
+						case EBlendType::BT_BURN: BufferLayerEmbeddedMask<BurnChannelMasked, BurnChannel, false>(Base->GetData(), Base.get(), pBlended.get(), bApplyColorBlendToAlpha, bBlendOnlyOneMip); break;
+						case EBlendType::BT_DODGE: BufferLayerEmbeddedMask<DodgeChannelMasked, DodgeChannel, false>(Base->GetData(), Base.get(), pBlended.get(), bApplyColorBlendToAlpha, bBlendOnlyOneMip); break;
+						case EBlendType::BT_SCREEN: BufferLayerEmbeddedMask<ScreenChannelMasked, ScreenChannel, false>(Base->GetData(), Base.get(), pBlended.get(), bApplyColorBlendToAlpha, bBlendOnlyOneMip); break;
+						case EBlendType::BT_OVERLAY: BufferLayerEmbeddedMask<OverlayChannelMasked, OverlayChannel, false>(Base->GetData(), Base.get(), pBlended.get(), bApplyColorBlendToAlpha, bBlendOnlyOneMip); break;
+						case EBlendType::BT_LIGHTEN: BufferLayerEmbeddedMask<LightenChannelMasked, LightenChannel, false>(Base->GetData(), Base.get(), pBlended.get(), bApplyColorBlendToAlpha, bBlendOnlyOneMip); break;
+						case EBlendType::BT_MULTIPLY: BufferLayerEmbeddedMask<MultiplyChannelMasked, MultiplyChannel, false>(Base->GetData(), Base.get(), pBlended.get(), bApplyColorBlendToAlpha, bBlendOnlyOneMip); break;
+						case EBlendType::BT_BLEND: BufferLayerEmbeddedMask<BlendChannelMasked, BlendChannel, false>(Base->GetData(), Base.get(), pBlended.get(), bApplyColorBlendToAlpha, bBlendOnlyOneMip); break;
 						default: check(false);
 						}
 					}
@@ -2643,22 +2630,16 @@ namespace mu
                     {
                         switch (EBlendType(args.blendType))
                         {
-						case EBlendType::BT_NORMAL_COMBINE: 
-							// \todo: optimise with ping pong buffers or blend-in-place
-							ImageNormalCombine(pNew.get(), pBase.get(), pBlended.get(), bBlendOnlyOneMip);
-							pResult = pNew;
-							pBase = pNew;
-							pNew = new Image(pBase->GetSizeX(), pBase->GetSizeY(), pBase->GetLODCount(), pBase->GetFormat());
-							break;
-                        case EBlendType::BT_SOFTLIGHT: BufferLayer<SoftLightChannel, false>(pNew.get(), pNew.get(), pBlended.get(), bApplyColorBlendToAlpha, bBlendOnlyOneMip); break;
-                        case EBlendType::BT_HARDLIGHT: BufferLayer<HardLightChannel, false>(pNew.get(), pNew.get(), pBlended.get(), bApplyColorBlendToAlpha, bBlendOnlyOneMip); break;
-                        case EBlendType::BT_BURN: BufferLayer<BurnChannel, false>(pNew.get(), pNew.get(), pBlended.get(), bApplyColorBlendToAlpha, bBlendOnlyOneMip); break;
-                        case EBlendType::BT_DODGE: BufferLayer<DodgeChannel, false>(pNew.get(), pNew.get(), pBlended.get(), bApplyColorBlendToAlpha, bBlendOnlyOneMip); break;
-                        case EBlendType::BT_SCREEN: BufferLayer<ScreenChannel, false>(pNew.get(), pNew.get(),  pBlended.get(), bApplyColorBlendToAlpha, bBlendOnlyOneMip); break;
-                        case EBlendType::BT_OVERLAY: BufferLayer<OverlayChannel, false>(pNew.get(), pNew.get(), pBlended.get(), bApplyColorBlendToAlpha, bBlendOnlyOneMip); break;
-                        case EBlendType::BT_LIGHTEN: BufferLayer<LightenChannel, false>(pNew.get(), pNew.get(), pBlended.get(), bApplyColorBlendToAlpha, bBlendOnlyOneMip); break;
-                        case EBlendType::BT_MULTIPLY: BufferLayer<MultiplyChannel, false>(pNew.get(), pNew.get(), pBlended.get(), bApplyColorBlendToAlpha, bBlendOnlyOneMip); break;
-                        case EBlendType::BT_BLEND: BufferLayer<BlendChannel, false>(pNew.get(), pNew.get(), pBlended.get(), bApplyColorBlendToAlpha, bBlendOnlyOneMip); break;
+						case EBlendType::BT_NORMAL_COMBINE: check(false); break;
+                        case EBlendType::BT_SOFTLIGHT: BufferLayer<SoftLightChannel, false>(Base.get(), Base.get(), pBlended.get(), bApplyColorBlendToAlpha, bBlendOnlyOneMip); break;
+                        case EBlendType::BT_HARDLIGHT: BufferLayer<HardLightChannel, false>(Base.get(), Base.get(), pBlended.get(), bApplyColorBlendToAlpha, bBlendOnlyOneMip); break;
+                        case EBlendType::BT_BURN: BufferLayer<BurnChannel, false>(Base.get(), Base.get(), pBlended.get(), bApplyColorBlendToAlpha, bBlendOnlyOneMip); break;
+                        case EBlendType::BT_DODGE: BufferLayer<DodgeChannel, false>(Base.get(), Base.get(), pBlended.get(), bApplyColorBlendToAlpha, bBlendOnlyOneMip); break;
+                        case EBlendType::BT_SCREEN: BufferLayer<ScreenChannel, false>(Base.get(), Base.get(),  pBlended.get(), bApplyColorBlendToAlpha, bBlendOnlyOneMip); break;
+                        case EBlendType::BT_OVERLAY: BufferLayer<OverlayChannel, false>(Base.get(), Base.get(), pBlended.get(), bApplyColorBlendToAlpha, bBlendOnlyOneMip); break;
+                        case EBlendType::BT_LIGHTEN: BufferLayer<LightenChannel, false>(Base.get(), Base.get(), pBlended.get(), bApplyColorBlendToAlpha, bBlendOnlyOneMip); break;
+                        case EBlendType::BT_MULTIPLY: BufferLayer<MultiplyChannel, false>(Base.get(), Base.get(), pBlended.get(), bApplyColorBlendToAlpha, bBlendOnlyOneMip); break;
+                        case EBlendType::BT_BLEND: BufferLayer<BlendChannel, false>(Base.get(), Base.get(), pBlended.get(), bApplyColorBlendToAlpha, bBlendOnlyOneMip); break;
                         default: check(false);
                         }
                     }
@@ -2669,56 +2650,72 @@ namespace mu
 						// Separate alpha operation ignores the mask.
 						switch (EBlendType(args.blendTypeAlpha))
 						{
-						case EBlendType::BT_SOFTLIGHT: BufferLayerInPlace<SoftLightChannel, false, 1>(pNew.get(), pBlended.get(), bBlendOnlyOneMip, 3, 3); break;
-						case EBlendType::BT_HARDLIGHT: BufferLayerInPlace<HardLightChannel, false, 1>(pNew.get(), pBlended.get(), bBlendOnlyOneMip, 3, 3); break;
-						case EBlendType::BT_BURN: BufferLayerInPlace<BurnChannel, false, 1>(pNew.get(), pBlended.get(), bBlendOnlyOneMip, 3, 3); break;
-						case EBlendType::BT_DODGE: BufferLayerInPlace<DodgeChannel, false, 1>(pNew.get(), pBlended.get(), bBlendOnlyOneMip, 3, 3); break;
-						case EBlendType::BT_SCREEN: BufferLayerInPlace<ScreenChannel, false, 1>(pNew.get(), pBlended.get(), bBlendOnlyOneMip, 3, 3); break;
-						case EBlendType::BT_OVERLAY: BufferLayerInPlace<OverlayChannel, false, 1>(pNew.get(), pBlended.get(), bBlendOnlyOneMip, 3, 3); break;
-						case EBlendType::BT_LIGHTEN: BufferLayerInPlace<LightenChannel, false, 1>(pNew.get(), pBlended.get(), bBlendOnlyOneMip, 3, 3); break;
-						case EBlendType::BT_MULTIPLY: BufferLayerInPlace<MultiplyChannel, false, 1>(pNew.get(), pBlended.get(), bBlendOnlyOneMip, 3, 3); break;
-						case EBlendType::BT_BLEND: BufferLayerInPlace<BlendChannel, false, 1>(pNew.get(), pBlended.get(), bBlendOnlyOneMip, 3, 3); break;
+						case EBlendType::BT_SOFTLIGHT: BufferLayerInPlace<SoftLightChannel, false, 1>(Base.get(), pBlended.get(), bBlendOnlyOneMip, 3, 3); break;
+						case EBlendType::BT_HARDLIGHT: BufferLayerInPlace<HardLightChannel, false, 1>(Base.get(), pBlended.get(), bBlendOnlyOneMip, 3, 3); break;
+						case EBlendType::BT_BURN: BufferLayerInPlace<BurnChannel, false, 1>(Base.get(), pBlended.get(), bBlendOnlyOneMip, 3, 3); break;
+						case EBlendType::BT_DODGE: BufferLayerInPlace<DodgeChannel, false, 1>(Base.get(), pBlended.get(), bBlendOnlyOneMip, 3, 3); break;
+						case EBlendType::BT_SCREEN: BufferLayerInPlace<ScreenChannel, false, 1>(Base.get(), pBlended.get(), bBlendOnlyOneMip, 3, 3); break;
+						case EBlendType::BT_OVERLAY: BufferLayerInPlace<OverlayChannel, false, 1>(Base.get(), pBlended.get(), bBlendOnlyOneMip, 3, 3); break;
+						case EBlendType::BT_LIGHTEN: BufferLayerInPlace<LightenChannel, false, 1>(Base.get(), pBlended.get(), bBlendOnlyOneMip, 3, 3); break;
+						case EBlendType::BT_MULTIPLY: BufferLayerInPlace<MultiplyChannel, false, 1>(Base.get(), pBlended.get(), bBlendOnlyOneMip, 3, 3); break;
+						case EBlendType::BT_BLEND: BufferLayerInPlace<BlendChannel, false, 1>(Base.get(), pBlended.get(), bBlendOnlyOneMip, 3, 3); break;
 						default: check(false);
 						}
 					}
 				}
 
-				if (bBlendOnlyOneMip)
+				// Are we done?
+				if (CurrentIteration + 1 == Iterations)
 				{
-					MUTABLE_CPUPROFILER_SCOPE(ImageLayer_MipFix);
-					FMipmapGenerationSettings DummyMipSettings{};
-					ImageMipmapInPlace(m_pSettings->GetPrivate()->m_imageCompressionQuality, pNew.get(), DummyMipSettings);
+					if (Data.MultiLayer.bBlendOnlyOneMip)
+					{
+						MUTABLE_CPUPROFILER_SCOPE(ImageLayer_MipFix);
+						FMipmapGenerationSettings DummyMipSettings{};
+						ImageMipmapInPlace(m_pSettings->GetPrivate()->m_imageCompressionQuality, Base.get(), DummyMipSettings);
+					}
+
+					// TODO: Reconvert to OriginalBaseFormat if necessary?
+
+					GetMemory().SetImage(item, Base);
+					Data.Resource = nullptr;
+					break;
 				}
+				else
+				{
+					// Request a new layer
+					++CurrentIteration;
+					FScheduledOp ItemCopy = item;
+					ExecutionIndex Index = GetMemory().GetRageIndex(item.ExecutionIndex);
+					Index.SetFromModelRangeIndex(args.rangeId, CurrentIteration);
+					ItemCopy.ExecutionIndex = GetMemory().GetRageIndexIndex(Index);
+					AddOp(FScheduledOp(item.At, item, 2+CurrentIteration, item.CustomState), FScheduledOp(args.blended, ItemCopy), FScheduledOp(args.mask, ItemCopy));
 
-				// TODO: Reconvert to baseFormat?
+				}
 
-                GetMemory().SetImage( item, pResult );
                 break;
             }
 
-            default:
-                check(false);
-            }
+            } // switch stage
 
             break;
         }
 
         case OP_TYPE::IM_DIFFERENCE:
         {
-            auto args = pModel->GetPrivate()->m_program.GetOpArgs<OP::ImageDifferenceArgs>(item.at);
-            switch (item.stage)
+            auto args = pModel->GetPrivate()->m_program.GetOpArgs<OP::ImageDifferenceArgs>(item.At);
+            switch (item.Stage)
             {
             case 0:
-                    AddOp( SCHEDULED_OP( item.at, item, 1),
-                           SCHEDULED_OP( args.a, item),
-                           SCHEDULED_OP( args.b, item) );
+                    AddOp( FScheduledOp( item.At, item, 1),
+                           FScheduledOp( args.a, item),
+                           FScheduledOp( args.b, item) );
                 break;
 
             case 1:
             {
                 // TODO: Reuse base if possible
-                Ptr<const Image> pA = GetMemory().GetImage( CACHE_ADDRESS(args.a,item) );
-                Ptr<const Image> pB = GetMemory().GetImage( CACHE_ADDRESS(args.b,item) );
+                Ptr<const Image> pA = GetMemory().GetImage( FCacheAddress(args.a,item) );
+                Ptr<const Image> pB = GetMemory().GetImage( FCacheAddress(args.b,item) );
 
                 ImagePtr pResult = ImageDifference( pA.get(), pB.get() );
 
@@ -2735,15 +2732,15 @@ namespace mu
 
 		case OP_TYPE::IM_NORMALCOMPOSITE:
 		{
-			OP::ImageNormalCompositeArgs args = pModel->GetPrivate()->m_program.GetOpArgs<OP::ImageNormalCompositeArgs>(item.at);
-			switch (item.stage)
+			OP::ImageNormalCompositeArgs args = pModel->GetPrivate()->m_program.GetOpArgs<OP::ImageNormalCompositeArgs>(item.At);
+			switch (item.Stage)
 			{
 			case 0:
 				if (args.base && args.normal)
 				{
-						AddOp(SCHEDULED_OP(item.at, item, 1),
-							  SCHEDULED_OP(args.base, item),
-							  SCHEDULED_OP(args.normal, item));
+						AddOp(FScheduledOp(item.At, item, 1),
+							  FScheduledOp(args.base, item),
+							  FScheduledOp(args.normal, item));
 					}
 					else
 					{
@@ -2753,8 +2750,8 @@ namespace mu
 
 			case 1:
 			{
-				Ptr<const Image> Base = GetMemory().GetImage(CACHE_ADDRESS(args.base, item));
-				Ptr<const Image> Normal = GetMemory().GetImage(CACHE_ADDRESS(args.normal, item));
+				Ptr<const Image> Base = GetMemory().GetImage(FCacheAddress(args.base, item));
+				Ptr<const Image> Normal = GetMemory().GetImage(FCacheAddress(args.normal, item));
 
 				if (Normal->GetLODCount() < Base->GetLODCount())
 				{
@@ -2790,12 +2787,12 @@ namespace mu
 
         case OP_TYPE::IM_PIXELFORMAT:
         {
-			OP::ImagePixelFormatArgs args = pModel->GetPrivate()->m_program.GetOpArgs<OP::ImagePixelFormatArgs>(item.at);
-            switch (item.stage)
+			OP::ImagePixelFormatArgs args = pModel->GetPrivate()->m_program.GetOpArgs<OP::ImagePixelFormatArgs>(item.At);
+            switch (item.Stage)
             {
             case 0:
-                    AddOp( SCHEDULED_OP( item.at, item, 1),
-                           SCHEDULED_OP( args.source, item) );
+                    AddOp( FScheduledOp( item.At, item, 1),
+                           FScheduledOp( args.source, item) );
                 break;
 
             case 1:
@@ -2814,12 +2811,12 @@ namespace mu
 
         case OP_TYPE::IM_MIPMAP:
         {
-			OP::ImageMipmapArgs args = pModel->GetPrivate()->m_program.GetOpArgs<OP::ImageMipmapArgs>(item.at);
-            switch (item.stage)
+			OP::ImageMipmapArgs args = pModel->GetPrivate()->m_program.GetOpArgs<OP::ImageMipmapArgs>(item.At);
+            switch (item.Stage)
             {
             case 0:
-				AddOp( SCHEDULED_OP( item.at, item, 1),
-					SCHEDULED_OP( args.source, item) );
+				AddOp( FScheduledOp( item.At, item, 1),
+					FScheduledOp( args.source, item) );
                 break;
 
             case 1:
@@ -2836,17 +2833,22 @@ namespace mu
 
         case OP_TYPE::IM_RESIZE:
         {
-			OP::ImageResizeArgs args = pModel->GetPrivate()->m_program.GetOpArgs<OP::ImageResizeArgs>(item.at);
-            switch (item.stage)
+			OP::ImageResizeArgs args = pModel->GetPrivate()->m_program.GetOpArgs<OP::ImageResizeArgs>(item.At);
+            switch (item.Stage)
             {
             case 0:
-                    AddOp( SCHEDULED_OP( item.at, item, 1),
-                           SCHEDULED_OP( args.source, item) );
+                    AddOp( FScheduledOp( item.At, item, 1),
+                           FScheduledOp( args.source, item) );
                 break;
 
             case 1:
             {
-                Ptr<const Image> pBase = GetMemory().GetImage( CACHE_ADDRESS(args.source,item) );
+                Ptr<const Image> pBase = GetMemory().GetImage( FCacheAddress(args.source,item) );
+
+				if (!pBase)
+				{
+					GetMemory().SetImage(item, nullptr);
+				}
 
                 FImageSize destSize = FImageSize
                     (
@@ -2855,6 +2857,7 @@ namespace mu
                     );
 
                 ImagePtr pResult;
+
                 if ( destSize[0]!=pBase->GetSizeX()
                      ||
                      destSize[1]!=pBase->GetSizeY() )
@@ -2905,19 +2908,19 @@ namespace mu
 
         case OP_TYPE::IM_RESIZELIKE:
         {
-			OP::ImageResizeLikeArgs args = pModel->GetPrivate()->m_program.GetOpArgs<OP::ImageResizeLikeArgs>(item.at);
-            switch (item.stage)
+			OP::ImageResizeLikeArgs args = pModel->GetPrivate()->m_program.GetOpArgs<OP::ImageResizeLikeArgs>(item.At);
+            switch (item.Stage)
             {
             case 0:
-                    AddOp( SCHEDULED_OP( item.at, item, 1),
-                           SCHEDULED_OP( args.source, item),
-                           SCHEDULED_OP( args.sizeSource, item) );
+                    AddOp( FScheduledOp( item.At, item, 1),
+                           FScheduledOp( args.source, item),
+                           FScheduledOp( args.sizeSource, item) );
                 break;
 
             case 1:
             {
-                Ptr<const Image> pBase = GetMemory().GetImage( CACHE_ADDRESS(args.source,item) );
-                Ptr<const Image> pSizeBase = GetMemory().GetImage( CACHE_ADDRESS(args.sizeSource,item) );
+                Ptr<const Image> pBase = GetMemory().GetImage( FCacheAddress(args.source,item) );
+                Ptr<const Image> pSizeBase = GetMemory().GetImage( FCacheAddress(args.sizeSource,item) );
                 ImagePtr pResult;
 
                 if ( pBase->GetSizeX()!=pSizeBase->GetSizeX()
@@ -2968,17 +2971,17 @@ namespace mu
 
         case OP_TYPE::IM_RESIZEREL:
         {
-            auto args = pModel->GetPrivate()->m_program.GetOpArgs<OP::ImageResizeRelArgs>(item.at);
-            switch (item.stage)
+            auto args = pModel->GetPrivate()->m_program.GetOpArgs<OP::ImageResizeRelArgs>(item.At);
+            switch (item.Stage)
             {
             case 0:
-                    AddOp( SCHEDULED_OP( item.at, item, 1),
-                           SCHEDULED_OP( args.source, item) );
+                    AddOp( FScheduledOp( item.At, item, 1),
+                           FScheduledOp( args.source, item) );
                 break;
 
             case 1:
             {
-                Ptr<const Image> pBase = GetMemory().GetImage( CACHE_ADDRESS(args.source,item) );
+                Ptr<const Image> pBase = GetMemory().GetImage( FCacheAddress(args.source,item) );
 
                 FImageSize destSize(
                             uint16( FMath::Max(1.0, pBase->GetSizeX()*args.factor[0] + 0.5f ) ),
@@ -3024,17 +3027,17 @@ namespace mu
 
         case OP_TYPE::IM_BLANKLAYOUT:
         {
-            auto args = pModel->GetPrivate()->m_program.GetOpArgs<OP::ImageBlankLayoutArgs>(item.at);
-            switch (item.stage)
+            auto args = pModel->GetPrivate()->m_program.GetOpArgs<OP::ImageBlankLayoutArgs>(item.At);
+            switch (item.Stage)
             {
             case 0:
-                    AddOp( SCHEDULED_OP( item.at, item, 1),
-                           SCHEDULED_OP::FromOpAndOptions( args.layout, item, 0) );
+                    AddOp( FScheduledOp( item.At, item, 1),
+                           FScheduledOp::FromOpAndOptions( args.layout, item, 0) );
                 break;
 
             case 1:
             {
-                Ptr<const Layout> pLayout = GetMemory().GetLayout(SCHEDULED_OP::FromOpAndOptions(args.layout, item, 0));
+                Ptr<const Layout> pLayout = GetMemory().GetLayout(FScheduledOp::FromOpAndOptions(args.layout, item, 0));
 
                 FIntPoint SizeInBlocks = pLayout->GetGridSize();
 
@@ -3045,7 +3048,7 @@ namespace mu
 				int32 FullImageMipCount = Image::GetMipmapCount(FullImageSizeInPixels.X, FullImageSizeInPixels.Y);
 
 				FIntPoint ImageSizeInPixels = FullImageSizeInPixels;
-				int32 MipsToSkip = item.executionOptions;
+				int32 MipsToSkip = item.ExecutionOptions;
 				MipsToSkip = FMath::Min(MipsToSkip, FullImageMipCount);
 				if (MipsToSkip > 0)
 				{
@@ -3101,35 +3104,34 @@ namespace mu
 
         case OP_TYPE::IM_COMPOSE:
         {
-            auto args = pModel->GetPrivate()->m_program.GetOpArgs<OP::ImageComposeArgs>( item.at );
-            switch ( item.stage )
+            auto args = pModel->GetPrivate()->m_program.GetOpArgs<OP::ImageComposeArgs>( item.At );
+            switch ( item.Stage )
             {
             case 0:
-                AddOp( SCHEDULED_OP( item.at, item, 1 ), 
-					SCHEDULED_OP::FromOpAndOptions( args.layout, item, 0 ) );
+                AddOp( FScheduledOp( item.At, item, 1 ), 
+					FScheduledOp::FromOpAndOptions( args.layout, item, 0 ) );
                 break;
 
             case 1:
             {
-                Ptr<const Layout> pLayout =
-                    GetMemory().GetLayout( CACHE_ADDRESS( args.layout, SCHEDULED_OP::FromOpAndOptions(args.layout, item, 0)) );
+                Ptr<const Layout> ComposeLayout = GetMemory().GetLayout( FCacheAddress( args.layout, FScheduledOp::FromOpAndOptions(args.layout, item, 0)) );
 
                 FScheduledOpData data;
-                data.layout = pLayout;
+                data.Resource = const_cast<Layout*>(ComposeLayout.get());
 				int32 dataPos = m_heapData.Add( data );
 
-                int relBlockIndex = pLayout->FindBlock( args.blockIndex );
+                int relBlockIndex = ComposeLayout->FindBlock( args.blockIndex );
                 if ( relBlockIndex >= 0 )
                 {
-                    AddOp( SCHEDULED_OP( item.at, item, 2, dataPos ),
-                           SCHEDULED_OP( args.base, item ),
-                           SCHEDULED_OP( args.blockImage, item ),
-                           SCHEDULED_OP( args.mask, item ) );
+                    AddOp( FScheduledOp( item.At, item, 2, dataPos ),
+                           FScheduledOp( args.base, item ),
+                           FScheduledOp( args.blockImage, item ),
+                           FScheduledOp( args.mask, item ) );
                 }
                 else
                 {
-                    AddOp( SCHEDULED_OP( item.at, item, 2, dataPos ),
-                           SCHEDULED_OP( args.base, item ) );
+                    AddOp( FScheduledOp( item.At, item, 2, dataPos ),
+                           FScheduledOp( args.base, item ) );
                 }
                 break;
             }
@@ -3148,12 +3150,12 @@ namespace mu
 
         case OP_TYPE::IM_INTERPOLATE:
         {
-            auto args = pModel->GetPrivate()->m_program.GetOpArgs<OP::ImageInterpolateArgs>(item.at);
-            switch (item.stage)
+            auto args = pModel->GetPrivate()->m_program.GetOpArgs<OP::ImageInterpolateArgs>(item.At);
+            switch (item.Stage)
             {
             case 0:
-                AddOp( SCHEDULED_OP( item.at, item, 1),
-                       SCHEDULED_OP( args.factor, item) );
+                AddOp( FScheduledOp( item.At, item, 1),
+                       FScheduledOp( args.factor, item) );
                 break;
 
             case 1:
@@ -3167,7 +3169,7 @@ namespace mu
                     count++;
                 }
 
-                float factor = GetMemory().GetScalar( CACHE_ADDRESS(args.factor,item) );
+                float factor = GetMemory().GetScalar( FCacheAddress(args.factor,item) );
 
                 float delta = 1.0f/(count-1);
                 int min = (int)floorf( factor/delta );
@@ -3176,26 +3178,26 @@ namespace mu
                 float bifactor = factor/delta - min;
 
                 FScheduledOpData data;
-                data.bifactor = bifactor;
-				data.min = FMath::Clamp(min, 0, count - 1);
-				data.max = FMath::Clamp(max, 0, count - 1);
+                data.Interpolate.Bifactor = bifactor;
+				data.Interpolate.Min = FMath::Clamp(min, 0, count - 1);
+				data.Interpolate.Max = FMath::Clamp(max, 0, count - 1);
 				uint32 dataPos = uint32(m_heapData.Add(data));
 
                 if ( bifactor < UE_SMALL_NUMBER )
                 {
-                        AddOp( SCHEDULED_OP( item.at, item, 2, dataPos),
-                               SCHEDULED_OP( args.targets[min], item) );
+                        AddOp( FScheduledOp( item.At, item, 2, dataPos),
+                               FScheduledOp( args.targets[min], item) );
                     }
                 else if ( bifactor > 1.0f-UE_SMALL_NUMBER )
                 {
-                        AddOp( SCHEDULED_OP( item.at, item, 2, dataPos),
-                               SCHEDULED_OP( args.targets[max], item) );
+                        AddOp( FScheduledOp( item.At, item, 2, dataPos),
+                               FScheduledOp( args.targets[max], item) );
                     }
                     else
                     {
-                        AddOp( SCHEDULED_OP( item.at, item, 2, dataPos),
-                               SCHEDULED_OP( args.targets[min], item),
-                               SCHEDULED_OP( args.targets[max], item) );
+                        AddOp( FScheduledOp( item.At, item, 2, dataPos),
+                               FScheduledOp( args.targets[min], item),
+                               FScheduledOp( args.targets[max], item) );
                     }
                 break;
             }
@@ -3212,26 +3214,26 @@ namespace mu
                 }
 
                 // Factor from 0 to 1 between the two targets
-                const FScheduledOpData& data = m_heapData[(size_t)item.customState];
-                float bifactor = data.bifactor;
-                int min = data.min;
-                int max = data.max;
+                const FScheduledOpData& data = m_heapData[(size_t)item.CustomState];
+                float bifactor = data.Interpolate.Bifactor;
+                int min = data.Interpolate.Min;
+                int max = data.Interpolate.Max;
 
                 ImagePtr pResult;
                 if ( bifactor < UE_SMALL_NUMBER )
                 {
-                    Ptr<const Image> pSource = GetMemory().GetImage( CACHE_ADDRESS(args.targets[min],item) );
+                    Ptr<const Image> pSource = GetMemory().GetImage( FCacheAddress(args.targets[min],item) );
                     pResult = pSource->Clone();
                 }
                 else if ( bifactor > 1.0f-UE_SMALL_NUMBER )
                 {
-                    Ptr<const Image> pSource = GetMemory().GetImage( CACHE_ADDRESS(args.targets[max],item) );
+                    Ptr<const Image> pSource = GetMemory().GetImage( FCacheAddress(args.targets[max],item) );
                     pResult = pSource->Clone();
                 }
                 else
                 {
-                    Ptr<const Image> pMin = GetMemory().GetImage( CACHE_ADDRESS(args.targets[min],item) );
-                    Ptr<const Image> pMax = GetMemory().GetImage( CACHE_ADDRESS(args.targets[max],item) );
+                    Ptr<const Image> pMin = GetMemory().GetImage( FCacheAddress(args.targets[min],item) );
+                    Ptr<const Image> pMax = GetMemory().GetImage( FCacheAddress(args.targets[max],item) );
 
                     if (pMin && pMax)
                     {
@@ -3292,28 +3294,28 @@ namespace mu
 
         case OP_TYPE::IM_INTERPOLATE3:
         {
-            auto args = pModel->GetPrivate()->m_program.GetOpArgs<OP::ImageInterpolate3Args>(item.at);
-            switch (item.stage)
+            auto args = pModel->GetPrivate()->m_program.GetOpArgs<OP::ImageInterpolate3Args>(item.At);
+            switch (item.Stage)
             {
             case 0:
-                    AddOp( SCHEDULED_OP( item.at, item, 1),
-                           SCHEDULED_OP::FromOpAndOptions( args.factor1, item, 0),
-                           SCHEDULED_OP::FromOpAndOptions( args.factor2, item, 0),
-                           SCHEDULED_OP( args.target0, item),
-                           SCHEDULED_OP( args.target1, item),
-                           SCHEDULED_OP( args.target2, item) );
+                    AddOp( FScheduledOp( item.At, item, 1),
+                           FScheduledOp::FromOpAndOptions( args.factor1, item, 0),
+                           FScheduledOp::FromOpAndOptions( args.factor2, item, 0),
+                           FScheduledOp( args.target0, item),
+                           FScheduledOp( args.target1, item),
+                           FScheduledOp( args.target2, item) );
                 break;
 
             case 1:
             {
                 // \TODO Optimise limit cases
 
-                float factor1 = GetMemory().GetScalar(SCHEDULED_OP::FromOpAndOptions(args.factor1, item, 0));
-                float factor2 = GetMemory().GetScalar(SCHEDULED_OP::FromOpAndOptions(args.factor2, item, 0));
+                float factor1 = GetMemory().GetScalar(FScheduledOp::FromOpAndOptions(args.factor1, item, 0));
+                float factor2 = GetMemory().GetScalar(FScheduledOp::FromOpAndOptions(args.factor2, item, 0));
 
-                Ptr<const Image> pTarget0 = GetMemory().GetImage( CACHE_ADDRESS(args.target0,item) );
-                Ptr<const Image> pTarget1 = GetMemory().GetImage( CACHE_ADDRESS(args.target1,item) );
-                Ptr<const Image> pTarget2 = GetMemory().GetImage( CACHE_ADDRESS(args.target2,item) );
+                Ptr<const Image> pTarget0 = GetMemory().GetImage( FCacheAddress(args.target0,item) );
+                Ptr<const Image> pTarget1 = GetMemory().GetImage( FCacheAddress(args.target1,item) );
+                Ptr<const Image> pTarget2 = GetMemory().GetImage( FCacheAddress(args.target2,item) );
 
                 ImagePtr pResult = ImageInterpolate( pTarget0.get(), pTarget1.get(), pTarget2.get(),
                                                      factor1, factor2 );
@@ -3331,19 +3333,19 @@ namespace mu
 
         case OP_TYPE::IM_SATURATE:
         {
-            auto args = pModel->GetPrivate()->m_program.GetOpArgs<OP::ImageSaturateArgs>(item.at);
-            switch (item.stage)
+            auto args = pModel->GetPrivate()->m_program.GetOpArgs<OP::ImageSaturateArgs>(item.At);
+            switch (item.Stage)
             {
             case 0:
-                    AddOp( SCHEDULED_OP( item.at, item, 1),
-                           SCHEDULED_OP( args.base, item ),
-                           SCHEDULED_OP::FromOpAndOptions( args.factor, item, 0 ));
+                    AddOp( FScheduledOp( item.At, item, 1),
+                           FScheduledOp( args.base, item ),
+                           FScheduledOp::FromOpAndOptions( args.factor, item, 0 ));
                 break;
 
             case 1:
             {
-                Ptr<const Image> pBase = GetMemory().GetImage( CACHE_ADDRESS(args.base,item) );
-                float factor = GetMemory().GetScalar(SCHEDULED_OP::FromOpAndOptions(args.factor, item, 0));
+                Ptr<const Image> pBase = GetMemory().GetImage( FCacheAddress(args.base,item) );
+                float factor = GetMemory().GetScalar(FScheduledOp::FromOpAndOptions(args.factor, item, 0));
 
                 ImagePtr pResult = ImageSaturate( pBase.get(), factor );
 
@@ -3360,17 +3362,17 @@ namespace mu
 
         case OP_TYPE::IM_LUMINANCE:
         {
-            auto args = pModel->GetPrivate()->m_program.GetOpArgs<OP::ImageLuminanceArgs>(item.at);
-            switch (item.stage)
+            auto args = pModel->GetPrivate()->m_program.GetOpArgs<OP::ImageLuminanceArgs>(item.At);
+            switch (item.Stage)
             {
             case 0:
-                AddOp( SCHEDULED_OP( item.at, item, 1),
-                        SCHEDULED_OP( args.base, item ) );
+                AddOp( FScheduledOp( item.At, item, 1),
+                        FScheduledOp( args.base, item ) );
                 break;
 
             case 1:
             {
-                Ptr<const Image> pBase = GetMemory().GetImage( CACHE_ADDRESS(args.base,item) );
+                Ptr<const Image> pBase = GetMemory().GetImage( FCacheAddress(args.base,item) );
 
                 ImagePtr pResult = ImageLuminance( pBase.get() );
                 GetMemory().SetImage( item, pResult );
@@ -3386,15 +3388,15 @@ namespace mu
 
         case OP_TYPE::IM_SWIZZLE:
         {
-			OP::ImageSwizzleArgs args = pModel->GetPrivate()->m_program.GetOpArgs<OP::ImageSwizzleArgs>(item.at);
-            switch (item.stage)
+			OP::ImageSwizzleArgs args = pModel->GetPrivate()->m_program.GetOpArgs<OP::ImageSwizzleArgs>(item.At);
+            switch (item.Stage)
             {
             case 0:
-                AddOp( SCHEDULED_OP( item.at, item, 1),
-                        SCHEDULED_OP( args.sources[0], item ),
-                        SCHEDULED_OP( args.sources[1], item ),
-                        SCHEDULED_OP( args.sources[2], item ),
-                        SCHEDULED_OP( args.sources[3], item ) );
+                AddOp( FScheduledOp( item.At, item, 1),
+                        FScheduledOp( args.sources[0], item ),
+                        FScheduledOp( args.sources[1], item ),
+                        FScheduledOp( args.sources[2], item ),
+                        FScheduledOp( args.sources[3], item ) );
                 break;
 
             case 1:
@@ -3411,19 +3413,19 @@ namespace mu
 
         case OP_TYPE::IM_SELECTCOLOUR:
         {
-			OP::ImageSelectColourArgs args = pModel->GetPrivate()->m_program.GetOpArgs<OP::ImageSelectColourArgs>(item.at);
-            switch (item.stage)
+			OP::ImageSelectColourArgs args = pModel->GetPrivate()->m_program.GetOpArgs<OP::ImageSelectColourArgs>(item.At);
+            switch (item.Stage)
             {
             case 0:
-                AddOp( SCHEDULED_OP( item.at, item, 1),
-                        SCHEDULED_OP( args.base, item ),
-                        SCHEDULED_OP::FromOpAndOptions( args.colour, item, 0 ) );
+                AddOp( FScheduledOp( item.At, item, 1),
+                        FScheduledOp( args.base, item ),
+                        FScheduledOp::FromOpAndOptions( args.colour, item, 0 ) );
                 break;
 
             case 1:
             {
-                Ptr<const Image> pBase = GetMemory().GetImage( CACHE_ADDRESS(args.base,item) );
-				FVector4f colour = GetMemory().GetColour(SCHEDULED_OP::FromOpAndOptions(args.colour, item, 0));
+                Ptr<const Image> pBase = GetMemory().GetImage( FCacheAddress(args.base,item) );
+				FVector4f colour = GetMemory().GetColour(FScheduledOp::FromOpAndOptions(args.colour, item, 0));
 
 				Ptr<Image> pResult = ImageSelectColour( pBase.get(), vec3f(colour) );
 
@@ -3440,21 +3442,21 @@ namespace mu
 
         case OP_TYPE::IM_COLOURMAP:
         {
-			OP::ImageColourMapArgs args = pModel->GetPrivate()->m_program.GetOpArgs<OP::ImageColourMapArgs>(item.at);
-            switch (item.stage)
+			OP::ImageColourMapArgs args = pModel->GetPrivate()->m_program.GetOpArgs<OP::ImageColourMapArgs>(item.At);
+            switch (item.Stage)
             {
             case 0:
-                    AddOp( SCHEDULED_OP( item.at, item, 1),
-                           SCHEDULED_OP( args.base, item ),
-                           SCHEDULED_OP( args.mask, item ),
-                           SCHEDULED_OP( args.map, item ) );
+                    AddOp( FScheduledOp( item.At, item, 1),
+                           FScheduledOp( args.base, item ),
+                           FScheduledOp( args.mask, item ),
+                           FScheduledOp( args.map, item ) );
                 break;
 
             case 1:
             {
-                Ptr<const Image> pSource = GetMemory().GetImage( CACHE_ADDRESS(args.base,item) );
-                Ptr<const Image> pMask = GetMemory().GetImage( CACHE_ADDRESS(args.mask,item) );
-                Ptr<const Image> pMap = GetMemory().GetImage( CACHE_ADDRESS(args.map,item) );
+                Ptr<const Image> pSource = GetMemory().GetImage( FCacheAddress(args.base,item) );
+                Ptr<const Image> pMask = GetMemory().GetImage( FCacheAddress(args.mask,item) );
+                Ptr<const Image> pMap = GetMemory().GetImage( FCacheAddress(args.map,item) );
 
 				// Be defensive: ensure image sizes match.
 				if (pMask->GetSize() != pSource->GetSize())
@@ -3479,19 +3481,19 @@ namespace mu
 
         case OP_TYPE::IM_GRADIENT:
         {
-			OP::ImageGradientArgs args = pModel->GetPrivate()->m_program.GetOpArgs<OP::ImageGradientArgs>(item.at);
-            switch (item.stage)
+			OP::ImageGradientArgs args = pModel->GetPrivate()->m_program.GetOpArgs<OP::ImageGradientArgs>(item.At);
+            switch (item.Stage)
             {
             case 0:
-                    AddOp( SCHEDULED_OP( item.at, item, 1),
-                           SCHEDULED_OP::FromOpAndOptions( args.colour0, item, 0 ),
-                           SCHEDULED_OP::FromOpAndOptions( args.colour1, item, 0 ) );
+                    AddOp( FScheduledOp( item.At, item, 1),
+                           FScheduledOp::FromOpAndOptions( args.colour0, item, 0 ),
+                           FScheduledOp::FromOpAndOptions( args.colour1, item, 0 ) );
                 break;
 
             case 1:
             {
-				FVector4f colour0 = GetMemory().GetColour(SCHEDULED_OP::FromOpAndOptions(args.colour0, item, 0));
-				FVector4f colour1 = GetMemory().GetColour(SCHEDULED_OP::FromOpAndOptions(args.colour1, item, 0));
+				FVector4f colour0 = GetMemory().GetColour(FScheduledOp::FromOpAndOptions(args.colour0, item, 0));
+				FVector4f colour1 = GetMemory().GetColour(FScheduledOp::FromOpAndOptions(args.colour1, item, 0));
 
                 ImagePtr pResult = ImageGradient( colour0, colour1,
                                                   args.size[0],
@@ -3510,20 +3512,20 @@ namespace mu
 
         case OP_TYPE::IM_BINARISE:
         {
-			OP::ImageBinariseArgs args = pModel->GetPrivate()->m_program.GetOpArgs<OP::ImageBinariseArgs>(item.at);
-            switch (item.stage)
+			OP::ImageBinariseArgs args = pModel->GetPrivate()->m_program.GetOpArgs<OP::ImageBinariseArgs>(item.At);
+            switch (item.Stage)
             {
             case 0:
-                    AddOp( SCHEDULED_OP( item.at, item, 1),
-                           SCHEDULED_OP( args.base, item ),
-                           SCHEDULED_OP::FromOpAndOptions( args.threshold, item, 0 ) );
+                    AddOp( FScheduledOp( item.At, item, 1),
+                           FScheduledOp( args.base, item ),
+                           FScheduledOp::FromOpAndOptions( args.threshold, item, 0 ) );
                 break;
 
             case 1:
             {
-                Ptr<const Image> pA = GetMemory().GetImage( CACHE_ADDRESS(args.base,item) );
+                Ptr<const Image> pA = GetMemory().GetImage( FCacheAddress(args.base,item) );
 
-                float c = GetMemory().GetScalar(SCHEDULED_OP::FromOpAndOptions(args.threshold, item, 0));
+                float c = GetMemory().GetScalar(FScheduledOp::FromOpAndOptions(args.threshold, item, 0));
 
                 ImagePtr pResult = ImageBinarise( pA.get(), c );
 
@@ -3540,17 +3542,17 @@ namespace mu
 
 		case OP_TYPE::IM_INVERT:
 		{
-			OP::ImageInvertArgs args = pModel->GetPrivate()->m_program.GetOpArgs<OP::ImageInvertArgs>(item.at);
-			switch (item.stage)
+			OP::ImageInvertArgs args = pModel->GetPrivate()->m_program.GetOpArgs<OP::ImageInvertArgs>(item.At);
+			switch (item.Stage)
 			{
 			case 0:
-					AddOp(SCHEDULED_OP(item.at, item, 1),
-						SCHEDULED_OP(args.base, item));
+					AddOp(FScheduledOp(item.At, item, 1),
+						FScheduledOp(args.base, item));
 				break;
 
 			case 1:
 			{
-				Ptr<const Image> pA = GetMemory().GetImage(CACHE_ADDRESS(args.base, item));
+				Ptr<const Image> pA = GetMemory().GetImage(FCacheAddress(args.base, item));
 
 				ImagePtr pResult = ImageInvert(pA.get());
 
@@ -3567,21 +3569,21 @@ namespace mu
 
         case OP_TYPE::IM_PLAINCOLOUR:
         {
-			OP::ImagePlainColourArgs args = pModel->GetPrivate()->m_program.GetOpArgs<OP::ImagePlainColourArgs>(item.at);
-            switch (item.stage)
+			OP::ImagePlainColourArgs args = pModel->GetPrivate()->m_program.GetOpArgs<OP::ImagePlainColourArgs>(item.At);
+            switch (item.Stage)
             {
             case 0:
-				AddOp( SCHEDULED_OP( item.at, item, 1),
-					SCHEDULED_OP::FromOpAndOptions( args.colour, item, 0 ) );
+				AddOp( FScheduledOp( item.At, item, 1),
+					FScheduledOp::FromOpAndOptions( args.colour, item, 0 ) );
                 break;
 
             case 1:
             {
-				FVector4f c = GetMemory().GetColour(SCHEDULED_OP::FromOpAndOptions(args.colour, item, 0));
+				FVector4f c = GetMemory().GetColour(FScheduledOp::FromOpAndOptions(args.colour, item, 0));
 
 				uint16 SizeX = args.size[0];
 				uint16 SizeY = args.size[1];
-				for (int l=0; l<item.executionOptions; ++l)
+				for (int l=0; l<item.ExecutionOptions; ++l)
 				{
 					SizeX = FMath::Max(uint16(1), FMath::DivideAndRoundUp(SizeX, uint16(2)));
 					SizeY = FMath::Max(uint16(1), FMath::DivideAndRoundUp(SizeY, uint16(2)));
@@ -3613,17 +3615,17 @@ namespace mu
 
         case OP_TYPE::IM_CROP:
         {
-			OP::ImageCropArgs args = pModel->GetPrivate()->m_program.GetOpArgs<OP::ImageCropArgs>(item.at);
-            switch (item.stage)
+			OP::ImageCropArgs args = pModel->GetPrivate()->m_program.GetOpArgs<OP::ImageCropArgs>(item.At);
+            switch (item.Stage)
             {
             case 0:
-                AddOp( SCHEDULED_OP( item.at, item, 1),
-					SCHEDULED_OP( args.source, item ) );
+                AddOp( FScheduledOp( item.At, item, 1),
+					FScheduledOp( args.source, item ) );
                 break;
 
             case 1:
             {
-                Ptr<const Image> pA = GetMemory().GetImage( CACHE_ADDRESS(args.source,item) );
+                Ptr<const Image> pA = GetMemory().GetImage( FCacheAddress(args.source,item) );
 
                 box< vec2<int> > rect;
                 rect.min[0] = args.minX;
@@ -3632,7 +3634,7 @@ namespace mu
                 rect.size[1] = args.sizeY;
 
 				// Apply ther mipmap reduction to the crop rectangle.
-				int32 MipsToSkip = item.executionOptions;
+				int32 MipsToSkip = item.ExecutionOptions;
 				while ( MipsToSkip>0 && rect.size[0]>0 && rect.size[1]>0 )
 				{
 					rect.ShrinkToHalf();
@@ -3658,21 +3660,21 @@ namespace mu
 
         case OP_TYPE::IM_PATCH:
         {
-			OP::ImagePatchArgs args = pModel->GetPrivate()->m_program.GetOpArgs<OP::ImagePatchArgs>(item.at);
-            switch (item.stage)
+			OP::ImagePatchArgs args = pModel->GetPrivate()->m_program.GetOpArgs<OP::ImagePatchArgs>(item.At);
+            switch (item.Stage)
             {
             case 0:
-                AddOp( SCHEDULED_OP( item.at, item, 1),
-					SCHEDULED_OP( args.base, item ),
-					SCHEDULED_OP( args.patch, item ) );
+                AddOp( FScheduledOp( item.At, item, 1),
+					FScheduledOp( args.base, item ),
+					FScheduledOp( args.patch, item ) );
                 break;
 
             case 1:
             {
                 MUTABLE_CPUPROFILER_SCOPE(ImagePatch);
 
-                Ptr<const Image> pA = GetMemory().GetImage( CACHE_ADDRESS(args.base,item) );
-                Ptr<const Image> pB = GetMemory().GetImage( CACHE_ADDRESS(args.patch,item) );
+                Ptr<const Image> pA = GetMemory().GetImage( FCacheAddress(args.base,item) );
+                Ptr<const Image> pB = GetMemory().GetImage( FCacheAddress(args.patch,item) );
 
 				// Failsafe
 				if (!pA || !pB)
@@ -3688,7 +3690,7 @@ namespace mu
                 rect.size[1] = pB->GetSizeY();
 
 				// Apply ther mipmap reduction to the crop rectangle.
-				int32 MipsToSkip = item.executionOptions;
+				int32 MipsToSkip = item.ExecutionOptions;
 				while (MipsToSkip > 0 && rect.size[0] > 0 && rect.size[1] > 0)
 				{
 					rect.min/=2;
@@ -3771,23 +3773,23 @@ namespace mu
 
         case OP_TYPE::IM_RASTERMESH:
         {
-            auto args = pModel->GetPrivate()->m_program.GetOpArgs<OP::ImageRasterMeshArgs>(item.at);
-            switch (item.stage)
+            auto args = pModel->GetPrivate()->m_program.GetOpArgs<OP::ImageRasterMeshArgs>(item.At);
+            switch (item.Stage)
             {
             case 0:
                     if (!args.image)
                     {
-                        AddOp( SCHEDULED_OP( item.at, item, 1),
-                            SCHEDULED_OP::FromOpAndOptions( args.mesh, item, 0 ) );
+                        AddOp( FScheduledOp( item.At, item, 1),
+                            FScheduledOp::FromOpAndOptions( args.mesh, item, 0 ) );
                     }
                     else
                     {
-                        AddOp( SCHEDULED_OP( item.at, item, 1),
-							SCHEDULED_OP::FromOpAndOptions( args.mesh, item, 0),
-							SCHEDULED_OP( args.image, item ),
-							SCHEDULED_OP( args.mask, item ),
-							SCHEDULED_OP::FromOpAndOptions( args.angleFadeProperties, item, 0 ),
-							SCHEDULED_OP::FromOpAndOptions( args.projector, item, 0 ) );
+                        AddOp( FScheduledOp( item.At, item, 1),
+							FScheduledOp::FromOpAndOptions( args.mesh, item, 0),
+							FScheduledOp( args.image, item ),
+							FScheduledOp( args.mask, item ),
+							FScheduledOp::FromOpAndOptions( args.angleFadeProperties, item, 0 ),
+							FScheduledOp::FromOpAndOptions( args.projector, item, 0 ) );
                     }
                 break;
 
@@ -3795,7 +3797,7 @@ namespace mu
             {
                 MUTABLE_CPUPROFILER_SCOPE(RunCode_RasterMesh);
 
-                Ptr<const Mesh> pMesh = GetMemory().GetMesh( SCHEDULED_OP::FromOpAndOptions(args.mesh, item, 0) );
+                Ptr<const Mesh> pMesh = GetMemory().GetMesh( FScheduledOp::FromOpAndOptions(args.mesh, item, 0) );
 
                 ImagePtr pNew;
 
@@ -3803,7 +3805,7 @@ namespace mu
 				uint16 SizeY = args.sizeY;
 
 				// Drop mips while possible
-				int32 MipsToDrop = item.executionOptions;
+				int32 MipsToDrop = item.ExecutionOptions;
 				while (MipsToDrop && !(SizeX % 2) && !(SizeY % 2))
 				{
 					SizeX = FMath::Max(uint16(1),FMath::DivideAndRoundUp(SizeX, uint16(2)));
@@ -3819,12 +3821,12 @@ namespace mu
                 else if (args.image)
                 {
                     // Raster with projection
-                    Ptr<const Image> pSource = GetMemory().GetImage( CACHE_ADDRESS(args.image,item) );
+                    Ptr<const Image> pSource = GetMemory().GetImage( FCacheAddress(args.image,item) );
 
                     Ptr<const Image> pMask = nullptr;
                     if ( args.mask )
                     {
-                        pMask = GetMemory().GetImage( CACHE_ADDRESS(args.mask,item) );
+                        pMask = GetMemory().GetImage( FCacheAddress(args.mask,item) );
 
 						// TODO: This shouldn't happen, but be defensive.
 						FImageSize ResultSize(SizeX, SizeY);
@@ -3839,7 +3841,7 @@ namespace mu
                     float fadeEnd = 180.0f;
                     if ( args.angleFadeProperties )
                     {
-                        auto fadeProperties = GetMemory().GetColour(SCHEDULED_OP::FromOpAndOptions(args.angleFadeProperties, item, 0));
+                        auto fadeProperties = GetMemory().GetColour(FScheduledOp::FromOpAndOptions(args.angleFadeProperties, item, 0));
                         fadeStart = fadeProperties[0];
                         fadeEnd = fadeProperties[1];
                     }
@@ -3867,7 +3869,7 @@ namespace mu
                     float projectionAngle = 0;
                     if ( args.projector )
                     {
-                        auto pProjector = GetMemory().GetProjector(SCHEDULED_OP::FromOpAndOptions(args.projector, item, 0));
+                        auto pProjector = GetMemory().GetProjector(FScheduledOp::FromOpAndOptions(args.projector, item, 0));
                         if (pProjector)
                         {
                             projectionAngle = pProjector->m_value.projectionAngle;
@@ -3926,19 +3928,19 @@ namespace mu
 
         case OP_TYPE::IM_MAKEGROWMAP:
         {
-            auto args = pModel->GetPrivate()->m_program.GetOpArgs<OP::ImageMakeGrowMapArgs>(item.at);
-            switch (item.stage)
+            auto args = pModel->GetPrivate()->m_program.GetOpArgs<OP::ImageMakeGrowMapArgs>(item.At);
+            switch (item.Stage)
             {
             case 0:
-                    AddOp( SCHEDULED_OP( item.at, item, 1),
-                           SCHEDULED_OP( args.mask, item) );
+                    AddOp( FScheduledOp( item.At, item, 1),
+                           FScheduledOp( args.mask, item) );
                 break;
 
             case 1:
             {
                 MUTABLE_CPUPROFILER_SCOPE(RunCode_ImageMakeGrowMap);
 
-                Ptr<const Image> pMask = GetMemory().GetImage( CACHE_ADDRESS(args.mask,item) );
+                Ptr<const Image> pMask = GetMemory().GetImage( FCacheAddress(args.mask,item) );
 
                 ImagePtr pNew = new Image( pMask->GetSizeX(), pMask->GetSizeY(), pMask->GetLODCount(), EImageFormat::IF_L_UBYTE);
 
@@ -3958,21 +3960,21 @@ namespace mu
 
         case OP_TYPE::IM_DISPLACE:
         {
-            auto args = pModel->GetPrivate()->m_program.GetOpArgs<OP::ImageDisplaceArgs>(item.at);
-            switch (item.stage)
+            auto args = pModel->GetPrivate()->m_program.GetOpArgs<OP::ImageDisplaceArgs>(item.At);
+            switch (item.Stage)
             {
             case 0:
-                    AddOp( SCHEDULED_OP( item.at, item, 1),
-                           SCHEDULED_OP( args.source, item ),
-                           SCHEDULED_OP( args.displacementMap, item ) );
+                    AddOp( FScheduledOp( item.At, item, 1),
+                           FScheduledOp( args.source, item ),
+                           FScheduledOp( args.displacementMap, item ) );
                 break;
 
             case 1:
             {
                 MUTABLE_CPUPROFILER_SCOPE(RunCode_ImageDisplace);
 
-                Ptr<const Image> pSource = GetMemory().GetImage( CACHE_ADDRESS(args.source,item) );
-                Ptr<const Image> pMap = GetMemory().GetImage( CACHE_ADDRESS(args.displacementMap,item) );
+                Ptr<const Image> pSource = GetMemory().GetImage( FCacheAddress(args.source,item) );
+                Ptr<const Image> pMap = GetMemory().GetImage( FCacheAddress(args.displacementMap,item) );
 
 				// TODO: This shouldn't happen: displacement maps cannot be scaled because their information
 				// is resolution sensitive (pixel offsets). If the size doesn't match, scale the source, apply 
@@ -4010,38 +4012,38 @@ namespace mu
 
         case OP_TYPE::IM_TRANSFORM:
         {
-            const OP::ImageTransformArgs Args = pModel->GetPrivate()->m_program.GetOpArgs<OP::ImageTransformArgs>(item.at);
+            const OP::ImageTransformArgs Args = pModel->GetPrivate()->m_program.GetOpArgs<OP::ImageTransformArgs>(item.At);
 
-            switch (item.stage)
+            switch (item.Stage)
             {
             case 0:
 			{
-				const TArray<SCHEDULED_OP, TInlineAllocator<6>> Deps = {
-						SCHEDULED_OP( Args.base, item),
-						SCHEDULED_OP( Args.offsetX, item),
-						SCHEDULED_OP( Args.offsetY, item),
-						SCHEDULED_OP( Args.scaleX, item),
-						SCHEDULED_OP( Args.scaleY, item),
-						SCHEDULED_OP( Args.rotation, item) };
+				const TArray<FScheduledOp, TInlineAllocator<6>> Deps = {
+						FScheduledOp( Args.base, item),
+						FScheduledOp( Args.offsetX, item),
+						FScheduledOp( Args.offsetY, item),
+						FScheduledOp( Args.scaleX, item),
+						FScheduledOp( Args.scaleY, item),
+						FScheduledOp( Args.rotation, item) };
 
-                AddOp( SCHEDULED_OP( item.at, item, 1), Deps );
+                AddOp( FScheduledOp( item.At, item, 1), Deps );
 
 				break;
 			}
             case 1:
             {
-                Ptr<const Image> pBaseImage = GetMemory().GetImage( CACHE_ADDRESS(Args.base, item) );
+                Ptr<const Image> pBaseImage = GetMemory().GetImage( FCacheAddress(Args.base, item) );
                 
                 const FVector2f Offset = FVector2f(
-                        Args.offsetX ? GetMemory().GetScalar( CACHE_ADDRESS(Args.offsetX, item) ) : 0.0f,
-                        Args.offsetY ? GetMemory().GetScalar( CACHE_ADDRESS(Args.offsetY, item) ) : 0.0f );
+                        Args.offsetX ? GetMemory().GetScalar( FCacheAddress(Args.offsetX, item) ) : 0.0f,
+                        Args.offsetY ? GetMemory().GetScalar( FCacheAddress(Args.offsetY, item) ) : 0.0f );
 
                 const FVector2f Scale = FVector2f(
-                        Args.scaleX ? GetMemory().GetScalar( CACHE_ADDRESS(Args.scaleX, item) ) : 1.0f,
-                        Args.scaleY ? GetMemory().GetScalar( CACHE_ADDRESS(Args.scaleY, item) ) : 1.0f );
+                        Args.scaleX ? GetMemory().GetScalar( FCacheAddress(Args.scaleX, item) ) : 1.0f,
+                        Args.scaleY ? GetMemory().GetScalar( FCacheAddress(Args.scaleY, item) ) : 1.0f );
 
 				// Map Range 0-1 to a full rotation
-                const float Rotation = GetMemory().GetScalar( CACHE_ADDRESS(Args.rotation, item) ) * UE_TWO_PI;
+                const float Rotation = GetMemory().GetScalar( FCacheAddress(Args.rotation, item) ) * UE_TWO_PI;
 
 
 				Ptr<Image> pResult = new Image( pBaseImage->GetSizeX(), pBaseImage->GetSizeY(), 1, pBaseImage->GetFormat());
@@ -4077,13 +4079,9 @@ namespace mu
 
 
     //---------------------------------------------------------------------------------------------
-    Ptr<RangeIndex> CodeRunner::BuildCurrentOpRangeIndex( const SCHEDULED_OP& item,
-                                                        const Parameters* pParams,
-                                                        const Model* pModel,
-                                                        int parameterIndex
-                                                        )
+    Ptr<RangeIndex> CodeRunner::BuildCurrentOpRangeIndex( const FScheduledOp& item, const Parameters* pParams, const Model* pModel, int32 parameterIndex )
     {
-        if (!item.executionIndex)
+        if (!item.ExecutionIndex)
         {
             return nullptr;
         }
@@ -4103,7 +4101,7 @@ namespace mu
              ++rangeIndexInParam )
         {
             auto rangeIndexInModel = paramDesc.m_ranges[rangeIndexInParam];
-            const ExecutionIndex& currentIndex = GetMemory().GetRageIndex( item.executionIndex );
+            const ExecutionIndex& currentIndex = GetMemory().GetRageIndex( item.ExecutionIndex );
             int position = currentIndex.GetFromModelRangeIndex(rangeIndexInModel);
             index->GetPrivate()->m_values[rangeIndexInParam] = position;
         }
@@ -4113,7 +4111,7 @@ namespace mu
 
 
     //---------------------------------------------------------------------------------------------
-    void CodeRunner::RunCode_Bool( SCHEDULED_OP& item,
+    void CodeRunner::RunCode_Bool( FScheduledOp& item,
                                    const Parameters* pParams,
                                    const Model* pModel
                                    )
@@ -4121,13 +4119,13 @@ namespace mu
         //MUTABLE_CPUPROFILER_SCOPE(RunCode_Bool);
 
         const auto& program = pModel->GetPrivate()->m_program;
-        auto type = program.GetOpType(item.at);
+        auto type = program.GetOpType(item.At);
         switch (type)
         {
 
         case OP_TYPE::BO_CONSTANT:
         {
-            auto args = program.GetOpArgs<OP::BoolConstantArgs>(item.at);
+            auto args = program.GetOpArgs<OP::BoolConstantArgs>(item.At);
             bool result = args.value;
             GetMemory().SetBool( item, result );
             break;
@@ -4135,7 +4133,7 @@ namespace mu
 
         case OP_TYPE::BO_PARAMETER:
         {
-            auto args = program.GetOpArgs<OP::ParameterArgs>(item.at);
+            auto args = program.GetOpArgs<OP::ParameterArgs>(item.At);
             bool result = false;
 			Ptr<RangeIndex> index = BuildCurrentOpRangeIndex( item, pParams, pModel, args.variable );
             result = pParams->GetBoolValue( args.variable, index );
@@ -4145,19 +4143,19 @@ namespace mu
 
         case OP_TYPE::BO_LESS:
         {
-            auto args = program.GetOpArgs<OP::BoolLessArgs>(item.at);
-            switch (item.stage)
+            auto args = program.GetOpArgs<OP::BoolLessArgs>(item.At);
+            switch (item.Stage)
             {
             case 0:
-                    AddOp( SCHEDULED_OP( item.at, item, 1),
-                           SCHEDULED_OP( args.a, item),
-                           SCHEDULED_OP( args.b, item) );
+                    AddOp( FScheduledOp( item.At, item, 1),
+                           FScheduledOp( args.a, item),
+                           FScheduledOp( args.b, item) );
                 break;
 
             case 1:
             {
-                float a = GetMemory().GetScalar( CACHE_ADDRESS(args.a,item) );
-                float b = GetMemory().GetScalar( CACHE_ADDRESS(args.b,item) );
+                float a = GetMemory().GetScalar( FCacheAddress(args.a,item) );
+                float b = GetMemory().GetScalar( FCacheAddress(args.b,item) );
                 bool result = a<b;
                 GetMemory().SetBool( item, result );
                 break;
@@ -4172,16 +4170,16 @@ namespace mu
 
         case OP_TYPE::BO_AND:
         {
-            auto args = program.GetOpArgs<OP::BoolBinaryArgs>(item.at);
-            switch (item.stage)
+            auto args = program.GetOpArgs<OP::BoolBinaryArgs>(item.At);
+            switch (item.Stage)
             {
             case 0:
                 {
                     // Try to avoid the op entirely if we have some children cached
                     bool skip = false;
-                    if ( args.a && GetMemory().IsValid( CACHE_ADDRESS(args.a,item) ) )
+                    if ( args.a && GetMemory().IsValid( FCacheAddress(args.a,item) ) )
                     {
-                         bool a = GetMemory().GetBool( CACHE_ADDRESS(args.a,item) );
+                         bool a = GetMemory().GetBool( FCacheAddress(args.a,item) );
                          if (!a)
                          {
                             GetMemory().SetBool( item, false );
@@ -4189,9 +4187,9 @@ namespace mu
                          }
                     }
 
-                    if ( !skip && args.b && GetMemory().IsValid( CACHE_ADDRESS(args.b,item) ) )
+                    if ( !skip && args.b && GetMemory().IsValid( FCacheAddress(args.b,item) ) )
                     {
-                         bool b = GetMemory().GetBool( CACHE_ADDRESS(args.b,item) );
+                         bool b = GetMemory().GetBool( FCacheAddress(args.b,item) );
                          if (!b)
                          {
                             GetMemory().SetBool( item, false );
@@ -4201,23 +4199,23 @@ namespace mu
 
                     if (!skip)
                     {
-                        AddOp( SCHEDULED_OP( item.at, item, 1),
-                               SCHEDULED_OP( args.a, item));
+                        AddOp( FScheduledOp( item.At, item, 1),
+                               FScheduledOp( args.a, item));
                     }
 				break;
                 }
 
             case 1:
             {
-                bool a = args.a ? GetMemory().GetBool( CACHE_ADDRESS(args.a,item) ) : true;
+                bool a = args.a ? GetMemory().GetBool( FCacheAddress(args.a,item) ) : true;
                 if (!a)
                 {
                     GetMemory().SetBool( item, false );
                 }
                 else
                 {
-                    AddOp( SCHEDULED_OP( item.at, item, 2),
-                           SCHEDULED_OP( args.b, item));
+                    AddOp( FScheduledOp( item.At, item, 2),
+                           FScheduledOp( args.b, item));
                 }
                 break;
             }
@@ -4225,7 +4223,7 @@ namespace mu
             case 2:
             {
                 // We arrived here because a is true
-                bool b = args.b ? GetMemory().GetBool( CACHE_ADDRESS(args.b,item) ) : true;
+                bool b = args.b ? GetMemory().GetBool( FCacheAddress(args.b,item) ) : true;
                 GetMemory().SetBool( item, b );
                 break;
             }
@@ -4238,16 +4236,16 @@ namespace mu
 
         case OP_TYPE::BO_OR:
         {
-            auto args = program.GetOpArgs<OP::BoolBinaryArgs>(item.at);
-            switch (item.stage)
+            auto args = program.GetOpArgs<OP::BoolBinaryArgs>(item.At);
+            switch (item.Stage)
             {
             case 0:
                 {
                     // Try to avoid the op entirely if we have some children cached
                     bool skip = false;
-                    if ( args.a && GetMemory().IsValid( CACHE_ADDRESS(args.a,item) ) )
+                    if ( args.a && GetMemory().IsValid( FCacheAddress(args.a,item) ) )
                     {
-                         bool a = GetMemory().GetBool( CACHE_ADDRESS(args.a,item) );
+                         bool a = GetMemory().GetBool( FCacheAddress(args.a,item) );
                          if (a)
                          {
                             GetMemory().SetBool( item, true );
@@ -4255,9 +4253,9 @@ namespace mu
                          }
                     }
 
-                    if ( !skip && args.b && GetMemory().IsValid( CACHE_ADDRESS(args.b,item) ) )
+                    if ( !skip && args.b && GetMemory().IsValid( FCacheAddress(args.b,item) ) )
                     {
-                         bool b = GetMemory().GetBool( CACHE_ADDRESS(args.b,item) );
+                         bool b = GetMemory().GetBool( FCacheAddress(args.b,item) );
                          if (b)
                          {
                             GetMemory().SetBool( item, true );
@@ -4267,23 +4265,23 @@ namespace mu
 
                     if (!skip)
                     {
-                        AddOp( SCHEDULED_OP( item.at, item, 1),
-                               SCHEDULED_OP( args.a, item));
+                        AddOp( FScheduledOp( item.At, item, 1),
+                               FScheduledOp( args.a, item));
                     }
 				break;
                 }
 
             case 1:
             {
-                bool a = args.a ? GetMemory().GetBool( CACHE_ADDRESS(args.a,item) ) : false;
+                bool a = args.a ? GetMemory().GetBool( FCacheAddress(args.a,item) ) : false;
                 if (a)
                 {
                     GetMemory().SetBool( item, true );
                 }
                 else
                 {
-                    AddOp( SCHEDULED_OP( item.at, item, 2),
-                           SCHEDULED_OP( args.b, item));
+                    AddOp( FScheduledOp( item.At, item, 2),
+                           FScheduledOp( args.b, item));
                 }
                 break;
             }
@@ -4291,7 +4289,7 @@ namespace mu
             case 2:
             {
                 // We arrived here because a is false
-                bool b = args.b ? GetMemory().GetBool( CACHE_ADDRESS(args.b,item) ) : false;
+                bool b = args.b ? GetMemory().GetBool( FCacheAddress(args.b,item) ) : false;
                 GetMemory().SetBool( item, b );
                 break;
             }
@@ -4304,17 +4302,17 @@ namespace mu
 
         case OP_TYPE::BO_NOT:
         {
-            auto args = program.GetOpArgs<OP::BoolNotArgs>(item.at);
-            switch (item.stage)
+            auto args = program.GetOpArgs<OP::BoolNotArgs>(item.At);
+            switch (item.Stage)
             {
             case 0:
-                    AddOp( SCHEDULED_OP( item.at, item, 1),
-                           SCHEDULED_OP( args.source, item) );
+                    AddOp( FScheduledOp( item.At, item, 1),
+                           FScheduledOp( args.source, item) );
                 break;
 
             case 1:
             {
-                bool result = !GetMemory().GetBool( CACHE_ADDRESS(args.source,item) );
+                bool result = !GetMemory().GetBool( FCacheAddress(args.source,item) );
                 GetMemory().SetBool( item, result );
                 break;
             }
@@ -4327,17 +4325,17 @@ namespace mu
 
         case OP_TYPE::BO_EQUAL_INT_CONST:
         {
-            auto args = program.GetOpArgs<OP::BoolEqualScalarConstArgs>(item.at);
-            switch (item.stage)
+            auto args = program.GetOpArgs<OP::BoolEqualScalarConstArgs>(item.At);
+            switch (item.Stage)
             {
             case 0:
-                    AddOp( SCHEDULED_OP( item.at, item, 1),
-                           SCHEDULED_OP( args.value, item) );
+                    AddOp( FScheduledOp( item.At, item, 1),
+                           FScheduledOp( args.value, item) );
                 break;
 
             case 1:
             {
-                int a = GetMemory().GetInt( CACHE_ADDRESS(args.value,item) );
+                int a = GetMemory().GetInt( FCacheAddress(args.value,item) );
                 bool result = a == args.constant;
                 GetMemory().SetBool( item, result );
                 break;
@@ -4357,20 +4355,20 @@ namespace mu
 
 
     //---------------------------------------------------------------------------------------------
-    void CodeRunner::RunCode_Int( SCHEDULED_OP& item,
+    void CodeRunner::RunCode_Int( FScheduledOp& item,
                                   const Parameters* pParams,
                                   const Model* pModel
                               )
     {
         MUTABLE_CPUPROFILER_SCOPE(RunCode_Int);
 
-		OP_TYPE type = pModel->GetPrivate()->m_program.GetOpType(item.at);
+		OP_TYPE type = pModel->GetPrivate()->m_program.GetOpType(item.At);
         switch (type)
         {
 
         case OP_TYPE::NU_CONSTANT:
         {
-            auto args = pModel->GetPrivate()->m_program.GetOpArgs<OP::IntConstantArgs>(item.at);
+            auto args = pModel->GetPrivate()->m_program.GetOpArgs<OP::IntConstantArgs>(item.At);
             int result = args.value;
             GetMemory().SetInt( item, result );
             break;
@@ -4378,7 +4376,7 @@ namespace mu
 
         case OP_TYPE::NU_PARAMETER:
         {
-            auto args = pModel->GetPrivate()->m_program.GetOpArgs<OP::ParameterArgs>(item.at);
+            auto args = pModel->GetPrivate()->m_program.GetOpArgs<OP::ParameterArgs>(item.At);
 			Ptr<RangeIndex> index = BuildCurrentOpRangeIndex( item, pParams, pModel, args.variable );
             int result = pParams->GetIntValue( args.variable, index );
 
@@ -4414,20 +4412,20 @@ namespace mu
 
 
     //---------------------------------------------------------------------------------------------
-    void CodeRunner::RunCode_Scalar( SCHEDULED_OP& item,
+    void CodeRunner::RunCode_Scalar( FScheduledOp& item,
                                      const Parameters* pParams,
                                      const Model* pModel
                                      )
     {
         MUTABLE_CPUPROFILER_SCOPE(RunCode_Scalar);
 
-		OP_TYPE type = pModel->GetPrivate()->m_program.GetOpType(item.at);
+		OP_TYPE type = pModel->GetPrivate()->m_program.GetOpType(item.At);
         switch (type)
         {
 
         case OP_TYPE::SC_CONSTANT:
         {
-            auto args = pModel->GetPrivate()->m_program.GetOpArgs<OP::ScalarConstantArgs>(item.at);
+            auto args = pModel->GetPrivate()->m_program.GetOpArgs<OP::ScalarConstantArgs>(item.At);
             float result = args.value;
             GetMemory().SetScalar( item, result );
             break;
@@ -4435,7 +4433,7 @@ namespace mu
 
         case OP_TYPE::SC_PARAMETER:
         {
-            auto args = pModel->GetPrivate()->m_program.GetOpArgs<OP::ParameterArgs>(item.at);
+            auto args = pModel->GetPrivate()->m_program.GetOpArgs<OP::ParameterArgs>(item.At);
 			Ptr<RangeIndex> index = BuildCurrentOpRangeIndex( item, pParams, pModel, args.variable );
             float result = pParams->GetFloatValue( args.variable, index );
             GetMemory().SetScalar( item, result );
@@ -4444,17 +4442,17 @@ namespace mu
 
         case OP_TYPE::SC_CURVE:
         {
-            auto args = pModel->GetPrivate()->m_program.GetOpArgs<OP::ScalarCurveArgs>(item.at);
-            switch (item.stage)
+            auto args = pModel->GetPrivate()->m_program.GetOpArgs<OP::ScalarCurveArgs>(item.At);
+            switch (item.Stage)
             {
             case 0:
-                    AddOp( SCHEDULED_OP( item.at, item, 1),
-                           SCHEDULED_OP( args.time, item) );
+                    AddOp( FScheduledOp( item.At, item, 1),
+                           FScheduledOp( args.time, item) );
                 break;
 
             case 1:
             {
-                float time = GetMemory().GetScalar( CACHE_ADDRESS(args.time,item) );
+                float time = GetMemory().GetScalar( FCacheAddress(args.time,item) );
 
                 const Curve& curve = pModel->GetPrivate()->m_program.m_constantCurves[args.curve];
                 float result = EvalCurve(curve, time);
@@ -4476,19 +4474,19 @@ namespace mu
 
         case OP_TYPE::SC_ARITHMETIC:
         {
-            auto args = pModel->GetPrivate()->m_program.GetOpArgs<OP::ArithmeticArgs>(item.at);
-            switch (item.stage)
+            auto args = pModel->GetPrivate()->m_program.GetOpArgs<OP::ArithmeticArgs>(item.At);
+            switch (item.Stage)
             {
             case 0:
-                    AddOp( SCHEDULED_OP( item.at, item, 1),
-                           SCHEDULED_OP( args.a, item),
-                           SCHEDULED_OP( args.b, item) );
+                    AddOp( FScheduledOp( item.At, item, 1),
+                           FScheduledOp( args.a, item),
+                           FScheduledOp( args.b, item) );
                 break;
 
             case 1:
             {
-                float a = GetMemory().GetScalar( CACHE_ADDRESS(args.a,item) );
-                float b = GetMemory().GetScalar( CACHE_ADDRESS(args.b,item) );
+                float a = GetMemory().GetScalar( FCacheAddress(args.a,item) );
+                float b = GetMemory().GetScalar( FCacheAddress(args.b,item) );
 
                 float result = 1.0f;
                 switch (args.operation)
@@ -4533,17 +4531,17 @@ namespace mu
 
     //---------------------------------------------------------------------------------------------
     void
-    CodeRunner::RunCode_String( SCHEDULED_OP& item, const Parameters* pParams, const Model* pModel )
+    CodeRunner::RunCode_String( FScheduledOp& item, const Parameters* pParams, const Model* pModel )
     {
         MUTABLE_CPUPROFILER_SCOPE(RunCode_String );
 
-		OP_TYPE type = pModel->GetPrivate()->m_program.GetOpType( item.at );
+		OP_TYPE type = pModel->GetPrivate()->m_program.GetOpType( item.At );
         switch ( type )
         {
 
         case OP_TYPE::ST_CONSTANT:
         {
-			OP::ResourceConstantArgs args = pModel->GetPrivate()->m_program.GetOpArgs<OP::ResourceConstantArgs>( item.at );
+			OP::ResourceConstantArgs args = pModel->GetPrivate()->m_program.GetOpArgs<OP::ResourceConstantArgs>( item.At );
             check( args.value < (uint32)pModel->GetPrivate()->m_program.m_constantStrings.Num() );
 
             auto result = pModel->GetPrivate()->m_program.m_constantStrings[args.value];
@@ -4554,7 +4552,7 @@ namespace mu
 
         case OP_TYPE::ST_PARAMETER:
         {
-			OP::ParameterArgs args = pModel->GetPrivate()->m_program.GetOpArgs<OP::ParameterArgs>( item.at );
+			OP::ParameterArgs args = pModel->GetPrivate()->m_program.GetOpArgs<OP::ParameterArgs>( item.At );
 			Ptr<RangeIndex> index = BuildCurrentOpRangeIndex( item, pParams, pModel, args.variable );
             string result = pParams->GetStringValue( args.variable, index );
             GetMemory().SetString( item, new String( result.c_str() ) );
@@ -4569,14 +4567,14 @@ namespace mu
 
 
     //---------------------------------------------------------------------------------------------
-    void CodeRunner::RunCode_Colour( SCHEDULED_OP& item,
+    void CodeRunner::RunCode_Colour( FScheduledOp& item,
                                      const Parameters* pParams,
                                      const Model* pModel
                                      )
     {
 		MUTABLE_CPUPROFILER_SCOPE(RunCode_Colour);
 
-		OP_TYPE type = pModel->GetPrivate()->m_program.GetOpType(item.at);
+		OP_TYPE type = pModel->GetPrivate()->m_program.GetOpType(item.At);
 
         const auto& program = pModel->GetPrivate()->m_program;
 
@@ -4585,7 +4583,7 @@ namespace mu
 
         case OP_TYPE::CO_CONSTANT:
         {
-			OP::ColourConstantArgs args = program.GetOpArgs<OP::ColourConstantArgs>(item.at);
+			OP::ColourConstantArgs args = program.GetOpArgs<OP::ColourConstantArgs>(item.At);
 			FVector4f result;
             result[0] = args.value[0];
             result[1] = args.value[1];
@@ -4597,7 +4595,7 @@ namespace mu
 
         case OP_TYPE::CO_PARAMETER:
         {
-			OP::ParameterArgs args = pModel->GetPrivate()->m_program.GetOpArgs<OP::ParameterArgs>(item.at);
+			OP::ParameterArgs args = pModel->GetPrivate()->m_program.GetOpArgs<OP::ParameterArgs>(item.At);
 			Ptr<RangeIndex> index = BuildCurrentOpRangeIndex( item, pParams, pModel, args.variable );
             float r=0.0f;
             float g=0.0f;
@@ -4609,23 +4607,23 @@ namespace mu
 
         case OP_TYPE::CO_SAMPLEIMAGE:
         {
-			OP::ColourSampleImageArgs args = pModel->GetPrivate()->m_program.GetOpArgs<OP::ColourSampleImageArgs>(item.at);
-            switch (item.stage)
+			OP::ColourSampleImageArgs args = pModel->GetPrivate()->m_program.GetOpArgs<OP::ColourSampleImageArgs>(item.At);
+            switch (item.Stage)
             {
             case 0:
-                    AddOp( SCHEDULED_OP( item.at, item, 1),
-                           SCHEDULED_OP( args.x, item),
-                           SCHEDULED_OP( args.y, item),
+                    AddOp( FScheduledOp( item.At, item, 1),
+                           FScheduledOp( args.x, item),
+                           FScheduledOp( args.y, item),
 						   // Don't skip mips for the texture to sample
-                           SCHEDULED_OP::FromOpAndOptions( args.image, item, 0) );
+                           FScheduledOp::FromOpAndOptions( args.image, item, 0) );
                 break;
 
             case 1:
             {
-                float x = args.x ? GetMemory().GetScalar( CACHE_ADDRESS(args.x,item) ) : 0.5f;
-                float y = args.y ? GetMemory().GetScalar( CACHE_ADDRESS(args.y,item) ) : 0.5f;
+                float x = args.x ? GetMemory().GetScalar( FCacheAddress(args.x,item) ) : 0.5f;
+                float y = args.y ? GetMemory().GetScalar( FCacheAddress(args.y,item) ) : 0.5f;
 
-                Ptr<const Image> pImage = GetMemory().GetImage(SCHEDULED_OP::FromOpAndOptions(args.image, item, 0));
+                Ptr<const Image> pImage = GetMemory().GetImage(FScheduledOp::FromOpAndOptions(args.image, item, 0));
 
 				FVector4f result;
                 if (pImage)
@@ -4658,15 +4656,15 @@ namespace mu
 
         case OP_TYPE::CO_SWIZZLE:
         {
-			OP::ColourSwizzleArgs args = pModel->GetPrivate()->m_program.GetOpArgs<OP::ColourSwizzleArgs>(item.at);
-            switch (item.stage)
+			OP::ColourSwizzleArgs args = pModel->GetPrivate()->m_program.GetOpArgs<OP::ColourSwizzleArgs>(item.At);
+            switch (item.Stage)
             {
             case 0:
-                    AddOp( SCHEDULED_OP( item.at, item, 1),
-                           SCHEDULED_OP( args.sources[0], item),
-                           SCHEDULED_OP( args.sources[1], item),
-                           SCHEDULED_OP( args.sources[2], item),
-                           SCHEDULED_OP( args.sources[3], item) );
+                    AddOp( FScheduledOp( item.At, item, 1),
+                           FScheduledOp( args.sources[0], item),
+                           FScheduledOp( args.sources[1], item),
+                           FScheduledOp( args.sources[2], item),
+                           FScheduledOp( args.sources[3], item) );
                 break;
 
             case 1:
@@ -4677,7 +4675,7 @@ namespace mu
                 {
                     if ( args.sources[t] )
                     {
-                        auto p = GetMemory().GetColour( CACHE_ADDRESS(args.sources[t],item) );
+                        auto p = GetMemory().GetColour( FCacheAddress(args.sources[t],item) );
                         result[t] = p[ args.sourceChannels[t] ];
                     }
                 }
@@ -4695,17 +4693,17 @@ namespace mu
 
         case OP_TYPE::CO_IMAGESIZE:
         {
-			OP::ColourImageSizeArgs args = pModel->GetPrivate()->m_program.GetOpArgs<OP::ColourImageSizeArgs>(item.at);
-            switch (item.stage)
+			OP::ColourImageSizeArgs args = pModel->GetPrivate()->m_program.GetOpArgs<OP::ColourImageSizeArgs>(item.At);
+            switch (item.Stage)
             {
             case 0:
-                    AddOp( SCHEDULED_OP( item.at, item, 1),
-                           SCHEDULED_OP( args.image, item) );
+                    AddOp( FScheduledOp( item.At, item, 1),
+                           FScheduledOp( args.image, item) );
                 break;
 
             case 1:
             {
-                Ptr<const Image> pImage = GetMemory().GetImage( CACHE_ADDRESS(args.image,item) );
+                Ptr<const Image> pImage = GetMemory().GetImage( FCacheAddress(args.image,item) );
 
 				FVector4f result = FVector4f( (float)pImage->GetSizeX(), (float)pImage->GetSizeY(), 0.0f, 0.0f );
 
@@ -4722,17 +4720,17 @@ namespace mu
 
         case OP_TYPE::CO_LAYOUTBLOCKTRANSFORM:
         {
-			OP::ColourLayoutBlockTransformArgs args = pModel->GetPrivate()->m_program.GetOpArgs<OP::ColourLayoutBlockTransformArgs>(item.at);
-            switch (item.stage)
+			OP::ColourLayoutBlockTransformArgs args = pModel->GetPrivate()->m_program.GetOpArgs<OP::ColourLayoutBlockTransformArgs>(item.At);
+            switch (item.Stage)
             {
             case 0:
-                    AddOp( SCHEDULED_OP( item.at, item, 1),
-                           SCHEDULED_OP( args.layout, item) );
+                    AddOp( FScheduledOp( item.At, item, 1),
+                           FScheduledOp( args.layout, item) );
                 break;
 
             case 1:
             {
-                Ptr<const Layout> pLayout = GetMemory().GetLayout( CACHE_ADDRESS(args.layout,item) );
+                Ptr<const Layout> pLayout = GetMemory().GetLayout( FCacheAddress(args.layout,item) );
 
 				FVector4f result = FVector4f(0,0,0,0);
                 if ( pLayout )
@@ -4772,15 +4770,15 @@ namespace mu
 
         case OP_TYPE::CO_FROMSCALARS:
         {
-			OP::ColourFromScalarsArgs args = pModel->GetPrivate()->m_program.GetOpArgs<OP::ColourFromScalarsArgs>(item.at);
-            switch (item.stage)
+			OP::ColourFromScalarsArgs args = pModel->GetPrivate()->m_program.GetOpArgs<OP::ColourFromScalarsArgs>(item.At);
+            switch (item.Stage)
             {
             case 0:
-                    AddOp( SCHEDULED_OP( item.at, item, 1),
-                           SCHEDULED_OP( args.x, item),
-                           SCHEDULED_OP( args.y, item),
-                           SCHEDULED_OP( args.z, item),
-                           SCHEDULED_OP( args.w, item));
+                    AddOp( FScheduledOp( item.At, item, 1),
+                           FScheduledOp( args.x, item),
+                           FScheduledOp( args.y, item),
+                           FScheduledOp( args.z, item),
+                           FScheduledOp( args.w, item));
                 break;
 
             case 1:
@@ -4789,22 +4787,22 @@ namespace mu
 
                 if (args.x)
                 {
-                    result[0] = GetMemory().GetScalar( CACHE_ADDRESS(args.x,item) );
+                    result[0] = GetMemory().GetScalar( FCacheAddress(args.x,item) );
                 }
 
                 if (args.y)
                 {
-                    result[1] = GetMemory().GetScalar( CACHE_ADDRESS(args.y,item) );
+                    result[1] = GetMemory().GetScalar( FCacheAddress(args.y,item) );
                 }
 
                 if (args.z)
                 {
-                    result[2] = GetMemory().GetScalar( CACHE_ADDRESS(args.z,item) );
+                    result[2] = GetMemory().GetScalar( FCacheAddress(args.z,item) );
                 }
 
                 if (args.w)
                 {
-                    result[3] = GetMemory().GetScalar( CACHE_ADDRESS(args.w,item) );
+                    result[3] = GetMemory().GetScalar( FCacheAddress(args.w,item) );
                 }
 
                 GetMemory().SetColour( item, result );
@@ -4820,13 +4818,13 @@ namespace mu
 
         case OP_TYPE::CO_ARITHMETIC:
         {
-			OP::ArithmeticArgs args = program.GetOpArgs<OP::ArithmeticArgs>(item.at);
-            switch (item.stage)
+			OP::ArithmeticArgs args = program.GetOpArgs<OP::ArithmeticArgs>(item.At);
+            switch (item.Stage)
             {
             case 0:
-                    AddOp( SCHEDULED_OP( item.at, item, 1),
-                           SCHEDULED_OP( args.a, item),
-                           SCHEDULED_OP( args.b, item));
+                    AddOp( FScheduledOp( item.At, item, 1),
+                           FScheduledOp( args.a, item),
+                           FScheduledOp( args.b, item));
                 break;
 
             case 1:
@@ -4837,9 +4835,9 @@ namespace mu
                 otype = program.GetOpType( args.b );
                 dtype = GetOpDataType( otype );
                 check( dtype == DT_COLOUR );
-				FVector4f a = args.a ? GetMemory().GetColour( CACHE_ADDRESS( args.a, item ) )
+				FVector4f a = args.a ? GetMemory().GetColour( FCacheAddress( args.a, item ) )
                                  : FVector4f( 0, 0, 0, 0 );
-				FVector4f b = args.b ? GetMemory().GetColour( CACHE_ADDRESS( args.b, item ) )
+				FVector4f b = args.b ? GetMemory().GetColour( FCacheAddress( args.b, item ) )
                                  : FVector4f( 0, 0, 0, 0 );
 
 				FVector4f result = FVector4f(0,0,0,0);
@@ -4885,7 +4883,7 @@ namespace mu
 
 
     //---------------------------------------------------------------------------------------------
-    void CodeRunner::RunCode_Projector( SCHEDULED_OP& item,
+    void CodeRunner::RunCode_Projector( FScheduledOp& item,
                                 const Parameters* pParams,
                                 const Model* pModel
                               )
@@ -4893,13 +4891,13 @@ namespace mu
         MUTABLE_CPUPROFILER_SCOPE(RunCode_Projector);
 
         const FProgram& program = pModel->GetPrivate()->m_program;
-		OP_TYPE type = program.GetOpType(item.at);
+		OP_TYPE type = program.GetOpType(item.At);
         switch (type)
         {
 
         case OP_TYPE::PR_CONSTANT:
         {
-			OP::ResourceConstantArgs args = pModel->GetPrivate()->m_program.GetOpArgs<OP::ResourceConstantArgs>(item.at);
+			OP::ResourceConstantArgs args = pModel->GetPrivate()->m_program.GetOpArgs<OP::ResourceConstantArgs>(item.At);
             ProjectorPtr pResult = new Projector();
             pResult->m_value = program.m_constantProjectors[args.value];
             GetMemory().SetProjector( item, pResult );
@@ -4908,7 +4906,7 @@ namespace mu
 
         case OP_TYPE::PR_PARAMETER:
         {
-			OP::ParameterArgs args = pModel->GetPrivate()->m_program.GetOpArgs<OP::ParameterArgs>(item.at);
+			OP::ParameterArgs args = pModel->GetPrivate()->m_program.GetOpArgs<OP::ParameterArgs>(item.At);
 			Ptr<RangeIndex> index = BuildCurrentOpRangeIndex( item, pParams, pModel, args.variable );
             ProjectorPtr pResult = new Projector();
             pResult->m_value = pParams->GetPrivate()->GetProjectorValue(args.variable,index);
@@ -4929,19 +4927,19 @@ namespace mu
 
 
     //---------------------------------------------------------------------------------------------
-    void CodeRunner::RunCode_Layout( SCHEDULED_OP& item,
+    void CodeRunner::RunCode_Layout( FScheduledOp& item,
                                      const Model* pModel
                                      )
     {
         //MUTABLE_CPUPROFILER_SCOPE(RunCode_Layout);
 
-		OP_TYPE type = pModel->GetPrivate()->m_program.GetOpType(item.at);
+		OP_TYPE type = pModel->GetPrivate()->m_program.GetOpType(item.At);
         switch (type)
         {
 
         case OP_TYPE::LA_CONSTANT:
         {
-			OP::ResourceConstantArgs args = pModel->GetPrivate()->m_program.GetOpArgs<OP::ResourceConstantArgs>(item.at);
+			OP::ResourceConstantArgs args = pModel->GetPrivate()->m_program.GetOpArgs<OP::ResourceConstantArgs>(item.At);
             check( args.value < (uint32)pModel->GetPrivate()->m_program.m_constantLayouts.Num() );
 
             LayoutPtrConst pResult = pModel->GetPrivate()->m_program.m_constantLayouts
@@ -4952,19 +4950,19 @@ namespace mu
 
         case OP_TYPE::LA_MERGE:
         {
-			OP::LayoutMergeArgs args = pModel->GetPrivate()->m_program.GetOpArgs<OP::LayoutMergeArgs>(item.at);
-            switch (item.stage)
+			OP::LayoutMergeArgs args = pModel->GetPrivate()->m_program.GetOpArgs<OP::LayoutMergeArgs>(item.At);
+            switch (item.Stage)
             {
             case 0:
-                    AddOp( SCHEDULED_OP( item.at, item, 1),
-                           SCHEDULED_OP( args.Base, item),
-                           SCHEDULED_OP( args.Added, item) );
+                    AddOp( FScheduledOp( item.At, item, 1),
+                           FScheduledOp( args.Base, item),
+                           FScheduledOp( args.Added, item) );
                 break;
 
             case 1:
             {
-                Ptr<const Layout> pA = GetMemory().GetLayout( CACHE_ADDRESS(args.Base,item) );
-                Ptr<const Layout> pB = GetMemory().GetLayout( CACHE_ADDRESS(args.Added,item) );
+                Ptr<const Layout> pA = GetMemory().GetLayout( FCacheAddress(args.Base,item) );
+                Ptr<const Layout> pB = GetMemory().GetLayout( FCacheAddress(args.Added,item) );
 
                 LayoutPtrConst pResult;
 
@@ -4994,17 +4992,17 @@ namespace mu
 
         case OP_TYPE::LA_PACK:
         {
-			OP::LayoutPackArgs args = pModel->GetPrivate()->m_program.GetOpArgs<OP::LayoutPackArgs>(item.at);
-            switch (item.stage)
+			OP::LayoutPackArgs args = pModel->GetPrivate()->m_program.GetOpArgs<OP::LayoutPackArgs>(item.At);
+            switch (item.Stage)
             {
             case 0:
-                    AddOp( SCHEDULED_OP( item.at, item, 1),
-                           SCHEDULED_OP( args.Source, item) );
+                    AddOp( FScheduledOp( item.At, item, 1),
+                           FScheduledOp( args.Source, item) );
                 break;
 
             case 1:
             {
-                Ptr<const Layout> pSource = GetMemory().GetLayout( CACHE_ADDRESS(args.Source,item) );
+                Ptr<const Layout> pSource = GetMemory().GetLayout( FCacheAddress(args.Source,item) );
 
 				LayoutPtr pResult;
 
@@ -5036,17 +5034,17 @@ namespace mu
 
 		case OP_TYPE::LA_FROMMESH:
 		{
-			OP::LayoutFromMeshArgs args = pModel->GetPrivate()->m_program.GetOpArgs<OP::LayoutFromMeshArgs>(item.at);
-			switch (item.stage)
+			OP::LayoutFromMeshArgs args = pModel->GetPrivate()->m_program.GetOpArgs<OP::LayoutFromMeshArgs>(item.At);
+			switch (item.Stage)
 			{
 			case 0:
-				AddOp(SCHEDULED_OP(item.at, item, 1),
-					SCHEDULED_OP(args.Mesh, item));
+				AddOp(FScheduledOp(item.At, item, 1),
+					FScheduledOp(args.Mesh, item));
 				break;
 
 			case 1:
 			{
-				Ptr<const Mesh> Mesh = GetMemory().GetMesh(CACHE_ADDRESS(args.Mesh, item));
+				Ptr<const Mesh> Mesh = GetMemory().GetMesh(FCacheAddress(args.Mesh, item));
 
 				Ptr<const Layout> Result = LayoutFromMesh_RemoveBlocks(Mesh.get(), args.LayoutIndex);
 
@@ -5063,19 +5061,19 @@ namespace mu
 
 		case OP_TYPE::LA_REMOVEBLOCKS:
 		{
-			OP::LayoutRemoveBlocksArgs args = pModel->GetPrivate()->m_program.GetOpArgs<OP::LayoutRemoveBlocksArgs>(item.at);
-			switch (item.stage)
+			OP::LayoutRemoveBlocksArgs args = pModel->GetPrivate()->m_program.GetOpArgs<OP::LayoutRemoveBlocksArgs>(item.At);
+			switch (item.Stage)
 			{
 			case 0:
-				AddOp(SCHEDULED_OP(item.at, item, 1),
-					SCHEDULED_OP(args.Source, item),
-					SCHEDULED_OP(args.ReferenceLayout, item));
+				AddOp(FScheduledOp(item.At, item, 1),
+					FScheduledOp(args.Source, item),
+					FScheduledOp(args.ReferenceLayout, item));
 				break;
 
 			case 1:
 			{
-				Ptr<const Layout> Source = GetMemory().GetLayout(CACHE_ADDRESS(args.Source, item));
-				Ptr<const Layout> ReferenceLayout = GetMemory().GetLayout(CACHE_ADDRESS(args.ReferenceLayout, item));
+				Ptr<const Layout> Source = GetMemory().GetLayout(FCacheAddress(args.Source, item));
+				Ptr<const Layout> ReferenceLayout = GetMemory().GetLayout(FCacheAddress(args.ReferenceLayout, item));
 
 				Ptr<const Layout> pResult;
 
@@ -5108,16 +5106,16 @@ namespace mu
 
 
     //---------------------------------------------------------------------------------------------
-    void CodeRunner::RunCode( SCHEDULED_OP& item,
+    void CodeRunner::RunCode( FScheduledOp& item,
                               const Parameters* pParams,
                               const Model* pModel,
                               uint32 lodMask)
     {
-		//UE_LOG( LogMutableCore, Log, TEXT("Running :%5d , %d "), item.at, item.stage );
-		check( item.Type == SCHEDULED_OP::EType::Full );
+		//UE_LOG( LogMutableCore, Log, TEXT("Running :%5d , %d "), item.At, item.Stage );
+		check( item.Type == FScheduledOp::EType::Full );
 
-		OP_TYPE type = pModel->GetPrivate()->m_program.GetOpType(item.at);
-		//UE_LOG(LogMutableCore, Log, TEXT("Running :%5d , %d, of type %d "), item.at, item.stage, type);
+		OP_TYPE type = pModel->GetPrivate()->m_program.GetOpType(item.At);
+		//UE_LOG(LogMutableCore, Log, TEXT("Running :%5d , %d, of type %d "), item.At, item.Stage, type);
 		switch ( type )
         {
         case OP_TYPE::NONE:
@@ -5210,7 +5208,7 @@ namespace mu
 
 
 	//---------------------------------------------------------------------------------------------
-	void CodeRunner::RunCodeImageDesc(SCHEDULED_OP& item,
+	void CodeRunner::RunCodeImageDesc(FScheduledOp& item,
 		const Parameters* pParams,
 		const Model* pModel, 
 		uint32 lodMask
@@ -5218,66 +5216,66 @@ namespace mu
 	{
 		MUTABLE_CPUPROFILER_SCOPE(RunCodeImageDesc);
 
-		check(item.Type == SCHEDULED_OP::EType::ImageDesc);
+		check(item.Type == FScheduledOp::EType::ImageDesc);
 
 		// Ensure there is room for the result in the heap.
-		if (item.customState >= uint32(m_heapData.Num()))
+		if (item.CustomState >= uint32(m_heapData.Num()))
 		{
-			m_heapImageDesc.SetNum(item.customState+1);
+			m_heapImageDesc.SetNum(item.CustomState+1);
 		}
 
 
 		const mu::FProgram& program = pModel->GetPrivate()->m_program;
 
-		OP_TYPE type = program.GetOpType(item.at);
+		OP_TYPE type = program.GetOpType(item.At);
 		switch (type)
 		{
 
 		case OP_TYPE::IM_CONSTANT:
 		{
-			check(item.stage == 0);
-			OP::ResourceConstantArgs args = program.GetOpArgs<OP::ResourceConstantArgs>(item.at);
+			check(item.Stage == 0);
+			OP::ResourceConstantArgs args = program.GetOpArgs<OP::ResourceConstantArgs>(item.At);
 			int32 ImageIndex = args.value;
-			m_heapImageDesc[item.customState].m_format = EImageFormat::IF_NONE;	// TODO: precalculate if necessary
-			m_heapImageDesc[item.customState].m_size[0] = program.m_constantImages[ImageIndex].ImageSizeX;
-			m_heapImageDesc[item.customState].m_size[1] = program.m_constantImages[ImageIndex].ImageSizeY;
-			m_heapImageDesc[item.customState].m_lods = program.m_constantImages[ImageIndex].LODCount;
+			m_heapImageDesc[item.CustomState].m_format = EImageFormat::IF_NONE;	// TODO: precalculate if necessary
+			m_heapImageDesc[item.CustomState].m_size[0] = program.m_constantImages[ImageIndex].ImageSizeX;
+			m_heapImageDesc[item.CustomState].m_size[1] = program.m_constantImages[ImageIndex].ImageSizeY;
+			m_heapImageDesc[item.CustomState].m_lods = program.m_constantImages[ImageIndex].LODCount;
 			GetMemory().SetValidDesc(item);
 			break;
 		}
 
 		case OP_TYPE::IM_PARAMETER:
 		{
-			check(item.stage == 0);
-			OP::ParameterArgs args = program.GetOpArgs<OP::ParameterArgs>(item.at);
+			check(item.Stage == 0);
+			OP::ParameterArgs args = program.GetOpArgs<OP::ParameterArgs>(item.At);
 			EXTERNAL_IMAGE_ID id = pParams->GetImageValue(args.variable);
 			ImagePtr pResult = LoadExternalImage(id);
-			m_heapImageDesc[item.customState].m_format = pResult->GetFormat();
-			m_heapImageDesc[item.customState].m_size = pResult->GetSize();
-			m_heapImageDesc[item.customState].m_lods = pResult->GetLODCount();
+			m_heapImageDesc[item.CustomState].m_format = pResult->GetFormat();
+			m_heapImageDesc[item.CustomState].m_size = pResult->GetSize();
+			m_heapImageDesc[item.CustomState].m_lods = pResult->GetLODCount();
 			GetMemory().SetValidDesc(item);
 			break;
 		}
 
 		case OP_TYPE::IM_CONDITIONAL:
 		{
-			OP::ConditionalArgs args = program.GetOpArgs<OP::ConditionalArgs>(item.at);
-			switch (item.stage)
+			OP::ConditionalArgs args = program.GetOpArgs<OP::ConditionalArgs>(item.At);
+			switch (item.Stage)
 			{
 			case 0:
 			{
 				// We need to run the full condition result
-				SCHEDULED_OP FullConditionOp(args.condition, item);
-				FullConditionOp.Type = SCHEDULED_OP::EType::Full;
-				AddOp(SCHEDULED_OP(item.at, item, 1), FullConditionOp);
+				FScheduledOp FullConditionOp(args.condition, item);
+				FullConditionOp.Type = FScheduledOp::EType::Full;
+				AddOp(FScheduledOp(item.At, item, 1), FullConditionOp);
 				break;
 			}
 
 			case 1:
 			{
-				bool value = GetMemory().GetBool(CACHE_ADDRESS(args.condition, item.executionIndex, item.executionOptions));
+				bool value = GetMemory().GetBool(FCacheAddress(args.condition, item.ExecutionIndex, item.ExecutionOptions));
 				OP::ADDRESS resultAt = value ? args.yes : args.no;
-				AddOp(SCHEDULED_OP(item.at, item, 2), SCHEDULED_OP(resultAt, item));
+				AddOp(FScheduledOp(item.At, item, 2), FScheduledOp(resultAt, item));
 				break;
 			}
 
@@ -5289,7 +5287,7 @@ namespace mu
 
 		case OP_TYPE::IM_SWITCH:
 		{
-			const uint8* data = program.GetOpArgsPointer(item.at);
+			const uint8* data = program.GetOpArgsPointer(item.At);
 		
 			OP::ADDRESS VarAddress;
 			FMemory::Memcpy( &VarAddress, data, sizeof(OP::ADDRESS));
@@ -5303,16 +5301,16 @@ namespace mu
 			FMemory::Memcpy( &CaseCount, data, sizeof(uint32));
 			data += sizeof(uint32);
 	
-			switch (item.stage)
+			switch (item.Stage)
 			{
 			case 0:
 			{
 				if (VarAddress)
 				{
 					// We need to run the full condition result
-					SCHEDULED_OP FullVariableOp(VarAddress, item);
-					FullVariableOp.Type = SCHEDULED_OP::EType::Full;
-					AddOp(SCHEDULED_OP(item.at, item, 1), FullVariableOp);
+					FScheduledOp FullVariableOp(VarAddress, item);
+					FullVariableOp.Type = FScheduledOp::EType::Full;
+					AddOp(FScheduledOp(item.At, item, 1), FullVariableOp);
 				}
 				else
 				{
@@ -5324,7 +5322,7 @@ namespace mu
 			case 1:
 			{
 				// Get the variable result
-				int var = GetMemory().GetInt(CACHE_ADDRESS(VarAddress, item));
+				int var = GetMemory().GetInt(FCacheAddress(VarAddress, item));
 
 				OP::ADDRESS valueAt = DefAddress;
 				for (uint32 C = 0; C < CaseCount; ++C)
@@ -5344,8 +5342,8 @@ namespace mu
 					}
 				}
 
-				AddOp(SCHEDULED_OP(item.at, item, 2, valueAt),
-					  SCHEDULED_OP(valueAt, item));
+				AddOp(FScheduledOp(item.At, item, 2, valueAt),
+					  FScheduledOp(valueAt, item));
 
 				break;
 			}
@@ -5358,10 +5356,10 @@ namespace mu
 
 		case OP_TYPE::IM_LAYERCOLOUR:
 		{
-			OP::ImageLayerColourArgs args = program.GetOpArgs<OP::ImageLayerColourArgs>(item.at);
-			switch (item.stage)
+			OP::ImageLayerColourArgs args = program.GetOpArgs<OP::ImageLayerColourArgs>(item.At);
+			switch (item.Stage)
 			{
-			case 0: AddOp(SCHEDULED_OP(item.at, item, 1), SCHEDULED_OP(args.base, item)); break;
+			case 0: AddOp(FScheduledOp(item.At, item, 1), FScheduledOp(args.base, item)); break;
 			case 1: GetMemory().SetValidDesc(item); break;
 			default: check(false);
 			}
@@ -5370,10 +5368,10 @@ namespace mu
 
 		case OP_TYPE::IM_LAYER:
 		{
-			OP::ImageLayerArgs args = program.GetOpArgs<OP::ImageLayerArgs>(item.at);
-			switch (item.stage)
+			OP::ImageLayerArgs args = program.GetOpArgs<OP::ImageLayerArgs>(item.At);
+			switch (item.Stage)
 			{
-			case 0: AddOp(SCHEDULED_OP(item.at, item, 1), SCHEDULED_OP(args.base, item)); break;
+			case 0: AddOp(FScheduledOp(item.At, item, 1), FScheduledOp(args.base, item)); break;
 			case 1: GetMemory().SetValidDesc(item); break;
 			default: check(false);
 			}
@@ -5382,10 +5380,10 @@ namespace mu
 
 		case OP_TYPE::IM_MULTILAYER:
 		{
-			OP::ImageMultiLayerArgs args = program.GetOpArgs<OP::ImageMultiLayerArgs>(item.at);
-			switch (item.stage)
+			OP::ImageMultiLayerArgs args = program.GetOpArgs<OP::ImageMultiLayerArgs>(item.At);
+			switch (item.Stage)
 			{
-			case 0: AddOp(SCHEDULED_OP(item.at, item, 1), SCHEDULED_OP(args.base, item)); break;
+			case 0: AddOp(FScheduledOp(item.At, item, 1), FScheduledOp(args.base, item)); break;
 			case 1: GetMemory().SetValidDesc(item); break;
 			default: check(false);
 			}
@@ -5394,10 +5392,10 @@ namespace mu
 
 		case OP_TYPE::IM_DIFFERENCE:
 		{
-			OP::ImageDifferenceArgs args = program.GetOpArgs<OP::ImageDifferenceArgs>(item.at);
-			switch (item.stage)
+			OP::ImageDifferenceArgs args = program.GetOpArgs<OP::ImageDifferenceArgs>(item.At);
+			switch (item.Stage)
 			{
-			case 0: AddOp(SCHEDULED_OP(item.at, item, 1), SCHEDULED_OP(args.a, item)); break;
+			case 0: AddOp(FScheduledOp(item.At, item, 1), FScheduledOp(args.a, item)); break;
 			case 1: GetMemory().SetValidDesc(item); break;
 			default: check(false);
 			}
@@ -5406,10 +5404,10 @@ namespace mu
 
 		case OP_TYPE::IM_NORMALCOMPOSITE:
 		{
-			OP::ImageNormalCompositeArgs args = program.GetOpArgs<OP::ImageNormalCompositeArgs>(item.at);
-			switch (item.stage)
+			OP::ImageNormalCompositeArgs args = program.GetOpArgs<OP::ImageNormalCompositeArgs>(item.At);
+			switch (item.Stage)
 			{
-			case 0: AddOp(SCHEDULED_OP(item.at, item, 1), SCHEDULED_OP(args.base, item)); break;
+			case 0: AddOp(FScheduledOp(item.At, item, 1), FScheduledOp(args.base, item)); break;
 			case 1: GetMemory().SetValidDesc(item); break;
 			default: check(false);
 			}
@@ -5418,17 +5416,17 @@ namespace mu
 
 		case OP_TYPE::IM_PIXELFORMAT:
 		{
-			OP::ImagePixelFormatArgs args = program.GetOpArgs<OP::ImagePixelFormatArgs>(item.at);
-			switch (item.stage)
+			OP::ImagePixelFormatArgs args = program.GetOpArgs<OP::ImagePixelFormatArgs>(item.At);
+			switch (item.Stage)
 			{
 			case 0:
-				AddOp(SCHEDULED_OP(item.at, item, 1), SCHEDULED_OP(args.source, item));
+				AddOp(FScheduledOp(item.At, item, 1), FScheduledOp(args.source, item));
 				break;
 
 			case 1:
 			{
 				// Update directly in the heap
-				EImageFormat OldFormat = m_heapImageDesc[item.customState].m_format;
+				EImageFormat OldFormat = m_heapImageDesc[item.CustomState].m_format;
 				EImageFormat NewFormat = args.format;
 				if (args.formatIfAlpha != EImageFormat::IF_NONE
 					&&
@@ -5436,7 +5434,7 @@ namespace mu
 				{
 					NewFormat = args.formatIfAlpha;
 				}
-				m_heapImageDesc[item.customState].m_format = NewFormat;				
+				m_heapImageDesc[item.CustomState].m_format = NewFormat;				
 				GetMemory().SetValidDesc(item);
 				break;
 			}
@@ -5450,17 +5448,17 @@ namespace mu
 
 		case OP_TYPE::IM_MIPMAP:
 		{
-			OP::ImageMipmapArgs args = program.GetOpArgs<OP::ImageMipmapArgs>(item.at);
-			switch (item.stage)
+			OP::ImageMipmapArgs args = program.GetOpArgs<OP::ImageMipmapArgs>(item.At);
+			switch (item.Stage)
 			{
 			case 0:
-				AddOp(SCHEDULED_OP(item.at, item, 1), SCHEDULED_OP(args.source, item));
+				AddOp(FScheduledOp(item.At, item, 1), FScheduledOp(args.source, item));
 				break;
 
 			case 1:
 			{
 				// Somewhat synched with Full op execution code.
-				FImageDesc BaseDesc = m_heapImageDesc[item.customState];
+				FImageDesc BaseDesc = m_heapImageDesc[item.CustomState];
 				int levelCount = args.levels;
 				int maxLevelCount = Image::GetMipmapCount(BaseDesc.m_size[0], BaseDesc.m_size[1]);
 				if (levelCount == 0)
@@ -5479,7 +5477,7 @@ namespace mu
 				levelCount = FMath::Max(startLevel, levelCount);
 
 				// Update result.
-				m_heapImageDesc[item.customState].m_lods = levelCount;
+				m_heapImageDesc[item.CustomState].m_lods = levelCount;
 				GetMemory().SetValidDesc(item);
 				break;
 			}
@@ -5492,16 +5490,16 @@ namespace mu
 
 		case OP_TYPE::IM_RESIZE:
 		{
-			auto args = program.GetOpArgs<OP::ImageResizeArgs>(item.at);
-			switch (item.stage)
+			auto args = program.GetOpArgs<OP::ImageResizeArgs>(item.At);
+			switch (item.Stage)
 			{
 			case 0:
-				AddOp(SCHEDULED_OP(item.at, item, 1), SCHEDULED_OP(args.source, item));
+				AddOp(FScheduledOp(item.At, item, 1), FScheduledOp(args.source, item));
 				break;
 
 			case 1:
-				m_heapImageDesc[item.customState].m_size[0] = args.size[0];
-				m_heapImageDesc[item.customState].m_size[1] = args.size[1];
+				m_heapImageDesc[item.CustomState].m_size[0] = args.size[0];
+				m_heapImageDesc[item.CustomState].m_size[1] = args.size[1];
 				GetMemory().SetValidDesc(item);
 				break;
 
@@ -5513,25 +5511,28 @@ namespace mu
 
 		case OP_TYPE::IM_RESIZELIKE:
 		{
-			auto args = program.GetOpArgs<OP::ImageResizeLikeArgs>(item.at);
-			switch (item.stage)
+			auto args = program.GetOpArgs<OP::ImageResizeLikeArgs>(item.At);
+			switch (item.Stage)
 			{
 			case 0:
 			{
-				int32 ResultAndBaseDesc = item.customState;
+				int32 ResultAndBaseDesc = item.CustomState;
 				int32 SourceDescAddress = m_heapImageDesc.Add({});
-				int32 SecondStageData = m_heapData.Add({ 0.0f, ResultAndBaseDesc, SourceDescAddress, nullptr });
-				AddOp(SCHEDULED_OP(item.at, item, 1, SecondStageData),
-					SCHEDULED_OP(args.source, item, 0, ResultAndBaseDesc),
-					SCHEDULED_OP(args.sizeSource, item, 0, SourceDescAddress));
+				FScheduledOpData Data;
+				Data.ResizeLike.ResultDescAt = ResultAndBaseDesc;
+				Data.ResizeLike.SourceDescAt = SourceDescAddress;
+				int32 SecondStageData = m_heapData.Add(Data);
+				AddOp(FScheduledOp(item.At, item, 1, SecondStageData),
+					FScheduledOp(args.source, item, 0, ResultAndBaseDesc),
+					FScheduledOp(args.sizeSource, item, 0, SourceDescAddress));
 				break;
 			}
 
 			case 1:
 			{
-				const FScheduledOpData& SecondStageData = m_heapData[ item.customState ];
-				FImageDesc& ResultAndBaseDesc = m_heapImageDesc[SecondStageData.min];
-				const FImageDesc& SourceDesc = m_heapImageDesc[SecondStageData.max];
+				const FScheduledOpData& SecondStageData = m_heapData[ item.CustomState ];
+				FImageDesc& ResultAndBaseDesc = m_heapImageDesc[SecondStageData.ResizeLike.ResultDescAt];
+				const FImageDesc& SourceDesc = m_heapImageDesc[SecondStageData.ResizeLike.SourceDescAt];
 				ResultAndBaseDesc.m_size = SourceDesc.m_size;
 				GetMemory().SetValidDesc(item);
 				break;
@@ -5546,16 +5547,16 @@ namespace mu
 
 		case OP_TYPE::IM_RESIZEREL:
 		{
-			auto args = program.GetOpArgs<OP::ImageResizeRelArgs>(item.at);
-			switch (item.stage)
+			auto args = program.GetOpArgs<OP::ImageResizeRelArgs>(item.At);
+			switch (item.Stage)
 			{
 			case 0:
-				AddOp(SCHEDULED_OP(item.at, item, 1), SCHEDULED_OP(args.source, item));
+				AddOp(FScheduledOp(item.At, item, 1), FScheduledOp(args.source, item));
 				break;
 
 			case 1:
 			{
-				FImageDesc& ResultAndBaseDesc = m_heapImageDesc[item.customState];
+				FImageDesc& ResultAndBaseDesc = m_heapImageDesc[item.CustomState];
 				FImageSize destSize(
 					uint16(ResultAndBaseDesc.m_size[0] * args.factor[0] + 0.5f),
 					uint16(ResultAndBaseDesc.m_size[1] * args.factor[1] + 0.5f));
@@ -5573,27 +5574,27 @@ namespace mu
 
 		case OP_TYPE::IM_BLANKLAYOUT:
 		{
-			OP::ImageBlankLayoutArgs args = program.GetOpArgs<OP::ImageBlankLayoutArgs>(item.at);
-			switch (item.stage)
+			OP::ImageBlankLayoutArgs args = program.GetOpArgs<OP::ImageBlankLayoutArgs>(item.At);
+			switch (item.Stage)
 			{
 			case 0:
 			{
 				// We need to run the full layout
-				SCHEDULED_OP FullLayoutOp(args.layout, item);
-				FullLayoutOp.Type = SCHEDULED_OP::EType::Full;
-				AddOp(SCHEDULED_OP(item.at, item, 1), FullLayoutOp);
+				FScheduledOp FullLayoutOp(args.layout, item);
+				FullLayoutOp.Type = FScheduledOp::EType::Full;
+				AddOp(FScheduledOp(item.At, item, 1), FullLayoutOp);
 				break;
 			}
 
 			case 1:
 			{
-				Ptr<const Layout> pLayout = GetMemory().GetLayout(CACHE_ADDRESS(args.layout, item));
+				Ptr<const Layout> pLayout = GetMemory().GetLayout(FCacheAddress(args.layout, item));
 
 				FIntPoint SizeInBlocks = pLayout->GetGridSize();
 				FIntPoint BlockSizeInPixels(args.blockSize[0], args.blockSize[1]);
 				FIntPoint ImageSizeInPixels = SizeInBlocks * BlockSizeInPixels;
 
-				FImageDesc& ResultAndBaseDesc = m_heapImageDesc[item.customState];
+				FImageDesc& ResultAndBaseDesc = m_heapImageDesc[item.CustomState];
 				FImageSize destSize(uint16(ImageSizeInPixels.X), uint16(ImageSizeInPixels.Y));
 				ResultAndBaseDesc.m_size = destSize;
 				
@@ -5621,10 +5622,10 @@ namespace mu
 
 		case OP_TYPE::IM_COMPOSE:
 		{
-			OP::ImageComposeArgs args = program.GetOpArgs<OP::ImageComposeArgs>(item.at);
-			switch (item.stage)
+			OP::ImageComposeArgs args = program.GetOpArgs<OP::ImageComposeArgs>(item.At);
+			switch (item.Stage)
 			{
-			case 0: AddOp(SCHEDULED_OP(item.at, item, 1), SCHEDULED_OP(args.base, item)); break;
+			case 0: AddOp(FScheduledOp(item.At, item, 1), FScheduledOp(args.base, item)); break;
 			case 1: GetMemory().SetValidDesc(item); break;
 			default: check(false);
 			}
@@ -5633,10 +5634,10 @@ namespace mu
 
 		case OP_TYPE::IM_INTERPOLATE:
 		{
-			OP::ImageInterpolateArgs args = program.GetOpArgs<OP::ImageInterpolateArgs>(item.at);
-			switch (item.stage)
+			OP::ImageInterpolateArgs args = program.GetOpArgs<OP::ImageInterpolateArgs>(item.At);
+			switch (item.Stage)
 			{
-			case 0: AddOp(SCHEDULED_OP(item.at, item, 1), SCHEDULED_OP(args.targets[0], item)); break;
+			case 0: AddOp(FScheduledOp(item.At, item, 1), FScheduledOp(args.targets[0], item)); break;
 			case 1: GetMemory().SetValidDesc(item); break;
 			default: check(false);
 			}
@@ -5645,10 +5646,10 @@ namespace mu
 
 		case OP_TYPE::IM_INTERPOLATE3:
 		{
-			OP::ImageInterpolate3Args args = program.GetOpArgs<OP::ImageInterpolate3Args>(item.at);
-			switch (item.stage)
+			OP::ImageInterpolate3Args args = program.GetOpArgs<OP::ImageInterpolate3Args>(item.At);
+			switch (item.Stage)
 			{
-			case 0: AddOp(SCHEDULED_OP(item.at, item, 1), SCHEDULED_OP(args.target0, item)); break;
+			case 0: AddOp(FScheduledOp(item.At, item, 1), FScheduledOp(args.target0, item)); break;
 			case 1: GetMemory().SetValidDesc(item); break;
 			default: check(false);
 			}
@@ -5657,10 +5658,10 @@ namespace mu
 
 		case OP_TYPE::IM_SATURATE:
 		{
-			OP::ImageSaturateArgs args = program.GetOpArgs<OP::ImageSaturateArgs>(item.at);
-			switch (item.stage)
+			OP::ImageSaturateArgs args = program.GetOpArgs<OP::ImageSaturateArgs>(item.At);
+			switch (item.Stage)
 			{
-			case 0: AddOp(SCHEDULED_OP(item.at, item, 1), SCHEDULED_OP(args.base, item)); break;
+			case 0: AddOp(FScheduledOp(item.At, item, 1), FScheduledOp(args.base, item)); break;
 			case 1: GetMemory().SetValidDesc(item); break;
 			default: check(false);
 			}
@@ -5669,15 +5670,15 @@ namespace mu
 
 		case OP_TYPE::IM_LUMINANCE:
 		{
-			OP::ImageLuminanceArgs args = program.GetOpArgs<OP::ImageLuminanceArgs>(item.at);
-			switch (item.stage)
+			OP::ImageLuminanceArgs args = program.GetOpArgs<OP::ImageLuminanceArgs>(item.At);
+			switch (item.Stage)
 			{
 			case 0:
-				AddOp(SCHEDULED_OP(item.at, item, 1), SCHEDULED_OP(args.base, item));
+				AddOp(FScheduledOp(item.At, item, 1), FScheduledOp(args.base, item));
 				break;
 
 			case 1:
-				m_heapImageDesc[item.customState].m_format = EImageFormat::IF_L_UBYTE;
+				m_heapImageDesc[item.CustomState].m_format = EImageFormat::IF_L_UBYTE;
 				GetMemory().SetValidDesc(item);
 				break;
 
@@ -5690,15 +5691,15 @@ namespace mu
 
 		case OP_TYPE::IM_SWIZZLE:
 		{
-			OP::ImageSwizzleArgs args = program.GetOpArgs<OP::ImageSwizzleArgs>(item.at);
-			switch (item.stage)
+			OP::ImageSwizzleArgs args = program.GetOpArgs<OP::ImageSwizzleArgs>(item.At);
+			switch (item.Stage)
 			{
 			case 0:
-				AddOp(SCHEDULED_OP(item.at, item, 1), SCHEDULED_OP(args.sources[0], item));
+				AddOp(FScheduledOp(item.At, item, 1), FScheduledOp(args.sources[0], item));
 				break;
 
 			case 1:
-				m_heapImageDesc[item.customState].m_format = args.format;
+				m_heapImageDesc[item.CustomState].m_format = args.format;
 				GetMemory().SetValidDesc(item);
 				break;
 
@@ -5711,15 +5712,15 @@ namespace mu
 
 		case OP_TYPE::IM_SELECTCOLOUR:
 		{
-			OP::ImageSelectColourArgs args = program.GetOpArgs<OP::ImageSelectColourArgs>(item.at);
-			switch (item.stage)
+			OP::ImageSelectColourArgs args = program.GetOpArgs<OP::ImageSelectColourArgs>(item.At);
+			switch (item.Stage)
 			{
 			case 0:
-				AddOp(SCHEDULED_OP(item.at, item, 1), SCHEDULED_OP(args.base, item));
+				AddOp(FScheduledOp(item.At, item, 1), FScheduledOp(args.base, item));
 				break;
 
 			case 1:
-				m_heapImageDesc[item.customState].m_format = EImageFormat::IF_L_UBYTE;
+				m_heapImageDesc[item.CustomState].m_format = EImageFormat::IF_L_UBYTE;
 				GetMemory().SetValidDesc(item);
 				break;
 
@@ -5732,10 +5733,10 @@ namespace mu
 
 		case OP_TYPE::IM_COLOURMAP:
 		{
-			OP::ImageColourMapArgs args = program.GetOpArgs<OP::ImageColourMapArgs>(item.at);
-			switch (item.stage)
+			OP::ImageColourMapArgs args = program.GetOpArgs<OP::ImageColourMapArgs>(item.At);
+			switch (item.Stage)
 			{
-			case 0: AddOp(SCHEDULED_OP(item.at, item, 1), SCHEDULED_OP(args.base, item)); break;
+			case 0: AddOp(FScheduledOp(item.At, item, 1), FScheduledOp(args.base, item)); break;
 			case 1: GetMemory().SetValidDesc(item); break;
 			default: check(false);
 			}
@@ -5744,26 +5745,26 @@ namespace mu
 
 		case OP_TYPE::IM_GRADIENT:
 		{
-			OP::ImageGradientArgs args = program.GetOpArgs<OP::ImageGradientArgs>(item.at);
-			m_heapImageDesc[item.customState].m_size[0] = args.size[0];
-			m_heapImageDesc[item.customState].m_size[1] = args.size[1];
-			m_heapImageDesc[item.customState].m_lods = 1;
-			m_heapImageDesc[item.customState].m_format = EImageFormat::IF_RGB_UBYTE;
+			OP::ImageGradientArgs args = program.GetOpArgs<OP::ImageGradientArgs>(item.At);
+			m_heapImageDesc[item.CustomState].m_size[0] = args.size[0];
+			m_heapImageDesc[item.CustomState].m_size[1] = args.size[1];
+			m_heapImageDesc[item.CustomState].m_lods = 1;
+			m_heapImageDesc[item.CustomState].m_format = EImageFormat::IF_RGB_UBYTE;
 			GetMemory().SetValidDesc(item);
 			break;
 		}
 
 		case OP_TYPE::IM_BINARISE:
 		{
-			OP::ImageBinariseArgs args = program.GetOpArgs<OP::ImageBinariseArgs>(item.at);
-			switch (item.stage)
+			OP::ImageBinariseArgs args = program.GetOpArgs<OP::ImageBinariseArgs>(item.At);
+			switch (item.Stage)
 			{
 			case 0:
-				AddOp(SCHEDULED_OP(item.at, item, 1), SCHEDULED_OP(args.base, item));
+				AddOp(FScheduledOp(item.At, item, 1), FScheduledOp(args.base, item));
 				break;
 
 			case 1:
-				m_heapImageDesc[item.customState].m_format = EImageFormat::IF_L_UBYTE;
+				m_heapImageDesc[item.CustomState].m_format = EImageFormat::IF_L_UBYTE;
 				GetMemory().SetValidDesc(item);
 				break;
 
@@ -5775,10 +5776,10 @@ namespace mu
 
 		case OP_TYPE::IM_INVERT:
 		{
-			OP::ImageInvertArgs args = program.GetOpArgs<OP::ImageInvertArgs>(item.at);
-			switch (item.stage)
+			OP::ImageInvertArgs args = program.GetOpArgs<OP::ImageInvertArgs>(item.At);
+			switch (item.Stage)
 			{
-			case 0: AddOp(SCHEDULED_OP(item.at, item, 1), SCHEDULED_OP(args.base, item)); break;
+			case 0: AddOp(FScheduledOp(item.At, item, 1), FScheduledOp(args.base, item)); break;
 			case 1: GetMemory().SetValidDesc(item); break;
 			default: check(false);
 			}
@@ -5787,28 +5788,28 @@ namespace mu
 
 		case OP_TYPE::IM_PLAINCOLOUR:
 		{
-			OP::ImagePlainColourArgs args = program.GetOpArgs<OP::ImagePlainColourArgs>(item.at);
-			m_heapImageDesc[item.customState].m_size[0] = args.size[0];
-			m_heapImageDesc[item.customState].m_size[1] = args.size[1];
-			m_heapImageDesc[item.customState].m_lods = 1;
-			m_heapImageDesc[item.customState].m_format = EImageFormat::IF_RGB_UBYTE;
+			OP::ImagePlainColourArgs args = program.GetOpArgs<OP::ImagePlainColourArgs>(item.At);
+			m_heapImageDesc[item.CustomState].m_size[0] = args.size[0];
+			m_heapImageDesc[item.CustomState].m_size[1] = args.size[1];
+			m_heapImageDesc[item.CustomState].m_lods = 1;
+			m_heapImageDesc[item.CustomState].m_format = EImageFormat::IF_RGB_UBYTE;
 			GetMemory().SetValidDesc(item);
 			break;
 		}
 
 		case OP_TYPE::IM_CROP:
 		{
-			OP::ImageCropArgs args = program.GetOpArgs<OP::ImageCropArgs>(item.at);
-			switch (item.stage)
+			OP::ImageCropArgs args = program.GetOpArgs<OP::ImageCropArgs>(item.At);
+			switch (item.Stage)
 			{
 			case 0:
-				AddOp(SCHEDULED_OP(item.at, item, 1), SCHEDULED_OP(args.source, item));
+				AddOp(FScheduledOp(item.At, item, 1), FScheduledOp(args.source, item));
 				break;
 
 			case 1:
-				m_heapImageDesc[item.customState].m_size[0] = args.sizeX;
-				m_heapImageDesc[item.customState].m_size[1] = args.sizeY;
-				m_heapImageDesc[item.customState].m_lods = 1;
+				m_heapImageDesc[item.CustomState].m_size[0] = args.sizeX;
+				m_heapImageDesc[item.CustomState].m_size[1] = args.sizeY;
+				m_heapImageDesc[item.CustomState].m_lods = 1;
 				GetMemory().SetValidDesc(item);
 				break;
 
@@ -5820,10 +5821,10 @@ namespace mu
 
 		case OP_TYPE::IM_PATCH:
 		{
-			OP::ImagePatchArgs args = program.GetOpArgs<OP::ImagePatchArgs>(item.at);
-			switch (item.stage)
+			OP::ImagePatchArgs args = program.GetOpArgs<OP::ImagePatchArgs>(item.At);
+			switch (item.Stage)
 			{
-			case 0: AddOp(SCHEDULED_OP(item.at, item, 1), SCHEDULED_OP(args.base, item)); break;
+			case 0: AddOp(FScheduledOp(item.At, item, 1), FScheduledOp(args.base, item)); break;
 			case 1: GetMemory().SetValidDesc(item); break;
 			default: check(false);
 			}
@@ -5832,27 +5833,27 @@ namespace mu
 
 		case OP_TYPE::IM_RASTERMESH:
 		{
-			OP::ImageRasterMeshArgs args = program.GetOpArgs<OP::ImageRasterMeshArgs>(item.at);
-			m_heapImageDesc[item.customState].m_size[0] = args.sizeX;
-			m_heapImageDesc[item.customState].m_size[1] = args.sizeY;
-			m_heapImageDesc[item.customState].m_lods = 1;
-			m_heapImageDesc[item.customState].m_format = EImageFormat::IF_L_UBYTE;
+			OP::ImageRasterMeshArgs args = program.GetOpArgs<OP::ImageRasterMeshArgs>(item.At);
+			m_heapImageDesc[item.CustomState].m_size[0] = args.sizeX;
+			m_heapImageDesc[item.CustomState].m_size[1] = args.sizeY;
+			m_heapImageDesc[item.CustomState].m_lods = 1;
+			m_heapImageDesc[item.CustomState].m_format = EImageFormat::IF_L_UBYTE;
 			GetMemory().SetValidDesc(item);
 			break;
 		}
 
 		case OP_TYPE::IM_MAKEGROWMAP:
 		{
-			OP::ImageMakeGrowMapArgs args = program.GetOpArgs<OP::ImageMakeGrowMapArgs>(item.at);
-			switch (item.stage)
+			OP::ImageMakeGrowMapArgs args = program.GetOpArgs<OP::ImageMakeGrowMapArgs>(item.At);
+			switch (item.Stage)
 			{
 			case 0:
-				AddOp(SCHEDULED_OP(item.at, item, 1), SCHEDULED_OP(args.mask, item));
+				AddOp(FScheduledOp(item.At, item, 1), FScheduledOp(args.mask, item));
 				break;
 
 			case 1:
-				m_heapImageDesc[item.customState].m_format = EImageFormat::IF_L_UBYTE;
-				m_heapImageDesc[item.customState].m_lods = 1;
+				m_heapImageDesc[item.CustomState].m_format = EImageFormat::IF_L_UBYTE;
+				m_heapImageDesc[item.CustomState].m_lods = 1;
 				GetMemory().SetValidDesc(item);
 				break;
 
@@ -5865,10 +5866,10 @@ namespace mu
 
 		case OP_TYPE::IM_DISPLACE:
 		{
-			OP::ImageDisplaceArgs args = program.GetOpArgs<OP::ImageDisplaceArgs>(item.at);
-			switch (item.stage)
+			OP::ImageDisplaceArgs args = program.GetOpArgs<OP::ImageDisplaceArgs>(item.At);
+			switch (item.Stage)
 			{
-			case 0: AddOp(SCHEDULED_OP(item.at, item, 1), SCHEDULED_OP(args.source, item)); break;
+			case 0: AddOp(FScheduledOp(item.At, item, 1), FScheduledOp(args.source, item)); break;
 			case 1: GetMemory().SetValidDesc(item); break;
 			default: check(false);
 			}
@@ -5878,13 +5879,13 @@ namespace mu
         case OP_TYPE::IM_TRANSFORM:
         {
 
-			OP::ImageTransformArgs Args = program.GetOpArgs<OP::ImageTransformArgs>(item.at);
+			OP::ImageTransformArgs Args = program.GetOpArgs<OP::ImageTransformArgs>(item.At);
 
-            switch (item.stage)
+            switch (item.Stage)
             {
             case 0:
 			{
-				AddOp(SCHEDULED_OP(item.at, item, 1), SCHEDULED_OP(Args.base, item));
+				AddOp(FScheduledOp(item.At, item, 1), FScheduledOp(Args.base, item));
                 break;
 			}
             case 1:
@@ -5905,7 +5906,7 @@ namespace mu
 			{
 				// Operation not implemented
 				check(false);
-				m_heapImageDesc[item.customState] = FImageDesc();
+				m_heapImageDesc[item.CustomState] = FImageDesc();
 			}
 			break;
 		}
diff --git a/Engine/Plugins/Experimental/Mutable/Source/MutableRuntime/Private/MuR/CodeRunner.h b/Engine/Plugins/Experimental/Mutable/Source/MutableRuntime/Private/MuR/CodeRunner.h
--- a/Engine/Plugins/Experimental/Mutable/Source/MutableRuntime/Private/MuR/CodeRunner.h
+++ b/Engine/Plugins/Experimental/Mutable/Source/MutableRuntime/Private/MuR/CodeRunner.h
@@ -31,9 +31,9 @@
 
 namespace  mu
 {
-class Model;
-class Parameters;
-class RangeIndex;
+	class Model;
+	class Parameters;
+	class RangeIndex;
 
 
     //---------------------------------------------------------------------------------------------
@@ -44,7 +44,7 @@ class RangeIndex;
     public:
 		CodeRunner(const SettingsPtrConst&, class System::Private*, 
 			const Model* pModel, const Parameters* pParams,
-			OP::ADDRESS at, uint32 lodMask, uint8 executionOptions, SCHEDULED_OP::EType );
+			OP::ADDRESS at, uint32 lodMask, uint8 executionOptions, FScheduledOp::EType );
 
     protected:
 
@@ -56,36 +56,51 @@ class RangeIndex;
         //! operation stages.
         struct FScheduledOpData
         {
-            float bifactor=0.0f;
-            int32 min=0, max=0;
-            Ptr<const Layout> layout;
-        };
+			union
+			{
+				struct
+				{
+					float Bifactor;
+					int32 Min, Max;
+				} Interpolate;
+
+				struct
+				{
+					int32 Iterations;
+					EImageFormat OriginalBaseFormat;
+					bool bBlendOnlyOneMip;
+				} MultiLayer;
 
+				struct
+				{
+					int32 ResultDescAt;
+					int32 SourceDescAt;
+				} ResizeLike;
+			};
+            Ptr<RefCounted> Resource;
+        };
 
-        Ptr<RangeIndex> BuildCurrentOpRangeIndex( const SCHEDULED_OP& item,
-                                                const Parameters* pParams,
-                                                const Model* pModel,
-                                                int parameterIndex );
 
+        Ptr<RangeIndex> BuildCurrentOpRangeIndex( const FScheduledOp&, const Parameters*, const Model*, int32 ParameterIndex );
 
-        void RunCode( SCHEDULED_OP& item, const Parameters* pParams, const Model* pModel, uint32_t lodMask);
+        void RunCode( FScheduledOp&, const Parameters*, const Model*, uint32 LodMask);
 
-        void RunCode_Conditional( SCHEDULED_OP& item, const Model* pModel );
-        void RunCode_Switch( SCHEDULED_OP& item, const Model* pModel );
-        void RunCode_Instance( SCHEDULED_OP& item, const Model* pModel, uint32_t lodMask );
-        void RunCode_InstanceAddResource( SCHEDULED_OP& item, const Model* pModel, const Parameters* pParams );
-        void RunCode_ConstantResource( SCHEDULED_OP& item, const Model* pModel );
-        void RunCode_Mesh( SCHEDULED_OP& item, const Model* pModel );
-        void RunCode_Image( SCHEDULED_OP& item, const Parameters* pParams, const Model* pModel );
-        void RunCode_Layout( SCHEDULED_OP& item, const Model* pModel );
-        void RunCode_Bool( SCHEDULED_OP& item, const Parameters* pParams, const Model* pModel );
-        void RunCode_Int( SCHEDULED_OP& item, const Parameters* pParams, const Model* pModel );
-        void RunCode_Scalar( SCHEDULED_OP& item, const Parameters* pParams, const Model* pModel );        
-        void RunCode_String( SCHEDULED_OP& item, const Parameters* pParams, const Model* pModel );
-        void RunCode_Colour( SCHEDULED_OP& item, const Parameters* pParams, const Model* pModel );
-        void RunCode_Projector( SCHEDULED_OP& item, const Parameters* pParams, const Model* pModel );
+        void RunCode_Conditional( FScheduledOp&, const Model* );
+        void RunCode_Switch( FScheduledOp&, const Model* );
+        void RunCode_Instance( FScheduledOp&, const Model*, uint32 LodMask );
+        void RunCode_InstanceAddResource( FScheduledOp&, const Model*, const Parameters* );
+        void RunCode_ConstantResource( FScheduledOp&, const Model* );
+        void RunCode_Mesh( FScheduledOp&, const Model* );
+        void RunCode_Image( FScheduledOp&, const Parameters*, const Model* );
+        void RunCode_Layout( FScheduledOp&, const Model* );
+        void RunCode_Bool( FScheduledOp&, const Parameters*, const Model* );
+        void RunCode_Int( FScheduledOp&, const Parameters*, const Model* );
+        void RunCode_Scalar( FScheduledOp&, const Parameters*, const Model* );        
+        void RunCode_String( FScheduledOp&, const Parameters*, const Model* );
+        void RunCode_Colour( FScheduledOp&, const Parameters*, const Model* );
+        void RunCode_Projector( FScheduledOp&, const Parameters*, const Model* );
 
-		void RunCodeImageDesc(SCHEDULED_OP& item, const Parameters* pParams, const Model* pModel, uint32_t lodMask);
+		void RunCodeImageDesc(FScheduledOp&, const Parameters*, const Model*, uint32 LodMask);
 
 
         ImagePtr LoadExternalImage( EXTERNAL_IMAGE_ID id );
@@ -109,21 +124,21 @@ class RangeIndex;
 		struct FTask
 		{
 			FTask() {}
-			FTask(const SCHEDULED_OP& InOp) : Op(InOp) {}
-			FTask(const SCHEDULED_OP& InOp, const SCHEDULED_OP& InDep0) : Op(InOp) { Deps.Add(InDep0); }
-			FTask(const SCHEDULED_OP& InOp, const SCHEDULED_OP& InDep0, const SCHEDULED_OP& InDep1) : Op(InOp) { Deps.Add(InDep0); Deps.Add(InDep1); }
-			FTask(const SCHEDULED_OP& InOp, const SCHEDULED_OP& InDep0, const SCHEDULED_OP& InDep1, const SCHEDULED_OP& InDep2) : Op(InOp) { Deps.Add(InDep0); Deps.Add(InDep1); Deps.Add(InDep2); }
-			FTask(const SCHEDULED_OP& InOp, const SCHEDULED_OP& InDep0, const SCHEDULED_OP& InDep1, const SCHEDULED_OP& InDep2, const SCHEDULED_OP& InDep3) : Op(InOp) { Deps.Add(InDep0); Deps.Add(InDep1); Deps.Add(InDep2); Deps.Add(InDep3); }
-			FTask(const SCHEDULED_OP& InOp, const SCHEDULED_OP& InDep0, const SCHEDULED_OP& InDep1, const SCHEDULED_OP& InDep2, const SCHEDULED_OP& InDep3, const SCHEDULED_OP& InDep4) : Op(InOp) { Deps.Add(InDep0); Deps.Add(InDep1); Deps.Add(InDep2); Deps.Add(InDep3); Deps.Add(InDep4); }
-
-			SCHEDULED_OP Op;
-			TArray<CACHE_ADDRESS, TInlineAllocator<3>> Deps;
+			FTask(const FScheduledOp& InOp) : Op(InOp) {}
+			FTask(const FScheduledOp& InOp, const FScheduledOp& InDep0) : Op(InOp) { Deps.Add(InDep0); }
+			FTask(const FScheduledOp& InOp, const FScheduledOp& InDep0, const FScheduledOp& InDep1) : Op(InOp) { Deps.Add(InDep0); Deps.Add(InDep1); }
+			FTask(const FScheduledOp& InOp, const FScheduledOp& InDep0, const FScheduledOp& InDep1, const FScheduledOp& InDep2) : Op(InOp) { Deps.Add(InDep0); Deps.Add(InDep1); Deps.Add(InDep2); }
+			FTask(const FScheduledOp& InOp, const FScheduledOp& InDep0, const FScheduledOp& InDep1, const FScheduledOp& InDep2, const FScheduledOp& InDep3) : Op(InOp) { Deps.Add(InDep0); Deps.Add(InDep1); Deps.Add(InDep2); Deps.Add(InDep3); }
+			FTask(const FScheduledOp& InOp, const FScheduledOp& InDep0, const FScheduledOp& InDep1, const FScheduledOp& InDep2, const FScheduledOp& InDep3, const FScheduledOp& InDep4) : Op(InOp) { Deps.Add(InDep0); Deps.Add(InDep1); Deps.Add(InDep2); Deps.Add(InDep3); Deps.Add(InDep4); }
+
+			FScheduledOp Op;
+			TArray<FCacheAddress, TInlineAllocator<3>> Deps;
 		};
 
 		class FIssuedTask
 		{
 		public:
-			SCHEDULED_OP Op;
+			FScheduledOp Op;
 
 #ifdef MUTABLE_USE_NEW_TASKGRAPH
 			UE::Tasks::FTask Event = {};
@@ -149,7 +164,7 @@ class RangeIndex;
 		class FLoadMeshRomTask : public CodeRunner::FIssuedTask
 		{
 		public:
-			FLoadMeshRomTask(SCHEDULED_OP InOp, int32 InRomIndex)
+			FLoadMeshRomTask(FScheduledOp InOp, int32 InRomIndex)
 			{
 				Op = InOp;
 				RomIndex = InRomIndex;
@@ -167,7 +182,7 @@ class RangeIndex;
 		class FLoadImageRomsTask : public CodeRunner::FIssuedTask
 		{
 		public:
-			FLoadImageRomsTask(SCHEDULED_OP InOp, int32 InLODIndexIndex, int32 InLODIndexCount )
+			FLoadImageRomsTask(FScheduledOp InOp, int32 InLODIndexIndex, int32 InLODIndexCount )
 			{
 				Op = InOp;
 				LODIndexIndex = InLODIndexIndex;
@@ -184,65 +199,65 @@ class RangeIndex;
 			int32 LODIndexCount = -1;
 		};
 
-		void AddOp(const SCHEDULED_OP& op)
+		void AddOp(const FScheduledOp& op)
 		{
 			OpenTasks.Add(op);
-			ScheduledStagePerOp[op] = op.stage + 1;
+			ScheduledStagePerOp[op] = op.Stage + 1;
 		}
 
-		void AddOp(const SCHEDULED_OP& op,
-			const SCHEDULED_OP& dep0)
+		void AddOp(const FScheduledOp& op,
+			const FScheduledOp& dep0)
 		{
 			ClosedTasks.Add(FTask(op, dep0));
-			ScheduledStagePerOp[op] = op.stage + 1;
+			ScheduledStagePerOp[op] = op.Stage + 1;
 			AddChildren(dep0);
 		}
 
-		void AddOp(const SCHEDULED_OP& op,
-			const SCHEDULED_OP& dep0,
-			const SCHEDULED_OP& dep1)
+		void AddOp(const FScheduledOp& op,
+			const FScheduledOp& dep0,
+			const FScheduledOp& dep1)
 		{
 			ClosedTasks.Add(FTask(op, dep0, dep1));
-			ScheduledStagePerOp[op] = op.stage + 1;
+			ScheduledStagePerOp[op] = op.Stage + 1;
 			AddChildren(dep0);
 			AddChildren(dep1);
 		}
 
-		void AddOp(const SCHEDULED_OP& op,
-			const SCHEDULED_OP& dep0,
-			const SCHEDULED_OP& dep1,
-			const SCHEDULED_OP& dep2)
+		void AddOp(const FScheduledOp& op,
+			const FScheduledOp& dep0,
+			const FScheduledOp& dep1,
+			const FScheduledOp& dep2)
 		{
 			ClosedTasks.Add(FTask(op, dep0, dep1, dep2));
-			ScheduledStagePerOp[op] = op.stage + 1;
+			ScheduledStagePerOp[op] = op.Stage + 1;
 			AddChildren(dep0);
 			AddChildren(dep1);
 			AddChildren(dep2);
 		}
 
-		void AddOp(const SCHEDULED_OP& op,
-			const SCHEDULED_OP& dep0,
-			const SCHEDULED_OP& dep1,
-			const SCHEDULED_OP& dep2,
-			const SCHEDULED_OP& dep3)
+		void AddOp(const FScheduledOp& op,
+			const FScheduledOp& dep0,
+			const FScheduledOp& dep1,
+			const FScheduledOp& dep2,
+			const FScheduledOp& dep3)
 		{
 			ClosedTasks.Add(FTask(op, dep0, dep1, dep2, dep3));
-			ScheduledStagePerOp[op] = op.stage + 1;
+			ScheduledStagePerOp[op] = op.Stage + 1;
 			AddChildren(dep0);
 			AddChildren(dep1);
 			AddChildren(dep2);
 			AddChildren(dep3);
 		}
 
-		void AddOp(const SCHEDULED_OP& op,
-			const SCHEDULED_OP& dep0,
-			const SCHEDULED_OP& dep1,
-			const SCHEDULED_OP& dep2,
-			const SCHEDULED_OP& dep3,
-			const SCHEDULED_OP& dep4)
+		void AddOp(const FScheduledOp& op,
+			const FScheduledOp& dep0,
+			const FScheduledOp& dep1,
+			const FScheduledOp& dep2,
+			const FScheduledOp& dep3,
+			const FScheduledOp& dep4)
 		{
 			ClosedTasks.Add(FTask(op, dep0, dep1, dep2, dep3, dep4));
-			ScheduledStagePerOp[op] = op.stage + 1;
+			ScheduledStagePerOp[op] = op.Stage + 1;
 			AddChildren(dep0);
 			AddChildren(dep1);
 			AddChildren(dep2);
@@ -250,35 +265,21 @@ class RangeIndex;
 			AddChildren(dep4);
 		}
 
-		void AddOp(const SCHEDULED_OP& op, const TArray<SCHEDULED_OP>& deps)
-		{
-			FTask Task(op);
-			for (const auto& d : deps)
-			{
-				Task.Deps.Add(d);
-			}
-			ClosedTasks.Add(Task);
-			for (const auto& d : deps)
-			{
-				AddChildren(d);
-			}
-		}
-
 		template<class RangeType>
-		void AddOp(const SCHEDULED_OP& Op, const RangeType& Deps)
+		void AddOp(const FScheduledOp& Op, const RangeType& Deps)
 		{
 			FTask Task(Op);
 
-			for (const SCHEDULED_OP& D : Deps)
+			for (const FScheduledOp& D : Deps)
 			{
 				Task.Deps.Add(D);
 			}
 
 			ClosedTasks.Add(Task);
 
-			ScheduledStagePerOp[Op] = Op.stage + 1;
+			ScheduledStagePerOp[Op] = Op.Stage + 1;
 
-			for (const SCHEDULED_OP& D : Deps)
+			for (const FScheduledOp& D : Deps)
 			{
 				AddChildren(D);
 			}
@@ -289,7 +290,7 @@ class RangeIndex;
 
 		//! Stack of pending operations, and the execution stage they are in.
 		TArray< FTask > ClosedTasks;
-		TArray< SCHEDULED_OP > OpenTasks;
+		TArray< FScheduledOp > OpenTasks;
 		CodeContainer<uint8> ScheduledStagePerOp;
 		TArray< TSharedPtr<FIssuedTask> > IssuedTasks;
 
@@ -307,29 +308,29 @@ class RangeIndex;
 		uint32 m_lodMask = 0;
 
 		// Async rom loading control
-		struct ROM_LOAD_OP
+		struct FRomLoadOp
 		{
 			int32 m_romIndex = 0;
 			DATATYPE ConstantType = DT_NONE;
 			ModelStreamer::OPERATION_ID m_streamID;
 			TArray<uint8> m_streamBuffer;
 		};
-		TArray<ROM_LOAD_OP> m_romLoadOps;
+		TArray<FRomLoadOp> m_romLoadOps;
 
 		//! Count of pending operations for every rom index
 		TArray<uint16> m_romPendingOps;
 
 	private:
 
-		inline void AddChildren(const SCHEDULED_OP& dep)
+		inline void AddChildren(const FScheduledOp& dep)
 		{
-			CACHE_ADDRESS at(dep);
-			if (dep.at && !GetMemory().IsValid(at))
+			FCacheAddress at(dep);
+			if (dep.At && !GetMemory().IsValid(at))
 			{
-				if (ScheduledStagePerOp.get(at) <= dep.stage)
+				if (ScheduledStagePerOp.get(at) <= dep.Stage)
 				{
 					OpenTasks.Add(dep);
-					ScheduledStagePerOp[at] = dep.stage + 1;
+					ScheduledStagePerOp[at] = dep.Stage + 1;
 				}
 			}
 
@@ -337,10 +338,10 @@ class RangeIndex;
 		}
 
 		/** Try to create a concurrent task for the given op. Return null if not possible. */
-		TSharedPtr<FIssuedTask> IssueOp(SCHEDULED_OP item);
+		TSharedPtr<FIssuedTask> IssueOp(FScheduledOp item);
 
 		/** */
-		void CompleteRomLoadOp(ROM_LOAD_OP& o);
+		void CompleteRomLoadOp(FRomLoadOp& o);
 
     };
 
diff --git a/Engine/Plugins/Experimental/Mutable/Source/MutableRuntime/Private/MuR/CodeRunnerTasked.cpp b/Engine/Plugins/Experimental/Mutable/Source/MutableRuntime/Private/MuR/CodeRunnerTasked.cpp
--- a/Engine/Plugins/Experimental/Mutable/Source/MutableRuntime/Private/MuR/CodeRunnerTasked.cpp
+++ b/Engine/Plugins/Experimental/Mutable/Source/MutableRuntime/Private/MuR/CodeRunnerTasked.cpp
@@ -48,7 +48,7 @@ DECLARE_CYCLE_STAT(TEXT("MutableCoreTask"), STAT_MutableCoreTask, STATGROUP_Game
 namespace mu
 {
 
-	void CodeRunner::CompleteRomLoadOp( ROM_LOAD_OP& o )
+	void CodeRunner::CompleteRomLoadOp( FRomLoadOp& o )
 	{
 		if (DebugRom && (DebugRomAll || o.m_romIndex == DebugRomIndex))
 			UE_LOG(LogMutableCore, Log, TEXT("CodeRunner::CompleteRomLoadOp for rom %d."), o.m_romIndex);
@@ -110,7 +110,7 @@ namespace mu
 			//	int32 ClosedReady = ClosedTasks.Num();
 			//	for (int Index = ClosedTasks.Num() - 1; Index >= 0; --Index)
 			//	{
-			//		for (const CACHE_ADDRESS& Dep : ClosedTasks[Index].Deps)
+			//		for (const FCacheAddress& Dep : ClosedTasks[Index].Deps)
 			//		{
 			//			if (Dep.at && !GetMemory().IsValid(Dep))
 			//			{
@@ -130,10 +130,10 @@ namespace mu
 				bool WorkDone = IssuedTasks[Index]->IsComplete(this);
 				if (WorkDone)
 				{
-					const SCHEDULED_OP& item = IssuedTasks[Index]->Op;
+					const FScheduledOp& item = IssuedTasks[Index]->Op;
 					IssuedTasks[Index]->Complete(this);
 
-					if (ScheduledStagePerOp[item] == item.stage + 1)
+					if (ScheduledStagePerOp[item] == item.Stage + 1)
 					{
 						// We completed everything that was requested, clear it otherwise if needed
 						// again it is not going to be rebuilt.
@@ -152,17 +152,17 @@ namespace mu
 
 			while (!OpenTasks.IsEmpty())
 			{
-				SCHEDULED_OP item = OpenTasks.Pop();
+				FScheduledOp item = OpenTasks.Pop();
 
 				// Special processing in case it is an ImageDesc operation
-				if ( item.Type==SCHEDULED_OP::EType::ImageDesc )
+				if ( item.Type==FScheduledOp::EType::ImageDesc )
 				{
 					RunCodeImageDesc(item, m_pParams, m_pModel, m_lodMask);
 					continue;
 				}
 
 				// Don't run it if we already have the result.
-				CACHE_ADDRESS cat(item);
+				FCacheAddress cat(item);
 				if (GetMemory().IsValid(cat))
 				{
 					continue;
@@ -213,7 +213,7 @@ namespace mu
 					// Run immediately
 					RunCode(item, m_pParams, m_pModel, m_lodMask);
 
-					if (ScheduledStagePerOp[item] == item.stage + 1)
+					if (ScheduledStagePerOp[item] == item.Stage + 1)
 					{
 						// We completed everything that was requested, clear it otherwise if needed
 						// again it is not going to be rebuilt.
@@ -225,7 +225,7 @@ namespace mu
 				if (bProfile)
 				{
 					++NumRunOps;
-					RunOpsPerType[size_t(m_pModel->GetPrivate()->m_program.GetOpType(item.at))]++;
+					RunOpsPerType[size_t(m_pModel->GetPrivate()->m_program.GetOpType(item.At))]++;
 				}
 			}
 
@@ -243,10 +243,10 @@ namespace mu
 			for (int Index = 0; Index<ClosedTasks.Num(); )
 			{
 				bool Ready = true;
-				for ( const CACHE_ADDRESS& Dep: ClosedTasks[Index].Deps )
+				for ( const FCacheAddress& Dep: ClosedTasks[Index].Deps )
 				{
 					bool bDependencySatisfied = false;					
-					bDependencySatisfied = Dep.at && !GetMemory().IsValid(Dep);
+					bDependencySatisfied = Dep.At && !GetMemory().IsValid(Dep);
 					
 					if (bDependencySatisfied)
 					{
@@ -275,12 +275,12 @@ namespace mu
 				// Log the task graph
 				for (int Index = 0; Index < ClosedTasks.Num(); ++Index)
 				{
-					FString TaskDesc = FString::Printf(TEXT("Closed task %d-%d-%d depends on : "), ClosedTasks[Index].Op.at, ClosedTasks[Index].Op.executionIndex, ClosedTasks[Index].Op.stage );
-					for (const CACHE_ADDRESS& Dep : ClosedTasks[Index].Deps)
+					FString TaskDesc = FString::Printf(TEXT("Closed task %d-%d-%d depends on : "), ClosedTasks[Index].Op.At, ClosedTasks[Index].Op.ExecutionIndex, ClosedTasks[Index].Op.Stage );
+					for (const FCacheAddress& Dep : ClosedTasks[Index].Deps)
 					{
-						if (Dep.at && !GetMemory().IsValid(Dep))
+						if (Dep.At && !GetMemory().IsValid(Dep))
 						{
-							TaskDesc += FString::Printf(TEXT("%d-%d, "), Dep.at, Dep.executionIndex);
+							TaskDesc += FString::Printf(TEXT("%d-%d, "), Dep.At, Dep.ExecutionIndex);
 						}
 					}
 
@@ -372,7 +372,7 @@ namespace mu
 	class FImageLayerTask : public CodeRunner::FIssuedTask
 	{
 	public:
-		FImageLayerTask(SCHEDULED_OP, FProgramCache&, const OP::ImageLayerArgs&, int imageCompressionQuality );
+		FImageLayerTask(FScheduledOp, FProgramCache&, const OP::ImageLayerArgs&, int imageCompressionQuality );
 
 		// FIssuedTask interface
 		void DoWork() override;
@@ -389,16 +389,16 @@ namespace mu
 
 
 	//---------------------------------------------------------------------------------------------
-	FImageLayerTask::FImageLayerTask(SCHEDULED_OP InOp, FProgramCache& Memory, const OP::ImageLayerArgs& InArgs, int imageCompressionQuality)
+	FImageLayerTask::FImageLayerTask(FScheduledOp InOp, FProgramCache& Memory, const OP::ImageLayerArgs& InArgs, int imageCompressionQuality)
 	{
 		Op = InOp;
 		m_imageCompressionQuality = imageCompressionQuality;
 		Args = InArgs;
-		m_base = Memory.GetImage({ Args.base, InOp.executionIndex, InOp.executionOptions });
-		m_blended = Memory.GetImage({ Args.blended, InOp.executionIndex, InOp.executionOptions });
+		m_base = Memory.GetImage({ Args.base, InOp.ExecutionIndex, InOp.ExecutionOptions });
+		m_blended = Memory.GetImage({ Args.blended, InOp.ExecutionIndex, InOp.ExecutionOptions });
 		if (Args.mask)
 		{
-			m_mask = Memory.GetImage({ Args.mask, InOp.executionIndex, InOp.executionOptions });
+			m_mask = Memory.GetImage({ Args.mask, InOp.ExecutionIndex, InOp.ExecutionOptions });
 		}
 	}
 
@@ -562,7 +562,7 @@ namespace mu
 	class FImageLayerColourTask : public CodeRunner::FIssuedTask
 	{
 	public:
-		FImageLayerColourTask(SCHEDULED_OP, FProgramCache&, const OP::ImageLayerColourArgs&, int imageCompressionQuality);
+		FImageLayerColourTask(FScheduledOp, FProgramCache&, const OP::ImageLayerColourArgs&, int imageCompressionQuality);
 
 		// FIssuedTask interface
 		void DoWork() override;
@@ -579,16 +579,16 @@ namespace mu
 
 
 	//---------------------------------------------------------------------------------------------
-	FImageLayerColourTask::FImageLayerColourTask(SCHEDULED_OP InOp, FProgramCache& Memory, const OP::ImageLayerColourArgs& InArgs, int imageCompressionQuality)
+	FImageLayerColourTask::FImageLayerColourTask(FScheduledOp InOp, FProgramCache& Memory, const OP::ImageLayerColourArgs& InArgs, int imageCompressionQuality)
 	{
 		Op = InOp;
 		m_imageCompressionQuality = imageCompressionQuality;
 		Args = InArgs;
-		m_base = Memory.GetImage({ Args.base, InOp.executionIndex, InOp.executionOptions });
-		m_col = Memory.GetColour({ Args.colour, InOp.executionIndex, 0 });
+		m_base = Memory.GetImage({ Args.base, InOp.ExecutionIndex, InOp.ExecutionOptions });
+		m_col = Memory.GetColour({ Args.colour, InOp.ExecutionIndex, 0 });
 		if (Args.mask)
 		{
-			m_mask = Memory.GetImage({ Args.mask, InOp.executionIndex, InOp.executionOptions });
+			m_mask = Memory.GetImage({ Args.mask, InOp.ExecutionIndex, InOp.ExecutionOptions });
 		}
 	}
 
@@ -747,7 +747,7 @@ namespace mu
 	class FImagePixelFormatTask : public CodeRunner::FIssuedTask
 	{
 	public:
-		FImagePixelFormatTask(SCHEDULED_OP, FProgramCache&, const OP::ImagePixelFormatArgs&, int imageCompressionQuality);
+		FImagePixelFormatTask(FScheduledOp, FProgramCache&, const OP::ImagePixelFormatArgs&, int imageCompressionQuality);
 
 		// FIssuedTask interface
 		void DoWork() override;
@@ -762,12 +762,12 @@ namespace mu
 
 
 	//---------------------------------------------------------------------------------------------
-	FImagePixelFormatTask::FImagePixelFormatTask(SCHEDULED_OP InOp, FProgramCache& Memory, const OP::ImagePixelFormatArgs& InArgs, int imageCompressionQuality)
+	FImagePixelFormatTask::FImagePixelFormatTask(FScheduledOp InOp, FProgramCache& Memory, const OP::ImagePixelFormatArgs& InArgs, int imageCompressionQuality)
 	{
 		Op = InOp;
 		m_imageCompressionQuality = imageCompressionQuality;
 		Args = InArgs;
-		m_base = Memory.GetImage({ Args.source, InOp.executionIndex, InOp.executionOptions });
+		m_base = Memory.GetImage({ Args.source, InOp.ExecutionIndex, InOp.ExecutionOptions });
 	}
 
 
@@ -833,7 +833,7 @@ namespace mu
 	class FImageMipmapTask : public CodeRunner::FIssuedTask
 	{
 	public:
-		FImageMipmapTask(SCHEDULED_OP, FProgramCache&, const OP::ImageMipmapArgs&, int imageCompressionQuality);
+		FImageMipmapTask(FScheduledOp, FProgramCache&, const OP::ImageMipmapArgs&, int imageCompressionQuality);
 
 		// FIssuedTask interface
 		void DoWork() override;
@@ -848,12 +848,12 @@ namespace mu
 
 
 	//---------------------------------------------------------------------------------------------
-	FImageMipmapTask::FImageMipmapTask(SCHEDULED_OP InOp, FProgramCache& Memory, const OP::ImageMipmapArgs& InArgs, int imageCompressionQuality)
+	FImageMipmapTask::FImageMipmapTask(FScheduledOp InOp, FProgramCache& Memory, const OP::ImageMipmapArgs& InArgs, int imageCompressionQuality)
 	{
 		Op = InOp;
 		m_imageCompressionQuality = imageCompressionQuality;
 		Args = InArgs;
-		m_base = Memory.GetImage({ Args.source, InOp.executionIndex, InOp.executionOptions });
+		m_base = Memory.GetImage({ Args.source, InOp.ExecutionIndex, InOp.ExecutionOptions });
 	}
 
 
@@ -914,7 +914,7 @@ namespace mu
 	class FImageSwizzleTask : public CodeRunner::FIssuedTask
 	{
 	public:
-		FImageSwizzleTask(SCHEDULED_OP, FProgramCache&, const OP::ImageSwizzleArgs&);
+		FImageSwizzleTask(FScheduledOp, FProgramCache&, const OP::ImageSwizzleArgs&);
 
 		// FIssuedTask interface
 		void DoWork() override;
@@ -928,7 +928,7 @@ namespace mu
 
 
 	//---------------------------------------------------------------------------------------------
-	FImageSwizzleTask::FImageSwizzleTask(SCHEDULED_OP InOp, FProgramCache& Memory, const OP::ImageSwizzleArgs& InArgs)
+	FImageSwizzleTask::FImageSwizzleTask(FScheduledOp InOp, FProgramCache& Memory, const OP::ImageSwizzleArgs& InArgs)
 	{
 		Op = InOp;
 		Args = InArgs;
@@ -936,7 +936,7 @@ namespace mu
 		{
 			if (Args.sources[i])
 			{
-				Sources[i] = Memory.GetImage({ Args.sources[i], InOp.executionIndex, InOp.executionOptions });
+				Sources[i] = Memory.GetImage({ Args.sources[i], InOp.ExecutionIndex, InOp.ExecutionOptions });
 			}
 		}
 	}
@@ -983,7 +983,7 @@ namespace mu
 	class FImageComposeTask : public CodeRunner::FIssuedTask
 	{
 	public:
-		FImageComposeTask(SCHEDULED_OP, FProgramCache&, const OP::ImageComposeArgs&, int imageCompressionQuality, const Ptr<const Layout>& );
+		FImageComposeTask(FScheduledOp, FProgramCache&, const OP::ImageComposeArgs&, int imageCompressionQuality, const Ptr<const Layout>& );
 
 		// FIssuedTask interface
 		void DoWork() override;
@@ -1001,23 +1001,23 @@ namespace mu
 
 
 	//---------------------------------------------------------------------------------------------
-	FImageComposeTask::FImageComposeTask(SCHEDULED_OP InOp, FProgramCache& Memory, const OP::ImageComposeArgs& InArgs, int imageCompressionQuality, const Ptr<const Layout>& InLayout)
+	FImageComposeTask::FImageComposeTask(FScheduledOp InOp, FProgramCache& Memory, const OP::ImageComposeArgs& InArgs, int imageCompressionQuality, const Ptr<const Layout>& InLayout)
 	{
 		Op = InOp;
 		m_imageCompressionQuality = imageCompressionQuality;
 		Args = InArgs;
 		m_layout = InLayout;
 
-		m_base = Memory.GetImage({ Args.base, Op.executionIndex, InOp.executionOptions });
+		m_base = Memory.GetImage({ Args.base, Op.ExecutionIndex, InOp.ExecutionOptions });
 		
 		int relBlockIndex = m_layout->FindBlock(Args.blockIndex);
 
 		if (relBlockIndex >= 0)
 		{
-			m_block = Memory.GetImage({ Args.blockImage, Op.executionIndex, InOp.executionOptions });
+			m_block = Memory.GetImage({ Args.blockImage, Op.ExecutionIndex, InOp.ExecutionOptions });
 			if (Args.mask)
 			{
-				m_mask = Memory.GetImage({ Args.mask, Op.executionIndex, InOp.executionOptions });
+				m_mask = Memory.GetImage({ Args.mask, Op.ExecutionIndex, InOp.ExecutionOptions });
 			}
 		}
 	}
@@ -1162,7 +1162,7 @@ namespace mu
 			uint32 RomSize = program.m_roms[RomIndex].Size;
 			check(RomSize>0);
 
-			CodeRunner::ROM_LOAD_OP* op = nullptr;
+			CodeRunner::FRomLoadOp* op = nullptr;
 			for (auto& o : runner->m_romLoadOps)
 			{
 				if (o.m_romIndex < 0)
@@ -1173,7 +1173,7 @@ namespace mu
 
 			if (!op)
 			{
-				runner->m_romLoadOps.Add(CodeRunner::ROM_LOAD_OP());
+				runner->m_romLoadOps.Add(CodeRunner::FRomLoadOp());
 				op = &runner->m_romLoadOps.Last();
 			}
 
@@ -1282,7 +1282,7 @@ namespace mu
 			uint32 RomSize = program.m_roms[RomIndex].Size;
 			check(RomSize > 0);
 
-			CodeRunner::ROM_LOAD_OP* op = nullptr;
+			CodeRunner::FRomLoadOp* op = nullptr;
 			for (auto& o : runner->m_romLoadOps)
 			{
 				if (o.m_romIndex < 0)
@@ -1293,7 +1293,7 @@ namespace mu
 
 			if (!op)
 			{
-				runner->m_romLoadOps.Add(CodeRunner::ROM_LOAD_OP());
+				runner->m_romLoadOps.Add(CodeRunner::FRomLoadOp());
 				op = &runner->m_romLoadOps.Last();
 			}
 
@@ -1387,19 +1387,19 @@ namespace mu
 	//---------------------------------------------------------------------------------------------
 	//---------------------------------------------------------------------------------------------
 	//---------------------------------------------------------------------------------------------
-	TSharedPtr<CodeRunner::FIssuedTask> CodeRunner::IssueOp(SCHEDULED_OP item)
+	TSharedPtr<CodeRunner::FIssuedTask> CodeRunner::IssueOp(FScheduledOp item)
 	{
 		TSharedPtr<FIssuedTask> Issued;
 
 		FProgram& program = m_pModel->GetPrivate()->m_program;
 
-		auto type = program.GetOpType(item.at);
+		auto type = program.GetOpType(item.At);
 
 		switch (type)
 		{
 		case OP_TYPE::ME_CONSTANT:
 		{
-			OP::MeshConstantArgs args = program.GetOpArgs<OP::MeshConstantArgs>(item.at);
+			OP::MeshConstantArgs args = program.GetOpArgs<OP::MeshConstantArgs>(item.At);
 			int32 RomIndex = program.m_constantMeshes[args.value].Key;
 			if (RomIndex >= 0 && !program.m_constantMeshes[args.value].Value )
 			{
@@ -1414,8 +1414,8 @@ namespace mu
 
 		case OP_TYPE::IM_CONSTANT:
 		{
-			OP::ResourceConstantArgs args = program.GetOpArgs<OP::ResourceConstantArgs>(item.at);
-			int32 MipsToSkip = item.executionOptions;
+			OP::ResourceConstantArgs args = program.GetOpArgs<OP::ResourceConstantArgs>(item.At);
+			int32 MipsToSkip = item.ExecutionOptions;
 			int32 ImageIndex = args.value;
 			int32 ReallySkip = FMath::Min(MipsToSkip, program.m_constantImages[ImageIndex].LODCount - 1);
 			int32 LODIndexIndex = program.m_constantImages[ImageIndex].FirstIndex + ReallySkip;
@@ -1451,9 +1451,9 @@ namespace mu
 
 		case OP_TYPE::IM_PIXELFORMAT:
 		{
-			if (item.stage == 1)
+			if (item.Stage == 1)
 			{
-				OP::ImagePixelFormatArgs args = program.GetOpArgs<OP::ImagePixelFormatArgs>(item.at);
+				OP::ImagePixelFormatArgs args = program.GetOpArgs<OP::ImagePixelFormatArgs>(item.At);
 				Issued = MakeShared<FImagePixelFormatTask>(item, GetMemory(), args, m_pSettings->GetPrivate()->m_imageCompressionQuality);
 			}
 			break;
@@ -1461,9 +1461,9 @@ namespace mu
 
 		case OP_TYPE::IM_LAYERCOLOUR:
 		{
-			if (item.stage == 1)
+			if (item.Stage == 1)
 			{
-				OP::ImageLayerColourArgs args = program.GetOpArgs<OP::ImageLayerColourArgs>(item.at);
+				OP::ImageLayerColourArgs args = program.GetOpArgs<OP::ImageLayerColourArgs>(item.At);
 				Issued = MakeShared<FImageLayerColourTask>(item, GetMemory(), args, m_pSettings->GetPrivate()->m_imageCompressionQuality);
 			}
 			break;
@@ -1471,9 +1471,9 @@ namespace mu
 
 		case OP_TYPE::IM_LAYER:
 		{
-			if (item.stage == 1)
+			if (item.Stage == 1)
 			{
-				OP::ImageLayerArgs args = program.GetOpArgs<OP::ImageLayerArgs>(item.at);
+				OP::ImageLayerArgs args = program.GetOpArgs<OP::ImageLayerArgs>(item.At);
 				Issued = MakeShared<FImageLayerTask>(item, GetMemory(), args, m_pSettings->GetPrivate()->m_imageCompressionQuality);
 			}
 			break;
@@ -1481,9 +1481,9 @@ namespace mu
 
 		case OP_TYPE::IM_MIPMAP:
 		{
-			if (item.stage == 1)
+			if (item.Stage == 1)
 			{
-				OP::ImageMipmapArgs args = program.GetOpArgs<OP::ImageMipmapArgs>(item.at);
+				OP::ImageMipmapArgs args = program.GetOpArgs<OP::ImageMipmapArgs>(item.At);
 				Issued = MakeShared<FImageMipmapTask>(item, GetMemory(), args, m_pSettings->GetPrivate()->m_imageCompressionQuality);
 			}
 			break;
@@ -1491,9 +1491,9 @@ namespace mu
 
 		case OP_TYPE::IM_SWIZZLE:
 		{
-			if (item.stage == 1)
+			if (item.Stage == 1)
 			{
-				OP::ImageSwizzleArgs args = program.GetOpArgs<OP::ImageSwizzleArgs>(item.at);
+				OP::ImageSwizzleArgs args = program.GetOpArgs<OP::ImageSwizzleArgs>(item.At);
 				Issued = MakeShared<FImageSwizzleTask>(item, GetMemory(), args);
 			}
 			break;
@@ -1501,11 +1501,11 @@ namespace mu
 
 		case OP_TYPE::IM_COMPOSE:
 		{
-			if (item.stage == 2)
+			if (item.Stage == 2)
 			{
-				OP::ImageComposeArgs args = program.GetOpArgs<OP::ImageComposeArgs>(item.at);
-				Ptr<const Layout> Layout = m_heapData[(size_t)item.customState].layout;
-				Issued = MakeShared<FImageComposeTask>(item, GetMemory(), args, m_pSettings->GetPrivate()->m_imageCompressionQuality, Layout);
+				OP::ImageComposeArgs args = program.GetOpArgs<OP::ImageComposeArgs>(item.At);
+				Ptr<const Layout> ComposeLayout = static_cast<const Layout*>( m_heapData[item.CustomState].Resource.get());
+				Issued = MakeShared<FImageComposeTask>(item, GetMemory(), args, m_pSettings->GetPrivate()->m_imageCompressionQuality, ComposeLayout);
 			}
 			break;
 		}
diff --git a/Engine/Plugins/Experimental/Mutable/Source/MutableRuntime/Private/MuR/InstancePrivate.h b/Engine/Plugins/Experimental/Mutable/Source/MutableRuntime/Private/MuR/InstancePrivate.h
--- a/Engine/Plugins/Experimental/Mutable/Source/MutableRuntime/Private/MuR/InstancePrivate.h
+++ b/Engine/Plugins/Experimental/Mutable/Source/MutableRuntime/Private/MuR/InstancePrivate.h
@@ -33,7 +33,7 @@ namespace mu
 			string m_name;
 		};
 
-		TArray<IMAGE> m_images;
+		TArray<IMAGE, TInlineAllocator<4>> m_images;
 
 		struct VECTOR
 		{
@@ -108,16 +108,16 @@ namespace mu
             RESOURCE_ID m_meshId;
 			string m_name;
 		};
-		TArray<MESH> m_meshes;
+		TArray<MESH, TInlineAllocator<2>> m_meshes;
 
 		// The order must match the meshes surfaces
-		TArray<INSTANCE_SURFACE> m_surfaces;
+		TArray<INSTANCE_SURFACE, TInlineAllocator<4>> m_surfaces;
 	};
 
 
     struct INSTANCE_LOD
     {
-		TArray<INSTANCE_COMPONENT> m_components;
+		TArray<INSTANCE_COMPONENT, TInlineAllocator<4>> m_components;
     };
 
 
@@ -129,7 +129,7 @@ namespace mu
         Instance::ID m_id;
 
 		//!
-		TArray<INSTANCE_LOD> m_lods;
+		TArray<INSTANCE_LOD,TInlineAllocator<4>> m_lods;
 
         Private()
         {
diff --git a/Engine/Plugins/Experimental/Mutable/Source/MutableRuntime/Private/MuR/Model.cpp b/Engine/Plugins/Experimental/Mutable/Source/MutableRuntime/Private/MuR/Model.cpp
--- a/Engine/Plugins/Experimental/Mutable/Source/MutableRuntime/Private/MuR/Model.cpp
+++ b/Engine/Plugins/Experimental/Mutable/Source/MutableRuntime/Private/MuR/Model.cpp
@@ -129,6 +129,8 @@ namespace mu
     //---------------------------------------------------------------------------------------------
     Model::~Model()
     {
+		MUTABLE_CPUPROFILER_SCOPE(ModelDestructor);
+
         check( m_pD );
         delete m_pD;
         m_pD = 0;
@@ -508,10 +510,10 @@ namespace mu
 
 
     //---------------------------------------------------------------------------------------------
-    uint32 Model::Private::GetResourceKey( uint32 paramListIndex,
-                                             OP::ADDRESS rootAt,
-                                             const Parameters* pParams )
+    uint32 Model::Private::GetResourceKey( uint32 paramListIndex, OP::ADDRESS rootAt, const Parameters* pParams )
     {
+		MUTABLE_CPUPROFILER_SCOPE(GetResourceKey)
+
         // Find the list of relevant parameters
         const TArray<uint16>* params = nullptr;
         if (paramListIndex<(uint32)m_program.m_parameterLists.Num())
@@ -529,8 +531,8 @@ namespace mu
 		parameterValuesBlob.Reserve(1024);
         for (int param: *params)
         {
-            size_t pos = parameterValuesBlob.Num();
-            size_t dataSize = 0;
+            int32 pos = parameterValuesBlob.Num();
+			int32 dataSize = 0;
 
             switch(m_program.m_parameters[param].m_type)
             {
@@ -555,9 +557,7 @@ namespace mu
             case PARAMETER_TYPE::T_INT:
                 dataSize = sizeof(int32);
                 parameterValuesBlob.SetNum( pos+dataSize );
-                FMemory::Memcpy( &parameterValuesBlob[pos],
-                        &pParams->GetPrivate()->m_values[param].m_int,
-                        dataSize );
+                FMemory::Memcpy( &parameterValuesBlob[pos], &pParams->GetPrivate()->m_values[param].m_int, dataSize );
                 pos += dataSize;
 
                 // Multi-values
@@ -578,9 +578,7 @@ namespace mu
             case PARAMETER_TYPE::T_FLOAT:
                 dataSize = sizeof(float);
                 parameterValuesBlob.SetNum( pos+dataSize );
-				FMemory::Memcpy( &parameterValuesBlob[pos],
-                        &pParams->GetPrivate()->m_values[param].m_float,
-                        dataSize );
+				FMemory::Memcpy( &parameterValuesBlob[pos], &pParams->GetPrivate()->m_values[param].m_float, dataSize );
                 pos += dataSize;
 
                 // Multi-values
@@ -601,9 +599,7 @@ namespace mu
             case PARAMETER_TYPE::T_COLOUR:
                 dataSize = 3*sizeof(float);
                 parameterValuesBlob.SetNum( pos+dataSize );
-				FMemory::Memcpy( &parameterValuesBlob[pos],
-                        &pParams->GetPrivate()->m_values[param].m_colour,
-                        dataSize );
+				FMemory::Memcpy( &parameterValuesBlob[pos], &pParams->GetPrivate()->m_values[param].m_colour, dataSize );
                 pos += dataSize;
 
                 // Multi-values
@@ -627,9 +623,7 @@ namespace mu
 
                 // \todo: padding will be random?
                 parameterValuesBlob.SetNum( pos+dataSize );
-				FMemory::Memcpy( &parameterValuesBlob[pos],
-                        &pParams->GetPrivate()->m_values[param].m_projector,
-                        dataSize );
+				FMemory::Memcpy( &parameterValuesBlob[pos], &pParams->GetPrivate()->m_values[param].m_projector, dataSize );
 				pos += dataSize;
 
                 // Multi-values
@@ -683,8 +677,8 @@ namespace mu
         {
             auto& key = m_generatedResources[i];
             if (key.m_rootAddress==rootAt
-                    &&
-                    key.m_parameterValuesBlob==parameterValuesBlob)
+				&&
+				key.m_parameterValuesBlob==parameterValuesBlob)
             {
                 key.m_lastRequestId = m_lastResourceResquestId;
                 return key.m_id;
@@ -706,13 +700,13 @@ namespace mu
         newKey.m_id = newId;
         newKey.m_lastRequestId = m_lastResourceResquestId;
         newKey.m_rootAddress = rootAt;
-        newKey.m_parameterValuesBlob = std::move(parameterValuesBlob);
+        newKey.m_parameterValuesBlob = MoveTemp(parameterValuesBlob);
 
         // TODO: Move the constant to settings?
         const size_t maxGeneratedResourcesIDCacheSize = 1024;
         if (m_generatedResources.Num()>=maxGeneratedResourcesIDCacheSize)
         {
-            m_generatedResources[oldestCachePosition] = std::move(newKey);
+            m_generatedResources[oldestCachePosition] = MoveTemp(newKey);
         }
         else
         {
diff --git a/Engine/Plugins/Experimental/Mutable/Source/MutableRuntime/Private/MuR/Octree.h b/Engine/Plugins/Experimental/Mutable/Source/MutableRuntime/Private/MuR/Octree.h
--- a/Engine/Plugins/Experimental/Mutable/Source/MutableRuntime/Private/MuR/Octree.h
+++ b/Engine/Plugins/Experimental/Mutable/Source/MutableRuntime/Private/MuR/Octree.h
@@ -8,8 +8,8 @@
 #include "MuR/Ptr.h"
 
 
-#define MAX_DIVISIONS 5
-#define MAX_VECTORS 20
+#define MUTABLE_OCTREE_MAX_DIVISIONS 5
+#define MUTABLE_OCTREE_MAX_VECTORS 32
 
 namespace mu
 {
@@ -54,7 +54,7 @@ namespace mu
 			{
 				elements.Add({ v, idx });
 
-				if (elements.Num() > MAX_VECTORS && division < MAX_DIVISIONS)
+				if (elements.Num() > MUTABLE_OCTREE_MAX_VECTORS && division < MUTABLE_OCTREE_MAX_DIVISIONS)
 				{
 					Divide();
 				}
@@ -165,15 +165,15 @@ namespace mu
 
 	private:
 		
-		int division = 0;
+		TArray<VertexInfo, TInlineAllocator<MUTABLE_OCTREE_MAX_VECTORS>> elements;
+
+		Octree* children[8];
 
 		vec3f minCorner;
 		vec3f maxCorner;
 
 		float padding = 0.0f;
+		int division = 0;
 
-		TArray<VertexInfo> elements;
-
-		Octree* children[8];
 	};
 }
diff --git a/Engine/Plugins/Experimental/Mutable/Source/MutableRuntime/Private/MuR/OpImageLayer.h b/Engine/Plugins/Experimental/Mutable/Source/MutableRuntime/Private/MuR/OpImageLayer.h
--- a/Engine/Plugins/Experimental/Mutable/Source/MutableRuntime/Private/MuR/OpImageLayer.h
+++ b/Engine/Plugins/Experimental/Mutable/Source/MutableRuntime/Private/MuR/OpImageLayer.h
@@ -830,26 +830,26 @@ namespace mu
 				] (uint32 i)
 				{
 					unsigned mask = pMaskBuf[i];
-			for (int c = 0; c < CHANNELS_TO_BLEND; ++c)
-			{
-				uint32 base = pBaseBuf[BASE_CHANNEL_STRIDE * i + c];
-				uint32 blended = pBlendedBuf[BLENDED_CHANNEL_STRIDE * i + c];
-				uint32 result = BLEND_FUNC_MASKED(base, blended, mask);
-				if (CLAMP)
-				{
-					pDestBuf[BASE_CHANNEL_STRIDE * i + c] = (uint8)FMath::Min(255u, result);
-				}
-				else
-				{
-					pDestBuf[BASE_CHANNEL_STRIDE * i + c] = (uint8)result;
-				}
-			}
-			// Copy the unblended channels
-			// \TODO: unnecessary when doing it in-place?
-			for (int c = 0; c < BASE_CHANNEL_STRIDE - CHANNELS_TO_BLEND; ++c)
-			{
-				pDestBuf[BASE_CHANNEL_STRIDE * i + CHANNELS_TO_BLEND + c] = pBaseBuf[BASE_CHANNEL_STRIDE * i + CHANNELS_TO_BLEND + c];
-			}
+					for (int c = 0; c < CHANNELS_TO_BLEND; ++c)
+					{
+						uint32 base = pBaseBuf[BASE_CHANNEL_STRIDE * i + c];
+						uint32 blended = pBlendedBuf[BLENDED_CHANNEL_STRIDE * i + c];
+						uint32 result = BLEND_FUNC_MASKED(base, blended, mask);
+						if (CLAMP)
+						{
+							pDestBuf[BASE_CHANNEL_STRIDE * i + c] = (uint8)FMath::Min(255u, result);
+						}
+						else
+						{
+							pDestBuf[BASE_CHANNEL_STRIDE * i + c] = (uint8)result;
+						}
+					}
+					// Copy the unblended channels
+					// \TODO: unnecessary when doing it in-place?
+					for (int c = 0; c < BASE_CHANNEL_STRIDE - CHANNELS_TO_BLEND; ++c)
+					{
+						pDestBuf[BASE_CHANNEL_STRIDE * i + CHANNELS_TO_BLEND + c] = pBaseBuf[BASE_CHANNEL_STRIDE * i + CHANNELS_TO_BLEND + c];
+					}
 				});
 		}
 		else if (maskFormat == EImageFormat::IF_L_UBYTE_RLE)
@@ -904,9 +904,13 @@ namespace mu
 
 								// Copy the unblended channels
 								// \TODO: unnecessary when doing it in-place?
-								for (int c = 0; c < BASE_CHANNEL_STRIDE - CHANNELS_TO_BLEND; ++c)
+								// Add an unnecessary check to see if static analysis is happier
+								if (BASE_CHANNEL_STRIDE - CHANNELS_TO_BLEND > 0)
 								{
-									pDestBuf[BASE_CHANNEL_STRIDE * i + CHANNELS_TO_BLEND + c] = pBaseBuf[BASE_CHANNEL_STRIDE * i + CHANNELS_TO_BLEND + c];
+									for (int c = 0; c < BASE_CHANNEL_STRIDE - CHANNELS_TO_BLEND; ++c)
+									{
+										pDestBuf[BASE_CHANNEL_STRIDE * i + CHANNELS_TO_BLEND + c] = pBaseBuf[BASE_CHANNEL_STRIDE * i + CHANNELS_TO_BLEND + c];
+									}
 								}
 							}
 						}
@@ -931,9 +935,13 @@ namespace mu
 
 								// Copy the unblended channels
 								// \TODO: unnecessary when doing it in-place?
-								for (int c = 0; c < BASE_CHANNEL_STRIDE - CHANNELS_TO_BLEND; ++c)
+								// Add an unnecessary check to see if static analysis is happier
+								if (BASE_CHANNEL_STRIDE - CHANNELS_TO_BLEND > 0)
 								{
-									pDestBuf[BASE_CHANNEL_STRIDE * i + CHANNELS_TO_BLEND + c] = pBaseBuf[BASE_CHANNEL_STRIDE * i + CHANNELS_TO_BLEND + c];
+									for (int c = 0; c < BASE_CHANNEL_STRIDE - CHANNELS_TO_BLEND; ++c)
+									{
+										pDestBuf[BASE_CHANNEL_STRIDE * i + CHANNELS_TO_BLEND + c] = pBaseBuf[BASE_CHANNEL_STRIDE * i + CHANNELS_TO_BLEND + c];
+									}
 								}
 							}
 						}
@@ -971,9 +979,13 @@ namespace mu
 
 							// Copy the unblended channels
 							// \TODO: unnecessary when doing it in-place?
-							for (int c = 0; c < BASE_CHANNEL_STRIDE - CHANNELS_TO_BLEND; ++c)
+							// Add an unnecessary check to see if static analysis is happier
+							if (BASE_CHANNEL_STRIDE - CHANNELS_TO_BLEND > 0)
 							{
-								pDestBuf[BASE_CHANNEL_STRIDE * i + CHANNELS_TO_BLEND + c] = pBaseBuf[BASE_CHANNEL_STRIDE * i + CHANNELS_TO_BLEND + c];
+								for (int c = 0; c < BASE_CHANNEL_STRIDE - CHANNELS_TO_BLEND; ++c)
+								{
+									pDestBuf[BASE_CHANNEL_STRIDE * i + CHANNELS_TO_BLEND + c] = pBaseBuf[BASE_CHANNEL_STRIDE * i + CHANNELS_TO_BLEND + c];
+								}
 							}
 						}
 
diff --git a/Engine/Plugins/Experimental/Mutable/Source/MutableRuntime/Private/MuR/OpImageProject.cpp b/Engine/Plugins/Experimental/Mutable/Source/MutableRuntime/Private/MuR/OpImageProject.cpp
--- a/Engine/Plugins/Experimental/Mutable/Source/MutableRuntime/Private/MuR/OpImageProject.cpp
+++ b/Engine/Plugins/Experimental/Mutable/Source/MutableRuntime/Private/MuR/OpImageProject.cpp
@@ -1269,7 +1269,7 @@ void ImageRasterProjectedCylindrical( const Mesh* pMesh,
 //#define DEBUG_PROJECTION 1
 
 #ifdef DEBUG_PROJECTION
-PRAGMA_DISABLE_OPTIMIZATION
+UE_DISABLE_OPTIMIZATION
 
 #undef assert
 int* assert_aux = 0;
@@ -1293,6 +1293,7 @@ inline void MeshCreateCollapsedVertexMap( const Mesh* pMesh,
     int vcount = pMesh->GetVertexCount();
     collapsedVertexMap.SetNum(vcount);
     vertices.SetNum(vcount);
+	collapsedVertsMap.Reserve(vcount);
 
     const FMeshBufferSet& MBSPriv = pMesh->GetVertexBuffers();
     for (int32 b = 0; b < MBSPriv.m_buffers.Num(); ++b)
@@ -1803,7 +1804,9 @@ void MeshProject_Optimised_Wrapping( const Mesh* pMesh,
 	TArray<AdjacentFaces> faceConnectivity;
 	faceConnectivity.SetNum(faceCount);
     TSet<int> processedFaces;
+	processedFaces.Reserve(faceCount/4);
 	TSet<int> discardedWrapAroundFaces;
+	discardedWrapAroundFaces.Reserve(faceCount/16);
 	TArray<int> faceStep;
 	faceStep.SetNum(faceCount);
 
@@ -1815,6 +1818,7 @@ void MeshProject_Optimised_Wrapping( const Mesh* pMesh,
 
     // Used to speed up connectivity building
 	TMultiMap<int, int> vertToFacesMap;
+	vertToFacesMap.Reserve(vertexCount);
 
     for (int f = 0; f < faceCount; ++f)
     {
@@ -1867,7 +1871,7 @@ void MeshProject_Optimised_Wrapping( const Mesh* pMesh,
             int v = collapsedVertexMap[pIndices[f * 3 + i]];
             //int v = pIndices[f * 3 + i];
 
-			TArray<int> FoundValues;
+			TArray<int,TInlineAllocator<32>> FoundValues;
 			vertToFacesMap.MultiFind(v, FoundValues);
             for (int f2 : FoundValues )
             {
@@ -1967,9 +1971,11 @@ void MeshProject_Optimised_Wrapping( const Mesh* pMesh,
     if (intersectedFace >= 0)
     {
         TArray<NeighborFace> pendingFaces; // Queue of pending face + new vertex
+		pendingFaces.Reserve(faceCount/64);
         
 		// \TODO: Bool array to speed up?
 		TSet<int> pendingFacesUnique; // Used to quickly check uniqueness in the pendingFaces queue
+		pendingFacesUnique.Reserve(faceCount / 64);
 
         //FVector3f hitFaceProjectedNormal;
         bool hitFaceHasPositiveArea = false;
@@ -2645,7 +2651,7 @@ MeshPtr MeshProject_Optimised( const Mesh* pMesh,
     return pResult;
 }
 #ifdef DEBUG_PROJECTION
-PRAGMA_ENABLE_OPTIMIZATION
+UE_ENABLE_OPTIMIZATION
 #endif
 
 
diff --git a/Engine/Plugins/Experimental/Mutable/Source/MutableRuntime/Private/MuR/System.cpp b/Engine/Plugins/Experimental/Mutable/Source/MutableRuntime/Private/MuR/System.cpp
--- a/Engine/Plugins/Experimental/Mutable/Source/MutableRuntime/Private/MuR/System.cpp
+++ b/Engine/Plugins/Experimental/Mutable/Source/MutableRuntime/Private/MuR/System.cpp
@@ -83,6 +83,7 @@ namespace mu
     System::~System()
     {
 		LLM_SCOPE_BYNAME(TEXT("MutableRuntime"));
+		MUTABLE_CPUPROFILER_SCOPE(SystemDestructor);
 
         check( m_pD );
 
@@ -121,6 +122,7 @@ namespace mu
     System::Private::~Private()
     {
 		LLM_SCOPE_BYNAME(TEXT("MutableRuntime"));
+		MUTABLE_CPUPROFILER_SCOPE(SystemPrivateDestructor);
 
         delete m_pStreamInterface;
         m_pStreamInterface = nullptr;
@@ -272,7 +274,7 @@ namespace mu
 
 		m_pD->RunCode(pLiveInstance->m_pModel.get(), pParams.get(), rootAt, lodMask);
 
-		InstancePtrConst pResult = pLiveInstance->m_memory->GetInstance(CACHE_ADDRESS(rootAt, 0, 0));
+		InstancePtrConst pResult = pLiveInstance->m_memory->GetInstance(FCacheAddress(rootAt, 0, 0));
 
 		pLiveInstance->m_pInstance = pResult;
 		if (pResult)
@@ -371,7 +373,7 @@ namespace mu
 				if (GetOpDataType(opType) == DT_IMAGE)
 				{
 					int8 executionOptions = 0;
-					CodeRunner Runner(m_pD->m_pSettings, m_pD, Model, pLiveInstance->m_pOldParameters.get(), at, System::AllLODs, executionOptions, SCHEDULED_OP::EType::ImageDesc);
+					CodeRunner Runner(m_pD->m_pSettings, m_pD, Model, pLiveInstance->m_pOldParameters.get(), at, System::AllLODs, executionOptions, FScheduledOp::EType::ImageDesc);
 					Runner.Run();
 					Runner.GetImageDescResult(OutDesc);
 				}
@@ -642,7 +644,7 @@ namespace mu
 		const Parameters* InParameters, OP::ADDRESS InCodeRoot, uint32 InLODs, uint8 executionOptions)
 	{
 		CodeRunner Runner(m_pSettings, this, InModel.get(), InParameters, InCodeRoot, InLODs, 
-			executionOptions, SCHEDULED_OP::EType::Full);
+			executionOptions, FScheduledOp::EType::Full);
 		Runner.Run();
 		bUnrecoverableError = Runner.bUnrecoverableError;
 	}
@@ -658,7 +660,7 @@ namespace mu
 		{
 			return false;
 		}
-		return m_memory->GetBool(CACHE_ADDRESS(at, 0, 0));
+		return m_memory->GetBool(FCacheAddress(at, 0, 0));
 	}
 
 
@@ -672,7 +674,7 @@ namespace mu
 		{
 			return 0.0f;
 		}
-		return m_memory->GetScalar(CACHE_ADDRESS(at, 0, 0));
+		return m_memory->GetScalar(FCacheAddress(at, 0, 0));
 	}
 
 
@@ -687,7 +689,7 @@ namespace mu
 			return 0;
 		}
 
-		return m_memory->GetInt(CACHE_ADDRESS(at, 0, 0));
+		return m_memory->GetInt(FCacheAddress(at, 0, 0));
 	}
 
 
@@ -712,7 +714,7 @@ namespace mu
 				if (pB) *pB = 0.0f;
 			}
 
-			col = m_memory->GetColour(CACHE_ADDRESS(at, 0, 0));
+			col = m_memory->GetColour(FCacheAddress(at, 0, 0));
 		}
 
 		if (pR) *pR = col[0];
@@ -729,7 +731,7 @@ namespace mu
 		{
 			return nullptr;
 		}
-    	return m_memory->GetProjector(CACHE_ADDRESS(at, 0, 0));
+    	return m_memory->GetProjector(FCacheAddress(at, 0, 0));
 	}
 
 	
@@ -747,7 +749,7 @@ namespace mu
 			{
 				return nullptr;
 			}
-			ImagePtrConst Result = m_memory->GetImage(CACHE_ADDRESS(at, 0, MipsToSkip));
+			ImagePtrConst Result = m_memory->GetImage(FCacheAddress(at, 0, MipsToSkip));
 			return Result;
 		}
 
@@ -769,7 +771,7 @@ namespace mu
 			{
 				return nullptr;
 			}
-			MeshPtrConst pResult = m_memory->GetMesh(CACHE_ADDRESS(at, 0, 0));
+			MeshPtrConst pResult = m_memory->GetMesh(FCacheAddress(at, 0, 0));
 			return pResult;
 		}
 
@@ -794,7 +796,7 @@ namespace mu
 				{
 					return nullptr;
 				}
-				pResult = m_memory->GetLayout(CACHE_ADDRESS(at, 0, 0));
+				pResult = m_memory->GetLayout(FCacheAddress(at, 0, 0));
 			}
 		}
 
@@ -819,7 +821,7 @@ namespace mu
 				{
 					return nullptr;
 				}
-				pResult = m_memory->GetString(CACHE_ADDRESS(at, 0, 0));
+				pResult = m_memory->GetString(FCacheAddress(at, 0, 0));
 			}
 		}
 
diff --git a/Engine/Plugins/Experimental/Mutable/Source/MutableRuntime/Private/MuR/SystemPrivate.h b/Engine/Plugins/Experimental/Mutable/Source/MutableRuntime/Private/MuR/SystemPrivate.h
--- a/Engine/Plugins/Experimental/Mutable/Source/MutableRuntime/Private/MuR/SystemPrivate.h
+++ b/Engine/Plugins/Experimental/Mutable/Source/MutableRuntime/Private/MuR/SystemPrivate.h
@@ -103,64 +103,61 @@ namespace mu
 	};
 
 
-	struct SCHEDULED_OP
+	/** This structure stores the data about an ongoing mutable operation that needs to be executed. */
+	struct FScheduledOp
 	{
-		SCHEDULED_OP()
+		inline FScheduledOp()
 		{
-			stage = 0;
+			Stage = 0;
 			Type = EType::Full;
 		}
 
-		SCHEDULED_OP(OP::ADDRESS _at,
-			const SCHEDULED_OP& opTemplate,
-			uint8 _stage = 0,
-			uint32 _customState = 0)
+		inline FScheduledOp(OP::ADDRESS InAt, const FScheduledOp& InOpTemplate, uint8 InStage = 0, uint32 InCustomState = 0)
 		{
-			at = _at;
-			executionOptions = opTemplate.executionOptions;
-			executionIndex = opTemplate.executionIndex;
-			stage = _stage;
-			customState = _customState;
-			Type = opTemplate.Type;
+			check(InStage < 120);
+			At = InAt;
+			ExecutionOptions = InOpTemplate.ExecutionOptions;
+			ExecutionIndex = InOpTemplate.ExecutionIndex;
+			Stage = InStage;
+			CustomState = InCustomState;
+			Type = InOpTemplate.Type;
 		}
 
-		static inline SCHEDULED_OP FromOpAndOptions(OP::ADDRESS _at,
-			const SCHEDULED_OP& opTemplate,
-			uint8 _executionOptions)
+		static inline FScheduledOp FromOpAndOptions(OP::ADDRESS InAt, const FScheduledOp& InOpTemplate, uint8 InExecutionOptions)
 		{
-			SCHEDULED_OP r;
-			r.at = _at;
-			r.executionOptions = _executionOptions;
-			r.executionIndex = opTemplate.executionIndex;
-			r.stage = 0;
-			r.customState = opTemplate.customState;
-			r.Type = opTemplate.Type;
+			FScheduledOp r;
+			r.At = InAt;
+			r.ExecutionOptions = InExecutionOptions;
+			r.ExecutionIndex = InOpTemplate.ExecutionIndex;
+			r.Stage = 0;
+			r.CustomState = InOpTemplate.CustomState;
+			r.Type = InOpTemplate.Type;
 			return r;
 		}
 
 		//! Address of the operation
-		OP::ADDRESS at = 0;
+		OP::ADDRESS At = 0;
 
 		//! Additional custom state data that the operation can store. This is usually used to pass information
 		//! between execution stages of an operation.
-		uint32 customState = 0;
+		uint32 CustomState = 0;
 
 		//! Index of the operation execution: This is used for iteration of different ranges.
 		//! It is an index into the CodeRunner::GetMemory()::m_rangeIndex vector.
 		//! executionIndex 0 is always used for empty ExecutionIndex, which is the most common
 		//! one.
-		uint16 executionIndex = 0;
+		uint16 ExecutionIndex = 0;
 
 		//! Additional execution options. Set externally to this op, it usually alters the result.
 		//! For example, this is used to keep track of the mipmaps to skip in image operations.
-		uint8 executionOptions = 0;
+		uint8 ExecutionOptions = 0;
 
 		//! Internal stage of the operation.
 		//! Stage 0 is usually scheduling of children, and 1 is execution. Some instructions
 		//! may have more steges to schedule children that are optional for execution, etc.
-		uint8 stage : 4;
+		uint8 Stage : 7;
 
-		//!
+		//! Type of calculation we are requestinf for this operation.
 		enum class EType : uint8
 		{
 			//! Execute the operation to calculate the full result
@@ -169,53 +166,54 @@ namespace mu
 			//! Execute the operation to obtain the descriptor of an image.
 			ImageDesc
 		};
-		EType Type : 4;
+		EType Type : 1;
 	};
 
-	inline uint32 GetTypeHash(const SCHEDULED_OP& op)
+	inline uint32 GetTypeHash(const FScheduledOp& Op)
 	{
-		return HashCombine(::GetTypeHash(op.at), HashCombine(::GetTypeHash(op.stage), ::GetTypeHash(op.executionIndex)));
+		return HashCombine(::GetTypeHash(Op.At), HashCombine(::GetTypeHash(Op.Stage), ::GetTypeHash(Op.ExecutionIndex)));
 	}
 
 
 	//! A cache address is the operation plus the context of execution (iteration indices, etc...).
-	struct CACHE_ADDRESS
+	struct FCacheAddress
 	{
-		OP::ADDRESS at = 0;
-		uint16 executionIndex = 0;
-		uint8 executionOptions = 0;
-		SCHEDULED_OP::EType Type = SCHEDULED_OP::EType::Full;
+		/** The meaning of all these fields is the same than the FScheduledOp struct. */
+		OP::ADDRESS At = 0;
+		uint16 ExecutionIndex = 0;
+		uint8 ExecutionOptions = 0;
+		FScheduledOp::EType Type = FScheduledOp::EType::Full;
 
-		CACHE_ADDRESS() {}
+		FCacheAddress() {}
 
-		CACHE_ADDRESS(OP::ADDRESS _at, uint16 _executionIndex, uint8 _executionOptions)
+		FCacheAddress(OP::ADDRESS InAt, uint16 InExecutionIndex, uint8 InExecutionOptions)
 		{
-			at = _at;
-			executionIndex = _executionIndex;
-			executionOptions = _executionOptions;
+			At = InAt;
+			ExecutionIndex = InExecutionIndex;
+			ExecutionOptions = InExecutionOptions;
 		}
 
-		CACHE_ADDRESS(OP::ADDRESS _at, const SCHEDULED_OP& item)
+		FCacheAddress(OP::ADDRESS InAt, const FScheduledOp& Item)
 		{
-			at = _at;
-			executionIndex = item.executionIndex;
-			executionOptions = item.executionOptions;
-			Type = item.Type;
+			At = InAt;
+			ExecutionIndex = Item.ExecutionIndex;
+			ExecutionOptions = Item.ExecutionOptions;
+			Type = Item.Type;
 		}
 
-		CACHE_ADDRESS(const SCHEDULED_OP& item)
+		FCacheAddress(const FScheduledOp& Item)
 		{
-			at = item.at;
-			executionIndex = item.executionIndex;
-			executionOptions = item.executionOptions;
-			Type = item.Type;
+			At = Item.At;
+			ExecutionIndex = Item.ExecutionIndex;
+			ExecutionOptions = Item.ExecutionOptions;
+			Type = Item.Type;
 		}
 	};
 
 
-	inline uint32 GetTypeHash(const CACHE_ADDRESS& a)
+	inline uint32 GetTypeHash(const FCacheAddress& a)
 	{
-		return HashCombine(::GetTypeHash(a.at), a.executionIndex);
+		return HashCombine(::GetTypeHash(a.At), a.ExecutionIndex);
 	}
 
 
@@ -243,11 +241,11 @@ namespace mu
 			m_otherIndex.Empty();
 		}
 
-		inline void erase(const CACHE_ADDRESS& at)
+		inline void erase(const FCacheAddress& at)
 		{
-			if (at.executionIndex == 0 && at.executionOptions == 0)
+			if (at.ExecutionIndex == 0 && at.ExecutionOptions == 0)
 			{
-				m_index0[at.at] = nullptr;
+				m_index0[at.At] = nullptr;
 			}
 			else
 			{
@@ -255,13 +253,13 @@ namespace mu
 			}
 		}
 
-		inline DATA get(const CACHE_ADDRESS& at) const
+		inline DATA get(const FCacheAddress& at) const
 		{
-			if (at.executionIndex == 0 && at.executionOptions == 0)
+			if (at.ExecutionIndex == 0 && at.ExecutionOptions == 0)
 			{
-				if (at.at < uint32(m_index0.Num()))
+				if (at.At < uint32(m_index0.Num()))
 				{
-					return m_index0[at.at];
+					return m_index0[at.At];
 				}
 				else
 				{
@@ -279,13 +277,13 @@ namespace mu
 			return 0;
 		}
 
-		inline const DATA* get_ptr(const CACHE_ADDRESS& at) const
+		inline const DATA* get_ptr(const FCacheAddress& at) const
 		{
-			if (at.executionIndex == 0 && at.executionOptions == 0)
+			if (at.ExecutionIndex == 0 && at.ExecutionOptions == 0)
 			{
-				if (at.at < uint32(m_index0.Num()))
+				if (at.At < uint32(m_index0.Num()))
 				{
-					return &m_index0[at.at];
+					return &m_index0[at.At];
 				}
 				else
 				{
@@ -303,11 +301,11 @@ namespace mu
 			return nullptr;
 		}
 
-		inline DATA& operator[](const CACHE_ADDRESS& at)
+		inline DATA& operator[](const FCacheAddress& at)
 		{
-			if (at.executionIndex == 0 && at.executionOptions == 0)
+			if (at.ExecutionIndex == 0 && at.ExecutionOptions == 0)
 			{
-				return m_index0[at.at];
+				return m_index0[at.At];
 			}
 			else
 			{
@@ -321,7 +319,7 @@ namespace mu
 			friend class CodeContainer<DATA>;
 			const CodeContainer<DATA>* container;
 			typename TArray<DATA>::TIterator it0;
-			typename TMap<CACHE_ADDRESS, DATA>::TIterator it1;
+			typename TMap<FCacheAddress, DATA>::TIterator it1;
 
 			iterator(CodeContainer<DATA>* InContainer)
 				: container(InContainer)
@@ -374,7 +372,7 @@ namespace mu
 				}
 			}
 
-			inline CACHE_ADDRESS get_address() const
+			inline FCacheAddress get_address() const
 			{
 				if (it0)
 				{
@@ -403,7 +401,7 @@ namespace mu
 		TArray<DATA> m_index0;
 
 		// For index>0
-		TMap<CACHE_ADDRESS, DATA> m_otherIndex;
+		TMap<FCacheAddress, DATA> m_otherIndex;
 	};
 
 
@@ -476,9 +474,9 @@ namespace mu
 			m_opHitCount.resize(size);
 		}
 
-		void SetUnused(CACHE_ADDRESS at)
+		void SetUnused(FCacheAddress at)
 		{
-			//UE_LOG(LogMutableCore, Log, TEXT("memory SetUnused : %5d "), at.at);
+			//UE_LOG(LogMutableCore, Log, TEXT("memory SetUnused : %5d "), at.At);
 			if (m_resources[at].Key >= 2)
 			{
 				// Keep the result anyway if it doesn't use any memory.
@@ -491,16 +489,16 @@ namespace mu
 		}
 
 
-		bool IsValid(CACHE_ADDRESS at)
+		bool IsValid(FCacheAddress at)
 		{
-			if (at.at == 0) return false;
+			if (at.At == 0) return false;
 
 			// Is it a desc data query?
-			if (at.Type == SCHEDULED_OP::EType::ImageDesc)
+			if (at.Type == FScheduledOp::EType::ImageDesc)
 			{
-				if (uint32(m_descCache.Num()) > at.at)
+				if (uint32(m_descCache.Num()) > at.At)
 				{
-					return m_descCache[at.at];
+					return m_descCache[at.At];
 				}
 				else
 				{
@@ -516,6 +514,8 @@ namespace mu
 
 		void ClearCacheLayer0()
 		{
+			MUTABLE_CPUPROFILER_SCOPE(ClearLayer0);
+
 			CodeContainer<int>::iterator it = m_opHitCount.begin();
 			for (; it.IsValid(); ++it)
 			{
@@ -531,6 +531,8 @@ namespace mu
 
 		void ClearCacheLayer1()
 		{
+			MUTABLE_CPUPROFILER_SCOPE(ClearLayer1);
+
 			CodeContainer<int>::iterator it = m_opHitCount.begin();
 			for (; it.IsValid(); ++it)
 			{
@@ -544,6 +546,8 @@ namespace mu
 
 		void Clear()
 		{
+			MUTABLE_CPUPROFILER_SCOPE(ProgramCacheClear);
+
 			size_t codeSize = m_resources.size_code();
 			m_resources.clear();
 			m_resources.resize(codeSize);
@@ -553,184 +557,198 @@ namespace mu
 		}
 
 
-		bool GetBool(CACHE_ADDRESS at)
+		bool GetBool(FCacheAddress at)
 		{
-			if (!at.at) return false;
+			if (!at.At) return false;
 			auto d = m_resources.get_ptr(at);
 			if (!d) return false;
 
 			Ptr<const Bool> pResult;
-			if (at.at)
+			if (at.At)
 			{
 				pResult = (const Bool*)d->Value.get();
 			}
 			return pResult ? pResult->m_value : false;
 		}
 
-		float GetScalar(CACHE_ADDRESS at)
+		float GetScalar(FCacheAddress at)
 		{
-			if (!at.at) return 0.0f;
+			if (!at.At) return 0.0f;
 			auto d = m_resources.get_ptr(at);
 			if (!d) return 0.0f;
 
 			Ptr<const Scalar> pResult;
-			if (at.at)
+			if (at.At)
 			{
 				pResult = (const Scalar*)d->Value.get();
 			}
 			return pResult ? pResult->m_value : 0.0f;
 		}
 
-		int GetInt(CACHE_ADDRESS at)
+		int GetInt(FCacheAddress at)
 		{
-			if (!at.at) return 0;
+			if (!at.At) return 0;
 			auto d = m_resources.get_ptr(at);
 			if (!d) return 0;
 
 			Ptr<const Int> pResult;
-			if (at.at)
+			if (at.At)
 			{
 				pResult = (const Int*)d->Value.get();
 			}
 			return pResult ? pResult->m_value : 0;
 		}
 
-		FVector4f GetColour(CACHE_ADDRESS at)
+		FVector4f GetColour(FCacheAddress at)
 		{
-			if (!at.at) return FVector4f();
+			if (!at.At) return FVector4f();
 			auto d = m_resources.get_ptr(at);
 			if (!d) return FVector4f();
 
 			Ptr<const Colour> pResult;
-			if (at.at)
+			if (at.At)
 			{
 				pResult = (const Colour*)d->Value.get();
 			}
 			return pResult ? pResult->m_colour : FVector4f();
 		}
 
-		Ptr<const Projector> GetProjector(CACHE_ADDRESS at)
+		Ptr<const Projector> GetProjector(FCacheAddress at)
 		{
-			if (!at.at) return nullptr;
+			if (!at.At) return nullptr;
 			auto d = m_resources.get_ptr(at);
 			if (!d) return nullptr;
 
 			Ptr<const Projector> pResult;
-			if (at.at)
+			if (at.At)
 			{
 				pResult = (const Projector*)d->Value.get();
 			}
 			return pResult;
 		}
 
-		Ptr<const Instance> GetInstance(CACHE_ADDRESS at)
+		Ptr<const Instance> GetInstance(FCacheAddress at)
 		{
-			if (!at.at) return nullptr;
+			if (!at.At) return nullptr;
 			auto d = m_resources.get_ptr(at);
 			if (!d) return nullptr;
 
-			Ptr<const Instance> pResult;
-			if (at.at)
+			Ptr<const Instance> pResult = (const Instance*)d->Value.get();
+
+			// We need to decrease the hitcount even if the result is null.
+			// Lower hit counts means we shouldn't clear the value
+			if (m_opHitCount[at] > 0)
 			{
-				pResult = (const Instance*)d->Value.get();
+				--m_opHitCount[at];
+
+				if (m_opHitCount[at] <= 0)
+				{
+					SetUnused(at);
+				}
 			}
+
 			return pResult;
 		}
 
 
-		Ptr<const Layout> GetLayout(CACHE_ADDRESS at)
+		Ptr<const Layout> GetLayout(FCacheAddress at)
 		{
-			if (!at.at) return nullptr;
+			if (!at.At)
+			{
+				return nullptr;
+			}
 			auto d = m_resources.get_ptr(at);
 			if (!d) return nullptr;
 
 			Ptr<const Layout> pResult;
-			if (at.at)
+			if (at.At)
 			{
 				pResult = (const Layout*)d->Value.get();
 			}
 			return pResult;
 		}
 
-		Ptr<const String> GetString(CACHE_ADDRESS at)
+		Ptr<const String> GetString(FCacheAddress at)
 		{
-			if (!at.at)
+			if (!at.At)
+			{
 				return nullptr;
+			}
 			auto d = m_resources.get_ptr(at);
 			if (!d)
 				return nullptr;
 
 			Ptr<const String> pResult;
-			if (at.at)
+			if (at.At)
 			{
 				pResult = (const String*)d->Value.get();
 			}
 			return pResult;
 		}
 
-		void SetBool(CACHE_ADDRESS at, bool v)
+		void SetBool(FCacheAddress at, bool v)
 		{
-			check(at.at < m_resources.size_code());
+			check(at.At < m_resources.size_code());
 
 			Ptr<Bool> pResult = new Bool;
 			pResult->m_value = v;
 			m_resources[at] = TPair<int, Ptr<const RefCounted>>(1, pResult);
 		}
 
-		void SetValidDesc(CACHE_ADDRESS at)
+		void SetValidDesc(FCacheAddress at)
 		{
-			check(at.Type == SCHEDULED_OP::EType::ImageDesc);
-			check(at.at < uint32(m_descCache.Num()));
+			check(at.Type == FScheduledOp::EType::ImageDesc);
+			check(at.At < uint32(m_descCache.Num()));
 
-			m_descCache[at.at] = true;
+			m_descCache[at.At] = true;
 		}
 
-		void SetScalar(CACHE_ADDRESS at, float v)
+		void SetScalar(FCacheAddress at, float v)
 		{
-			check(at.at < m_resources.size_code());
+			check(at.At < m_resources.size_code());
 
 			Ptr<Scalar> pResult = new Scalar;
 			pResult->m_value = v;
 			m_resources[at] = TPair<int, Ptr<const RefCounted>>(1, pResult);
 		}
 
-		void SetInt(CACHE_ADDRESS at, int v)
+		void SetInt(FCacheAddress at, int v)
 		{
-			check(at.at < m_resources.size_code());
+			check(at.At < m_resources.size_code());
 
 			Ptr<Int> pResult = new Int;
 			pResult->m_value = v;
 			m_resources[at] = TPair<int, Ptr<const RefCounted>>(1, pResult);
 		}
 
-		void SetColour(CACHE_ADDRESS at, const FVector4f& v)
+		void SetColour(FCacheAddress at, const FVector4f& v)
 		{
-			check(at.at < m_resources.size_code());
+			check(at.At < m_resources.size_code());
 
 			Ptr<Colour> pResult = new Colour;
 			pResult->m_colour = v;
 			m_resources[at] = TPair<int, Ptr<const RefCounted>>(1, pResult);
 		}
 
-		void SetProjector(CACHE_ADDRESS at, Ptr<const Projector> v)
+		void SetProjector(FCacheAddress at, Ptr<const Projector> v)
 		{
-			check(at.at < m_resources.size_code());
+			check(at.At < m_resources.size_code());
 			m_resources[at] = TPair<int, Ptr<const RefCounted>>(1, v);
 		}
 
-		void SetInstance(CACHE_ADDRESS at, Ptr<const Instance> v)
+		void SetInstance(FCacheAddress at, Ptr<const Instance> v)
 		{
-			check(at.at < m_resources.size_code());
-			m_resources[at] = TPair<int, Ptr<const RefCounted>>(1, v);
+			check(at.At < m_resources.size_code());
+			m_resources[at] = TPair<int, Ptr<const RefCounted>>(2, v);
 		}
 
-		void SetImage(CACHE_ADDRESS at, Ptr<const Image> v)
+		void SetImage(FCacheAddress at, Ptr<const Image> v)
 		{
-			check(at.at < m_resources.size_code());
+			check(at.At < m_resources.size_code());
 			m_resources[at] = TPair<int, Ptr<const RefCounted>>(2, v);
 		}
 
-		void SetMesh(CACHE_ADDRESS at, Ptr<const Mesh> v)
+		void SetMesh(FCacheAddress at, Ptr<const Mesh> v)
 		{
 			// debug
 //            if (v)
@@ -738,24 +756,24 @@ namespace mu
 //                v->GetPrivate()->CheckIntegrity();
 //            }
 
-			check(at.at < m_resources.size_code());
+			check(at.At < m_resources.size_code());
 			m_resources[at] = TPair<int, Ptr<const RefCounted>>(2, v);
 		}
 
-		void SetLayout(CACHE_ADDRESS at, Ptr<const Layout> v)
+		void SetLayout(FCacheAddress at, Ptr<const Layout> v)
 		{
-			check(at.at < m_resources.size_code());
+			check(at.At < m_resources.size_code());
 			m_resources[at] = TPair<int, Ptr<const RefCounted>>(1, v);
 		}
 
-		void SetString(CACHE_ADDRESS at, Ptr<const String> v)
+		void SetString(FCacheAddress at, Ptr<const String> v)
 		{
-			check(at.at < m_resources.size_code());
+			check(at.At < m_resources.size_code());
 			m_resources[at] = TPair<int, Ptr<const RefCounted>>(1, v);
 		}
 
 
-		inline void IncreaseHitCount(CACHE_ADDRESS at)
+		inline void IncreaseHitCount(FCacheAddress at)
 		{
 			m_opHitCount[at] += 1;
 		}
@@ -766,10 +784,10 @@ namespace mu
 		}
 
 
-		Ptr<const Image> GetImage(CACHE_ADDRESS at)
+		Ptr<const Image> GetImage(FCacheAddress at)
 		{
-			if (!at.at) return nullptr;
-			if (size_t(at.at) >= m_resources.size_code()) return nullptr;
+			if (!at.At) return nullptr;
+			if (size_t(at.At) >= m_resources.size_code()) return nullptr;
 			auto d = m_resources.get_ptr(at);
 			if (!d) return nullptr;
 
@@ -789,10 +807,10 @@ namespace mu
 			return pResult;
 		}
 
-		Ptr<const Mesh> GetMesh(CACHE_ADDRESS at)
+		Ptr<const Mesh> GetMesh(FCacheAddress at)
 		{
-			if (!at.at) return nullptr;
-			if (size_t(at.at) >= m_resources.size_code()) return nullptr;
+			if (!at.At) return nullptr;
+			if (size_t(at.At) >= m_resources.size_code()) return nullptr;
 			auto d = m_resources.get_ptr(at);
 			if (!d) return nullptr;
 
@@ -815,25 +833,25 @@ namespace mu
 	};
 
 
-	inline bool operator==(const CACHE_ADDRESS& a, const CACHE_ADDRESS& b)
+	inline bool operator==(const FCacheAddress& a, const FCacheAddress& b)
 	{
-		return a.at == b.at
+		return a.At == b.At
 			&&
-			a.executionIndex == b.executionIndex
+			a.ExecutionIndex == b.ExecutionIndex
 			&&
-			a.executionOptions == b.executionOptions
+			a.ExecutionOptions == b.ExecutionOptions
 			&&
 			a.Type == b.Type;
 	}
 
-	inline bool operator<(const CACHE_ADDRESS& a, const CACHE_ADDRESS& b)
+	inline bool operator<(const FCacheAddress& a, const FCacheAddress& b)
 	{
-		if (a.at < b.at) return true;
-		if (a.at > b.at) return false;
-		if (a.executionIndex < b.executionIndex) return true;
-		if (a.executionIndex > b.executionIndex) return false;
-		if (a.executionOptions < b.executionOptions) return true;
-		if (a.executionOptions > b.executionOptions) return false;
+		if (a.At < b.At) return true;
+		if (a.At > b.At) return false;
+		if (a.ExecutionIndex < b.ExecutionIndex) return true;
+		if (a.ExecutionIndex > b.ExecutionIndex) return false;
+		if (a.ExecutionOptions < b.ExecutionOptions) return true;
+		if (a.ExecutionOptions > b.ExecutionOptions) return false;
 		return a.Type < b.Type;
 	}
 
@@ -844,6 +862,11 @@ namespace mu
     //---------------------------------------------------------------------------------------------
     struct FModelCache
     {
+		~FModelCache()
+		{
+			MUTABLE_CPUPROFILER_SCOPE(FModelCacheDestructor);
+		}
+
         struct FModelCacheEntry
         {
             //!
@@ -865,8 +888,7 @@ namespace mu
 			TFunctionRef<bool(const Model*, int)> IsRomLockedFunc = [](const Model*, int) {return false;});
 
         //!
-        void UpdateForLoad( int romIndex, const Model* pModel,
-                            TFunctionRef<bool(const Model*,int)> isRomLockedFunc );
+        void UpdateForLoad( int romIndex, const Model* pModel, TFunctionRef<bool(const Model*,int)> isRomLockedFunc );
         void MarkRomUsed( int romIndex, const Model* pModel );
 
         //! Private helper
@@ -939,6 +961,16 @@ namespace mu
 
 			//! An entry for every instruction in the program to cache resources (meshes, images) if necessary.
 			TSharedPtr<FProgramCache> m_memory;
+
+			~FLiveInstance()
+			{
+				// Manually done to trace mem deallocations
+				MUTABLE_CPUPROFILER_SCOPE(LiveInstanceDestructor);
+				m_memory = nullptr;
+				m_pOldParameters = nullptr;
+				m_pInstance = nullptr;
+				m_pModel = nullptr;
+			}
         };
 
         TArray<FLiveInstance> m_liveInstances;
