[mutable] Fix crash when initializing a texture that belongs to the UI group + streaming enabled.
Fix Texture's render asset using an UnknownRefHeuristic which caused all mips to be requested at once.

#rnx

diff --git a/Engine/Plugins/Experimental/Mutable/Source/CustomizableObject/Private/MuCO/CustomizableObjectInstance.cpp b/Engine/Plugins/Experimental/Mutable/Source/CustomizableObject/Private/MuCO/CustomizableObjectInstance.cpp
--- a/Engine/Plugins/Experimental/Mutable/Source/CustomizableObject/Private/MuCO/CustomizableObjectInstance.cpp
+++ b/Engine/Plugins/Experimental/Mutable/Source/CustomizableObject/Private/MuCO/CustomizableObjectInstance.cpp
@@ -4248,6 +4248,11 @@ bool UCustomizableInstancePrivateData::BuildSkeletalMeshRenderData(const TShared
 		LODModel.bIsLODOptional = RefLODRenderData.bIsLODOptional;
 		LODModel.bStreamedDataInlined = RefLODRenderData.bStreamedDataInlined;
 
+		// WARNING! BufferSize must be > 0 or all texture mips in this LOD will be requested at once. 
+		// USkeletalMesh::IsMaterialUsed checks this size to see if a material is being used. If it fails the textures used by it won't be included 
+		// in the map of textures to stream.
+		LODModel.BuffersSize = 1;
+
 	}
 	
 
diff --git a/Engine/Plugins/Experimental/Mutable/Source/CustomizableObjectEditor/Private/MuCOE/GenerateMutableSource/GenerateMutableSourceSurface.cpp b/Engine/Plugins/Experimental/Mutable/Source/CustomizableObjectEditor/Private/MuCOE/GenerateMutableSource/GenerateMutableSourceSurface.cpp
--- a/Engine/Plugins/Experimental/Mutable/Source/CustomizableObjectEditor/Private/MuCOE/GenerateMutableSource/GenerateMutableSourceSurface.cpp
+++ b/Engine/Plugins/Experimental/Mutable/Source/CustomizableObjectEditor/Private/MuCOE/GenerateMutableSource/GenerateMutableSourceSurface.cpp
@@ -541,6 +541,18 @@ mu::NodeSurfacePtr GenerateMutableSourceSurface(const UEdGraphPin * Pin, FMutabl
 						Props.AddressX = ReferenceTexture->AddressX;
 						Props.AddressY = ReferenceTexture->AddressY;
 						Props.bFlipGreenChannel = ReferenceTexture->bFlipGreenChannel;
+
+						// TODO: MTBL-1081
+						// TextureGroup::TEXTUREGROUP_UI does not support streaming. If we generate a texture that requires streaming and set this group, it will crash when initializing the resource. 
+						// If LODGroup == TEXTUREGROUP_UI, UTexture::IsPossibleToStream() will return false and UE will assume all mips are loaded, when they're not, and crash.
+						if (Props.LODGroup == TEXTUREGROUP_UI)
+						{
+							Props.LODGroup = TextureGroup::TEXTUREGROUP_Character;
+							
+							FString msg = FString::Printf(TEXT("The Reference texture [%s] is using TEXTUREGROUP_UI which does not support streaming. Please set a different TEXTURE group."),
+								*ReferenceTexture->GetName(), *ImageName);
+							GenerationContext.Compiler->CompilerLog(FText::FromString(msg), Node);
+						}
 					}
 					else if (!GroupProjectionImg.get())
 					{
