[mutable] Fixes at CO group names update.

#rnx

diff --git a/Engine/Plugins/Experimental/Mutable/Source/CustomizableObjectEditor/Classes/MuCOE/Nodes/CustomizableObjectNodeObjectGroup.h b/Engine/Plugins/Experimental/Mutable/Source/CustomizableObjectEditor/Classes/MuCOE/Nodes/CustomizableObjectNodeObjectGroup.h
--- a/Engine/Plugins/Experimental/Mutable/Source/CustomizableObjectEditor/Classes/MuCOE/Nodes/CustomizableObjectNodeObjectGroup.h
+++ b/Engine/Plugins/Experimental/Mutable/Source/CustomizableObjectEditor/Classes/MuCOE/Nodes/CustomizableObjectNodeObjectGroup.h
@@ -57,20 +57,11 @@ public:
 	bool IsSingleOutputNode() const override;
 
 	// Own interface
-	UEdGraphPin* ObjectsPin() const
-	{
-		return FindPin(TEXT("Objects"));
-	}
-
-	UEdGraphPin* GroupProjectorsPin() const
-	{
-		return FindPin(TEXT("Projectors"));
-	}
-
-	UEdGraphPin* GroupPin() const
-	{
-		return FindPin(TEXT("Group"));
-	}
+	UEdGraphPin* ObjectsPin() const;
+
+	UEdGraphPin* GroupProjectorsPin() const;
+
+	UEdGraphPin* GroupPin() const;
 
 private:
 	FString LastGroupName;
diff --git a/Engine/Plugins/Experimental/Mutable/Source/CustomizableObjectEditor/Private/MuCOE/Nodes/CustomizableObjectNodeObjectGroup.cpp b/Engine/Plugins/Experimental/Mutable/Source/CustomizableObjectEditor/Private/MuCOE/Nodes/CustomizableObjectNodeObjectGroup.cpp
--- a/Engine/Plugins/Experimental/Mutable/Source/CustomizableObjectEditor/Private/MuCOE/Nodes/CustomizableObjectNodeObjectGroup.cpp
+++ b/Engine/Plugins/Experimental/Mutable/Source/CustomizableObjectEditor/Private/MuCOE/Nodes/CustomizableObjectNodeObjectGroup.cpp
@@ -44,38 +44,41 @@ void UCustomizableObjectNodeObjectGroup::Serialize(FArchive& Ar)
 void UCustomizableObjectNodeObjectGroup::PostEditChangeProperty(FPropertyChangedEvent& PropertyChangedEvent)
 {
 	Super::PostEditChangeProperty(PropertyChangedEvent);
-	for (UEdGraphPin* LinkedPin : FollowOutputPinArray(*GroupPin()))
-	{
-		UCustomizableObjectNode* Root = Cast<UCustomizableObjectNode>(LinkedPin->GetOwningNode());
-		check(Root); // All nodes inherit from UCustomizableObjectNode.
 
-		if (UCustomizableObjectNodeObject* CurrentRootNode = Cast<UCustomizableObjectNodeObject>(Root))
+	const FName PropertyName = (PropertyChangedEvent.Property ? PropertyChangedEvent.Property->GetFName() : NAME_None);
+	if (PropertyName == GET_MEMBER_NAME_CHECKED(UCustomizableObjectNodeObjectGroup, GroupName))
+	{
+		for (UEdGraphPin* LinkedPin : FollowOutputPinArray(*GroupPin()))
 		{
-			if (CurrentRootNode->ParentObject)
-			{
-				TArray<UCustomizableObject*> VisitedObjects;
-				CurrentRootNode = GetFullGraphRootNodeObject(CurrentRootNode, VisitedObjects);
-			}
-			if (!CurrentRootNode->ParentObject)
+			UCustomizableObjectNode* Root = CastChecked<UCustomizableObjectNode>(LinkedPin->GetOwningNode());
+
+			if (UCustomizableObjectNodeObject* CurrentRootNode = Cast<UCustomizableObjectNodeObject>(Root))
 			{
-				for (FCustomizableObjectState& State : CurrentRootNode->States)
+				if (CurrentRootNode->ParentObject)
+				{
+					TArray<UCustomizableObject*> VisitedObjects;
+					CurrentRootNode = GetFullGraphRootNodeObject(CurrentRootNode, VisitedObjects);
+				}
+				if (!CurrentRootNode->ParentObject)
 				{
-					for (int p = State.RuntimeParameters.Num() - 1; p >= 0; --p)
+					for (FCustomizableObjectState& State : CurrentRootNode->States)
 					{
-						if (State.RuntimeParameters[p].Equals(LastGroupName))
+						for (int32 p = 0; p < State.RuntimeParameters.Num(); ++p)
 						{
-							State.RuntimeParameters.RemoveAt(p);
-							State.RuntimeParameters.Add(GroupName);
+							if (State.RuntimeParameters[p].Equals(LastGroupName))
+							{
+								State.RuntimeParameters[p]= GroupName;
+							}
+						}
+						if (State.ForcedParameterValues.Contains(GroupName))
+						{
+							// Forced parameter already contains the NEW name of the parameter. We currently allow inconsistencies while working before compile time, so no warning needed here. When this changes to ID based instead of String based, this would become a warning worthy check (TODO: MTBL-1071).
+						}
+						else if (State.ForcedParameterValues.Contains(LastGroupName))
+						{
+							FString LastForcedValue = State.ForcedParameterValues.FindAndRemoveChecked(LastGroupName);
+							State.ForcedParameterValues.Emplace(GroupName, LastForcedValue);
 						}
-					}
-					if (State.ForcedParameterValues.Contains(GroupName))
-					{
-						// Forced parameter already contains the NEW name of the parameter. TODO: Warning?
-					}
-					else if (State.ForcedParameterValues.Contains(LastGroupName))
-					{
-						FString LastForcedValue = State.ForcedParameterValues.FindAndRemoveChecked(LastGroupName);
-						State.ForcedParameterValues.Emplace(GroupName, LastForcedValue);
 					}
 				}
 			}
@@ -85,20 +88,6 @@ void UCustomizableObjectNodeObjectGroup::PostEditChangeProperty(FPropertyChanged
 }
 
 
-void UCustomizableObjectNodeObjectGroup::AllocateDefaultPins(UCustomizableObjectNodeRemapPins* RemapPins)
-{
-	const UEdGraphSchema_CustomizableObject* Schema = GetDefault<UEdGraphSchema_CustomizableObject>();
-
-	UEdGraphPin* ObjectsPin = CustomCreatePin(EGPD_Input, Schema->PC_Object, FName("Objects"), true );
-	ObjectsPin->bDefaultValueIsIgnored = true;
-
-	UEdGraphPin* ProjectorsPin = CustomCreatePin(EGPD_Input, Schema->PC_GroupProjector, FName("Projectors"), true );
-	ObjectsPin->bDefaultValueIsIgnored = true;
-
-	CustomCreatePin(EGPD_Output, Schema->PC_Object, FName("Group"));
-}
-
-
 FText UCustomizableObjectNodeObjectGroup::GetNodeTitle(ENodeTitleType::Type TitleType) const
 {
 	FFormatNamedArguments Args;
@@ -122,10 +111,42 @@ FText UCustomizableObjectNodeObjectGroup::GetTooltipText() const
 }
 
 
+void UCustomizableObjectNodeObjectGroup::AllocateDefaultPins(UCustomizableObjectNodeRemapPins* RemapPins)
+{
+	const UEdGraphSchema_CustomizableObject* Schema = GetDefault<UEdGraphSchema_CustomizableObject>();
+
+	UEdGraphPin* ObjectsPin = CustomCreatePin(EGPD_Input, Schema->PC_Object, FName("Objects"), true);
+	ObjectsPin->bDefaultValueIsIgnored = true;
+
+	UEdGraphPin* ProjectorsPin = CustomCreatePin(EGPD_Input, Schema->PC_GroupProjector, FName("Projectors"), true);
+	ObjectsPin->bDefaultValueIsIgnored = true;
+
+	CustomCreatePin(EGPD_Output, Schema->PC_Object, FName("Group"));
+}
+
+
 bool UCustomizableObjectNodeObjectGroup::IsSingleOutputNode() const
 {
 	return true;
 }
 
 
+UEdGraphPin* UCustomizableObjectNodeObjectGroup::ObjectsPin() const
+{
+	return FindPin(TEXT("Objects"));
+}
+
+
+UEdGraphPin* UCustomizableObjectNodeObjectGroup::GroupProjectorsPin() const
+{
+	return FindPin(TEXT("Projectors"));
+}
+
+
+UEdGraphPin* UCustomizableObjectNodeObjectGroup::GroupPin() const
+{
+	return FindPin(TEXT("Group"));
+}
+
+
 #undef LOCTEXT_NAMESPACE
