[Mutable] Minor code improvements.

diff --git a/Engine/Plugins/Experimental/Mutable/Source/CustomizableObject/Classes/MuCO/CustomizableObject.h b/Engine/Plugins/Experimental/Mutable/Source/CustomizableObject/Classes/MuCO/CustomizableObject.h
--- a/Engine/Plugins/Experimental/Mutable/Source/CustomizableObject/Classes/MuCO/CustomizableObject.h
+++ b/Engine/Plugins/Experimental/Mutable/Source/CustomizableObject/Classes/MuCO/CustomizableObject.h
@@ -1154,7 +1154,7 @@ public:
 
 	// Get the index of a parameter
 	UFUNCTION(BlueprintCallable, Category = CustomizableObject)
-	int FindParameter(const FString& Name) const;
+	int32 FindParameter(const FString& Name) const;
 
 	// Get the type of a parameter
 	UFUNCTION(BlueprintCallable, Category = CustomizableObject)
diff --git a/Engine/Plugins/Experimental/Mutable/Source/CustomizableObject/Private/MuCO/CustomizableObject.cpp b/Engine/Plugins/Experimental/Mutable/Source/CustomizableObject/Private/MuCO/CustomizableObject.cpp
--- a/Engine/Plugins/Experimental/Mutable/Source/CustomizableObject/Private/MuCO/CustomizableObject.cpp
+++ b/Engine/Plugins/Experimental/Mutable/Source/CustomizableObject/Private/MuCO/CustomizableObject.cpp
@@ -1316,7 +1316,7 @@ const FString& UCustomizableObject::GetIntParameterAvailableOption(int32 ParamIn
 }
 
 
-int UCustomizableObject::FindParameter(const FString& Name) const
+int32 UCustomizableObject::FindParameter(const FString& Name) const
 {
 	const int32 * Found = ParameterPropertiesLookupTable.Find(Name);
 	if (Found == nullptr)
diff --git a/Engine/Plugins/Experimental/Mutable/Source/CustomizableObject/Private/MuCO/CustomizableObjectInstanceDescriptor.cpp b/Engine/Plugins/Experimental/Mutable/Source/CustomizableObject/Private/MuCO/CustomizableObjectInstanceDescriptor.cpp
--- a/Engine/Plugins/Experimental/Mutable/Source/CustomizableObject/Private/MuCO/CustomizableObjectInstanceDescriptor.cpp
+++ b/Engine/Plugins/Experimental/Mutable/Source/CustomizableObject/Private/MuCO/CustomizableObjectInstanceDescriptor.cpp
@@ -444,22 +444,22 @@ UCustomizableObject* FCustomizableObjectInstanceDescriptor::GetCustomizableObjec
 }
 
 
-bool FCustomizableObjectInstanceDescriptor::GetBuildParameterDecorations() const
+void FCustomizableObjectInstanceDescriptor::SetCustomizableObject(UCustomizableObject& InCustomizableObject)
 {
-	return bBuildParameterDecorations;
+	CustomizableObject = &InCustomizableObject;
+	ReloadParameters();
 }
 
 
-void FCustomizableObjectInstanceDescriptor::SetBuildParameterDecorations(const bool Value)
+bool FCustomizableObjectInstanceDescriptor::GetBuildParameterDecorations() const
 {
-	bBuildParameterDecorations = Value;
+	return bBuildParameterDecorations;
 }
 
 
-void FCustomizableObjectInstanceDescriptor::SetCustomizableObject(UCustomizableObject& InCustomizableObject)
+void FCustomizableObjectInstanceDescriptor::SetBuildParameterDecorations(const bool Value)
 {
-	CustomizableObject = &InCustomizableObject;
-	ReloadParameters();
+	bBuildParameterDecorations = Value;
 }
 
 
diff --git a/Engine/Plugins/Experimental/Mutable/Source/CustomizableObjectEditor/Private/MuCOE/GenerateMutableSource/GenerateMutableSource.h b/Engine/Plugins/Experimental/Mutable/Source/CustomizableObjectEditor/Private/MuCOE/GenerateMutableSource/GenerateMutableSource.h
--- a/Engine/Plugins/Experimental/Mutable/Source/CustomizableObjectEditor/Private/MuCOE/GenerateMutableSource/GenerateMutableSource.h
+++ b/Engine/Plugins/Experimental/Mutable/Source/CustomizableObjectEditor/Private/MuCOE/GenerateMutableSource/GenerateMutableSource.h
@@ -794,7 +794,7 @@ void PopulateReferenceSkeletalMeshesData(FMutableGraphGenerationContext& Generat
 void CheckNumOutputs(const UEdGraphPin& Pin, const FMutableGraphGenerationContext& GenerationContext);
 
 
-// TODO GMT Remove generation context dependency and move to GraphTraversal.
+// TODO FutureGMT Remove generation context dependency and move to GraphTraversal.
 UTexture2D* FindReferenceImage(const UEdGraphPin* Pin, FMutableGraphGenerationContext& GenerationContext);
 
 
