[mutable] Workaround for static analysis false positives.
#rnx

diff --git a/Engine/Plugins/Experimental/Mutable/Source/MutableRuntime/Private/MuR/OpImageLayer.h b/Engine/Plugins/Experimental/Mutable/Source/MutableRuntime/Private/MuR/OpImageLayer.h
--- a/Engine/Plugins/Experimental/Mutable/Source/MutableRuntime/Private/MuR/OpImageLayer.h
+++ b/Engine/Plugins/Experimental/Mutable/Source/MutableRuntime/Private/MuR/OpImageLayer.h
@@ -808,6 +808,7 @@ namespace mu
 
 		EImageFormat maskFormat = pMask->GetFormat();
 		bool bIsMaskUncompressed = (maskFormat == EImageFormat::IF_L_UBYTE);
+		int32 UnblendedChannels = BASE_CHANNEL_STRIDE - CHANNELS_TO_BLEND;
 
 		// The base determines the number of lods to process.
 		int32 LODCount = bOnlyFirstLOD ? 1 : pBase->GetLODCount();
@@ -904,13 +905,9 @@ namespace mu
 
 								// Copy the unblended channels
 								// \TODO: unnecessary when doing it in-place?
-								// Add an unnecessary check to see if static analysis is happier
-								if (BASE_CHANNEL_STRIDE - CHANNELS_TO_BLEND > 0)
+								for (int32 c = 0; c < UnblendedChannels; ++c)
 								{
-									for (int c = 0; c < BASE_CHANNEL_STRIDE - CHANNELS_TO_BLEND; ++c)
-									{
-										pDestBuf[BASE_CHANNEL_STRIDE * i + CHANNELS_TO_BLEND + c] = pBaseBuf[BASE_CHANNEL_STRIDE * i + CHANNELS_TO_BLEND + c];
-									}
+									pDestBuf[BASE_CHANNEL_STRIDE * i + CHANNELS_TO_BLEND + c] = pBaseBuf[BASE_CHANNEL_STRIDE * i + CHANNELS_TO_BLEND + c];
 								}
 							}
 						}
@@ -935,13 +932,9 @@ namespace mu
 
 								// Copy the unblended channels
 								// \TODO: unnecessary when doing it in-place?
-								// Add an unnecessary check to see if static analysis is happier
-								if (BASE_CHANNEL_STRIDE - CHANNELS_TO_BLEND > 0)
+								for (int32 c = 0; c < UnblendedChannels; ++c)
 								{
-									for (int c = 0; c < BASE_CHANNEL_STRIDE - CHANNELS_TO_BLEND; ++c)
-									{
-										pDestBuf[BASE_CHANNEL_STRIDE * i + CHANNELS_TO_BLEND + c] = pBaseBuf[BASE_CHANNEL_STRIDE * i + CHANNELS_TO_BLEND + c];
-									}
+									pDestBuf[BASE_CHANNEL_STRIDE * i + CHANNELS_TO_BLEND + c] = pBaseBuf[BASE_CHANNEL_STRIDE * i + CHANNELS_TO_BLEND + c];
 								}
 							}
 						}
@@ -979,13 +972,9 @@ namespace mu
 
 							// Copy the unblended channels
 							// \TODO: unnecessary when doing it in-place?
-							// Add an unnecessary check to see if static analysis is happier
-							if (BASE_CHANNEL_STRIDE - CHANNELS_TO_BLEND > 0)
+							for (int32 c = 0; c < UnblendedChannels; ++c)
 							{
-								for (int c = 0; c < BASE_CHANNEL_STRIDE - CHANNELS_TO_BLEND; ++c)
-								{
-									pDestBuf[BASE_CHANNEL_STRIDE * i + CHANNELS_TO_BLEND + c] = pBaseBuf[BASE_CHANNEL_STRIDE * i + CHANNELS_TO_BLEND + c];
-								}
+								pDestBuf[BASE_CHANNEL_STRIDE * i + CHANNELS_TO_BLEND + c] = pBaseBuf[BASE_CHANNEL_STRIDE * i + CHANNELS_TO_BLEND + c];
 							}
 						}
 
