[mutable] Fix optimization bad performance case with mipmaps and image patch.
More coding standard fixes.
#rnx

diff --git a/Engine/Plugins/Experimental/Mutable/Source/MutableRuntime/Private/MuR/CodeRunner.cpp b/Engine/Plugins/Experimental/Mutable/Source/MutableRuntime/Private/MuR/CodeRunner.cpp
--- a/Engine/Plugins/Experimental/Mutable/Source/MutableRuntime/Private/MuR/CodeRunner.cpp
+++ b/Engine/Plugins/Experimental/Mutable/Source/MutableRuntime/Private/MuR/CodeRunner.cpp
@@ -356,7 +356,7 @@ namespace mu
                 }
                 else
                 {
-                    pResult = pBase->Clone();
+					pResult = CloneOrTakeOver<Instance>(pBase.get());
                 }
 
                 if ( args.value )
@@ -2815,8 +2815,7 @@ namespace mu
             switch (item.Stage)
             {
             case 0:
-				AddOp( FScheduledOp( item.At, item, 1),
-					FScheduledOp( args.source, item) );
+				AddOp( FScheduledOp( item.At, item, 1), FScheduledOp( args.source, item) );
                 break;
 
             case 1:
@@ -2837,8 +2836,7 @@ namespace mu
             switch (item.Stage)
             {
             case 0:
-                    AddOp( FScheduledOp( item.At, item, 1),
-                           FScheduledOp( args.source, item) );
+                    AddOp( FScheduledOp( item.At, item, 1), FScheduledOp( args.source, item) );
                 break;
 
             case 1:
@@ -2848,6 +2846,7 @@ namespace mu
 				if (!pBase)
 				{
 					GetMemory().SetImage(item, nullptr);
+					break;
 				}
 
                 FImageSize destSize = FImageSize
@@ -2863,9 +2862,7 @@ namespace mu
                      destSize[1]!=pBase->GetSizeY() )
                 {
                     //pResult = ImageResize( pBase.get(), destSize );
-                    pResult =
-                        ImageResizeLinear( m_pSettings->GetPrivate()->m_imageCompressionQuality,
-                                           pBase.get(), destSize );
+                    pResult = ImageResizeLinear( m_pSettings->GetPrivate()->m_imageCompressionQuality, pBase.get(), destSize );
 
                     // If the source image had mips, generate them as well for the resized image.
                     // This shouldn't happen often since "ResizeLike" should be usually optimised out
@@ -3848,10 +3845,10 @@ namespace mu
                     fadeStart *= PI / 180.0f;
                     fadeEnd *= PI / 180.0f;
 
-                    auto format = GetUncompressedFormat( pSource->GetFormat() );
+					EImageFormat format = pSource ? GetUncompressedFormat(pSource->GetFormat()) : EImageFormat::IF_L_UBYTE;
                     pNew = new Image( SizeX, SizeY, 1, format );
 
-                    if (pSource->GetFormat()!=format)
+                    if (pSource && pSource->GetFormat()!=format)
                     {
 						MUTABLE_CPUPROFILER_SCOPE(RunCode_RasterMesh_ReformatSource);
                         pSource = ImagePixelFormat( m_pSettings->GetPrivate()->m_imageCompressionQuality, pSource.get(), format );
@@ -3867,7 +3864,7 @@ namespace mu
                     int layout = 0;
 
                     float projectionAngle = 0;
-                    if ( args.projector )
+                    if ( args.projector && pSource )
                     {
                         auto pProjector = GetMemory().GetProjector(FScheduledOp::FromOpAndOptions(args.projector, item, 0));
                         if (pProjector)
diff --git a/Engine/Plugins/Experimental/Mutable/Source/MutableRuntime/Private/MuR/CodeRunnerTasked.cpp b/Engine/Plugins/Experimental/Mutable/Source/MutableRuntime/Private/MuR/CodeRunnerTasked.cpp
--- a/Engine/Plugins/Experimental/Mutable/Source/MutableRuntime/Private/MuR/CodeRunnerTasked.cpp
+++ b/Engine/Plugins/Experimental/Mutable/Source/MutableRuntime/Private/MuR/CodeRunnerTasked.cpp
@@ -599,7 +599,11 @@ namespace mu
 		MUTABLE_CPUPROFILER_SCOPE(FImageLayerColourTask);
 
 		// Warning: This runs in a worker thread.
-		ImagePtr pNew;
+		Ptr<Image> pNew;
+		if (!m_base)
+		{
+			return;
+		}
 		
 		if (Args.mask)
 		{
@@ -615,16 +619,13 @@ namespace mu
 				//MUTABLE_CPUPROFILER_SCOPE(ImageLayerColour_EmergencyFix);
 
 				int32 levelCount = m_base->GetLODCount();
-				ImagePtr pDest = new Image(m_mask->GetSizeX(), m_mask->GetSizeY(),
-					levelCount,
-					m_mask->GetFormat());
+				Ptr<Image> pDest = new Image(m_mask->GetSizeX(), m_mask->GetSizeY(), levelCount, m_mask->GetFormat());
 
 				SCRATCH_IMAGE_MIPMAP scratch;
 				FMipmapGenerationSettings settings{};
 
 				ImageMipmap_PrepareScratch(pDest.get(), m_mask.get(), levelCount, &scratch);
-				ImageMipmap(m_imageCompressionQuality, pDest.get(), m_mask.get(), levelCount,
-					&scratch, settings);
+				ImageMipmap(m_imageCompressionQuality, pDest.get(), m_mask.get(), levelCount, &scratch, settings);
 				m_mask = pDest;
 			}
 		}
diff --git a/Engine/Plugins/Experimental/Mutable/Source/MutableRuntime/Private/MuR/OpImageBinarise.h b/Engine/Plugins/Experimental/Mutable/Source/MutableRuntime/Private/MuR/OpImageBinarise.h
--- a/Engine/Plugins/Experimental/Mutable/Source/MutableRuntime/Private/MuR/OpImageBinarise.h
+++ b/Engine/Plugins/Experimental/Mutable/Source/MutableRuntime/Private/MuR/OpImageBinarise.h
@@ -13,17 +13,22 @@ namespace mu
 	//---------------------------------------------------------------------------------------------
 	//!
 	//---------------------------------------------------------------------------------------------
-	inline ImagePtr ImageBinarise( const Image* pA, float threshold )
+	inline Ptr<Image> ImageBinarise( const Image* pA, float threshold )
 	{
-        ImagePtr pDest = new Image( pA->GetSizeX(), pA->GetSizeY(), pA->GetLODCount(), EImageFormat::IF_L_UBYTE );
+		if (!pA)
+		{
+			return nullptr;
+		}
+
+		Ptr<Image> pDest = new Image( pA->GetSizeX(), pA->GetSizeY(), pA->GetLODCount(), EImageFormat::IF_L_UBYTE );
 
-        uint8_t* pDestBuf = pDest->GetData();
-        const uint8_t* pABuf = pA->GetData();
+        uint8* pDestBuf = pDest->GetData();
+        const uint8* pABuf = pA->GetData();
 
 		// Generic implementation
-		int pixelCount = (int)pA->CalculatePixelCount();
+		int32 pixelCount = pA->CalculatePixelCount();
 
-        uint32_t t_8 = uint32_t( threshold * 255.0f );
+        uint32 t_8 = uint32( threshold * 255.0f );
 
 		switch ( pA->GetFormat() )
 		{
@@ -31,7 +36,7 @@ namespace mu
 		{
 			for ( int i=0; i<pixelCount; ++i )
 			{
-                uint32_t a_8 = pABuf[i];
+                uint32 a_8 = pABuf[i];
 				pDestBuf[i] = (a_8>=t_8) ? 255 : 0;
 			}
 			break;
@@ -41,7 +46,7 @@ namespace mu
 		{
 			for ( int i=0; i<pixelCount; ++i )
 			{
-                uint32_t a_8 = pABuf[3*i+0];
+                uint32 a_8 = pABuf[3*i+0];
 				a_8 += pABuf[3*i+1];
 				a_8 += pABuf[3*i+2];
 				a_8 /= 3;
@@ -55,7 +60,7 @@ namespace mu
         {
 			for ( int i=0; i<pixelCount; ++i )
 			{
-                uint32_t a_8 = pABuf[4*i+0];
+                uint32 a_8 = pABuf[4*i+0];
 				a_8 += pABuf[4*i+1];
 				a_8 += pABuf[4*i+2];
 				a_8 /= 3;
diff --git a/Engine/Plugins/Experimental/Mutable/Source/MutableRuntime/Private/MuR/OpImageSwizzle.h b/Engine/Plugins/Experimental/Mutable/Source/MutableRuntime/Private/MuR/OpImageSwizzle.h
--- a/Engine/Plugins/Experimental/Mutable/Source/MutableRuntime/Private/MuR/OpImageSwizzle.h
+++ b/Engine/Plugins/Experimental/Mutable/Source/MutableRuntime/Private/MuR/OpImageSwizzle.h
@@ -9,15 +9,20 @@
 namespace mu
 {
 
-	inline ImagePtr ImageSwizzle
+	inline Ptr<Image> ImageSwizzle
 		(
 			EImageFormat format,
-			const ImagePtrConst pSources[],
-            const uint8_t channels[]
+			const Ptr<const Image> pSources[],
+            const uint8 channels[]
 		)
 	{
 		MUTABLE_CPUPROFILER_SCOPE(ImageSwizzle);
 
+		if (!pSources[0])
+		{
+			return nullptr;
+		}
+
         ImagePtr pDest = new Image( pSources[0]->GetSizeX(), pSources[0]->GetSizeY(),
                                     pSources[0]->GetLODCount(),
                                     format );
@@ -68,7 +73,7 @@ namespace mu
 
 		for (int i = 0; i < NumDestChannels; ++i)
 		{
-			uint8_t* pDestBuf = pDest->GetData() + i;
+			uint8* pDestBuf = pDest->GetData() + i;
 
 			if (format == EImageFormat::IF_BGRA_UBYTE)
 			{
@@ -86,7 +91,7 @@ namespace mu
 
 			if (pSources[i])
 			{
-				const uint8_t* pSourceBuf = pSources[i]->GetData() + channels[i];
+				const uint8* pSourceBuf = pSources[i]->GetData() + channels[i];
 
 				switch (pSources[i]->GetFormat())
 				{
diff --git a/Engine/Plugins/Experimental/Mutable/Source/MutableTools/Private/MuT/ASTOpImageMipmap.cpp b/Engine/Plugins/Experimental/Mutable/Source/MutableTools/Private/MuT/ASTOpImageMipmap.cpp
--- a/Engine/Plugins/Experimental/Mutable/Source/MutableTools/Private/MuT/ASTOpImageMipmap.cpp
+++ b/Engine/Plugins/Experimental/Mutable/Source/MutableTools/Private/MuT/ASTOpImageMipmap.cpp
@@ -125,7 +125,7 @@ namespace mu
 		case OP_TYPE::IM_BLANKLAYOUT:
 		{
 			// Set the mipmap generation on the blank layout operation
-			auto mop = mu::Clone<ASTOpFixed>(sourceAt);
+			Ptr<ASTOpFixed> mop = mu::Clone<ASTOpFixed>(sourceAt);
 			mop->op.args.ImageBlankLayout.generateMipmaps = 1;    // true
 			mop->op.args.ImageBlankLayout.mipmapCount = Levels;
 			at = mop;
@@ -329,7 +329,7 @@ namespace mu
 		case OP_TYPE::IM_CONDITIONAL:
 		{
 			// We move the op down the two paths
-			auto newOp = mu::Clone<ASTOpConditional>(at);
+			Ptr<ASTOpConditional> newOp = mu::Clone<ASTOpConditional>(at);
 			newOp->yes = Visit(newOp->yes.child(), currentMipmapOp);
 			newOp->no = Visit(newOp->no.child(), currentMipmapOp);
 			newAt = newOp;
@@ -339,9 +339,9 @@ namespace mu
 		case OP_TYPE::IM_SWITCH:
 		{
 			// We move the op down all the paths
-			auto newOp = mu::Clone<ASTOpSwitch>(at);
+			Ptr<ASTOpSwitch> newOp = mu::Clone<ASTOpSwitch>(at);
 			newOp->def = Visit(newOp->def.child(), currentMipmapOp);
-			for (auto& c : newOp->cases)
+			for (ASTOpSwitch::FCase& c : newOp->cases)
 			{
 				c.branch = Visit(c.branch.child(), currentMipmapOp);
 			}
@@ -359,12 +359,12 @@ namespace mu
 				!typedAt->Mask
 				)
 			{
-				auto newOp = mu::Clone<ASTOpImageCompose>(at);
+				Ptr<ASTOpImageCompose> newOp = mu::Clone<ASTOpImageCompose>(at);
 
-				auto baseOp = newOp->Base.child();
+				Ptr<ASTOp> baseOp = newOp->Base.child();
 				newOp->Base = Visit(baseOp, currentMipmapOp);
 
-				auto blockOp = newOp->BlockImage.child();
+				Ptr<ASTOp> blockOp = newOp->BlockImage.child();
 				newOp->BlockImage = Visit(blockOp, currentMipmapOp);
 
 				newAt = newOp;
@@ -383,11 +383,12 @@ namespace mu
 				// mipmaps after the patch is done.
 
 				const ASTOpImagePatch* typedSource = dynamic_cast<const ASTOpImagePatch*>(at.get());
-				auto rectOp = typedSource->patch.child();
+				Ptr<ASTOp> rectOp = typedSource->patch.child();
 				ASTOp::FGetImageDescContext context;
-				auto patchDesc = rectOp->GetImageDesc(false, &context);
+				FImageDesc patchDesc = rectOp->GetImageDesc(false, &context);
 
-				uint8_t blockLevels = 0;
+				// Calculate the mip levels that can be calculated for the patch
+				uint8 PatchBlockLevels = 0;
 				{
 					uint16 minX = typedSource->location[0];
 					uint16 minY = typedSource->location[1];
@@ -397,7 +398,7 @@ namespace mu
 						&&
 						minX % 2 == 0 && minY % 2 == 0 && sizeX % 2 == 0 && sizeY % 2 == 0)
 					{
-						blockLevels++;
+						PatchBlockLevels++;
 						minX /= 2;
 						minY /= 2;
 						sizeX /= 2;
@@ -405,23 +406,30 @@ namespace mu
 					}
 				}
 
-				mu::Ptr<ASTOpImageMipmap> newMip = mu::Clone<ASTOpImageMipmap>(currentMipmapOp);
-				newMip->Levels = blockLevels;
-				newMip->BlockLevels = blockLevels;
-				newMip->bOnlyTail = false;
+				if (currentMipmapOp->Levels!= PatchBlockLevels || currentMipmapOp->BlockLevels!= PatchBlockLevels)
+				{
+					mu::Ptr<ASTOpImageMipmap> newMip = mu::Clone<ASTOpImageMipmap>(currentMipmapOp);
+					newMip->Levels = PatchBlockLevels;
+					newMip->BlockLevels = PatchBlockLevels;
+					newMip->bOnlyTail = false;
 
-				auto newOp = mu::Clone<ASTOpImagePatch>(at);
-				newOp->base = Visit(newOp->base.child(), newMip.get());
-				newOp->patch = Visit(newOp->patch.child(), newMip.get());
-				newAt = newOp;
+					Ptr<ASTOpImagePatch> newOp = mu::Clone<ASTOpImagePatch>(at);
+					newOp->base = Visit(newOp->base.child(), newMip.get());
+					newOp->patch = Visit(newOp->patch.child(), newMip.get());
 
-				// We need to add a mipmap on top to finish the mipmapping
-				{
+					// We need to add a mipmap on top to finish the mipmapping
 					mu::Ptr<ASTOpImageMipmap> topMipOp = mu::Clone<ASTOpImageMipmap>(currentMipmapOp);
 					topMipOp->Source = newOp;
 					topMipOp->bOnlyTail = true;
 					newAt = topMipOp;
 				}
+				else
+				{
+					// The patch supports the same amount of mips that we are currently sinking.
+					Ptr<ASTOpImagePatch> newOp = mu::Clone<ASTOpImagePatch>(at);
+					newOp->base = Visit(newOp->base.child(), currentMipmapOp);
+					newOp->patch = Visit(newOp->patch.child(), currentMipmapOp);
+				}
 			}
 
 			break;
