#rnx

diff --git a/Engine/Plugins/Experimental/Mutable/Source/CustomizableObject/Classes/MuCO/CustomizableObjectSystem.h b/Engine/Plugins/Experimental/Mutable/Source/CustomizableObject/Classes/MuCO/CustomizableObjectSystem.h
--- a/Engine/Plugins/Experimental/Mutable/Source/CustomizableObject/Classes/MuCO/CustomizableObjectSystem.h
+++ b/Engine/Plugins/Experimental/Mutable/Source/CustomizableObject/Classes/MuCO/CustomizableObjectSystem.h
@@ -243,6 +243,11 @@ public:
 	// This is usually only used in the editor
 	bool LockObject(const class UCustomizableObject*);
 	void UnlockObject(const class UCustomizableObject*);
+
+	/** Checks if there are any outstading disk or mip update operations in flight for the parameter Customizable Object that may
+	* make it unsafe to compile at the moment.
+	* @return true if there are operations in flight and it's not safe to compile */
+	bool CheckIfDiskOrMipUpdateOperationsPending(const UCustomizableObject& Object) const;
 	
 	// Called whenever the Mutable Editor Settings change, copying the new value of the current needed settings to the Customizable Object System
 	void EditorSettingsChanged(const FEditorCompileSettings& InEditorSettings);
diff --git a/Engine/Plugins/Experimental/Mutable/Source/CustomizableObject/Private/MuCO/CustomizableObjectMipDataProvider.cpp b/Engine/Plugins/Experimental/Mutable/Source/CustomizableObject/Private/MuCO/CustomizableObjectMipDataProvider.cpp
--- a/Engine/Plugins/Experimental/Mutable/Source/CustomizableObject/Private/MuCO/CustomizableObjectMipDataProvider.cpp
+++ b/Engine/Plugins/Experimental/Mutable/Source/CustomizableObject/Private/MuCO/CustomizableObjectMipDataProvider.cpp
@@ -44,11 +44,30 @@ FMutableTextureMipDataProvider::FMutableTextureMipDataProvider(const UTexture* T
 }
 
 
+void FMutableTextureMipDataProvider::PrintWarningAndAdvanceToCleanup()
+{
+	UE_LOG(LogMutable, Warning, TEXT("Tried to update a mip from a Customizable Object being compiled, cancelling mip update."));
+	AdvanceTo(ETickState::CleanUp, ETickThread::Async);
+}
+
+
 void FMutableTextureMipDataProvider::Init(const FTextureUpdateContext& Context, const FTextureUpdateSyncOptions& SyncOptions)
 {
+#if WITH_EDITOR
+	check(Context.Texture->HasPendingInitOrStreaming());
+	check(CustomizableObjectInstance->GetCustomizableObject());
+	if (CustomizableObjectInstance->GetCustomizableObject()->IsLocked())
+	{
+		PrintWarningAndAdvanceToCleanup();
+
+		return;
+	}
+#endif
+
 	AdvanceTo(ETickState::GetMips, ETickThread::Async);
 }
 
+
 namespace impl
 {
 	void Task_Mutable_UpdateImage(TSharedPtr<FMutableImageOperationData> OperationData)
@@ -139,6 +158,17 @@ namespace impl
 
 int32 FMutableTextureMipDataProvider::GetMips(const FTextureUpdateContext& Context, int32 StartingMipIndex, const FTextureMipInfoArray& MipInfos, const FTextureUpdateSyncOptions& SyncOptions)
 {
+#if WITH_EDITOR
+	check(Context.Texture->HasPendingInitOrStreaming());
+	check(CustomizableObjectInstance->GetCustomizableObject());
+	if (CustomizableObjectInstance->GetCustomizableObject()->IsLocked())
+	{
+		PrintWarningAndAdvanceToCleanup();
+
+		return CurrentFirstLODIdx;
+	}
+#endif
+
 	const UTexture2D* Texture = Cast<UTexture2D>(Context.Texture);
 	check(Texture);
 	check(!Texture->NeverStream);
@@ -224,6 +254,16 @@ int32 FMutableTextureMipDataProvider::GetMips(const FTextureUpdateContext& Conte
 
 bool FMutableTextureMipDataProvider::PollMips(const FTextureUpdateSyncOptions& SyncOptions)
 {
+#if WITH_EDITOR
+	check(CustomizableObjectInstance->GetCustomizableObject());
+	if (CustomizableObjectInstance->GetCustomizableObject()->IsLocked())
+	{
+		PrintWarningAndAdvanceToCleanup();
+
+		return false;
+	}
+#endif
+
 	if (!bRequestAborted && UpdateImageMutableTaskEvent)
 	{
 		if (OperationData && OperationData->Result && OperationData->Levels.Num())
diff --git a/Engine/Plugins/Experimental/Mutable/Source/CustomizableObject/Private/MuCO/CustomizableObjectPrivate.h b/Engine/Plugins/Experimental/Mutable/Source/CustomizableObject/Private/MuCO/CustomizableObjectPrivate.h
--- a/Engine/Plugins/Experimental/Mutable/Source/CustomizableObject/Private/MuCO/CustomizableObjectPrivate.h
+++ b/Engine/Plugins/Experimental/Mutable/Source/CustomizableObject/Private/MuCO/CustomizableObjectPrivate.h
@@ -19,7 +19,7 @@ public:
 	mu::Model* GetModel();
 	const mu::Model* GetModel() const;
 
-	// See UCustomizableObjectSystem::LockObject.
+	// See UCustomizableObjectSystem::LockObject. Must only be modified from the game thread
 	bool bLocked = false;
 
 #if WITH_EDITOR
diff --git a/Engine/Plugins/Experimental/Mutable/Source/CustomizableObject/Private/MuCO/CustomizableObjectSystem.cpp b/Engine/Plugins/Experimental/Mutable/Source/CustomizableObject/Private/MuCO/CustomizableObjectSystem.cpp
--- a/Engine/Plugins/Experimental/Mutable/Source/CustomizableObject/Private/MuCO/CustomizableObjectSystem.cpp
+++ b/Engine/Plugins/Experimental/Mutable/Source/CustomizableObject/Private/MuCO/CustomizableObjectSystem.cpp
@@ -900,13 +900,45 @@ void UCustomizableObjectSystem::ClearResourceCacheProtected()
 #if WITH_EDITOR
 bool UCustomizableObjectSystem::LockObject(const class UCustomizableObject* InObject)
 {
+	check(InObject);
 	check(!InObject->GetPrivate()->bLocked);
+	check(IsInGameThread() && !IsInParallelGameThread());
 
 	// If the current instance is for this object, make the lock fail by returning false
 	if (Private->CurrentInstanceBeingUpdated &&
 		Private->CurrentInstanceBeingUpdated->GetCustomizableObject() == InObject)
 	{
 		UE_LOG(LogMutable, Warning, TEXT("---- failed to lock object %s"), *InObject->GetName());
+
+		return false;
+	}
+
+	FString Message = FString::Printf(TEXT("Customizable Object %s has pending texture streaming operations. Please wait a few seconds and try again."), 
+		*InObject->GetName());
+
+	// Pre-check pending operations before locking. This check is redundant and incomplete because it's checked again after locking 
+	// and some operations may start between here and the actual lock. But in the CO Editor preview it will prevent some 
+	// textures getting stuck at low resolution when they try to update mips and are cancelled when the user presses 
+	// the compile button but the compilation quits anyway because there are pending operations
+	if (CheckIfDiskOrMipUpdateOperationsPending(*InObject))
+	{
+		UE_LOG(LogMutable, Warning, TEXT("%s"), *Message);
+
+		return false;
+	}
+
+	// Lock the object, no new file or mip streaming operations should start from this point
+	InObject->GetPrivate()->bLocked = true;
+
+	// But some could have started between the first CheckIfDiskOrMipUpdateOperationsPending and the lock a few lines back, so check again
+	if (CheckIfDiskOrMipUpdateOperationsPending(*InObject))
+	{
+		UE_LOG(LogMutable, Warning, TEXT("%s"), *Message);
+
+		// Unlock and return because the pending operations cannot be easily stopped now, the compilation hasn't started and the CO
+		// hasn't changed state yet. It's simpler to quit the compilation, unlock and let the user try to compile again
+		InObject->GetPrivate()->bLocked = false;
+
 		return false;
 	}
 
@@ -916,12 +948,12 @@ bool UCustomizableObjectSystem::LockObject(const class UCustomizableObject* InOb
 		GetPrivate()->Streamer->CancelStreamingForObject(InObject);
 	}
 
-	InObject->GetPrivate()->bLocked = true;
-
 	// Clear the cache for the instance, since we will remake it
 	FMutableResourceCache& Cache = GetPrivate()->GetObjectCache(InObject);
 	Cache.Clear();
 
+	check(InObject->GetPrivate()->bLocked);
+
 	return true;
 }
 
@@ -929,10 +961,41 @@ bool UCustomizableObjectSystem::LockObject(const class UCustomizableObject* InOb
 void UCustomizableObjectSystem::UnlockObject(const class UCustomizableObject* Obj)
 {
 	check(Obj->GetPrivate()->bLocked);
+	check(IsInGameThread() && !IsInParallelGameThread());
 	Obj->GetPrivate()->bLocked = false;
 }
 
 
+bool UCustomizableObjectSystem::CheckIfDiskOrMipUpdateOperationsPending(const UCustomizableObject& Object) const
+{
+	for (TObjectIterator<UCustomizableObjectInstance> CustomizableObjectInstance; CustomizableObjectInstance; ++CustomizableObjectInstance)
+	{
+		if (CustomizableObjectInstance->GetCustomizableObject() == &Object)
+		{
+			for (const FGeneratedTexture& GeneratedTexture : CustomizableObjectInstance->GetPrivate()->GeneratedTextures)
+			{
+				if (GeneratedTexture.Texture->HasPendingInitOrStreaming())
+				{
+					return true;
+				}
+			}
+		}
+	}
+
+	// Ensure that we don't try to handle any further streaming operations for this object
+	check(GetPrivate());
+	if (const FUnrealMutableModelBulkStreamer* Streamer = GetPrivate()->Streamer)
+	{
+		if (Streamer->AreTherePendingStreamingOperationsForObject(&Object))
+		{
+			return true;
+		}
+	}
+
+	return false;
+}
+
+
 void UCustomizableObjectSystem::EditorSettingsChanged(const FEditorCompileSettings& InEditorSettings)
 {
 	EditorSettings = InEditorSettings;
diff --git a/Engine/Plugins/Experimental/Mutable/Source/CustomizableObject/Private/MuCO/UnrealMutableModelDiskStreamer.cpp b/Engine/Plugins/Experimental/Mutable/Source/CustomizableObject/Private/MuCO/UnrealMutableModelDiskStreamer.cpp
--- a/Engine/Plugins/Experimental/Mutable/Source/CustomizableObject/Private/MuCO/UnrealMutableModelDiskStreamer.cpp
+++ b/Engine/Plugins/Experimental/Mutable/Source/CustomizableObject/Private/MuCO/UnrealMutableModelDiskStreamer.cpp
@@ -236,6 +236,32 @@ void FUnrealMutableModelBulkStreamer::CancelStreamingForObject(const UCustomizab
 	}
 }
 
+
+bool FUnrealMutableModelBulkStreamer::AreTherePendingStreamingOperationsForObject(const UCustomizableObject* CustomizableObject) const
+{
+	// This happens in the game thread
+	check(IsInGameThread());
+
+	if (!CustomizableObject)
+	{
+		check(false);
+		return false;
+	}
+
+	for (int32 ObjectIndex = 0; ObjectIndex < Objects.Num(); ++ObjectIndex)
+	{
+		if (Objects[ObjectIndex].Model.Pin() == CustomizableObject->GetModel())
+		{
+			if(!Objects[ObjectIndex].CurrentReadRequests.IsEmpty())
+			{
+				return true;
+			}
+		}
+	}
+
+	return false;
+}
+
 #endif // WITH_EDITOR
 
 
diff --git a/Engine/Plugins/Experimental/Mutable/Source/CustomizableObject/Public/MuCO/UnrealMutableModelDiskStreamer.h b/Engine/Plugins/Experimental/Mutable/Source/CustomizableObject/Public/MuCO/UnrealMutableModelDiskStreamer.h
--- a/Engine/Plugins/Experimental/Mutable/Source/CustomizableObject/Public/MuCO/UnrealMutableModelDiskStreamer.h
+++ b/Engine/Plugins/Experimental/Mutable/Source/CustomizableObject/Public/MuCO/UnrealMutableModelDiskStreamer.h
@@ -66,9 +66,14 @@ public:
 #if WITH_EDITOR
 	/** Cancel any further streaming operations for the given object. This is necessary if the object compiled data is
 	 * going to be modified. This can only happen in the editor, when recompiling.
-	 * Any additional straming requests for this object will fail.
+	 * Any additional streaming requests for this object will fail.
 	 */
 	void CancelStreamingForObject(const UCustomizableObject* CustomizableObject);
+
+	/** Checks if there are any streaming operations for the parameter object.
+	* @return true if there are streaming operations in flight
+	*/
+	bool AreTherePendingStreamingOperationsForObject(const UCustomizableObject* CustomizableObject) const;
 #endif
 
 	/** Release all the pending resources. This disables treamings for all objects. */
diff --git a/Engine/Plugins/Experimental/Mutable/Source/CustomizableObjectEditor/Private/MuCOE/CustomizableObjectCompiler.cpp b/Engine/Plugins/Experimental/Mutable/Source/CustomizableObjectEditor/Private/MuCOE/CustomizableObjectCompiler.cpp
--- a/Engine/Plugins/Experimental/Mutable/Source/CustomizableObjectEditor/Private/MuCOE/CustomizableObjectCompiler.cpp
+++ b/Engine/Plugins/Experimental/Mutable/Source/CustomizableObjectEditor/Private/MuCOE/CustomizableObjectCompiler.cpp
@@ -41,9 +41,11 @@
 #include "MuCO/CustomizableObjectClothingTypes.h"
 #include "MuCO/CustomizableObjectIdentifier.h"
 #include "MuCO/CustomizableObjectInstance.h"
+#include "MuCO/CustomizableInstancePrivateData.h"
 #include "MuCO/CustomizableObjectParameterTypeDefinitions.h"
 #include "MuCO/CustomizableObjectSystem.h"
 #include "MuCO/CustomizableObjectUIData.h"
+#include "MuCO/UnrealMutableModelDiskStreamer.h"
 #include "MuCO/ICustomizableObjectModule.h"
 #include "MuCOE/CustomizableObjectCompileRunnable.h"
 #include "MuCOE/CustomizableObjectEditorLogger.h"
@@ -232,7 +234,10 @@ void FCustomizableObjectCompiler::PreloadingReferencerAssetsCallback(UCustomizab
 
 void FCustomizableObjectCompiler::Compile(UCustomizableObject& Object, const FCompilationOptions& InOptions, bool bAsync)
 {
-	if (UCustomizableObjectSystem::GetInstance()->IsCompilationDisabled())
+	UCustomizableObjectSystem* System = UCustomizableObjectSystem::GetInstance();
+	check(System);
+
+	if (System->IsCompilationDisabled())
 	{
 		CompileInternal(&Object, InOptions, bAsync);
 		return;
@@ -240,32 +245,29 @@ void FCustomizableObjectCompiler::Compile(UCustomizableObject& Object, const FCo
 
 	UE_LOG(LogMutable, Log, TEXT("PROFILE: [ %16.8f ] Preload asynchronously assets start."), FPlatformTime::Seconds());
 
+	FString Message = FString::Printf(TEXT("Customizable Object %s is already being compiled or updated. Please wait a few seconds and try again."), *Object.GetName());
+	FNotificationInfo Info(LOCTEXT("CustomizableObjectBeingCompilerOrUpdated", "Customizable Object compile and/or update still in process. Please wait a few seconds and try again."));
+	Info.bFireAndForget = true;
+	Info.bUseThrobber = true;
+	Info.FadeOutDuration = 1.0f;
+	Info.ExpireDuration = 1.0f;
+
 	if (PreloadingReferencerAssets || CompilationLaunchPending || (Object.IsLocked()))
 	{
-		FString Message = FString::Printf(TEXT("Customizable Object %s is already being compiled or updated. Please wait a few seconds and try again."), *Object.GetName());
 		UE_LOG(LogMutable, Warning, TEXT("%s"), *Message);
-		FNotificationInfo Info(LOCTEXT("CustomizableObjectBeingCompilerOrUpdated", "Customizable object compile and update still in process"));
-		Info.bFireAndForget = true;
-		Info.bUseThrobber = true;
-		Info.FadeOutDuration = 1.0f;
-		Info.ExpireDuration = 1.0f;
 		FSlateNotificationManager::Get().AddNotification(Info);
+
 		return;
 	}
 
-	// Lock object during asynchronous asset loading to avoid instance updates
-	bool LockResult = UCustomizableObjectSystem::GetInstance()->LockObject(&Object);
+	// Lock object during asynchronous asset loading to avoid instance/mip updates and reentrant compilations
+	bool LockResult = System->LockObject(&Object);
 
 	if (!LockResult)
-	{
-		FString Message = FString::Printf(TEXT("Customizable Object %s is already being compiled or updated. Please wait a few seconds and try again."), *Object.GetName());
+	{		
 		UE_LOG(LogMutable, Warning, TEXT("%s"), *Message);
-		FNotificationInfo Info(LOCTEXT("CustomizableObjectBeingCompilerOrUpdated", "Customizable object compile and update still in process"));
-		Info.bFireAndForget = true;
-		Info.bUseThrobber = true;
-		Info.FadeOutDuration = 1.0f;
-		Info.ExpireDuration = 1.0f;
 		FSlateNotificationManager::Get().AddNotification(Info);
+
 		return;
 	}
 
@@ -291,7 +293,7 @@ void FCustomizableObjectCompiler::Compile(UCustomizableObject& Object, const FCo
 
 	if (ArrayAssetToStream.Num() > 0)
 	{
-		FStreamableManager& Streamable = UCustomizableObjectSystem::GetInstance()->GetStreamableManager();
+		FStreamableManager& Streamable = System->GetStreamableManager();
 
 		if (bAsync) 
 		{
