#rnx

diff --git a/Engine/Plugins/Experimental/Mutable/Source/CustomizableObject/Public/MuCO/CustomizableObjectMipDataProvider.h b/Engine/Plugins/Experimental/Mutable/Source/CustomizableObject/Public/MuCO/CustomizableObjectMipDataProvider.h
--- a/Engine/Plugins/Experimental/Mutable/Source/CustomizableObject/Public/MuCO/CustomizableObjectMipDataProvider.h
+++ b/Engine/Plugins/Experimental/Mutable/Source/CustomizableObject/Public/MuCO/CustomizableObjectMipDataProvider.h
@@ -90,6 +90,7 @@ public:
 
 private:
 	void CancelCounterSafely();
+	void PrintWarningAndAdvanceToCleanup();
 
 public:
 	// Todo: Simplify by replacing the reference to the Instance with some static parametrization or hash with enough info to reconstruct the texture
