[mutable] Fix crash when generating layout blocks from UVs

#rnx

diff --git a/Engine/Plugins/Experimental/Mutable/Source/CustomizableObjectEditor/Private/MuCOE/CustomizableObjectLayout.cpp b/Engine/Plugins/Experimental/Mutable/Source/CustomizableObjectEditor/Private/MuCOE/CustomizableObjectLayout.cpp
--- a/Engine/Plugins/Experimental/Mutable/Source/CustomizableObjectEditor/Private/MuCOE/CustomizableObjectLayout.cpp
+++ b/Engine/Plugins/Experimental/Mutable/Source/CustomizableObjectEditor/Private/MuCOE/CustomizableObjectLayout.cpp
@@ -95,6 +95,12 @@ void UCustomizableObjectLayout::GenerateBlocksFromUVs()
 		
 		if (USkeletalMesh* SkeletalMesh = Cast<USkeletalMesh>(Mesh))
 		{
+			// We don't need all the data to generate the blocks
+			const uint32 ShapeFlags = uint32(EMutableMeshConversionFlags::IgnoreSkinning) &
+				uint32(EMutableMeshConversionFlags::IgnorePhysics);
+
+			GenerationContext.MeshGenerationFlags.Push(ShapeFlags);
+
 			GenerationContext.ComponentInfos.Add(SkeletalMesh);
 			MutableMesh = ConvertSkeletalMeshToMutable(SkeletalMesh, LOD, Material, GenerationContext, Node);
 		}
diff --git a/Engine/Plugins/Experimental/Mutable/Source/CustomizableObjectEditor/Private/MuCOE/GenerateMutableSource/GenerateMutableSourceMesh.cpp b/Engine/Plugins/Experimental/Mutable/Source/CustomizableObjectEditor/Private/MuCOE/GenerateMutableSource/GenerateMutableSourceMesh.cpp
--- a/Engine/Plugins/Experimental/Mutable/Source/CustomizableObjectEditor/Private/MuCOE/GenerateMutableSource/GenerateMutableSourceMesh.cpp
+++ b/Engine/Plugins/Experimental/Mutable/Source/CustomizableObjectEditor/Private/MuCOE/GenerateMutableSource/GenerateMutableSourceMesh.cpp
@@ -543,7 +543,7 @@ mu::MeshPtr ConvertSkeletalMeshToMutable(USkeletalMesh* InSkeletalMesh, int LOD,
 	if (!bIgnoreSkeleton)
 	{
 		// Add the RefSkeleton ID to the mesh.
-		const int32 RefSkeletonID = GenerationContext.ReferencedSkeletons.Find(ComponentRefSkeleton);
+		const int32 RefSkeletonID = GenerationContext.ReferencedSkeletons.AddUnique(ComponentRefSkeleton);
 		MutableMesh->AddSkeletonID(RefSkeletonID);
 
 		if (InSkeleton != ComponentRefSkeleton)
