[mutable] optimizations for bone and physics reshape operations and other improvements.

-Physics reshape methods have been reimplemented to be more robust, add more sampling points and have a better method of extracting the shape characteristics from those points.
-Bind data that does not track orientation has it's own more compact data type.
-Physics mesh merging will only merge when necessary.
-Physics asset creation does not create a duplicate which is slow but creates a new asset and populate it with relevant data.

diff --git a/Engine/Plugins/Experimental/Mutable/Source/CustomizableObject/Classes/MuCO/CustomizableObject.h b/Engine/Plugins/Experimental/Mutable/Source/CustomizableObject/Classes/MuCO/CustomizableObject.h
--- a/Engine/Plugins/Experimental/Mutable/Source/CustomizableObject/Classes/MuCO/CustomizableObject.h
+++ b/Engine/Plugins/Experimental/Mutable/Source/CustomizableObject/Classes/MuCO/CustomizableObject.h
@@ -1220,7 +1220,7 @@ private:
 	// This is a manual version number for the binary blobs in this asset.
 	// Increasing it invalidates all the previously compiled models.
 	// Warning: If while merging code both versions have changed, take the highest+1.
-	static const int32 CurrentSupportedVersion = 357;
+	static const int32 CurrentSupportedVersion = 358;
 
 public:
 
diff --git a/Engine/Plugins/Experimental/Mutable/Source/CustomizableObject/Private/MuCO/CustomizableInstancePrivateData.h b/Engine/Plugins/Experimental/Mutable/Source/CustomizableObject/Private/MuCO/CustomizableInstancePrivateData.h
--- a/Engine/Plugins/Experimental/Mutable/Source/CustomizableObject/Private/MuCO/CustomizableInstancePrivateData.h
+++ b/Engine/Plugins/Experimental/Mutable/Source/CustomizableObject/Private/MuCO/CustomizableInstancePrivateData.h
@@ -260,7 +260,7 @@ private:
 	USkeleton* MergeSkeletons(UCustomizableObjectInstance* Public, const FMutableRefSkeletalMeshData* RefSkeletalMeshData, int32 ComponentIndex);
 
 	//
-	UPhysicsAsset* BuildPhysicsAsset(TObjectPtr<class UPhysicsAsset> TamplateAsset, const mu::PhysicsBody* PhysicsBody, int32 ComponentIndex, bool bDisableCollisionBetweenAssets);
+	UPhysicsAsset* GetOrBuildPhysicsAsset(TObjectPtr<class UPhysicsAsset> TamplateAsset, const mu::PhysicsBody* PhysicsBody, int32 ComponentIndex, bool bDisableCollisionBetweenAssets);
 	
 	
 	int32 GetLastMeshId(int32 ComponentIndex, int32 LODIndex) const;
diff --git a/Engine/Plugins/Experimental/Mutable/Source/CustomizableObject/Private/MuCO/CustomizableObjectInstance.cpp b/Engine/Plugins/Experimental/Mutable/Source/CustomizableObject/Private/MuCO/CustomizableObjectInstance.cpp
--- a/Engine/Plugins/Experimental/Mutable/Source/CustomizableObject/Private/MuCO/CustomizableObjectInstance.cpp
+++ b/Engine/Plugins/Experimental/Mutable/Source/CustomizableObject/Private/MuCO/CustomizableObjectInstance.cpp
@@ -1050,10 +1050,10 @@ USkeleton* UCustomizableInstancePrivateData::MergeSkeletons(UCustomizableObjectI
 	return FinalSkeleton;
 }
 
-UPhysicsAsset* UCustomizableInstancePrivateData::BuildPhysicsAsset(
-		TObjectPtr<UPhysicsAsset> TemplateAsset,
-		const mu::PhysicsBody* MutablePhysics, int32 ComponentId,
-		bool bDisableCollisionsBetweenDifferentAssets)
+UPhysicsAsset* UCustomizableInstancePrivateData::GetOrBuildPhysicsAsset(
+	TObjectPtr<UPhysicsAsset> TemplateAsset,
+	const mu::PhysicsBody* MutablePhysics, int32 ComponentId,
+	bool bDisableCollisionsBetweenDifferentAssets)
 {
 
 	MUTABLE_CPUPROFILER_SCOPE(MergePhysicsAssets);
@@ -1062,12 +1062,12 @@ UPhysicsAsset* UCustomizableInstancePrivateData::BuildPhysicsAsset(
 
 	UPhysicsAsset* Result = nullptr;
 
-	FCustomizableInstanceComponentData* ComponentData =	GetComponentData(ComponentId);
+	FCustomizableInstanceComponentData* ComponentData = GetComponentData(ComponentId);
 	check(ComponentData);
 
 	TArray<TObjectPtr<UPhysicsAsset>>& PhysicsAssets = ComponentData->PhysicsAssets.PhysicsAssetsToMerge;
 
-	TArray<TObjectPtr<UPhysicsAsset>> ValidAssets;		
+	TArray<TObjectPtr<UPhysicsAsset>> ValidAssets;
 
 	const int32 NumPhysicsAssets = ComponentData->PhysicsAssets.PhysicsAssetsToMerge.Num();
 	for (int32 I = 0; I < NumPhysicsAssets; ++I)
@@ -1079,29 +1079,33 @@ UPhysicsAsset* UCustomizableInstancePrivateData::BuildPhysicsAsset(
 			ValidAssets.AddUnique(PhysicsAsset);
 		}
 	}
-	
+
 	if (!ValidAssets.Num())
 	{
 		return Result;
 	}
 
-	if (!TemplateAsset)
-	{
-		Result = DuplicateObject(ValidAssets[0], nullptr);
-	}
-	else
+	// Just get the referenced asset if no recontrution or merge is needed.
+	if (ValidAssets.Num() == 1 && !MutablePhysics->bBodiesModified)
 	{
-		Result = DuplicateObject(TemplateAsset, nullptr);	
+		return ValidAssets[0];
 	}
 
+	TemplateAsset = TemplateAsset ? TemplateAsset : ValidAssets[0];
+	check(TemplateAsset);
+
+	Result = NewObject<UPhysicsAsset>();
+
 	if (!Result)
 	{
 		return nullptr;
 	}
-	// TODO: We are duplicating just to remove the data that will be replaced.
-	Result->CollisionDisableTable.Empty();
-	Result->SkeletalBodySetups.Empty();
-	Result->ConstraintSetup.Empty();
+
+	Result->SolverSettings = TemplateAsset->SolverSettings;
+	Result->SolverType = TemplateAsset->SolverType;
+
+	Result->bNotForDedicatedServer = TemplateAsset->bNotForDedicatedServer;
+
 
 	auto MakeAggGeomFromMutablePhysics = [](int32 BodyIndex, const mu::PhysicsBody* MutablePhysicsBody) -> FKAggregateGeom
 	{
@@ -1116,7 +1120,7 @@ UPhysicsAsset* UCustomizableInstancePrivateData::BuildPhysicsAsset(
 		{
 			return static_cast<bool>( (Flags >> 8 ) & 1 );
 		};
-				
+
 		const int32 NumSpheres = MutablePhysicsBody->GetSphereCount( BodyIndex );
 		TArray<FKSphereElem>& AggSpheres = BodyAggGeom.SphereElems;
 		AggSpheres.Empty(NumSpheres);
@@ -1163,25 +1167,25 @@ UPhysicsAsset* UCustomizableInstancePrivateData::BuildPhysicsAsset(
 			NewElem.SetName( FName(*Name) );
 		}
 
-		const int32 NumConvexes = MutablePhysicsBody->GetConvexCount( BodyIndex );
-		TArray<FKConvexElem>& AggConvexes = BodyAggGeom.ConvexElems;
-		AggConvexes.Empty();
-		for (int32 I = 0; I < NumConvexes; ++I)
-		{
-			uint32 Flags = MutablePhysicsBody->GetConvexFlags( BodyIndex, I );
-			FString Name = MutablePhysicsBody->GetConvexName( BodyIndex, I );
-
-			const FVector3f* Vertices;
-			const int32* Indices;
-			int32 NumVertices;
-			int32 NumIndices;
-			FTransform3f Transform;
-
-			MutablePhysicsBody->GetConvex( BodyIndex, I, Vertices, NumVertices, Indices, NumIndices, Transform );
-			
-			TArrayView<const FVector3f> VerticesView( Vertices, NumVertices );
-			TArrayView<const int32> IndicesView( Indices, NumIndices );
-		}
+		//const int32 NumConvexes = MutablePhysicsBody->GetConvexCount( BodyIndex );
+		//TArray<FKConvexElem>& AggConvexes = BodyAggGeom.ConvexElems;
+		//AggConvexes.Empty();
+		//for (int32 I = 0; I < NumConvexes; ++I)
+		//{
+		//	uint32 Flags = MutablePhysicsBody->GetConvexFlags( BodyIndex, I );
+		//	FString Name = MutablePhysicsBody->GetConvexName( BodyIndex, I );
+
+		//	const FVector3f* Vertices;
+		//	const int32* Indices;
+		//	int32 NumVertices;
+		//	int32 NumIndices;
+		//	FTransform3f Transform;
+
+		//	MutablePhysicsBody->GetConvex( BodyIndex, I, Vertices, NumVertices, Indices, NumIndices, Transform );
+		//	
+		//	TArrayView<const FVector3f> VerticesView( Vertices, NumVertices );
+		//	TArrayView<const int32> IndicesView( Indices, NumIndices );
+		//}
 
 		TArray<FKSphylElem>& AggSphyls = BodyAggGeom.SphylElems;
 		const int32 NumSphyls = MutablePhysicsBody->GetSphylCount( BodyIndex );
@@ -1315,8 +1319,14 @@ UPhysicsAsset* UCustomizableInstancePrivateData::BuildPhysicsAsset(
 				NewBodySetup->WalkableSlopeOverride = BodySetup->WalkableSlopeOverride;
 				NewBodySetup->BuildScale3D = BodySetup->BuildScale3D;	
 				NewBodySetup->bSkipScaleFromAnimation = BodySetup->bSkipScaleFromAnimation;
+				
+				// PhysicalAnimationProfiles can't be added with the current UPhysicsAsset API outside the editor.
+				// Don't poulate them for now.	
+				//NewBodySetup->PhysicalAnimationData = BodySetup->PhysicalAnimationData;
+
 				NewBodySetup->AggGeom = MakeAggGeomFromMutablePhysics(*MutableBodyIndex, MutablePhysics);
 
+
 				Result->SkeletalBodySetups.Add(NewBodySetup);
 				
 				int32 IndexInSet = CollisionSets[CollisionSetIndex].Add(CurrentBodyIndex);
@@ -1574,7 +1584,7 @@ bool UCustomizableInstancePrivateData::UpdateSkeletalMesh_PostBeginUpdate0(UCust
 					if (MutablePhysics)
 					{
 						constexpr bool bDisallowCollisionBetweenAssets = true;
-						UPhysicsAsset* PhysicsAssetResult = BuildPhysicsAsset( 
+						UPhysicsAsset* PhysicsAssetResult = GetOrBuildPhysicsAsset( 
 								RefSkeletalMeshData->PhysicsAsset.Get(), MutablePhysics.get(), SkeletalMeshIndex, bDisallowCollisionBetweenAssets);
 					
 						SkeletalMesh->SetPhysicsAsset(PhysicsAssetResult);
@@ -3070,7 +3080,7 @@ bool UCustomizableInstancePrivateData::BuildSkeletalMeshSkeletonData(const TShar
 		}
 
 		// First step is to update the RefBasesInvMatrix for the bones.
-		for (const auto& BoneMatrix : MutSkeletonData.BoneMatricesWithScale)
+		for (const TPair<FName, FMatrix44f>& BoneMatrix : MutSkeletonData.BoneMatricesWithScale)
 		{
 			const int32 RefSkelBoneIndex = BoneToFinalBoneIndexMap[BoneMatrix.Key];
 			RefBasesInvMatrix[RefSkelBoneIndex] = BoneMatrix.Value;
@@ -3083,12 +3093,8 @@ bool UCustomizableInstancePrivateData::BuildSkeletalMeshSkeletonData(const TShar
 			int32 ParentBoneIndex = ReferenceSkeleton.GetParentIndex(RefSkelBoneIndex);
 			if (ParentBoneIndex >= 0)
 			{
-				FMatrix44f BoneBaseInvMatrix = RefBasesInvMatrix[RefSkelBoneIndex];
-				FMatrix44f ParentBaseInvMatrix = RefBasesInvMatrix[ParentBoneIndex];
-				FMatrix44f BonePoseMatrix = BoneBaseInvMatrix.Inverse() * ParentBaseInvMatrix;
-
-				FTransform3f BonePoseTransform;
-				BonePoseTransform.SetFromMatrix(BonePoseMatrix);
+				const FTransform3f BonePoseTransform(
+						RefBasesInvMatrix[RefSkelBoneIndex].Inverse() * RefBasesInvMatrix[ParentBoneIndex]);
 
 				SkeletonModifier.UpdateRefPoseTransform(RefSkelBoneIndex, (FTransform)BonePoseTransform);
 			}
diff --git a/Engine/Plugins/Experimental/Mutable/Source/CustomizableObject/Private/MuCO/CustomizableObjectSystem.cpp b/Engine/Plugins/Experimental/Mutable/Source/CustomizableObject/Private/MuCO/CustomizableObjectSystem.cpp
--- a/Engine/Plugins/Experimental/Mutable/Source/CustomizableObject/Private/MuCO/CustomizableObjectSystem.cpp
+++ b/Engine/Plugins/Experimental/Mutable/Source/CustomizableObject/Private/MuCO/CustomizableObjectSystem.cpp
@@ -1442,7 +1442,7 @@ namespace impl
 			TArray<uint32>& SkeletonIds = OperationData->InstanceUpdateData.Skeletons[ComponentIndex].SkeletonIds;
 			TArray<FName>& BoneNames = OperationData->InstanceUpdateData.Skeletons[ComponentIndex].BoneNames;
 			TMap<FName, FMatrix44f>& BoneMatricesWithScale = OperationData->InstanceUpdateData.Skeletons[ComponentIndex].BoneMatricesWithScale;
-
+			
 			// Use first valid LOD bone count as a potential total number of bones, used for pre-allocating data arrays
 			if (MinLODComponent.Mesh && MinLODComponent.Mesh->GetSkeleton())
 			{
@@ -1487,71 +1487,71 @@ namespace impl
 					const int32 NumVerticesLODModel = Mesh->GetVertexCount();
 					const int32 SurfaceCount = Mesh->GetSurfaceCount();
 
-					mu::MESH_BUFFER_FORMAT boneIndexFormat = mu::MBF_NONE;
-					int boneIndexComponents = 0;
-					int boneIndexOffset = 0;
-					int boneIndexBuffer = -1;
-					int boneIndexChannel = -1;
-					MutableMeshVertexBuffers.FindChannel(mu::MBS_BONEINDICES, 0, &boneIndexBuffer, &boneIndexChannel);
-					if (boneIndexBuffer >= 0 || boneIndexChannel >= 0)
+					mu::MESH_BUFFER_FORMAT BoneIndexFormat = mu::MBF_NONE;
+					int32 BoneIndexComponents = 0;
+					int32 BoneIndexOffset = 0;
+					int32 BoneIndexBuffer = -1;
+					int32 BoneIndexChannel = -1;
+					MutableMeshVertexBuffers.FindChannel(mu::MBS_BONEINDICES, 0, &BoneIndexBuffer, &BoneIndexChannel);
+					if (BoneIndexBuffer >= 0 || BoneIndexChannel >= 0)
 					{
-						MutableMeshVertexBuffers.GetChannel(boneIndexBuffer, boneIndexChannel,
-							nullptr, nullptr, &boneIndexFormat, &boneIndexComponents, &boneIndexOffset);
+						MutableMeshVertexBuffers.GetChannel(BoneIndexBuffer, BoneIndexChannel,
+							nullptr, nullptr, &BoneIndexFormat, &BoneIndexComponents, &BoneIndexOffset);
 					}
 
-					int32 elemSize = MutableMeshVertexBuffers.GetElementSize(boneIndexBuffer);
+					const int32 ElemSize = MutableMeshVertexBuffers.GetElementSize(BoneIndexBuffer);
 
-					const uint8_t* BufferStart = MutableMeshVertexBuffers.GetBufferData(boneIndexBuffer) + boneIndexOffset;
+					const uint8_t* BufferStart = MutableMeshVertexBuffers.GetBufferData(BoneIndexBuffer) + BoneIndexOffset;
 
 					UsedBones.Empty(BoneCount);
 					UsedBones.AddDefaulted(BoneCount);
 
 					for (int32 Surface = 0; Surface < SurfaceCount; Surface++)
 					{
-						int32 is, ic, vs, VertexCount;
-						Mesh->GetSurface(Surface, &vs, &VertexCount, &is, &ic);
+						int32 IndexStart, IndexCount, VertexStart, VertexCount;
+						Mesh->GetSurface(Surface, &VertexStart, &VertexCount, &IndexStart, &IndexCount);
 
-						if (VertexCount == 0 || ic == 0)
+						if (VertexCount == 0 || IndexCount == 0)
 						{
 							continue;
 						}
 
-						const uint8_t* VertexBoneIndexPtr = BufferStart + vs * elemSize;
+						const uint8_t* VertexBoneIndexPtr = BufferStart + VertexStart * ElemSize;
 
-						if (boneIndexFormat == mu::MBF_UINT8)
+						if (BoneIndexFormat == mu::MBF_UINT8)
 						{
 							for (int32 v = 0; v < VertexCount; ++v)
 							{
-								for (int32 i = 0; i < boneIndexComponents; ++i)
+								for (int32 i = 0; i < BoneIndexComponents; ++i)
 								{
 									size_t SectionBoneIndex = VertexBoneIndexPtr[i];
 									UsedBones[SectionBoneIndex] = true;
 								}
-								VertexBoneIndexPtr += elemSize;
+								VertexBoneIndexPtr += ElemSize;
 							}
 						}
-						else if (boneIndexFormat == mu::MBF_UINT16)
+						else if (BoneIndexFormat == mu::MBF_UINT16)
 						{
 							for (int32 v = 0; v < VertexCount; ++v)
 							{
-								for (int32 i = 0; i < boneIndexComponents; ++i)
+								for (int32 i = 0; i < BoneIndexComponents; ++i)
 								{
 									size_t SectionBoneIndex = ((uint16*)VertexBoneIndexPtr)[i];
 									UsedBones[SectionBoneIndex] = true;
 								}
-								VertexBoneIndexPtr += elemSize;
+								VertexBoneIndexPtr += ElemSize;
 							}
 						}
-						else if (boneIndexFormat == mu::MBF_UINT32)
+						else if (BoneIndexFormat == mu::MBF_UINT32)
 						{
 							for (int32 v = 0; v < VertexCount; ++v)
 							{
-								for (int32 i = 0; i < boneIndexComponents; ++i)
+								for (int32 i = 0; i < BoneIndexComponents; ++i)
 								{
 									size_t SectionBoneIndex = ((uint32_t*)VertexBoneIndexPtr)[i];
 									UsedBones[SectionBoneIndex] = true;
 								}
-								VertexBoneIndexPtr += elemSize;
+								VertexBoneIndexPtr += ElemSize;
 							}
 						}
 						else
@@ -1559,11 +1559,9 @@ namespace impl
 							// Unsupported bone index format in generated mutable mesh
 							check(false);
 						}
-
 					}
 				}
 
-
 				{
 					MUTABLE_CPUPROFILER_SCOPE(PrepareSkeletonData_ActiveBones);
 
@@ -1576,9 +1574,9 @@ namespace impl
 					ParentChain.SetNumZeroed(BoneCount);
 
 					// Add new bones and their poses to the final hierarchy of active bones
-					auto AddBone = [&](const char* strName) {
+					auto AddBone = [&](const char* StrName) {
 
-						const FName BoneName = strName;
+						const FName BoneName = StrName;
 						const int16 FinalBoneIndex = BoneNames.Num();
 
 						// Add bone
@@ -1588,12 +1586,12 @@ namespace impl
 						ComponentActiveBones.Add(FinalBoneIndex);
 
 						// Add bone pose
-						const int32 BonePoseIndex = Mesh->FindBonePose(strName);
+						const int32 BonePoseIndex = Mesh->FindBonePose(StrName);
 						if (BonePoseIndex != INDEX_NONE)
 						{
-							FTransform3f transform;
-							Mesh->GetBoneTransform(BonePoseIndex, transform);
-							BoneMatricesWithScale.Add(BoneName, transform.Inverse().ToMatrixWithScale());
+							FTransform3f Transform;
+							Mesh->GetBoneTransform(BonePoseIndex, Transform);
+							BoneMatricesWithScale.Emplace(BoneName, Transform.Inverse().ToMatrixWithScale());
 						}
 
 						return FinalBoneIndex;
@@ -1690,7 +1688,7 @@ namespace impl
 
 					ComponentBoneMap.Shrink();
 					ComponentActiveBones.Shrink();
-				}
+				}	
 			}
 
 			BoneNames.Shrink();
diff --git a/Engine/Plugins/Experimental/Mutable/Source/CustomizableObjectEditor/Private/MuCOE/GenerateMutableSource/GenerateMutableSourceMesh.cpp b/Engine/Plugins/Experimental/Mutable/Source/CustomizableObjectEditor/Private/MuCOE/GenerateMutableSource/GenerateMutableSourceMesh.cpp
--- a/Engine/Plugins/Experimental/Mutable/Source/CustomizableObjectEditor/Private/MuCOE/GenerateMutableSource/GenerateMutableSourceMesh.cpp
+++ b/Engine/Plugins/Experimental/Mutable/Source/CustomizableObjectEditor/Private/MuCOE/GenerateMutableSource/GenerateMutableSourceMesh.cpp
@@ -1548,10 +1548,11 @@ mu::MeshPtr ConvertSkeletalMeshToMutable(USkeletalMesh* InSkeletalMesh, int LOD,
 					VertexData[Elem] = FVector3f(ConvexElem.VertexData[Elem]);
 				}
 				
-				PhysicsBody->SetConvex( B, I, 
-						VertexData.GetData(), ConvexElem.VertexData.Num(), 
-						ConvexElem.IndexData.GetData(), ConvexElem.IndexData.Num(), 
-						FTransform3f(ConvexElem.GetTransform()) );
+				PhysicsBody->SetConvexMesh(B, I,
+						TArrayView<const FVector3f>(VertexData.GetData(), ConvexElem.VertexData.Num()),
+						TArrayView<const int32>(ConvexElem.IndexData.GetData(), ConvexElem.IndexData.Num()));
+
+				PhysicsBody->SetConvexTransform(B, I, FTransform3f(ConvexElem.GetTransform()));
 				
 				const FString KElemName = ConvexElem.GetName().ToString();
 				PhysicsBody->SetConvexName(B, I, TCHAR_TO_ANSI(*KElemName));
diff --git a/Engine/Plugins/Experimental/Mutable/Source/MutableRuntime/Private/MuR/OpMeshBind.h b/Engine/Plugins/Experimental/Mutable/Source/MutableRuntime/Private/MuR/OpMeshBind.h
--- a/Engine/Plugins/Experimental/Mutable/Source/MutableRuntime/Private/MuR/OpMeshBind.h
+++ b/Engine/Plugins/Experimental/Mutable/Source/MutableRuntime/Private/MuR/OpMeshBind.h
@@ -149,6 +149,32 @@ namespace mu
 		}
 	};
 
+	struct FReshapePointBindingData
+	{
+		float S, T;
+		float D;
+		int32 Triangle;
+		float Weight;
+	};
+
+	static_assert(sizeof(FReshapePointBindingData) == 4*5);
+
+	struct FReshapePointBindingDataBufferDescriptor
+	{
+		constexpr static int Channels = 4;
+		constexpr static MESH_BUFFER_SEMANTIC Semantics[Channels] = { MBS_BARYCENTRICCOORDS, MBS_DISTANCE, MBS_TRIANGLEINDEX, MBS_OTHER };
+		constexpr static MESH_BUFFER_FORMAT Formats[Channels] = { MBF_FLOAT32, MBF_FLOAT32, MBF_INT32, MBF_FLOAT32 };
+		constexpr static int Components[Channels] = { 2, 1, 1, 1 };
+		constexpr static int Offsets[Channels] = { 0, 8, 12, 16 };
+
+		const int SemanticIndices[Channels] = { 0, 0, 0, 0 };
+
+		FReshapePointBindingDataBufferDescriptor(int32 DataSetIndex)
+			: SemanticIndices{ DataSetIndex, DataSetIndex, DataSetIndex, DataSetIndex }
+		{
+		}
+	};
+	
 	struct FUsedIndicesBufferDescriptor
 	{
 		constexpr static int Channels = 1;
@@ -163,6 +189,9 @@ namespace mu
 		}
 	};
 
+	using ReshapePoseBindingType = FReshapePointBindingData;
+	using ReshapePhysicsBindingType = FReshapePointBindingData;
+
 	// Structure used for vertex bing data in vertexbuffers for Clip deform operations.
 	struct FClipDeformVertexBindingData
 	{
@@ -192,7 +221,6 @@ namespace mu
 		}
 	};
 
-
 	//---------------------------------------------------------------------------------------------
     //! Generate the mesh-shape binding data for Reshape operations
     //---------------------------------------------------------------------------------------------
@@ -458,11 +486,11 @@ namespace mu
 		BindReshapePoint(ShapeMeshTree, BoundPoint, OutBindingData, BindTolerance);
 	}
 
-	inline TTuple<TArray<FReshapeVertexBindingData>, TArray<int32>> BindPhysicsBodies( 
+	inline TTuple<TArray<FReshapePointBindingData>, TArray<int32>> BindPhysicsBodies( 
 			const PhysicsBody* PBody, FShapeMeshTree& ShapeMeshTree, const Mesh* pMesh, 
 			const TArray<string>& PhysicsToDeform )
 	{
-		TTuple<TArray<FReshapeVertexBindingData>, TArray<int32>> ReturnValue;
+		TTuple<TArray<FReshapePointBindingData>, TArray<int32>> ReturnValue;
 
 		TArray<int32>& BodiesToDeformIndices = ReturnValue.Get<1>();
 	
@@ -470,7 +498,7 @@ namespace mu
 		const int32 NumBodies = PBody->GetBodyCount();
 		for (int32 I = 0; I < NumBodies; ++I)
 		{
-			if ( PhysicsToDeform.Contains(PBody->GetBodyBoneName(I)))
+			if (PhysicsToDeform.Contains(PBody->GetBodyBoneName(I)))
 			{
 				BodiesToDeformIndices.Add(I);
 			}
@@ -479,24 +507,23 @@ namespace mu
 		// Count how many poinst will be needed.
 		const int32 NumBodiesToDeformIndices = BodiesToDeformIndices.Num();
 		int32 NumPoints = 0;
-		for ( const int32 I : BodiesToDeformIndices )
+		for (const int32 I : BodiesToDeformIndices)
 		{
-			NumPoints += PBody->GetSphereCount(I) * 8;
-			NumPoints += PBody->GetBoxCount(I) * 8;
-			NumPoints += PBody->GetSphylCount(I) * 8;
-			NumPoints += PBody->GetTaperedCapsuleCount(I) * 8;
+			NumPoints += PBody->GetSphereCount(I) * 6;
+			NumPoints += PBody->GetBoxCount(I) * 14;
+			NumPoints += PBody->GetSphylCount(I) * 14;
+			NumPoints += PBody->GetTaperedCapsuleCount(I) * 14;
 
 			const int32 ConvexCount = PBody->GetConvexCount(I);
 			for ( int C = 0; C < ConvexCount; ++C )
 			{
-				const FVector3f* VertData;
-				int32 VertCount;
-				const int32* IndicesData;
-				int32 IndicesCount;
+				TArrayView<const FVector3f> Vertices;
+				TArrayView<const int32> Indices;
+
 				FTransform3f Transform;
-				PBody->GetConvex(I, C, VertData, VertCount, IndicesData, IndicesCount, Transform);
+				PBody->GetConvex(I, C, Vertices, Indices, Transform);
 
-				NumPoints += VertCount;
+				NumPoints += Vertices.Num();
 			}
 		}
 		
@@ -509,10 +536,10 @@ namespace mu
 		// Currently using oriented bounding boxes corners for all shapes except convex which uses the deformed mesh directly.
 
 		int32 AddedPoints = 0;
-		for ( const int32 B : BodiesToDeformIndices )
+		for (const int32 B : BodiesToDeformIndices)
 		{
-			const char* BoneName = PBody->GetBodyBoneName( B );
-			int32 BoneIdx = pMesh->FindBonePose( BoneName );
+			const char* BoneName = PBody->GetBodyBoneName(B);
+			int32 BoneIdx = pMesh->FindBonePose(BoneName);
 
 			FTransform3f T = FTransform3f::Identity;
 			if (BoneIdx > 0)
@@ -521,28 +548,25 @@ namespace mu
 			}
 			
 			const int32 SphereCount = PBody->GetSphereCount(B);
-			for ( int32 I = 0; I < SphereCount; ++I, AddedPoints += 8)
+			for ( int32 I = 0; I < SphereCount; ++I, AddedPoints += 6)
 			{
 				FVector3f P;
 				float R;
 
 				PBody->GetSphere(B, I, P, R);	
 
-				Points[AddedPoints + 0] = T.TransformPosition( P + FVector3f(R, R, R)  );
-				Points[AddedPoints + 1] = T.TransformPosition( P + FVector3f(-R, R, R)  );
-				Points[AddedPoints + 2] = T.TransformPosition( P + FVector3f(-R, -R, R)  );
-				Points[AddedPoints + 3] = T.TransformPosition( P + FVector3f(R, -R, R)  );
+				Points[AddedPoints + 0] = T.TransformPosition(P + FVector3f(R, 0.0f, 0.0f));
+				Points[AddedPoints + 1] = T.TransformPosition(P - FVector3f(R, 0.0f, 0.0f));
 
-														  
-				Points[AddedPoints + 4] = T.TransformPosition( P + FVector3f(R, R, -R) );
-				Points[AddedPoints + 5] = T.TransformPosition( P + FVector3f(-R, R, -R) );
-				Points[AddedPoints + 6] = T.TransformPosition( P + FVector3f(-R, -R, -R) );
-				Points[AddedPoints + 7] = T.TransformPosition( P + FVector3f(R, -R, -R) );
+				Points[AddedPoints + 2] = T.TransformPosition(P + FVector3f(0.0f, R, 0.0f));
+				Points[AddedPoints + 3] = T.TransformPosition(P - FVector3f(0.0f, R, 0.0f));
 
+				Points[AddedPoints + 4] = T.TransformPosition(P + FVector3f(0.0f, 0.0f, R));
+				Points[AddedPoints + 5] = T.TransformPosition(P - FVector3f(0.0f, 0.0f, R));
 			}
 
 			const int32 BoxCount = PBody->GetBoxCount(B);
-			for ( int32 I = 0; I < BoxCount; ++I, AddedPoints += 8)
+			for ( int32 I = 0; I < BoxCount; ++I, AddedPoints += 14)
 			{
 				FVector3f P;
 				FQuat4f Q;
@@ -550,20 +574,31 @@ namespace mu
 			
 				PBody->GetBox(B, I, P, Q, S);
 
-				Points[AddedPoints + 0] = T.TransformPosition( P + Q.RotateVector(FVector3f(S.X, S.Y, S.Z) ) );
-				Points[AddedPoints + 1] = T.TransformPosition( P + Q.RotateVector(FVector3f(-S.X, S.Y, S.Z) ) );
-				Points[AddedPoints + 2] = T.TransformPosition( P + Q.RotateVector(FVector3f(-S.X, -S.Y, S.Z) ) );
-				Points[AddedPoints + 3] = T.TransformPosition( P + Q.RotateVector(FVector3f(S.X, -S.Y, S.Z) ) );
+				const FVector3f BasisX = Q.RotateVector(FVector3f::UnitX());
+				const FVector3f BasisY = Q.RotateVector(FVector3f::UnitY());
+				const FVector3f BasisZ = Q.RotateVector(FVector3f::UnitZ());
+
+				Points[AddedPoints + 0] = T.TransformPosition(P + BasisX*S.X + BasisY*S.Y + BasisZ*S.Z);
+				Points[AddedPoints + 1] = T.TransformPosition(P + BasisX*S.X - BasisY*S.Y + BasisZ*S.Z);
+				Points[AddedPoints + 2] = T.TransformPosition(P - BasisX*S.X + BasisY*S.Y + BasisZ*S.Z);
+				Points[AddedPoints + 3] = T.TransformPosition(P - BasisX*S.X - BasisY*S.Y + BasisZ*S.Z);
+
+				Points[AddedPoints + 4] = T.TransformPosition(P + BasisX*S.X + BasisY*S.Y - BasisZ*S.Z);
+				Points[AddedPoints + 5] = T.TransformPosition(P + BasisX*S.X - BasisY*S.Y - BasisZ*S.Z);
+				Points[AddedPoints + 6] = T.TransformPosition(P - BasisX*S.X + BasisY*S.Y - BasisZ*S.Z);
+				Points[AddedPoints + 7] = T.TransformPosition(P - BasisX*S.X - BasisY*S.Y - BasisZ*S.Z);
 
-														  
-				Points[AddedPoints + 4] = T.TransformPosition( P + Q.RotateVector(FVector3f(S.X, S.Y, -S.Z) ) );
-				Points[AddedPoints + 5] = T.TransformPosition( P + Q.RotateVector(FVector3f(-S.X, S.Y, -S.Z) ) );
-				Points[AddedPoints + 6] = T.TransformPosition( P + Q.RotateVector(FVector3f(-S.X, -S.Y, -S.Z) ) );
-				Points[AddedPoints + 7] = T.TransformPosition( P + Q.RotateVector(FVector3f(S.X, -S.Y, -S.Z) ) );
+				Points[AddedPoints + 8]  = T.TransformPosition(P + BasisX*S.X);
+				Points[AddedPoints + 9]  = T.TransformPosition(P + BasisY*S.Y);
+				Points[AddedPoints + 10] = T.TransformPosition(P + BasisZ*S.Z);
+
+				Points[AddedPoints + 11] = T.TransformPosition(P - BasisX*S.X);
+				Points[AddedPoints + 12] = T.TransformPosition(P - BasisY*S.Y);
+				Points[AddedPoints + 13] = T.TransformPosition(P - BasisZ*S.Z);
 			}
 
 			const int32 SphylCount = PBody->GetSphylCount(B);
-			for ( int32 I = 0; I < SphylCount; ++I, AddedPoints += 8)
+			for ( int32 I = 0; I < SphylCount; ++I, AddedPoints += 14)
 			{
 				FVector3f P;
 				FQuat4f Q;
@@ -572,22 +607,38 @@ namespace mu
 			
 				PBody->GetSphyl(B, I, P, Q, R, L);
 				
-				const FVector3f S = FVector3f(R, R, (L + R*2.0f) * 0.5f);
-
-				Points[AddedPoints + 0] = T.TransformPosition( P + Q.RotateVector(FVector3f(S.X, S.Y, S.Z) ) );
-				Points[AddedPoints + 1] = T.TransformPosition( P + Q.RotateVector(FVector3f(-S.X, S.Y, S.Z) ) );
-				Points[AddedPoints + 2] = T.TransformPosition( P + Q.RotateVector(FVector3f(-S.X, -S.Y, S.Z) ) );
-				Points[AddedPoints + 3] = T.TransformPosition( P + Q.RotateVector(FVector3f(S.X, -S.Y, S.Z) ) );
-
-														  
-				Points[AddedPoints + 4] = T.TransformPosition( P + Q.RotateVector(FVector3f(S.X, S.Y, -S.Z) ) );
-				Points[AddedPoints + 5] = T.TransformPosition( P + Q.RotateVector(FVector3f(-S.X, S.Y, -S.Z) ) );
-				Points[AddedPoints + 6] = T.TransformPosition( P + Q.RotateVector(FVector3f(-S.X, -S.Y, -S.Z) ) );
-				Points[AddedPoints + 7] = T.TransformPosition( P + Q.RotateVector(FVector3f(S.X, -S.Y, -S.Z) ) );
+				const float H = L * 0.5f;
+
+				const FVector3f BasisX = Q.RotateVector(FVector3f::UnitX());
+				const FVector3f BasisY = Q.RotateVector(FVector3f::UnitY());
+				const FVector3f BasisZ = Q.RotateVector(FVector3f::UnitZ());
+
+				// Top and Bottom
+				Points[AddedPoints + 0] = T.TransformPosition(P + BasisZ*(H+R));
+				Points[AddedPoints + 1] = T.TransformPosition(P - BasisZ*(H+R));
+
+				// Top ring
+				Points[AddedPoints + 2] = T.TransformPosition(P + BasisX*R + BasisZ*H);
+				Points[AddedPoints + 3] = T.TransformPosition(P - BasisX*R + BasisZ*H);
+				Points[AddedPoints + 4] = T.TransformPosition(P + BasisY*R + BasisZ*H);
+				Points[AddedPoints + 5] = T.TransformPosition(P - BasisY*R + BasisZ*H);
+
+				// Center ring
+				Points[AddedPoints + 6] = T.TransformPosition(P + BasisX*R);
+				Points[AddedPoints + 7] = T.TransformPosition(P - BasisX*R);
+				Points[AddedPoints + 8] = T.TransformPosition(P + BasisY*R);
+				Points[AddedPoints + 9] = T.TransformPosition(P - BasisY*R);
+	
+				// Bottom ring
+				Points[AddedPoints + 10] = T.TransformPosition(P + BasisX*R - BasisZ*H);
+				Points[AddedPoints + 11] = T.TransformPosition(P - BasisX*R - BasisZ*H);
+				Points[AddedPoints + 12] = T.TransformPosition(P + BasisY*R - BasisZ*H);
+				Points[AddedPoints + 13] = T.TransformPosition(P - BasisY*R - BasisZ*H);
+
 			}
 	
 			const int32 TaperedCapsuleCount = PBody->GetTaperedCapsuleCount(B);
-			for ( int32 I = 0; I < TaperedCapsuleCount; ++I, AddedPoints += 8)
+			for ( int32 I = 0; I < TaperedCapsuleCount; ++I, AddedPoints += 14)
 			{
 				FVector3f P;
 				FQuat4f Q;
@@ -596,48 +647,65 @@ namespace mu
 				float L;
 			
 				PBody->GetTaperedCapsule(B, I, P, Q, R0, R1, L);
-				
-				const FVector3f ST = FVector3f(R0, R0, L*0.5f);
-				const FVector3f SB = FVector3f(R1, R1, L*0.5f);
-				
-				Points[AddedPoints + 0] = T.TransformPosition( P + Q.RotateVector(FVector3f(ST.X, ST.Y, ST.Z + R0) ) );
-				Points[AddedPoints + 1] = T.TransformPosition( P + Q.RotateVector(FVector3f(-ST.X, ST.Y, ST.Z + R0) ) );
-				Points[AddedPoints + 2] = T.TransformPosition( P + Q.RotateVector(FVector3f(-ST.X, -ST.Y, ST.Z + R0) ) );
-				Points[AddedPoints + 3] = T.TransformPosition( P + Q.RotateVector(FVector3f(ST.X, -ST.Y, ST.Z + R0) ) );
-
-				Points[AddedPoints + 4] = T.TransformPosition( P + Q.RotateVector(FVector3f(SB.X, SB.Y, -SB.Z - R1) ) );
-				Points[AddedPoints + 5] = T.TransformPosition( P + Q.RotateVector(FVector3f(-SB.X, SB.Y, -SB.Z - R1) ) );
-				Points[AddedPoints + 6] = T.TransformPosition( P + Q.RotateVector(FVector3f(-SB.X, -SB.Y, -SB.Z - R1) ) );
-				Points[AddedPoints + 7] = T.TransformPosition( P + Q.RotateVector(FVector3f(SB.X, -SB.Y, -SB.Z - R1) ) );
+		
+				const float H = L*0.5f;
+				const float RCenter = (R0 + R1) * 0.5f;
+
+				const FVector3f BasisX = Q.RotateVector(FVector3f::UnitX());
+				const FVector3f BasisY = Q.RotateVector(FVector3f::UnitY());
+				const FVector3f BasisZ = Q.RotateVector(FVector3f::UnitZ());
+
+				// Top and Bottom
+				Points[AddedPoints + 0] = T.TransformPosition(P + BasisZ*(H+R0));
+				Points[AddedPoints + 1] = T.TransformPosition(P - BasisZ*(H+R1));
+
+				// Top ring
+				Points[AddedPoints + 2] = T.TransformPosition(P + BasisX*R0 + BasisZ*H);
+				Points[AddedPoints + 3] = T.TransformPosition(P - BasisX*R0 + BasisZ*H);
+				Points[AddedPoints + 4] = T.TransformPosition(P + BasisY*R0 + BasisZ*H);
+				Points[AddedPoints + 5] = T.TransformPosition(P - BasisY*R0 + BasisZ*H);
+
+				// Center ring
+				Points[AddedPoints + 6] = T.TransformPosition(P + BasisX*RCenter);
+				Points[AddedPoints + 7] = T.TransformPosition(P - BasisX*RCenter);
+				Points[AddedPoints + 8] = T.TransformPosition(P + BasisY*RCenter);
+				Points[AddedPoints + 9] = T.TransformPosition(P - BasisY*RCenter);
+	
+				// Bottom ring
+				Points[AddedPoints + 10] = T.TransformPosition(P + BasisX*R1 - BasisZ*H);
+				Points[AddedPoints + 11] = T.TransformPosition(P - BasisX*R1 - BasisZ*H);
+				Points[AddedPoints + 12] = T.TransformPosition(P + BasisY*R1 - BasisZ*H);
+				Points[AddedPoints + 13] = T.TransformPosition(P - BasisY*R1 - BasisZ*H);
 			}
 
 			const int32 ConvexCount = PBody->GetConvexCount(B);
 			for ( int32 I = 0; I < ConvexCount; ++I )
 			{
-				const FVector3f* VertexData;
-				int32 VertexCount;
-				const int32* IndicesData;
-				int32 IndicesCount;
+			
+				TArrayView<const FVector3f> VerticesView;
+				TArrayView<const int32> IndicesView;
 				FTransform3f ConvexT;
-				
-				PBody->GetConvex(B, I, VertexData, VertexCount, IndicesData, IndicesCount, ConvexT);
-				
-				TArrayView<const FVector3f> VertexDataView( VertexData, VertexCount );
+
+				PBody->GetConvex(B, I, VerticesView, IndicesView, ConvexT);
 
 				ConvexT = T * ConvexT;
-				for ( const FVector3f& P : VertexDataView )
+				for ( const FVector3f& P : VerticesView )
 				{
 					Points[AddedPoints++] = ConvexT.TransformPosition(P);
 				}	
 			}
 		}
 		
-		TArray<FReshapeVertexBindingData>& PhysicsBodyBindData = ReturnValue.Get<0>();
+		TArray<FReshapePointBindingData>& PhysicsBodyBindData = ReturnValue.Get<0>();
 		PhysicsBodyBindData.SetNumUninitialized( NumPoints );
 
+		FReshapeVertexBindingData VertexBindData;
+
 		for (int32 I = 0; I < NumPoints; ++I)
 		{
-			BindReshapePoint(ShapeMeshTree, Points[I], PhysicsBodyBindData[I], 0.01f);
+			BindReshapePoint(ShapeMeshTree, Points[I], VertexBindData, 0.1f);
+			PhysicsBodyBindData[I] = FReshapePointBindingData
+				{ VertexBindData.S, VertexBindData.T, VertexBindData.D, VertexBindData.Triangle, VertexBindData.Weight };
 		}
 
 		return ReturnValue;
@@ -722,7 +790,7 @@ namespace mu
 		return BindData;
 	}
 
-	inline TTuple<TArray<FReshapeVertexBindingData>, TArray<int32>> BindPose(
+	inline TTuple<TArray<FReshapePointBindingData>, TArray<int32>> BindPose(
 			const Mesh* Mesh, FShapeMeshTree& ShapeMeshTree, const TArray<string>& BonesToDeform )
 	{
 		UE::Geometry::FAxisAlignedBox3d ShapeAABBox = ShapeMeshTree.GetBoundingBox();
@@ -731,26 +799,19 @@ namespace mu
 										  ? UE_KINDA_SMALL_NUMBER
 										  : static_cast<float>(ShapeAABBox.MaxDim()) * 1e-3f;
 		
-		TTuple<TArray<FReshapeVertexBindingData>, TArray<int32>> ReturnValue;
+		TTuple<TArray<FReshapePointBindingData>, TArray<int32>> ReturnValue;
 
-		TArray<FReshapeVertexBindingData>& SkeletonBindDataArray = ReturnValue.Get<0>();
+		TArray<FReshapePointBindingData>& SkeletonBindDataArray = ReturnValue.Get<0>();
 		TArray<int32>& BoneIndices = ReturnValue.Get<1>();
 
 		const int32 BoneCount = Mesh->GetBonePoseCount();
 		SkeletonBindDataArray.Reserve(BoneCount);
 		BoneIndices.Reserve(BoneCount);
 
-		for ( int32 BoneIndex = 0; BoneIndex < BoneCount; ++BoneIndex)
+		for (int32 BoneIndex = 0; BoneIndex < BoneCount; ++BoneIndex)
 		{
 			const EBoneUsageFlags BoneUsageFlags = Mesh->BonePoses[BoneIndex].BoneUsageFlags;
 
-			// Only deform skinned bones.
-			//if (EnumHasAnyFlags(BoneUsageFlags, EBoneUsageFlags::SkinningParent | EBoneUsageFlags::Root) &&
-			//	!EnumHasAnyFlags(BoneUsageFlags, EBoneUsageFlags::Skinning))
-			//{
-			//	continue;
-			//}
-
 			if (EnumHasAnyFlags(BoneUsageFlags, EBoneUsageFlags::Root))
 			{
 				continue;
@@ -761,13 +822,17 @@ namespace mu
 				continue;
 			}
 
-			BoneIndices.Add(BoneIndex);
-			
-			BindReshapePoint(
-					ShapeMeshTree, 
-					Mesh->BonePoses[BoneIndex].BoneTransform.GetLocation(),
-					SkeletonBindDataArray.AddDefaulted_GetRef(),
-					BindValidityTolerance);
+			FReshapeVertexBindingData BindData;
+
+			BindReshapePoint(ShapeMeshTree, Mesh->BonePoses[BoneIndex].BoneTransform.GetLocation(), BindData, BindValidityTolerance);
+
+			// Only add binding  if there is a chance of the bone moving.
+			if (BindData.Weight > UE_SMALL_NUMBER && BindData.Triangle >= 0)
+			{
+				SkeletonBindDataArray.Emplace(
+					FReshapePointBindingData{ BindData.S, BindData.T, BindData.D, BindData.Triangle, BindData.Weight });
+				BoneIndices.Add(BoneIndex);
+			}	
 		}
 
 		return ReturnValue;
@@ -909,9 +974,9 @@ namespace mu
 		{
 			MUTABLE_CPUPROFILER_SCOPE(BindSkeleton);
 
-			TTuple<TArray<FReshapeVertexBindingData>, TArray<int32>> SkeletonBindingData = 
+			TTuple<TArray<FReshapePointBindingData>, TArray<int32>> SkeletonBindingData = 
 					BindPose(Result.get(), ShapeMeshTree, BonesToDeform);
-			const TArray<FReshapeVertexBindingData>& SkeletonBindDataArray = SkeletonBindingData.Get<0>();
+			const TArray<FReshapePointBindingData>& SkeletonBindDataArray = SkeletonBindingData.Get<0>();
 			const TArray<int32>& BoneIndices = SkeletonBindingData.Get<1>();
 
 			check(BoneIndices.Num() == SkeletonBindDataArray.Num());
@@ -922,8 +987,8 @@ namespace mu
 			SkeletonBuffer.SetBufferCount(2);
 			SkeletonBuffer.SetElementCount(SkeletonBindDataArray.Num());
 
-			FReshapeVertexBindingDataBufferDescriptor BufDesc(BindingDataIndex);
-			SkeletonBuffer.SetBuffer(0, sizeof(FReshapeVertexBindingData), BufDesc.Channels, BufDesc.Semantics, BufDesc.SemanticIndices, BufDesc.Formats, BufDesc.Components);
+			FReshapePointBindingDataBufferDescriptor BufDesc(BindingDataIndex);
+			SkeletonBuffer.SetBuffer(0, sizeof(FReshapePointBindingData), BufDesc.Channels, BufDesc.Semantics, BufDesc.SemanticIndices, BufDesc.Formats, BufDesc.Components);
 
 			// Bone indices buffer
 			MESH_BUFFER_SEMANTIC BoneSemantics[1] = { MBS_OTHER };
@@ -934,7 +999,7 @@ namespace mu
 
 			SkeletonBuffer.SetBuffer(1, sizeof(int32), 1, BoneSemantics, BoneSemanticIndices, BoneFormats, BoneComponents, BoneOffsets);
 
-			FMemory::Memcpy(SkeletonBuffer.GetBufferData(0), SkeletonBindDataArray.GetData(), NumBonesToDeform * sizeof(FReshapeVertexBindingData));
+			FMemory::Memcpy(SkeletonBuffer.GetBufferData(0), SkeletonBindDataArray.GetData(), NumBonesToDeform * sizeof(FReshapePointBindingData));
 			FMemory::Memcpy(SkeletonBuffer.GetBufferData(1), BoneIndices.GetData(), NumBonesToDeform * sizeof(int32));
 
 			Result->m_AdditionalBuffers.Emplace(EMeshBufferType::SkeletonDeformBinding, MoveTemp(SkeletonBuffer));
@@ -945,9 +1010,9 @@ namespace mu
 		{
 			MUTABLE_CPUPROFILER_SCOPE(BindPhysicsBody);
 
-			TTuple<TArray<FReshapeVertexBindingData>, TArray<int32>> PhysicsBindingData = 
+			TTuple<TArray<FReshapePointBindingData>, TArray<int32>> PhysicsBindingData = 
 					BindPhysicsBodies(ResultPhysicsBody, ShapeMeshTree, Result.get(), PhysicsToDeform);
-			const TArray<FReshapeVertexBindingData>& PhysicsBindDataArray = PhysicsBindingData.Get<0>();
+			const TArray<FReshapePointBindingData>& PhysicsBindDataArray = PhysicsBindingData.Get<0>();
 			const TArray<int32>& DeformedBodyIndices = PhysicsBindingData.Get<1>();
 
 			FMeshBufferSet PhysicsBodyBuffer;
@@ -955,9 +1020,9 @@ namespace mu
 			PhysicsBodyBuffer.SetElementCount(PhysicsBindDataArray.Num());
 
 			constexpr int32 BindDataIndex = 0; 
-			FReshapeVertexBindingDataBufferDescriptor BufDesc(BindingDataIndex);
-			PhysicsBodyBuffer.SetBuffer(0, sizeof(FReshapeVertexBindingData), BufDesc.Channels, BufDesc.Semantics, BufDesc.SemanticIndices, BufDesc.Formats, BufDesc.Components, BufDesc.Offsets);	
-			FMemory::Memcpy(PhysicsBodyBuffer.GetBufferData(0), PhysicsBindDataArray.GetData(), PhysicsBindDataArray.Num() * sizeof(FReshapeVertexBindingData));
+			FReshapePointBindingDataBufferDescriptor BufDesc(BindingDataIndex);
+			PhysicsBodyBuffer.SetBuffer(0, sizeof(FReshapePointBindingData), BufDesc.Channels, BufDesc.Semantics, BufDesc.SemanticIndices, BufDesc.Formats, BufDesc.Components, BufDesc.Offsets);	
+			FMemory::Memcpy(PhysicsBodyBuffer.GetBufferData(0), PhysicsBindDataArray.GetData(), PhysicsBindDataArray.Num() * sizeof(FReshapePointBindingData));
 
 			// Bone indices buffer
 			MESH_BUFFER_SEMANTIC BoneSemantics[1] = { MBS_OTHER };
diff --git a/Engine/Plugins/Experimental/Mutable/Source/MutableRuntime/Private/MuR/OpMeshMerge.h b/Engine/Plugins/Experimental/Mutable/Source/MutableRuntime/Private/MuR/OpMeshMerge.h
--- a/Engine/Plugins/Experimental/Mutable/Source/MutableRuntime/Private/MuR/OpMeshMerge.h
+++ b/Engine/Plugins/Experimental/Mutable/Source/MutableRuntime/Private/MuR/OpMeshMerge.h
@@ -16,11 +16,11 @@ namespace mu
 	//---------------------------------------------------------------------------------------------
 	//! Merge two meshes into one new mesh
 	//---------------------------------------------------------------------------------------------
-	inline MeshPtr MeshMerge( const Mesh* pFirst, const Mesh* pSecond )
+	inline MeshPtr MeshMerge(const Mesh* pFirst, const Mesh* pSecond)
 	{
-        MUTABLE_CPUPROFILER_SCOPE(MeshMerge);
+		MUTABLE_CPUPROFILER_SCOPE(MeshMerge);
 
-        MeshPtr pResult = new Mesh;
+		MeshPtr pResult = new Mesh;
 
 		// Should never happen, but fixes static analysis warnings.
 		if (!(pFirst && pSecond))
@@ -30,328 +30,328 @@ namespace mu
 
 		// Indices
 		//-----------------
-		if ( pFirst->GetIndexBuffers().GetBufferCount()>0 )
+		if (pFirst->GetIndexBuffers().GetBufferCount() > 0)
 		{
-            MUTABLE_CPUPROFILER_SCOPE(Indices);
+			MUTABLE_CPUPROFILER_SCOPE(Indices);
 
-            int firstCount = pFirst->GetIndexBuffers().GetElementCount();
+			int firstCount = pFirst->GetIndexBuffers().GetElementCount();
 			int secondCount = pSecond->GetIndexBuffers().GetElementCount();
-			pResult->GetIndexBuffers().SetElementCount( firstCount+secondCount );
+			pResult->GetIndexBuffers().SetElementCount(firstCount + secondCount);
 
-			check( pFirst->GetIndexBuffers().GetBufferCount() <= 1 );
-			check( pSecond->GetIndexBuffers().GetBufferCount() <= 1 );
-			pResult->GetIndexBuffers().SetBufferCount( 1 );
+			check(pFirst->GetIndexBuffers().GetBufferCount() <= 1);
+			check(pSecond->GetIndexBuffers().GetBufferCount() <= 1);
+			pResult->GetIndexBuffers().SetBufferCount(1);
 
-            // This will be changed below if need to change the format of the index buffers.
-            MESH_BUFFER_FORMAT requiredBufferFormatChange = MBF_NONE;
+			// This will be changed below if need to change the format of the index buffers.
+			MESH_BUFFER_FORMAT requiredBufferFormatChange = MBF_NONE;
 
-            if (firstCount && secondCount)
+			if (firstCount && secondCount)
 			{
 				const MESH_BUFFER& first = pFirst->GetIndexBuffers().m_buffers[0];
 				const MESH_BUFFER& second = pSecond->GetIndexBuffers().m_buffers[0];
 				check(first.m_channels == second.m_channels);
 				check(first.m_elementSize == second.m_elementSize);
-                check(first.m_channels.IsEmpty() || first.m_channels[0].m_format == second.m_channels[0].m_format);
-				
+				check(first.m_channels.IsEmpty() || first.m_channels[0].m_format == second.m_channels[0].m_format);
+
 				// Avoid unused variable warnings
 				(void)first;
 				(void)second;
 
-                // We need to know the total number of vertices in case we need to adjust the index buffer format.
-                uint64_t totalVertexCount = pFirst->GetVertexBuffers().GetElementCount()
-                        + pSecond->GetVertexBuffers().GetElementCount();
-                uint64_t maxValueBits = GetMeshFormatData(pFirst->GetIndexBuffers().m_buffers[0].m_channels[0].m_format).m_maxValueBits;
-                uint64_t maxSupportedVertices = uint64_t(1)<<maxValueBits;
-                if (totalVertexCount>maxSupportedVertices)
-                {
-                    if (totalVertexCount > 0xFFFF)
-                    {
-                        requiredBufferFormatChange = MBF_UINT32;
-                    }
-                    else
-                    {
-                        requiredBufferFormatChange = MBF_UINT16;
-                    }
-                }
-            }
+				// We need to know the total number of vertices in case we need to adjust the index buffer format.
+				uint64_t totalVertexCount = pFirst->GetVertexBuffers().GetElementCount()
+					+ pSecond->GetVertexBuffers().GetElementCount();
+				uint64_t maxValueBits = GetMeshFormatData(pFirst->GetIndexBuffers().m_buffers[0].m_channels[0].m_format).m_maxValueBits;
+				uint64_t maxSupportedVertices = uint64_t(1) << maxValueBits;
+				if (totalVertexCount > maxSupportedVertices)
+				{
+					if (totalVertexCount > 0xFFFF)
+					{
+						requiredBufferFormatChange = MBF_UINT32;
+					}
+					else
+					{
+						requiredBufferFormatChange = MBF_UINT16;
+					}
+				}
+			}
 
 			MESH_BUFFER& result = pResult->GetIndexBuffers().m_buffers[0];
 
-            if (requiredBufferFormatChange != MBF_NONE)
-            {
-                // We only support vertex indices in case of having to change the format.
-                check(pFirst->GetIndexBuffers().m_buffers[0].m_channels.Num()==1);
-
-                result.m_channels.SetNum(1);
-                result.m_channels[0].m_semantic = MBS_VERTEXINDEX;
-                result.m_channels[0].m_format = requiredBufferFormatChange;
-                result.m_channels[0].m_componentCount = 1;
-                result.m_channels[0].m_semanticIndex = 0;
-                result.m_channels[0].m_offset = 0;
-                result.m_elementSize = GetMeshFormatData(requiredBufferFormatChange).m_size;
-            }
-            else if (firstCount)
+			if (requiredBufferFormatChange != MBF_NONE)
+			{
+				// We only support vertex indices in case of having to change the format.
+				check(pFirst->GetIndexBuffers().m_buffers[0].m_channels.Num() == 1);
+
+				result.m_channels.SetNum(1);
+				result.m_channels[0].m_semantic = MBS_VERTEXINDEX;
+				result.m_channels[0].m_format = requiredBufferFormatChange;
+				result.m_channels[0].m_componentCount = 1;
+				result.m_channels[0].m_semanticIndex = 0;
+				result.m_channels[0].m_offset = 0;
+				result.m_elementSize = GetMeshFormatData(requiredBufferFormatChange).m_size;
+			}
+			else if (firstCount)
 			{
 				const MESH_BUFFER& first = pFirst->GetIndexBuffers().m_buffers[0];
 				result.m_channels = first.m_channels;
 				result.m_elementSize = first.m_elementSize;
 			}
-            else if (secondCount)
+			else if (secondCount)
 			{
 				const MESH_BUFFER& second = pSecond->GetIndexBuffers().m_buffers[0];
 				result.m_channels = second.m_channels;
 				result.m_elementSize = second.m_elementSize;
 			}
 
-			result.m_data.SetNum(result.m_elementSize*(firstCount + secondCount));
+			result.m_data.SetNum(result.m_elementSize * (firstCount + secondCount));
 
-			check( result.m_channels.Num()==1 );
-			check( result.m_channels[0].m_semantic == MBS_VERTEXINDEX );
+			check(result.m_channels.Num() == 1);
+			check(result.m_channels[0].m_semantic == MBS_VERTEXINDEX);
 
-			if ( result.m_data.Num() )
+			if (result.m_data.Num())
 			{
-				if ( firstCount )
+				if (firstCount)
 				{
 					const MESH_BUFFER& first = pFirst->GetIndexBuffers().m_buffers[0];
 
-                    if (requiredBufferFormatChange==MBF_NONE
-                            || requiredBufferFormatChange==first.m_channels[0].m_format)
-                    {
-                        FMemory::Memcpy( &result.m_data[0],
-                                        &first.m_data[0],
-                                        first.m_elementSize*firstCount );
-                    }
-                    else
-                    {
-                        // Conversion required
-                        const uint8_t* pSource = &first.m_data[0];
-                        uint8_t* pDest = &result.m_data[0];
-                        switch( requiredBufferFormatChange )
-                        {
-                        case MBF_UINT32:
-                        {
-                            switch( first.m_channels[0].m_format )
-                            {
-                            case MBF_UINT16:
-                            {
-                                for ( int v=0; v<firstCount; ++v )
-                                {
-                                    *(uint32_t*)pDest = *(const uint16*)pSource;
-                                    pSource+=first.m_elementSize;
-                                    pDest+=result.m_elementSize;
-                                }
-                                break;
-                            }
+					if (requiredBufferFormatChange == MBF_NONE
+						|| requiredBufferFormatChange == first.m_channels[0].m_format)
+					{
+						FMemory::Memcpy(&result.m_data[0],
+							&first.m_data[0],
+							first.m_elementSize * firstCount);
+					}
+					else
+					{
+						// Conversion required
+						const uint8_t* pSource = &first.m_data[0];
+						uint8_t* pDest = &result.m_data[0];
+						switch (requiredBufferFormatChange)
+						{
+						case MBF_UINT32:
+						{
+							switch (first.m_channels[0].m_format)
+							{
+							case MBF_UINT16:
+							{
+								for (int v = 0; v < firstCount; ++v)
+								{
+									*(uint32_t*)pDest = *(const uint16*)pSource;
+									pSource += first.m_elementSize;
+									pDest += result.m_elementSize;
+								}
+								break;
+							}
 
-                            case MBF_UINT8:
-                            {
-                                for ( int v=0; v<firstCount; ++v )
-                                {
-                                    *(uint32_t*)pDest = *(const uint8_t*)pSource;
-                                    pSource+=first.m_elementSize;
-                                    pDest+=result.m_elementSize;
-                                }
-                                break;
-                            }
+							case MBF_UINT8:
+							{
+								for (int v = 0; v < firstCount; ++v)
+								{
+									*(uint32_t*)pDest = *(const uint8_t*)pSource;
+									pSource += first.m_elementSize;
+									pDest += result.m_elementSize;
+								}
+								break;
+							}
 
-                            default:
-                                checkf( false, TEXT("Format not supported.") );
-                                break;
-                            }
-                            break;
-                        }
+							default:
+								checkf(false, TEXT("Format not supported."));
+								break;
+							}
+							break;
+						}
 
-                        case MBF_UINT16:
-                        {
-                            switch( first.m_channels[0].m_format )
-                            {
+						case MBF_UINT16:
+						{
+							switch (first.m_channels[0].m_format)
+							{
 
-                            case MBF_UINT8:
-                            {
-                                for ( int v=0; v<firstCount; ++v )
-                                {
-                                    *(uint16*)pDest = *(const uint8_t*)pSource;
-                                    pSource+=first.m_elementSize;
-                                    pDest+=result.m_elementSize;
-                                }
-                                break;
-                            }
+							case MBF_UINT8:
+							{
+								for (int v = 0; v < firstCount; ++v)
+								{
+									*(uint16*)pDest = *(const uint8_t*)pSource;
+									pSource += first.m_elementSize;
+									pDest += result.m_elementSize;
+								}
+								break;
+							}
 
-                            default:
+							default:
 								checkf(false, TEXT("Format not supported."));
 								break;
-                            }
-                            break;
-                        }
+							}
+							break;
+						}
 
-                        default:
+						default:
 							checkf(false, TEXT("Format not supported."));
 							break;
-                        }
-                    }
-                }
+						}
+					}
+				}
 
-				if ( secondCount )
+				if (secondCount)
 				{
 					const MESH_BUFFER& second = pSecond->GetIndexBuffers().m_buffers[0];
-                    const uint8_t* pSource = &second.m_data[0];
-                    uint8_t* pDest = &result.m_data[result.m_elementSize*firstCount ];
+					const uint8_t* pSource = &second.m_data[0];
+					uint8_t* pDest = &result.m_data[result.m_elementSize * firstCount];
 
-                    uint32_t firstVertexCount = pFirst->GetVertexBuffers().GetElementCount();
+					uint32_t firstVertexCount = pFirst->GetVertexBuffers().GetElementCount();
 
-                    if (requiredBufferFormatChange==MBF_NONE
-                            || requiredBufferFormatChange==second.m_channels[0].m_format)
-                    {
-                        switch( second.m_channels[0].m_format )
-                        {
-                        case MBF_INT32:
-                        case MBF_UINT32:
-                        case MBF_NINT32:
-                        case MBF_NUINT32:
-                        {
-                            for ( int v=0; v<secondCount; ++v )
-                            {
-                                *(uint32_t*)pDest = firstVertexCount + *(const uint32_t*)pSource;
-                                pSource+=second.m_elementSize;
-                                pDest+=result.m_elementSize;
-                            }
-                            break;
-                        }
+					if (requiredBufferFormatChange == MBF_NONE
+						|| requiredBufferFormatChange == second.m_channels[0].m_format)
+					{
+						switch (second.m_channels[0].m_format)
+						{
+						case MBF_INT32:
+						case MBF_UINT32:
+						case MBF_NINT32:
+						case MBF_NUINT32:
+						{
+							for (int v = 0; v < secondCount; ++v)
+							{
+								*(uint32_t*)pDest = firstVertexCount + *(const uint32_t*)pSource;
+								pSource += second.m_elementSize;
+								pDest += result.m_elementSize;
+							}
+							break;
+						}
 
-                        case MBF_INT16:
-                        case MBF_UINT16:
-                        case MBF_NINT16:
-                        case MBF_NUINT16:
-                        {
-                            for ( int v=0; v<secondCount; ++v )
-                            {
-                                *(uint16*)pDest = uint16(firstVertexCount) + *(const uint16*)pSource;
-                                pSource+=second.m_elementSize;
-                                pDest+=result.m_elementSize;
-                            }
-                            break;
-                        }
+						case MBF_INT16:
+						case MBF_UINT16:
+						case MBF_NINT16:
+						case MBF_NUINT16:
+						{
+							for (int v = 0; v < secondCount; ++v)
+							{
+								*(uint16*)pDest = uint16(firstVertexCount) + *(const uint16*)pSource;
+								pSource += second.m_elementSize;
+								pDest += result.m_elementSize;
+							}
+							break;
+						}
 
-                        case MBF_INT8:
-                        case MBF_UINT8:
-                        case MBF_NINT8:
-                        case MBF_NUINT8:
-                        {
-                            for ( int v=0; v<secondCount; ++v )
-                            {
-                                *(uint8_t*)pDest = uint8_t(firstVertexCount) + *(const uint8_t*)pSource;
-                                pSource+=second.m_elementSize;
-                                pDest+=result.m_elementSize;
-                            }
-                            break;
-                        }
+						case MBF_INT8:
+						case MBF_UINT8:
+						case MBF_NINT8:
+						case MBF_NUINT8:
+						{
+							for (int v = 0; v < secondCount; ++v)
+							{
+								*(uint8_t*)pDest = uint8_t(firstVertexCount) + *(const uint8_t*)pSource;
+								pSource += second.m_elementSize;
+								pDest += result.m_elementSize;
+							}
+							break;
+						}
 
-                        default:
+						default:
 							checkf(false, TEXT("Format not supported."));
 							break;
-                        }
-                    }
-                    else
-                    {
-                        // Format conversion required
-                        switch (requiredBufferFormatChange)
-                        {
+						}
+					}
+					else
+					{
+						// Format conversion required
+						switch (requiredBufferFormatChange)
+						{
 
-                        case MBF_UINT32:
-                        {
-                            switch( second.m_channels[0].m_format )
-                            {
-                            case MBF_INT16:
-                            case MBF_UINT16:
-                            case MBF_NINT16:
-                            case MBF_NUINT16:
-                            {
-                                for ( int v=0; v<secondCount; ++v )
-                                {
-                                    *(uint32_t*)pDest = uint32_t(firstVertexCount) + *(const uint16*)pSource;
-                                    pSource+=second.m_elementSize;
-                                    pDest+=result.m_elementSize;
-                                }
-                                break;
-                            }
+						case MBF_UINT32:
+						{
+							switch (second.m_channels[0].m_format)
+							{
+							case MBF_INT16:
+							case MBF_UINT16:
+							case MBF_NINT16:
+							case MBF_NUINT16:
+							{
+								for (int v = 0; v < secondCount; ++v)
+								{
+									*(uint32_t*)pDest = uint32_t(firstVertexCount) + *(const uint16*)pSource;
+									pSource += second.m_elementSize;
+									pDest += result.m_elementSize;
+								}
+								break;
+							}
 
-                            case MBF_INT8:
-                            case MBF_UINT8:
-                            case MBF_NINT8:
-                            case MBF_NUINT8:
-                            {
-                                for ( int v=0; v<secondCount; ++v )
-                                {
-                                    *(uint32_t*)pDest = uint32_t(firstVertexCount) + *(const uint8_t*)pSource;
-                                    pSource+=second.m_elementSize;
-                                    pDest+=result.m_elementSize;
-                                }
-                                break;
-                            }
+							case MBF_INT8:
+							case MBF_UINT8:
+							case MBF_NINT8:
+							case MBF_NUINT8:
+							{
+								for (int v = 0; v < secondCount; ++v)
+								{
+									*(uint32_t*)pDest = uint32_t(firstVertexCount) + *(const uint8_t*)pSource;
+									pSource += second.m_elementSize;
+									pDest += result.m_elementSize;
+								}
+								break;
+							}
 
-                            default:
+							default:
 								checkf(false, TEXT("Format not supported."));
 								break;
-                            }
+							}
 
-                            break;
-                        }
+							break;
+						}
 
-                        case MBF_UINT16:
-                        {
-                            switch( second.m_channels[0].m_format )
-                            {
-                            case MBF_INT8:
-                            case MBF_UINT8:
-                            case MBF_NINT8:
-                            case MBF_NUINT8:
-                            {
-                                for ( int v=0; v<secondCount; ++v )
-                                {
-                                    *(uint16*)pDest = uint16(firstVertexCount) + *(const uint8_t*)pSource;
-                                    pSource+=second.m_elementSize;
-                                    pDest+=result.m_elementSize;
-                                }
-                                break;
-                            }
+						case MBF_UINT16:
+						{
+							switch (second.m_channels[0].m_format)
+							{
+							case MBF_INT8:
+							case MBF_UINT8:
+							case MBF_NINT8:
+							case MBF_NUINT8:
+							{
+								for (int v = 0; v < secondCount; ++v)
+								{
+									*(uint16*)pDest = uint16(firstVertexCount) + *(const uint8_t*)pSource;
+									pSource += second.m_elementSize;
+									pDest += result.m_elementSize;
+								}
+								break;
+							}
 
-                            default:
+							default:
 								checkf(false, TEXT("Format not supported."));
 								break;
-                            }
+							}
 
-                            break;
-                        }
+							break;
+						}
 
-                        default:
+						default:
 							checkf(false, TEXT("Format not supported."));
 							break;
 
-                        }
-                    }
-                }
+						}
+					}
+				}
 			}
 		}
 
 		// Faces
 		//-----------------
 		{
-            MUTABLE_CPUPROFILER_SCOPE(Faces);
+			MUTABLE_CPUPROFILER_SCOPE(Faces);
 
-            int firstCount = pFirst->GetFaceBuffers().GetElementCount();
+			int firstCount = pFirst->GetFaceBuffers().GetElementCount();
 			int secondCount = pSecond->GetFaceBuffers().GetElementCount();
-			pResult->GetFaceBuffers().SetElementCount( firstCount+secondCount );
-			pResult->GetFaceBuffers().SetBufferCount( pFirst->GetFaceBuffers().GetBufferCount() );
+			pResult->GetFaceBuffers().SetElementCount(firstCount + secondCount);
+			pResult->GetFaceBuffers().SetBufferCount(pFirst->GetFaceBuffers().GetBufferCount());
 
 			// Merge only the buffers present in the first mesh
-			for ( int b = 0; b<pResult->GetFaceBuffers().GetBufferCount(); ++b )
+			for (int b = 0; b < pResult->GetFaceBuffers().GetBufferCount(); ++b)
 			{
 				const MESH_BUFFER& first = pFirst->GetFaceBuffers().m_buffers[b];
 
 				MESH_BUFFER& result = pResult->GetFaceBuffers().m_buffers[b];
 				result.m_channels = first.m_channels;
 				result.m_elementSize = first.m_elementSize;
-				result.m_data.SetNum( result.m_elementSize * (firstCount+secondCount) );
+				result.m_data.SetNum(result.m_elementSize * (firstCount + secondCount));
 
 				MESH_BUFFER_SEMANTIC semantic = MBS_NONE;
 				int semanticIndex = 0;
@@ -359,79 +359,79 @@ namespace mu
 				int components = 0;
 				int offset = 0;
 				pFirst->GetFaceBuffers().GetChannel
-						( b, 0, &semantic, &semanticIndex, &type, &components, &offset );
+				(b, 0, &semantic, &semanticIndex, &type, &components, &offset);
 
-				if ( firstCount )
+				if (firstCount)
 				{
-					FMemory::Memcpy( &result.m_data[0],
-									&first.m_data[0],
-									result.m_elementSize * firstCount
-									);
+					FMemory::Memcpy(&result.m_data[0],
+						&first.m_data[0],
+						result.m_elementSize * firstCount
+					);
 				}
 
-				if ( secondCount )
+				if (secondCount)
 				{
 					// Find in the second mesh
 					int otherBuffer = -1;
 					int otherChannel = -1;
 					pSecond->GetFaceBuffers().FindChannel
-							( semantic, 0, &otherBuffer, &otherChannel );
+					(semantic, 0, &otherBuffer, &otherChannel);
 
-					if (otherBuffer>=0)
+					if (otherBuffer >= 0)
 					{
 						const MESH_BUFFER& second =
-								pSecond->GetFaceBuffers().m_buffers[otherBuffer];
-						check( first.m_channels == second.m_channels );
+							pSecond->GetFaceBuffers().m_buffers[otherBuffer];
+						check(first.m_channels == second.m_channels);
 
 						// Raw copy
-						FMemory::Memcpy( &result.m_data[ result.m_elementSize * firstCount ],
-										&second.m_data[0],
-										result.m_elementSize * secondCount );
+						FMemory::Memcpy(&result.m_data[result.m_elementSize * firstCount],
+							&second.m_data[0],
+							result.m_elementSize * secondCount);
 					}
 					else
 					{
 						// Fill with zeroes
-						FMemory::Memzero( &result.m_data[ result.m_elementSize * firstCount ],
-										result.m_elementSize * secondCount );
+						FMemory::Memzero(&result.m_data[result.m_elementSize * firstCount],
+							result.m_elementSize * secondCount);
 					}
 				}
 			}
 		}
 
 
-        // Layouts
-        //-----------------
-        {
-            MUTABLE_CPUPROFILER_SCOPE(Layouts);
+		// Layouts
+		//-----------------
+		{
+			MUTABLE_CPUPROFILER_SCOPE(Layouts);
 
-            pResult->m_layouts.SetNum( pFirst->m_layouts.Num() );
-            for ( int i=0; i<pFirst->m_layouts.Num(); ++i )
-            {
-                const Layout* pF = pFirst->m_layouts[i].get();
-                LayoutPtr pR = pF->Clone();
+			pResult->m_layouts.SetNum(pFirst->m_layouts.Num());
+			for (int i = 0; i < pFirst->m_layouts.Num(); ++i)
+			{
+				const Layout* pF = pFirst->m_layouts[i].get();
+				LayoutPtr pR = pF->Clone();
 
-                if ( i<pSecond->m_layouts.Num() )
-                {
-                    const Layout* pS = pSecond->m_layouts[i].get();
+				if (i < pSecond->m_layouts.Num())
+				{
+					const Layout* pS = pSecond->m_layouts[i].get();
 
-                    pR->m_blocks.Append(pS->m_blocks);
-                }
+					pR->m_blocks.Append(pS->m_blocks);
+				}
 
-                pResult->m_layouts[i] = pR;
-            }
-        }
+				pResult->m_layouts[i] = pR;
+			}
+		}
 
 
-        // Surfaces
-        //-----------------
-        pResult->m_surfaces = pFirst->m_surfaces;
+		// Surfaces
+		//-----------------
+		pResult->m_surfaces = pFirst->m_surfaces;
 
-        for ( const auto& s : pSecond->m_surfaces )
-        {
-            pResult->m_surfaces.Add(s);
-            pResult->m_surfaces.Last().m_firstVertex += pFirst->GetVertexCount();
-            pResult->m_surfaces.Last().m_firstIndex += pFirst->GetIndexCount();
-        }
+		for (const auto& s : pSecond->m_surfaces)
+		{
+			pResult->m_surfaces.Add(s);
+			pResult->m_surfaces.Last().m_firstVertex += pFirst->GetVertexCount();
+			pResult->m_surfaces.Last().m_firstIndex += pFirst->GetIndexCount();
+		}
 
 
 		// Skeleton
@@ -445,76 +445,76 @@ namespace mu
 			pResult->SkeletonIDs.AddUnique(SkeletonID);
 		}
 
-        // Do they have the same skeleton?
-        bool remapNeeded = pFirst->GetSkeleton() != pSecond->GetSkeleton();
+		// Do they have the same skeleton?
+		bool remapNeeded = pFirst->GetSkeleton() != pSecond->GetSkeleton();
 
-        // Are they different skeletons but with the same data?
-        if ( remapNeeded && pFirst->GetSkeleton() && pSecond->GetSkeleton() )
-        {
-            remapNeeded = ! ( *pFirst->GetSkeleton()
-                              ==
-                              *pSecond->GetSkeleton() );
-        }
+		// Are they different skeletons but with the same data?
+		if (remapNeeded && pFirst->GetSkeleton() && pSecond->GetSkeleton())
+		{
+			remapNeeded = !(*pFirst->GetSkeleton()
+				==
+				*pSecond->GetSkeleton());
+		}
 
 
 		TArray<int32> secondToFirstBones;
 
-        if ( remapNeeded )
-        {
-            MUTABLE_CPUPROFILER_SCOPE(Skeleton);
-            Ptr<Skeleton> pResultSkeleton;
+		if (remapNeeded)
+		{
+			MUTABLE_CPUPROFILER_SCOPE(Skeleton);
+			Ptr<Skeleton> pResultSkeleton;
 
 			mu::SkeletonPtrConst pFirstSkeleton = pFirst->GetSkeleton();
 			mu::SkeletonPtrConst pSecondSkeleton = pSecond->GetSkeleton();
 
 			const int32 pFirstBoneCount = pFirstSkeleton ? pFirstSkeleton->GetBoneCount() : 0;
 			const int32 pSecondBoneCount = pSecondSkeleton ? pSecondSkeleton->GetBoneCount() : 0;
-			
-			pResultSkeleton = pFirstSkeleton ? pFirstSkeleton->Clone() : new Skeleton; 
+
+			pResultSkeleton = pFirstSkeleton ? pFirstSkeleton->Clone() : new Skeleton;
 			pResult->SetSkeleton(pResultSkeleton);
 
 			secondToFirstBones.SetNumUninitialized(pSecondBoneCount);
 
 			// Merge pSecond and build the remap array 
-            for ( int32 ob=0; ob< pSecondBoneCount; ++ob )
-            {
-                int32 tb = INDEX_NONE;
-                for ( int32 c=0; tb<0 && c<pResultSkeleton->m_bones.Num(); ++c )
-                {
-                    if ( pResultSkeleton->m_bones[c] == pSecondSkeleton->GetBoneName( ob ) )
-                    {
-                        tb = c;
+			for (int32 ob = 0; ob < pSecondBoneCount; ++ob)
+			{
+				int32 tb = INDEX_NONE;
+				for (int32 c = 0; tb < 0 && c < pResultSkeleton->m_bones.Num(); ++c)
+				{
+					if (pResultSkeleton->m_bones[c] == pSecondSkeleton->GetBoneName(ob))
+					{
+						tb = c;
 						break;
-                    }
-                }
+					}
+				}
+
+				// Add a new bone
+				if (tb == INDEX_NONE)
+				{
+					tb = pResultSkeleton->m_bones.Num();
+					pResultSkeleton->m_bones.Add(pSecondSkeleton->m_bones[ob]);
 
-                // Add a new bone
-                if( tb == INDEX_NONE)
-                {
-                    tb = pResultSkeleton->m_bones.Num();
-                    pResultSkeleton->m_bones.Add( pSecondSkeleton->m_bones[ob] );
-					
 					// Add an incorrect index, to be fixed below in case the parent index is later in the bone array.
 					pResultSkeleton->m_boneParents.Add(pSecondSkeleton->m_boneParents[ob]);
-                }
+				}
 
-                secondToFirstBones[ob] = tb;
+				secondToFirstBones[ob] = tb;
 			}
 
-            // Fix second mesh bone parents
-            for (int32 ob = pFirstBoneCount; ob < pResultSkeleton->m_boneParents.Num(); ++ob)
-            {
-                int16_t secondMeshIndex = pResultSkeleton->m_boneParents[ob];
-                if (secondMeshIndex != INDEX_NONE)
-                {
-                    pResultSkeleton->m_boneParents[ob] = (uint16)secondToFirstBones[secondMeshIndex];
-                }
-            }
-        }
-        else
-        {
-            pResult->SetSkeleton( pFirst->GetSkeleton() );
-        }
+			// Fix second mesh bone parents
+			for (int32 ob = pFirstBoneCount; ob < pResultSkeleton->m_boneParents.Num(); ++ob)
+			{
+				int16_t secondMeshIndex = pResultSkeleton->m_boneParents[ob];
+				if (secondMeshIndex != INDEX_NONE)
+				{
+					pResultSkeleton->m_boneParents[ob] = (uint16)secondToFirstBones[secondMeshIndex];
+				}
+			}
+		}
+		else
+		{
+			pResult->SetSkeleton(pFirst->GetSkeleton());
+		}
 
 
 		// Pose
@@ -523,7 +523,7 @@ namespace mu
 			MUTABLE_CPUPROFILER_SCOPE(Pose);
 
 			pResult->BonePoses.Reserve(pResult->GetSkeleton()->GetBoneCount());
-			
+
 			// Copy poses from the first mesh
 			pResult->BonePoses = pFirst->BonePoses;
 
@@ -541,7 +541,7 @@ namespace mu
 					auto ComputeBoneMergePriority = [](const Mesh::FBonePose& BonePose)
 					{
 						return (EnumHasAnyFlags(BonePose.BoneUsageFlags, EBoneUsageFlags::Skinning) ? 1 : 0) +
-							   (EnumHasAnyFlags(BonePose.BoneUsageFlags, EBoneUsageFlags::Reshaped) ? 1 : 0);
+							(EnumHasAnyFlags(BonePose.BoneUsageFlags, EBoneUsageFlags::Reshaped) ? 1 : 0);
 					};
 
 					if (ComputeBoneMergePriority(ResultBonePose) < ComputeBoneMergePriority(SecondBonePose))
@@ -549,8 +549,8 @@ namespace mu
 						//ResultBonePose.BoneName = SecondBonePose.BoneName;
 						ResultBonePose.BoneTransform = SecondBonePose.BoneTransform;
 						// Merge usage flags
-						EnumAddFlags(ResultBonePose.BoneUsageFlags, SecondBonePose.BoneUsageFlags);					
-					}	
+						EnumAddFlags(ResultBonePose.BoneUsageFlags, SecondBonePose.BoneUsageFlags);
+					}
 				}
 				else
 				{
@@ -561,97 +561,127 @@ namespace mu
 			pResult->BonePoses.Shrink();
 		}
 
-		
+
 		// PhysicsBodies
 		//---------------------------
+		{
+			MUTABLE_CPUPROFILER_SCOPE(PhysicsBodies);
 
-		// Appends InPhysicsBody to OutPhysicsBody removing Bodies that are equal, have same bone and customId and its properies are identical.	
-		auto AppendPhysicsBodiesUnique = [](PhysicsBody& OutPhysicsBody, const PhysicsBody& InPhysicsBody)
-		{	
-			TArray<string>& OutBones = OutPhysicsBody.Bones;
-			TArray<int32>& OutCustomIds = OutPhysicsBody.CustomIds;
-			TArray<FPhysicsBodyAggregate>& OutBodies = OutPhysicsBody.Bodies;
-			
-			const TArray<string>& InBones = InPhysicsBody.Bones;
-			const TArray<int32>& InCustomIds = InPhysicsBody.CustomIds;
-			const TArray<FPhysicsBodyAggregate>& InBodies = InPhysicsBody.Bodies;
-			
-			const int32 InBodyCount = InPhysicsBody.GetBodyCount();
-			const int32 OutBodyCount = OutPhysicsBody.GetBodyCount();
-
-			for (int32 I = 0; I < InBodyCount; ++I)
+			// Appends InPhysicsBody to OutPhysicsBody removing Bodies that are equal, have same bone and customId and its properies are identical.	
+			auto AppendPhysicsBodiesUnique = [](PhysicsBody& OutPhysicsBody, const PhysicsBody& InPhysicsBody) -> bool
 			{
-				int32 FoundIndex = INDEX_NONE;
-				for (int32 O = 0; O < OutBodyCount; ++O)
+				TArray<string>& OutBones = OutPhysicsBody.Bones;
+				TArray<int32>& OutCustomIds = OutPhysicsBody.CustomIds;
+				TArray<FPhysicsBodyAggregate>& OutBodies = OutPhysicsBody.Bodies;
+
+				const TArray<string>& InBones = InPhysicsBody.Bones;
+				const TArray<int32>& InCustomIds = InPhysicsBody.CustomIds;
+				const TArray<FPhysicsBodyAggregate>& InBodies = InPhysicsBody.Bodies;
+
+				const int32 InBodyCount = InPhysicsBody.GetBodyCount();
+				const int32 OutBodyCount = OutPhysicsBody.GetBodyCount();
+
+				bool bModified = false;
+
+				for (int32 InBodyIndex = 0; InBodyIndex < InBodyCount; ++InBodyIndex)
 				{
-					if (InCustomIds[I] == OutCustomIds[O] && InBones[I] == OutBones[O])
+					int32 FoundIndex = INDEX_NONE;
+					for (int32 OutBodyIndex = 0; OutBodyIndex < OutBodyCount; ++OutBodyIndex)
 					{
-						FoundIndex = O;
-						break;
+						if (InCustomIds[InBodyIndex] == OutCustomIds[OutBodyIndex] && InBones[InBodyIndex] == OutBones[OutBodyIndex])
+						{
+							FoundIndex = OutBodyIndex;
+							break;
+						}
 					}
-				}
 
-				if (FoundIndex == INDEX_NONE)
-				{
-					OutBones.Add(InBones[I]);
-					OutCustomIds.Add(InCustomIds[I]);
-					OutBodies.Add(InBodies[I]);
-				}
-				else
-				{
-					for (const FSphereBody& Body : InBodies[I].Spheres)
+					if (FoundIndex == INDEX_NONE)
 					{
-						OutBodies[FoundIndex].Spheres.AddUnique(Body);
+						OutBones.Add(InBones[InBodyIndex]);
+						OutCustomIds.Add(InCustomIds[InBodyIndex]);
+						OutBodies.Add(InBodies[InBodyIndex]);
+
+						bModified |= true;
+
+						continue;
 					}
 
-					for (const FBoxBody& Body : InBodies[I].Boxes)
+					for (const FSphereBody& Body : InBodies[InBodyIndex].Spheres)
 					{
-						OutBodies[FoundIndex].Boxes.AddUnique(Body);
+						const int32 NumBeforeAddition = OutBodies[FoundIndex].Spheres.Num();
+						bModified |= NumBeforeAddition == OutBodies[FoundIndex].Spheres.AddUnique(Body);
 					}
 
-					for (const FSphylBody& Body : InBodies[I].Sphyls)
+					for (const FBoxBody& Body : InBodies[InBodyIndex].Boxes)
 					{
-						OutBodies[FoundIndex].Sphyls.AddUnique(Body);
+						const int32 NumBeforeAddition = OutBodies[FoundIndex].Boxes.Num();
+						bModified |= NumBeforeAddition == OutBodies[FoundIndex].Boxes.AddUnique(Body);
 					}
 
-					for (const FTaperedCapsuleBody& Body : InBodies[I].TaperedCapsules)
+					for (const FSphylBody& Body : InBodies[InBodyIndex].Sphyls)
 					{
-						OutBodies[FoundIndex].TaperedCapsules.AddUnique(Body);
+						const int32 NumBeforeAddition = OutBodies[FoundIndex].Sphyls.Num();
+						bModified |= NumBeforeAddition == OutBodies[FoundIndex].Sphyls.AddUnique(Body);
 					}
 
-					for (const FConvexBody& Body : InBodies[I].Convex)
+					for (const FTaperedCapsuleBody& Body : InBodies[InBodyIndex].TaperedCapsules)
 					{
-						OutBodies[FoundIndex].Convex.AddUnique(Body);
+						const int32 NumBeforeAddition = OutBodies[FoundIndex].TaperedCapsules.Num();
+						bModified |= NumBeforeAddition == OutBodies[FoundIndex].TaperedCapsules.AddUnique(Body);
+					}
+
+					for (const FConvexBody& Body : InBodies[InBodyIndex].Convex)
+					{
+						const int32 NumBeforeAddition = OutBodies[FoundIndex].Convex.Num();
+						bModified |= NumBeforeAddition == OutBodies[FoundIndex].Convex.AddUnique(Body);
 					}
 				}
-			}
-		};
 
+				return bModified;
+			};
 
-		if (pFirst->GetPhysicsBody() == pSecond->GetPhysicsBody())
-		{
-			// Both meshes are sharing the physics body, let the result share it as well.
-			pResult->SetPhysicsBody( pFirst->GetPhysicsBody() );
-		}
-		else
-		{
-			Ptr<PhysicsBody> ResultPhysicsBody = new PhysicsBody;
+			TTuple<const PhysicsBody*, bool> SharedResultPhysicsBody = Invoke([&]()
+				-> TTuple<const PhysicsBody*, bool>
+			{
+				if (pFirst->GetPhysicsBody() == pSecond->GetPhysicsBody())
+				{
+					return MakeTuple(pFirst->GetPhysicsBody().get(), true);
+				}
+
+				if (pFirst->GetPhysicsBody() && !pSecond->GetPhysicsBody())
+				{
+					return MakeTuple(pFirst->GetPhysicsBody().get(), true);
+				}
+
+				if (!pFirst->GetPhysicsBody() && pSecond->GetPhysicsBody())
+				{
+					return MakeTuple(pSecond->GetPhysicsBody().get(), true);
+				}
 
-			if (pFirst->GetPhysicsBody())
+				return MakeTuple(nullptr, false);
+			});
+
+			if (SharedResultPhysicsBody.Get<1>())
 			{
-				AppendPhysicsBodiesUnique(*ResultPhysicsBody, *pFirst->GetPhysicsBody());
+				check(!pFirst->GetPhysicsBody() || !pSecond->GetPhysicsBody())
+
+					// Only one or non of the meshes has physics, share the result.
+					pResult->SetPhysicsBody(SharedResultPhysicsBody.Get<0>());
 			}
-			
-			if (pSecond->GetPhysicsBody())
+			else
 			{
-				AppendPhysicsBodiesUnique(*ResultPhysicsBody, *pSecond->GetPhysicsBody());
+				check(pFirst->GetPhysicsBody() && pSecond->GetPhysicsBody());
+
+				Ptr<PhysicsBody> MergedResultPhysicsBody = pFirst->GetPhysicsBody()->Clone();
+
+				MergedResultPhysicsBody->bBodiesModified =
+					AppendPhysicsBodiesUnique(*MergedResultPhysicsBody, *pSecond->GetPhysicsBody()) ||
+					pFirst->GetPhysicsBody()->bBodiesModified || pSecond->GetPhysicsBody()->bBodiesModified;
+
+				pResult->SetPhysicsBody(MergedResultPhysicsBody);
 			}
-			bool FirstBodyModified = pFirst->GetPhysicsBody() ? pFirst->GetPhysicsBody()->bBodiesModified : false;
-			bool SecondBodyModified = pSecond->GetPhysicsBody() ? pSecond->GetPhysicsBody()->bBodiesModified : false;
-			ResultPhysicsBody->bBodiesModified = FirstBodyModified || SecondBodyModified;
-			pResult->SetPhysicsBody( ResultPhysicsBody );
-		}
 
+		}
 		// Vertices
 		//-----------------
 		{
diff --git a/Engine/Plugins/Experimental/Mutable/Source/MutableRuntime/Private/MuR/OpMeshReshape.h b/Engine/Plugins/Experimental/Mutable/Source/MutableRuntime/Private/MuR/OpMeshReshape.h
--- a/Engine/Plugins/Experimental/Mutable/Source/MutableRuntime/Private/MuR/OpMeshReshape.h
+++ b/Engine/Plugins/Experimental/Mutable/Source/MutableRuntime/Private/MuR/OpMeshReshape.h
@@ -16,6 +16,7 @@
 #include "Distance/DistPoint3Triangle3.h"
 #include "LineTypes.h"
 
+#include "Algo/AnyOf.h" 
 
 
 // TODO: Make the handling of rotations an option. It is more expensive on CPU and memory, and for some
@@ -41,124 +42,130 @@ namespace mu
 		TArray<UE::Geometry::FIndex3i> Triangles;
 	};
 
-	// Method to actually deform a point
-	inline bool GetDeform( const FShapeMeshDescriptorApply& Shape, const FReshapeVertexBindingData& Binding, FVector3f& NewPosition, FQuat4f& Rotation)
+	// Methods to actually deform a point
+	inline void GetDeform( const FShapeMeshDescriptorApply& Shape, const FReshapeVertexBindingData& Binding, FVector3f& NewPosition, FQuat4f& Rotation)
 	{
-		const TArray<FVector3f>& ShapePositions = Shape.Positions;
-		const TArray<FVector3f>& ShapeNormals = Shape.Normals;
-	    const TArray<UE::Geometry::FIndex3i>& ShapeTriangles = Shape.Triangles;
-		
-		if (Binding.Triangle < 0 || Binding.Triangle >= ShapeTriangles.Num() )
-		{
-			return false;
-		}
+		const UE::Geometry::FIndex3i& Triangle = Shape.Triangles[Binding.Triangle];
 
-		const UE::Geometry::FIndex3i& Triangle = ShapeTriangles[Binding.Triangle];
-
-		FVector3f ProjectedVertexPosition
-			= ShapePositions[Triangle.A] * (1.0f - Binding.S - Binding.T)
-			+ ShapePositions[Triangle.B] * Binding.S
-			+ ShapePositions[Triangle.C] * Binding.T;
+		FVector3f ProjectedVertexPosition = 
+			Shape.Positions[Triangle.A] * (1.0f - Binding.S - Binding.T)
+			+ Shape.Positions[Triangle.B] * Binding.S
+			+ Shape.Positions[Triangle.C] * Binding.T;
 
 		// This method approximates the shape face rotation
-		FVector3f InterpolatedNormal
-			= ShapeNormals[Triangle.A] * (1.0f - Binding.S - Binding.T)
-			+ ShapeNormals[Triangle.B] * Binding.S
-			+ ShapeNormals[Triangle.C] * Binding.T;
+		FVector3f InterpolatedNormal = 
+			Shape.Normals[Triangle.A] * (1.0f - Binding.S - Binding.T)
+			+ Shape.Normals[Triangle.B] * Binding.S
+			+ Shape.Normals[Triangle.C] * Binding.T;
 		
-		FVector3f PositionOffset = InterpolatedNormal * Binding.D;
-		NewPosition = ProjectedVertexPosition + PositionOffset;
+		NewPosition = ProjectedVertexPosition + InterpolatedNormal*Binding.D;
+
+		FVector3f CurrentShapeNormal = ((Shape.Positions[Triangle.B] - Shape.Positions[Triangle.A]) ^ 
+										(Shape.Positions[Triangle.C] - Shape.Positions[Triangle.A])).GetSafeNormal();
 
-		FVector3f CurrentShapeNormal = ((ShapePositions[Triangle.B] - ShapePositions[Triangle.A]) ^ (ShapePositions[Triangle.C] - ShapePositions[Triangle.A])).GetSafeNormal();
 		Rotation = FQuat4f::FindBetween(Binding.ShapeNormal, CurrentShapeNormal);
+	}
+
+	inline void GetDeform( const FShapeMeshDescriptorApply& Shape, const FReshapePointBindingData& Binding, FVector3f& NewPosition)
+	{
+		const UE::Geometry::FIndex3i& Triangle = Shape.Triangles[Binding.Triangle];
+
+		FVector3f ProjectedVertexPosition = 
+			Shape.Positions[Triangle.A] * (1.0f - Binding.S - Binding.T)
+			+ Shape.Positions[Triangle.B] * Binding.S
+			+ Shape.Positions[Triangle.C] * Binding.T;
+
+		// This method approximates the shape face rotation
+		FVector3f InterpolatedNormal = 
+			Shape.Normals[Triangle.A] * (1.0f - Binding.S - Binding.T)
+			+ Shape.Normals[Triangle.B] * Binding.S
+			+ Shape.Normals[Triangle.C] * Binding.T;
 		
-		return true;
+		NewPosition = ProjectedVertexPosition + InterpolatedNormal*Binding.D;
 	}
 
+
     //---------------------------------------------------------------------------------------------
     //! Physics Bodies Reshape 
     //---------------------------------------------------------------------------------------------
-	inline bool GetDeformedBox( const FShapeMeshDescriptorApply& Shape, const FReshapeVertexBindingData* BindingData,  TStaticArray<FVector3f, 8>& OutBox )
+	template<uint32 NumPoints>
+	inline bool GetDeformedPoints( const FShapeMeshDescriptorApply& Shape, const FReshapePointBindingData* BindingData,  TStaticArray<FVector3f, NumPoints>& OutPoints )
 	{
-		FQuat4f UnusedRotation;
-
-		for ( int32 I = 0; I < OutBox.Num(); ++I )
+		const int32 ShapeNumTris = Shape.Triangles.Num();
+		for (int32 I = 0; I < NumPoints; ++I)
 		{
-			bool bModified = GetDeform(Shape, *(BindingData + I), OutBox[I], UnusedRotation);
-
-			if (!bModified)
+			const FReshapePointBindingData& BindingDataPoint = *(BindingData + I);
+			if ((BindingDataPoint.Triangle < 0) | (BindingDataPoint.Triangle >= ShapeNumTris))
 			{
-				// Invalidate the whole box. 
 				return false;
 			}
 		}
 
+		for (int32 I = 0; I < NumPoints; ++I)
+		{
+			GetDeform(Shape, *(BindingData + I), OutPoints[I]);
+		}
+
 		return true;
 	}
 
-	inline void GetDeformedConvex( const FShapeMeshDescriptorApply& Shape, const FReshapeVertexBindingData* BindingData, TArray<FVector3f>& InOutDeformedVertices )
+	inline void GetDeformedConvex( const FShapeMeshDescriptorApply& Shape, const FReshapePointBindingData* BindingData, TArrayView<FVector3f>& InOutDeformedVertices )
 	{
+		const int32 ShapeNumTris = Shape.Triangles.Num();
 		const int32 ConvexVertCount = InOutDeformedVertices.Num();
-		FQuat4f UnusedTangentCorrection;
-		for ( int32 I = 0; I < ConvexVertCount; ++I )
+
+		for (int32 I = 0; I < ConvexVertCount; ++I)
 		{
-			GetDeform(Shape, *(BindingData + I), InOutDeformedVertices[I], UnusedTangentCorrection);	
+			const FReshapePointBindingData& BindingDataPoint = *(BindingData + I);
+
+			if ((BindingDataPoint.Triangle < 0) | (BindingDataPoint.Triangle >= ShapeNumTris))
+			{
+				continue;
+			}
+
+			GetDeform(Shape, BindingDataPoint, InOutDeformedVertices[I]);	
 		}
 	}
 
-	inline void ComputeSphereFromDeformedBox( 
-			const TStaticArray<FVector3f, 8>& Box, 
-			FVector3f& InOutP, float& InOutR, 
+	inline void ComputeSphereFromDeformedPoints( 
+			const TStaticArray<FVector3f, 6>& Points, 
+			FVector3f& OutP, float& OutR, 
 			const FTransform3f& InvBoneT )
-	{
-		TStaticArray<FVector3f, 8> TransformedBox = Box;
-		for (int32 I = 0; I < 8; ++I)
-		{
-			TransformedBox[I] = InvBoneT.TransformPosition(Box[I]) - InOutP;
-		}
-		
+	{	
 		FVector3f Centroid(ForceInitToZero);
-		for ( const FVector3f& V : TransformedBox )
+		for (const FVector3f& V : Points)
 		{
 			Centroid += V;	
 		}
 
-		Centroid *= (1.0f/8.0f);
+		Centroid *= (1.0f/6.0f);
 
-		InOutP = Centroid + InOutP;
+		OutP = Centroid + InvBoneT.GetTranslation();
 		
-		FVector3f Dim(ForceInitToZero);
-		for ( const FVector3f& V : TransformedBox )
+		float Radius = 0.0f;
+		for (const FVector3f& V : Points)
 		{
-			Dim += (V - Centroid).GetAbs();
+			Radius += (V - Centroid).Length();
 		}
 
-		Dim *= (1.0f/8.0f);
+		Radius *= (1.0f/6.0f);
 		
-		InOutR = (Dim.X + Dim.Y + Dim.Z) * (1.0f/3.0f);
+		OutR = Radius;
 	}
 
-	inline void ComputeBoxFromDeformedBox( 
-			const TStaticArray<FVector3f, 8>& Box, 
-			FVector3f& InOutP, FQuat4f& InOutQ, FVector3f& InOutS, 
+	inline void ComputeBoxFromDeformedPoints( 
+			const TStaticArray<FVector3f, 14>& Points, 
+			FVector3f& OutP, FQuat4f& OutQ, FVector3f& OutS, 
 			const FTransform3f& InvBoneT )
-	{
-		const FQuat4f InvQ = InOutQ.Inverse();
-		
-		TStaticArray<FVector3f, 8> TranformedBox;
-		for (int32 I = 0; I < 8; ++I)
-		{
-			TranformedBox[I] = InvQ.RotateVector(InvBoneT.TransformPosition(Box[I]) - InOutP);
-		}
-	
-		const FVector3f TopC = (Box[0] + Box[1] + Box[2] + Box[3])*0.25f;
-		const FVector3f BottomC = (Box[4+0] + Box[4+1] + Box[4+2] + Box[4+3])*0.25f;
+	{	
+		const FVector3f TopC = (Points[0] + Points[1] + Points[2] + Points[3])*0.25f;
+		const FVector3f BottomC = (Points[4+0] + Points[4+1] + Points[4+2] + Points[4+3])*0.25f;
     
-		const FVector3f FrontC = (Box[0] + Box[1] + Box[4+0] + Box[4+1])*0.25f;
-		const FVector3f BackC = (Box[2] + Box[3] + Box[4+2] + Box[4+3])*0.25f;
+		const FVector3f FrontC = (Points[0] + Points[1] + Points[4+0] + Points[4+1])*0.25f;
+		const FVector3f BackC = (Points[2] + Points[3] + Points[4+2] + Points[4+3])*0.25f;
     
-		const FVector3f RightC =  (Box[1] + Box[2] + Box[4+1] + Box[4+2])*0.25f;
-		const FVector3f LeftC =  (Box[3] + Box[0] + Box[4+3] + Box[4+0])*0.25f;
+		const FVector3f RightC =  (Points[1] + Points[2] + Points[4+1] + Points[4+2])*0.25f;
+		const FVector3f LeftC =  (Points[3] + Points[0] + Points[4+3] + Points[4+0])*0.25f;
     
 		FVector3f ZB = (TopC - BottomC).GetSafeNormal();
 		FVector3f XB = (RightC - LeftC).GetSafeNormal();
@@ -191,100 +198,186 @@ namespace mu
 					? FRotationMatrix44f::MakeFromZX(ZB, XB)
 					: FRotationMatrix44f::MakeFromZY(ZB, YB);	
 		} 
-		
-		InOutQ = RotationMatrix.ToQuat() * InOutQ;
-		InOutP = (TopC + BottomC) * 0.5f + InOutP;
-		InOutS = FVector3f((RightC - LeftC).Size(), (FrontC - BackC).Size(), (TopC - BottomC).Size()) * 0.5f;	
+	
+		FTransform3f ShapeToBone = FTransform3f(RotationMatrix.ToQuat(), (TopC + BottomC) * 0.5f) * InvBoneT;
+
+		OutQ = ShapeToBone.GetRotation();
+		OutP = ShapeToBone.GetTranslation();
+		OutS = FVector3f((RightC - LeftC).Size(), (FrontC - BackC).Size(), (TopC - BottomC).Size()) * 0.5f;	
 	}
 
-	inline void ComputeSphylFromDeformedBox( 
-			const TStaticArray<FVector3f, 8>& Box, 
-			FVector3f& InOutP, FQuat4f& InOutQ, float& InOutR, float& InOutL, 
+	inline void ComputeSphylFromDeformedPoints( 
+			const TStaticArray<FVector3f, 14>& Points, 
+			FVector3f& OutP, FQuat4f& OutQ, float& OutR, float& OutL, 
 			const FTransform3f& InvBoneT)
 	{
-		const FQuat4f InvQ = InOutQ.Inverse();
-    	
-		TStaticArray<FVector3f, 8> TranformedBox;
-		for (int32 I = 0; I < 8; ++I)
+		constexpr int32 NumCentroids = 5;
+		TStaticArray<FVector3f, NumCentroids> Centroids;
+		Centroids[0] = Points[0];
+		Centroids[1] = Points[1];
+
+		for (int32 I = 0; I < NumCentroids - 2; ++I)
 		{
-			TranformedBox[I] = InvQ.RotateVector(InvBoneT.TransformPosition(Box[I]) - InOutP);
+			Centroids[2 + I] = 
+				(Points[2 + I*4 + 0] + Points[2 + I*4 + 1] +
+				 Points[2 + I*4 + 2] + Points[2 + I*4 + 3] ) * 0.25f;
 		}
-       
-		FVector3f TopCentroid(ForceInitToZero);
-		FVector3f BottomCentroid(ForceInitToZero);
-		for (int32 I = 0; I < 4; ++I)
+
+		// Geometric linear regression of top, bottom and rings centroids.
+		FVector3f Centroid = FVector3f::Zero();
+		for (const FVector3f& C : Centroids)
 		{
-			TopCentroid += TranformedBox[I];
-			BottomCentroid += TranformedBox[4+I];
+			Centroid += C;
 		}
+
+		constexpr float OneOverNumCentroids = 1.0f / static_cast<float>(NumCentroids);
+		Centroid *= OneOverNumCentroids;
 		
-		TopCentroid *= 0.25f;
-		BottomCentroid *= 0.25f;
+		for (FVector3f& C : Centroids)
+		{
+			C -= Centroid;
+		}
 
-		FVector3f DirVector = (TopCentroid - BottomCentroid).GetSafeNormal();
+		constexpr int32 NumIters = 3;
+		FVector3f Direction = (Centroids[0] - Centroids[1]).GetSafeNormal();
+		for (int32 Iter = 0; Iter < NumIters; ++Iter)
+		{
+			FVector3f IterDirRefinement = Direction;
+			for (const FVector3f& C : Centroids)
+			{
+				IterDirRefinement += C * FVector3f::DotProduct(Direction, C);
+			}
 
-		InOutQ = FQuat4f::FindBetweenNormals(FVector3f::UnitZ(), DirVector) * InOutQ; 
-		InOutP = (TopCentroid + BottomCentroid) * 0.5f + InOutP;
+			Direction = IterDirRefinement.GetSafeNormal();
+		}
 
-		FVector3f TopDim(ForceInitToZero);
-		FVector3f BottomDim(ForceInitToZero);
-		for (int32 I = 0; I < 4; ++I)
+		// Project centroids to the line described by Direction and Centroid.
+		for (FVector3f& C : Centroids)
 		{
-			TopDim += (TopCentroid - TranformedBox[I]).GetAbs();
-			BottomDim += (BottomCentroid - TranformedBox[4+I]).GetAbs();
+			C = Centroid + Direction * FVector3f::DotProduct(C, Direction);
 		}
 
-		TopDim *= 0.25f;
-		BottomDim *= 0.25f;
+		// Quaternion form {0,0,1} to Direction.
+		const FQuat4f Rotation = FQuat4f(-Direction.Y, Direction.X, 0.0f, 1.0f + FMath::Max(Direction.Z, -1.0f + UE_SMALL_NUMBER)).GetNormalized(); 
 
-		InOutR = (TopDim.X + TopDim.Y + BottomDim.X + BottomDim.Y) * 0.25f;
-		InOutL = (TopCentroid - BottomCentroid).Size() - InOutR*2.0f;
+		FTransform3f ShapeToBone = FTransform3f(Rotation, Centroid) * InvBoneT;
+		OutQ = ShapeToBone.GetRotation();
+		OutP = ShapeToBone.GetTranslation();
+
+		// Project ring points to plane formed by ring centroid and direction to extract ring radius.
+		const auto ComputeRadiusContribution = [&](const FVector3f& P, const FVector3f& Origin, const FVector3f& Dir) -> float
+		{
+			return ((P + Dir * FVector3f::DotProduct(Dir, P - Origin)) - Origin).Length();
+		};
+
+		const float R0 = 
+			ComputeRadiusContribution(Points[2], Centroids[2], Direction) +
+			ComputeRadiusContribution(Points[3], Centroids[2], Direction) +
+			ComputeRadiusContribution(Points[4], Centroids[2], Direction) +
+			ComputeRadiusContribution(Points[5], Centroids[2], Direction);
+
+		const float R1 = 
+			ComputeRadiusContribution(Points[6], Centroids[3], Direction) +
+			ComputeRadiusContribution(Points[7], Centroids[3], Direction) +
+			ComputeRadiusContribution(Points[8], Centroids[3], Direction) +
+			ComputeRadiusContribution(Points[9], Centroids[3], Direction);
+
+		const float R2 = 
+			ComputeRadiusContribution(Points[10], Centroids[4], Direction) +
+			ComputeRadiusContribution(Points[11], Centroids[4], Direction) +
+			ComputeRadiusContribution(Points[12], Centroids[4], Direction) +
+			ComputeRadiusContribution(Points[13], Centroids[4], Direction);
+
+		OutR =  (R0 + R1 + R2) * (0.25f/3.0f);
+		OutL = FMath::Max(0.0f, (Centroids[0] - Centroids[1]).Length() - OutR*2.0f);
 	}
 
-	inline void ComputeTaperedCapsuleFromDeformedBox( 
-			const TStaticArray<FVector3f, 8>& Box, 
-			FVector3f& InOutP, FQuat4f& InOutQ, float& InOutR0, float& InOutR1, float& InOutL, 
+	inline void ComputeTaperedCapsuleFromDeformedPoints( 
+			const TStaticArray<FVector3f, 14>& Points, 
+			FVector3f& OutP, FQuat4f& OutQ, float& OutR0, float& OutR1, float& OutL, 
 			const FTransform3f& InvBoneT)
 	{
-		const FQuat4f InvQ = InOutQ.Inverse();
-    
-		TStaticArray<FVector3f, 8> TranformedBox;
-		for (int32 I = 0; I < 8; ++I)
+		constexpr int32 NumCentroids = 5;
+		TStaticArray<FVector3f, NumCentroids> Centroids;
+		Centroids[0] = Points[0];
+		Centroids[1] = Points[1];
+
+		for (int32 I = 0; I < NumCentroids - 2; ++I)
 		{
-			TranformedBox[I] = InvQ.RotateVector(InvBoneT.TransformPosition(Box[I]) - InOutP);
+			Centroids[2 + I] = 
+				(Points[2 + I*4 + 0] + Points[2 + I*4 + 1] +
+				 Points[2 + I*4 + 2] + Points[2 + I*4 + 3] ) * 0.25f;
 		}
-       
-		FVector3f TopCentroid(ForceInitToZero);
-		FVector3f BottomCentroid(ForceInitToZero);
-		for (int32 I = 0; I < 4; ++I)
+	
+		// Geometric linear regression of top, bottom and ring centroids.
+		FVector3f Centroid = FVector3f::Zero();
+		for (const FVector3f& C : Centroids)
 		{
-			TopCentroid += TranformedBox[I];
-			BottomCentroid += TranformedBox[4+I];
+			Centroid += C;
 		}
-		
-		TopCentroid *= 0.25f;
-		BottomCentroid *= 0.25f;
 
-		FVector3f DirVector = (TopCentroid - BottomCentroid).GetSafeNormal();
+		constexpr float OneOverNumCentroids = 1.0f / static_cast<float>(NumCentroids);
+		Centroid *= OneOverNumCentroids;
+		
+		for (FVector3f& C : Centroids)
+		{
+			C -= Centroid;
+		}
 
-		InOutQ = InOutQ * FQuat4f::FindBetweenNormals(FVector3f::UnitZ(), DirVector); 
-		InOutP = (TopCentroid + BottomCentroid) * 0.5f + InOutP;
+		constexpr int32 NumIters = 3;
+		FVector3f Direction = (Centroids[0] - Centroids[1]).GetSafeNormal();
+		for (int32 Iter = 0; Iter < NumIters; ++Iter)
+		{
+			FVector3f IterDirRefinement = Direction;
+			for (const FVector3f& C : Centroids)
+			{
+				IterDirRefinement += C * FVector3f::DotProduct(Direction, C);
+			}
 
-		FVector3f TopDim(ForceInitToZero);
-		FVector3f BottomDim(ForceInitToZero);
+			Direction = IterDirRefinement.GetSafeNormal();
+		}
 
-		for (int32 I = 0; I < 4; ++I)
+		// Project centroids to the line described by Direction and Centroid.
+		for (FVector3f& C : Centroids)
 		{
-			TopDim += (TopCentroid - TranformedBox[I]).GetAbs();
-			BottomDim += (BottomCentroid - TranformedBox[4+I]).GetAbs();
+			C = Centroid + Direction * FVector3f::DotProduct(C, Direction);
 		}
 
-		TopDim *= 0.25f;
-		BottomDim *= 0.25f;
+		// Quaternion form {0,0,1} to Direction.
+		const FQuat4f Rotation = FQuat4f(-Direction.Y, Direction.X, 0.0f, 1.0f + FMath::Max(Direction.Z, -1.0f + UE_SMALL_NUMBER)).GetNormalized(); 
+
+		FTransform3f ShapeToBone = FTransform3f(Rotation, Centroid) * InvBoneT;
+		OutQ = ShapeToBone.GetRotation();
+		OutP = ShapeToBone.GetTranslation();
 
-		InOutR0 = (TopDim.X + TopDim.Y) * 0.5f;
-		InOutR1 = (BottomDim.X + BottomDim.Y) * 0.5f;
-		InOutL = (TopCentroid - BottomCentroid).Size() - (InOutR0 + InOutR1);
+		// Project ring points to plane formed by ring centroid and direction to extract ring radius.
+		const auto ComputeRadiusContribution = [&](const FVector3f& P, const FVector3f& Origin, const FVector3f& Dir) -> float
+		{
+			return ((P + Dir * FVector3f::DotProduct(Dir, P - Origin)) - Origin).Length();
+		};
+
+		const float R0 = 
+			ComputeRadiusContribution(Points[2], Centroids[2], Direction) +
+			ComputeRadiusContribution(Points[3], Centroids[2], Direction) +
+			ComputeRadiusContribution(Points[4], Centroids[2], Direction) +
+			ComputeRadiusContribution(Points[5], Centroids[2], Direction);
+
+		const float R1 = 
+			ComputeRadiusContribution(Points[6], Centroids[3], Direction) +
+			ComputeRadiusContribution(Points[7], Centroids[3], Direction) +
+			ComputeRadiusContribution(Points[8], Centroids[3], Direction) +
+			ComputeRadiusContribution(Points[9], Centroids[3], Direction);
+
+		const float R2 = 
+			ComputeRadiusContribution(Points[10], Centroids[4], Direction) +
+			ComputeRadiusContribution(Points[11], Centroids[4], Direction) +
+			ComputeRadiusContribution(Points[12], Centroids[4], Direction) +
+			ComputeRadiusContribution(Points[13], Centroids[4], Direction);
+
+		// TODO: Ajust for R1, center ring radius.
+		OutR0 = R0*0.25f;
+		OutR1 = R2*0.25f;
+		OutL = FMath::Max(0.0f, (Centroids[0] - Centroids[1]).Length() - (OutR0 + OutR1));
 	}
 
 	inline void ApplyToVertices(Mesh* Mesh, TArrayView<const FReshapeVertexBindingData> BindingData, const FShapeMeshDescriptorApply& Shape)
@@ -292,42 +385,46 @@ namespace mu
 		check(Mesh);
 		check(Mesh->GetVertexCount() == BindingData.Num());
 		
-
 		UntypedMeshBufferIterator ItPosition(Mesh->GetVertexBuffers(), MBS_POSITION);
 		UntypedMeshBufferIterator ItNormal(Mesh->GetVertexBuffers(), MBS_NORMAL);
 		UntypedMeshBufferIterator ItTangent(Mesh->GetVertexBuffers(), MBS_TANGENT);
 
 #if DO_CHECK
 		// checking if the Base shape has more triangles than the target shape
-		for (const FReshapeVertexBindingData& Binding : BindingData)
+		const bool bTriangleOutOfScopeFound = Algo::AnyOf(BindingData, 
+			[NumShapeTriangles = Shape.Triangles.Num()](const FReshapeVertexBindingData& B) { return B.Triangle >= NumShapeTriangles; });
+
+		if (bTriangleOutOfScopeFound)
 		{
-			if (Binding.Triangle >= Shape.Triangles.Num())
-			{
-				UE_LOG(LogMutableCore, Warning, TEXT("Performing a Mesh Reshape where base shape and target shape do not have the same number of triangles."));
-				break;
-			}
-		}
+			UE_LOG(LogMutableCore, Warning, TEXT("Performing a Mesh Reshape where base shape and target shape do not have the same number of triangles."));
+		}	
 #endif
 
+		const int32 ShapeTriangleCount = Shape.Triangles.Num();
+		const int32 ShapePositionCount = Shape.Positions.Num();
+
 		const int32 MeshVertexCount = BindingData.Num();
+
 		for (int32 MeshVertexIndex = 0; MeshVertexIndex < MeshVertexCount; ++MeshVertexIndex)
 		{
 			const FReshapeVertexBindingData& Binding = BindingData[MeshVertexIndex];
 
 			FVector3f NewPosition;
 			FQuat4f TangentSpaceCorrection;
-			bool bModified = GetDeform( Shape, Binding, NewPosition, TangentSpaceCorrection );
 
+			const bool bModified = (Binding.Triangle >= 0) & (Binding.Triangle < ShapeTriangleCount);
 			if (bModified)
 			{
+				GetDeform(Shape, Binding, NewPosition, TangentSpaceCorrection);
+
 				FVector3f OldPosition = ItPosition.GetAsVec3f();
 				
-				if ( !FMath::IsNearlyEqual( Binding.Weight, 1.0f ) )
+				if (!FMath::IsNearlyEqual(Binding.Weight, 1.0f))
 				{
 					// Zero weighted vertices are already discarded at the binding phase so there is no gain
 					// checking here.
-					NewPosition = FMath::Lerp( OldPosition, NewPosition, Binding.Weight );
-					TangentSpaceCorrection = FQuat4f::Slerp( FQuat4f::Identity, TangentSpaceCorrection, Binding.Weight );
+					NewPosition = FMath::Lerp(OldPosition, NewPosition, Binding.Weight);
+					TangentSpaceCorrection = FQuat4f::Slerp(FQuat4f::Identity, TangentSpaceCorrection, Binding.Weight).GetNormalized();
 				}
 				
 				// Non rigid vertices will not rotate since are attached to themselves 
@@ -366,29 +463,45 @@ namespace mu
 	}
 
 	inline void ApplyToPose(Mesh* Result, 
-			TArrayView<const FReshapeVertexBindingData> BindingData, 
-			TArrayView<const int> BoneIndices,
+			TArrayView<const FReshapePointBindingData> BindingData, 
+			TArrayView<const int32> BoneIndices,
 			const FShapeMeshDescriptorApply& Shape)
 	{
 		check(Result);
 
+#if DO_CHECK
+		// checking if the Base shape has more triangles than the target shape
+		const bool bTriangleOutOfScopeFound = Algo::AnyOf(BindingData, 
+			[NumShapeTriangles = Shape.Triangles.Num()](const FReshapePointBindingData& B) { return B.Triangle >= NumShapeTriangles; });
+
+		if (bTriangleOutOfScopeFound)
+		{
+			UE_LOG(LogMutableCore, Warning, TEXT("Performing a Mesh Reshape where base shape and target shape do not have the same number of triangles."));
+		}	
+#endif
+
+		const int32 NumShapeTris = Shape.Triangles.Num();
 		const int32 NumBoneIndices = BoneIndices.Num();
-		for (int32 b = 0; b < NumBoneIndices; ++b)
+		for (int32 BoneSelectionIndex = 0; BoneSelectionIndex < NumBoneIndices; ++BoneSelectionIndex)
 		{
-			int32 BoneIndex = BoneIndices[b];
+			int32 BoneIndex = BoneIndices[BoneSelectionIndex];
 
-			const FReshapeVertexBindingData& Binding = BindingData[b];
+			const FReshapePointBindingData& Binding = BindingData[BoneSelectionIndex];
 
 			check(!EnumHasAnyFlags(Result->BonePoses[BoneIndex].BoneUsageFlags, EBoneUsageFlags::Root));
 
 			FTransform3f& T = Result->BonePoses[BoneIndex].BoneTransform;
 
 			FVector3f NewPosition(ForceInitToZero);
-			FQuat4f TangentSpaceCorrection;
 
-			const bool bModified = GetDeform(Shape, Binding, NewPosition, TangentSpaceCorrection);
-			const bool bHasChanged = bModified && FVector3f::DistSquared(NewPosition, T.GetLocation()) > UE_SMALL_NUMBER;
+			const bool bModified = (Binding.Triangle >= 0) & (Binding.Triangle < NumShapeTris);
+			
+			if (bModified)
+			{
+				GetDeform(Shape, Binding, NewPosition);
+			}
 
+			const bool bHasChanged = bModified && FVector3f::DistSquared(NewPosition, T.GetLocation()) > UE_SMALL_NUMBER;
 			// Only set it if has actually moved.
 			if (bHasChanged)
 			{	
@@ -402,15 +515,20 @@ namespace mu
 	}
 
 	inline void ApplyToPhysicsBodies(
-			PhysicsBody* PBody, const Mesh* BaseMesh, 
-			TArrayView<const FReshapeVertexBindingData> BindingData, TArrayView<const int32> UsedIndices, 
+			PhysicsBody& PBody, const Mesh& BaseMesh, 
+			TArrayView<const FReshapePointBindingData> BindingData, TArrayView<const int32> UsedIndices, 
 			const FShapeMeshDescriptorApply& Shape)
 	{
-		check(PBody);
-		check(BaseMesh);
+#if DO_CHECK
+		// checking if the Base shape has more triangles than the target shape
+		const bool bTriangleOutOfScopeFound = Algo::AnyOf(BindingData, 
+			[NumShapeTriangles = Shape.Triangles.Num()](const FReshapePointBindingData& B) { return B.Triangle >= NumShapeTriangles; });
 
-		// Not used.	
-		FQuat4f DummyTangentSpaceCorrection;
+		if (bTriangleOutOfScopeFound)
+		{
+			UE_LOG(LogMutableCore, Warning, TEXT("Performing a Mesh Reshape where base shape and target shape do not have the same number of triangles."));
+		}
+#endif
 		
 		int32 NumProcessedBindPoints = 0;
 
@@ -419,40 +537,39 @@ namespace mu
 		// Retrieve them in the same order the boxes where put in, so they can be linked to the physics body volumes.
 		for ( const int32 B : UsedIndices )
 		{	
-			int32 BoneIdx = BaseMesh->FindBonePose(PBody->GetBodyBoneName(B));
+			int32 BoneIdx = BaseMesh.FindBonePose(PBody.GetBodyBoneName(B));
 			FTransform3f BoneTransform = FTransform3f::Identity;
 			if (BoneIdx >= 0)
 			{
-				BaseMesh->GetBoneTransform(BoneIdx, BoneTransform);
+				BaseMesh.GetBoneTransform(BoneIdx, BoneTransform);
 			}
 			
 			FTransform3f InvBoneTransform = BoneTransform.Inverse();
 
-			const int32 NumSpheres = PBody->GetSphereCount(B); 
+			const int32 NumSpheres = PBody.GetSphereCount(B); 
 			for ( int32 I = 0; I < NumSpheres; ++I )
 			{
 				FVector3f P;
 				float R;
 
-				TStaticArray<FVector3f, 8> Box; 
-				const bool bDeformed = GetDeformedBox( Shape, &BindingData[NumProcessedBindPoints], Box );
+				TStaticArray<FVector3f, 6> Points; 
+				const bool bDeformed = GetDeformedPoints(Shape, &BindingData[NumProcessedBindPoints], Points);
 
 				if (bDeformed)
 				{
-					PBody->GetSphere(B, I, P, R);
-					ComputeSphereFromDeformedBox(Box, P, R, InvBoneTransform);
-					PBody->SetSphere(B, I, P, R);
+					ComputeSphereFromDeformedPoints(Points, P, R, InvBoneTransform);
+					PBody.SetSphere(B, I, P, R);
 					bAnyModified = true;
 				}
 
-				NumProcessedBindPoints += Box.Num();
+				NumProcessedBindPoints += Points.Num();
 			}
 
-			const int32 NumBoxes = PBody->GetBoxCount(B);
+			const int32 NumBoxes = PBody.GetBoxCount(B);
 			for ( int32 I = 0; I < NumBoxes; ++I )
 			{
-				TStaticArray<FVector3f, 8> Box;
-				const bool bDeformed = GetDeformedBox( Shape, &BindingData[NumProcessedBindPoints], Box );
+				TStaticArray<FVector3f, 14> Points;
+				const bool bDeformed = GetDeformedPoints(Shape, &BindingData[NumProcessedBindPoints], Points);
 				
 				if (bDeformed)
 				{
@@ -460,20 +577,19 @@ namespace mu
 					FQuat4f Q;
 					FVector3f S;
 					
-					PBody->GetBox(B, I, P, Q, S);
-					ComputeBoxFromDeformedBox(Box, P, Q, S, InvBoneTransform);
-					PBody->SetBox(B, I, P, Q, S);
+					ComputeBoxFromDeformedPoints(Points, P, Q, S, InvBoneTransform);
+					PBody.SetBox(B, I, P, Q, S);
 					bAnyModified = true;
 				}
 
-				NumProcessedBindPoints += Box.Num();
+				NumProcessedBindPoints += Points.Num();
 			}
 			
-			const int32 NumSphyls = PBody->GetSphylCount(B);
+			const int32 NumSphyls = PBody.GetSphylCount(B);
 			for ( int32 I = 0; I < NumSphyls; ++I )
 			{
-				TStaticArray<FVector3f, 8> Box;
-				const bool bDeformed = GetDeformedBox( Shape, &BindingData[NumProcessedBindPoints], Box );
+				TStaticArray<FVector3f, 14> Points;
+				const bool bDeformed = GetDeformedPoints(Shape, &BindingData[NumProcessedBindPoints], Points);
 
 				if ( bDeformed )
 				{
@@ -482,20 +598,19 @@ namespace mu
 					float R;
 					float L;						
 
-					PBody->GetSphyl(B, I, P, Q, R, L);
-					ComputeSphylFromDeformedBox(Box, P, Q, R, L, InvBoneTransform);
-					PBody->SetSphyl(B, I, P, Q, R, L);
+					ComputeSphylFromDeformedPoints(Points, P, Q, R, L, InvBoneTransform);
+					PBody.SetSphyl(B, I, P, Q, R, L);
 					bAnyModified = true;
 				}
 
-				NumProcessedBindPoints += Box.Num();
+				NumProcessedBindPoints += Points.Num();
 			}
 
-			const int32 NumTaperedCapsules = PBody->GetTaperedCapsuleCount(B);
+			const int32 NumTaperedCapsules = PBody.GetTaperedCapsuleCount(B);
 			for ( int32 I = 0; I < NumTaperedCapsules; ++I )
 			{
-				TStaticArray<FVector3f, 8> Box;
-				const bool bDeformed = GetDeformedBox( Shape, &BindingData[NumProcessedBindPoints], Box );
+				TStaticArray<FVector3f, 14> Points;
+				const bool bDeformed = GetDeformedPoints(Shape, &BindingData[NumProcessedBindPoints], Points);
 
 				if ( bDeformed )
 				{
@@ -505,46 +620,38 @@ namespace mu
 					float R1;
 					float L;
 				
-					PBody->GetTaperedCapsule(B, I, P, Q, R0, R1, L);
-					ComputeTaperedCapsuleFromDeformedBox(Box, P, Q, R0, R1, L, InvBoneTransform);
-					PBody->SetTaperedCapsule(B, I, P, Q, R0, R1, L);
+					ComputeTaperedCapsuleFromDeformedPoints(Points, P, Q, R0, R1, L, InvBoneTransform);
+					PBody.SetTaperedCapsule(B, I, P, Q, R0, R1, L);
 					bAnyModified = true;
 				}
 				
-				NumProcessedBindPoints += Box.Num();
+				NumProcessedBindPoints += Points.Num();
 			}
 
-			const int32 NumConvex = PBody->GetConvexCount(B);
+			const int32 NumConvex = PBody.GetConvexCount(B);
 			for ( int32 I = 0; I < NumConvex; ++I )
 			{
-				FVector3f const* VertexData;
-				int32 VertexCount;
-				int32 const* IndexData;
-				int32 IndexCount;
-				FTransform3f ConvexT;
-
-				PBody->GetConvex(B, I, VertexData, VertexCount, IndexData, IndexCount, ConvexT);
+				TArrayView<FVector3f> VerticesView;
+				TArrayView<int32> IndicesView;
+				PBody.GetConvexMeshView(B, I, VerticesView, IndicesView);
 
-				// TODO: Allow access and modification to the convex data directly to avoid a copy here.
-				TArray<FVector3f> DeformedVertexData(VertexData, VertexCount);
-				TArray<int32> Indices(IndexData, IndexCount);
+				FTransform3f ConvexTransform;
+				PBody.GetConvexTransform(B, I, ConvexTransform);
 
-				GetDeformedConvex( Shape, &BindingData[NumProcessedBindPoints], DeformedVertexData);
+				GetDeformedConvex(Shape, &BindingData[NumProcessedBindPoints], VerticesView);
 
-				FTransform3f InvConvexT = InvBoneTransform * ConvexT.Inverse();
-				for ( FVector3f& V : DeformedVertexData )
+				FTransform3f InvConvexT = InvBoneTransform * ConvexTransform.Inverse();
+				for ( FVector3f& V : VerticesView )
 				{
 					V = InvConvexT.TransformPosition( V );				
 				}
 				
-				PBody->SetConvex(B, I, DeformedVertexData.GetData(), DeformedVertexData.Num(), Indices.GetData(), Indices.Num(), ConvexT);
-
-				NumProcessedBindPoints += VertexCount;
+				NumProcessedBindPoints += VerticesView.Num();
 				bAnyModified = true;
 			}
 		}
 
-		PBody->bBodiesModified = bAnyModified;
+		PBody.bBodiesModified = bAnyModified;
 
 		check(NumProcessedBindPoints == BindingData.Num());	
 	}
@@ -718,10 +825,10 @@ namespace mu
 			{
 				// \TODO: More checks
 				check(BarycentricDataChannel == 0);
-				check(SkeletonBindBuffer && SkeletonBindBuffer->GetElementSize(BarycentricDataBuffer) == (int)sizeof(FReshapeVertexBindingData));
+				check(SkeletonBindBuffer && SkeletonBindBuffer->GetElementSize(BarycentricDataBuffer) == (int)sizeof(FReshapePointBindingData));
 
-				TArrayView<const FReshapeVertexBindingData> SkeletonBindingData( 
-						(const FReshapeVertexBindingData*)SkeletonBindBuffer->GetBufferData(BarycentricDataBuffer),
+				TArrayView<const FReshapePointBindingData> SkeletonBindingData( 
+						(const FReshapePointBindingData*)SkeletonBindBuffer->GetBufferData(BarycentricDataBuffer),
 						SkeletonBindBuffer->GetElementCount());
 			
 				check(SkeletonBindBuffer->GetBufferCount() >= 2);
@@ -732,7 +839,10 @@ namespace mu
 			}
 		}
 
-		// Transform Physics Volumes based on the deformed bounding box points.
+		// When transforming the physics volumes, the resulting pose of of the skeleton reshape operation will be used, so 
+		// order of operation is important.
+
+		// Transform physics volumes based on the deformed sampling points.
 		const PhysicsBody* OldPhysicsBody = Result->m_pPhysicsBody.get();
 
 		if (bReshapePhysicsVolumes && OldPhysicsBody)
@@ -759,10 +869,10 @@ namespace mu
 
 				// \TODO: More checks
 				check(BarycentricDataChannel == 0);
-				check(PhysicsBindBuffer.GetElementSize(BarycentricDataBuffer) == (int)sizeof(FReshapeVertexBindingData));
+				check(PhysicsBindBuffer.GetElementSize(BarycentricDataBuffer) == (int)sizeof(FReshapePointBindingData));
 					
-				TArrayView<const FReshapeVertexBindingData> PhysicsBindingData(  
-						(const FReshapeVertexBindingData*)PhysicsBindBuffer.GetBufferData(BarycentricDataBuffer),
+				TArrayView<const FReshapePointBindingData> PhysicsBindingData(  
+						(const FReshapePointBindingData*)PhysicsBindBuffer.GetBufferData(BarycentricDataBuffer),
 						PhysicsBindBuffer.GetElementCount() );
 
 				TArrayView<const int32> UsedIndices( 
@@ -772,7 +882,8 @@ namespace mu
 				Ptr<PhysicsBody> NewPhysicsBody = OldPhysicsBody->Clone();
 				Result->SetPhysicsBody(NewPhysicsBody);
 
-				ApplyToPhysicsBodies( NewPhysicsBody.get(), BaseMesh, PhysicsBindingData, UsedIndices, ShapeDescriptor );
+				check(NewPhysicsBody);
+				ApplyToPhysicsBodies(*NewPhysicsBody, *Result, PhysicsBindingData, UsedIndices, ShapeDescriptor);
 			}
 		}
 		
diff --git a/Engine/Plugins/Experimental/Mutable/Source/MutableRuntime/Private/MuR/PhysicsBody.cpp b/Engine/Plugins/Experimental/Mutable/Source/MutableRuntime/Private/MuR/PhysicsBody.cpp
--- a/Engine/Plugins/Experimental/Mutable/Source/MutableRuntime/Private/MuR/PhysicsBody.cpp
+++ b/Engine/Plugins/Experimental/Mutable/Source/MutableRuntime/Private/MuR/PhysicsBody.cpp
@@ -166,21 +166,29 @@ namespace mu
 	}
 
 	//-------------------------------------------------------------------------------------------------
-	void PhysicsBody::SetConvex(
+	void PhysicsBody::SetConvexMesh(
+		int32 B, int32 I,
+		TArrayView<const FVector3f> Vertices, TArrayView<const int32> Indices)
+	{
+		check(B >= 0 && B < Bodies.Num());
+		check(I >= 0 && I < Bodies[B].Convex.Num());
+
+		Bodies[B].Convex[I].Vertices = TArray<FVector3f>(Vertices);
+		Bodies[B].Convex[I].Indices = TArray<int32>(Indices);
+	}
+
+	//-------------------------------------------------------------------------------------------------
+	void PhysicsBody::SetConvexTransform(
 		int32 B, int32 I,
-		const FVector3f* Vertices, int32 VerticesCount,
-		const int32* Indices, int32 IndicesCount,
 		const FTransform3f& Transform)
 	{
 		check(B >= 0 && B < Bodies.Num());
 		check(I >= 0 && I < Bodies[B].Convex.Num());
 
-		Bodies[B].Convex[I].Vertices = TArray<FVector3f>(Vertices, VerticesCount);
-		Bodies[B].Convex[I].Indices = TArray<int32>(Indices, IndicesCount);
 		Bodies[B].Convex[I].Transform = Transform;
-
 	}
 
+
 	//-------------------------------------------------------------------------------------------------
 	void PhysicsBody::SetSphyl(
 		int32 B, int32 I,
@@ -329,18 +337,31 @@ namespace mu
 	//-------------------------------------------------------------------------------------------------
 	void PhysicsBody::GetConvex(
 		int32 B, int32 I,
-		FVector3f const*& OutVertices, int32& OutVerticesCount,
-		int32 const*& OutIndices, int32& OutIndicesCount,
+		TArrayView<const FVector3f>& OutVertices, TArrayView<const int32>& OutIndices,
 		FTransform3f& OutTransform) const
 	{
 		check(B >= 0 && B < Bodies.Num());
 		check(I >= 0 && I < Bodies[B].Convex.Num());
 
-		OutVertices = Bodies[B].Convex[I].Vertices.GetData();
-		OutVerticesCount = Bodies[B].Convex[I].Vertices.Num();
+		OutVertices = TArrayView<const FVector3f>(Bodies[B].Convex[I].Vertices.GetData(), Bodies[B].Convex[I].Vertices.Num());
+		OutIndices = TArrayView<const int32>(Bodies[B].Convex[I].Indices.GetData(), Bodies[B].Convex[I].Indices.Num());
+		OutTransform = Bodies[B].Convex[I].Transform;
+	}
+
+	void PhysicsBody::GetConvexMeshView(int32 B, int32 I,
+		TArrayView<FVector3f>& OutVerticesView, TArrayView<int32>& OutIndicesView)
+	{
+		check(B >= 0 && B < Bodies.Num());
+		check(I >= 0 && I < Bodies[B].Convex.Num());
 
-		OutIndices = Bodies[B].Convex[I].Indices.GetData();
-		OutIndicesCount = Bodies[B].Convex[I].Indices.Num();
+		OutVerticesView = TArrayView<FVector3f>(Bodies[B].Convex[I].Vertices.GetData(), Bodies[B].Convex[I].Vertices.Num());
+		OutIndicesView = TArrayView<int32>(Bodies[B].Convex[I].Indices.GetData(), Bodies[B].Convex[I].Indices.Num());
+	}
+
+	void PhysicsBody::GetConvexTransform(int32 B, int32 I, FTransform3f& OutTransform) const
+	{
+		check(B >= 0 && B < Bodies.Num());
+		check(I >= 0 && I < Bodies[B].Convex.Num());
 
 		OutTransform = Bodies[B].Convex[I].Transform;
 	}
diff --git a/Engine/Plugins/Experimental/Mutable/Source/MutableRuntime/Public/MuR/PhysicsBody.h b/Engine/Plugins/Experimental/Mutable/Source/MutableRuntime/Public/MuR/PhysicsBody.h
--- a/Engine/Plugins/Experimental/Mutable/Source/MutableRuntime/Public/MuR/PhysicsBody.h
+++ b/Engine/Plugins/Experimental/Mutable/Source/MutableRuntime/Public/MuR/PhysicsBody.h
@@ -352,10 +352,12 @@ namespace mu
 				int32 B, int32 I, 
 				FVector3f Position, FQuat4f Orientation, FVector3f Size );
 
-		void SetConvex( 
+		void SetConvexMesh( 
 				int32 B, int32 I,
-				const FVector3f* Vertices, int32 VerticesCount, 
-				const int32* Indices, int32 IndicesCount, 
+				TArrayView<const FVector3f> Vertices, TArrayView<const int32> Indices);
+
+		void SetConvexTransform( 
+				int32 B, int32 I, 
 				const FTransform3f& Transform );
 
 		void SetSphyl( 
@@ -390,8 +392,15 @@ namespace mu
 
 		void GetConvex( 
 				int32 B, int32 I, 
-				FVector3f const*& OutVertices, int32& OutVerticesCount, 
-				int32 const*& OutIndices, int32& OutIndicesCount, FTransform3f& OutTransform ) const;
+				TArrayView<const FVector3f>& OutVertices, TArrayView<const int32>& OutIndices, FTransform3f& OutTransform ) const;
+
+		void GetConvexMeshView(
+				int32 B, int32 I, 
+			    TArrayView<FVector3f>& OutVerticesView, TArrayView<int32>& OutIndicesView);
+
+		void GetConvexTransform(
+				int32 B, int32 I,
+				FTransform3f& OutTransform) const;
 
 		void GetSphyl( 
 				int32 B, int32 I, 
