[mutable] Additional workarounds the static analysis false positives.
#rnx

diff --git a/Engine/Plugins/Experimental/Mutable/Source/MutableRuntime/Private/MuR/OpImageLayer.h b/Engine/Plugins/Experimental/Mutable/Source/MutableRuntime/Private/MuR/OpImageLayer.h
--- a/Engine/Plugins/Experimental/Mutable/Source/MutableRuntime/Private/MuR/OpImageLayer.h
+++ b/Engine/Plugins/Experimental/Mutable/Source/MutableRuntime/Private/MuR/OpImageLayer.h
@@ -666,9 +666,11 @@ namespace mu
 			PixelCount = pBase->CalculatePixelCount();
 		}
 
+		int32 UnblendedChannels = BASE_CHANNEL_STRIDE - CHANNELS_TO_BLEND;
+
 		ParallelFor(PixelCount,
 			[
-				pBaseBuf, pBlendedBuf, pDestBuf
+				pBaseBuf, pBlendedBuf, pDestBuf, UnblendedChannels
 			] (uint32 i)
 			{
 				for (int c = 0; c < CHANNELS_TO_BLEND; ++c)
@@ -688,7 +690,7 @@ namespace mu
 
 				// Copy the unblended channels
 				// \TODO: unnecessary when doing it in-place?
-				for (uint32 c = 0; c < BASE_CHANNEL_STRIDE - CHANNELS_TO_BLEND; ++c)
+				for (int32 c = 0; c < UnblendedChannels; ++c)
 				{
 					pDestBuf[BASE_CHANNEL_STRIDE * i + CHANNELS_TO_BLEND + c] = pBaseBuf[BASE_CHANNEL_STRIDE * i + CHANNELS_TO_BLEND + c];
 				}
@@ -827,7 +829,7 @@ namespace mu
 
 			ParallelFor(PixelCount,
 				[
-					pBaseBuf, pBlendedBuf, pMaskBuf, pDestBuf
+					pBaseBuf, pBlendedBuf, pMaskBuf, pDestBuf, UnblendedChannels
 				] (uint32 i)
 				{
 					unsigned mask = pMaskBuf[i];
@@ -847,7 +849,7 @@ namespace mu
 					}
 					// Copy the unblended channels
 					// \TODO: unnecessary when doing it in-place?
-					for (int c = 0; c < BASE_CHANNEL_STRIDE - CHANNELS_TO_BLEND; ++c)
+					for (int c = 0; c < UnblendedChannels; ++c)
 					{
 						pDestBuf[BASE_CHANNEL_STRIDE * i + CHANNELS_TO_BLEND + c] = pBaseBuf[BASE_CHANNEL_STRIDE * i + CHANNELS_TO_BLEND + c];
 					}
@@ -1034,9 +1036,11 @@ namespace mu
 			PixelCount = pBase->CalculatePixelCount();
 		}
 
+		int32 UnblendedChannels = BASE_CHANNEL_STRIDE - CHANNELS_TO_BLEND;
+
 		ParallelFor(PixelCount,
 			[
-				pBaseBuf, pBlendedBuf, pDestBuf, pMaskBuf
+				pBaseBuf, pBlendedBuf, pDestBuf, pMaskBuf, UnblendedChannels
 			] (uint32 i)
 			{
 				uint32 mask = pMaskBuf[BLENDED_CHANNEL_STRIDE * i];
@@ -1056,7 +1060,7 @@ namespace mu
 				}
 				// Copy the unblended channels
 				// \TODO: unnecessary when doing it in-place?
-				for (int c = 0; c < BASE_CHANNEL_STRIDE - CHANNELS_TO_BLEND; ++c)
+				for (int c = 0; c < UnblendedChannels; ++c)
 				{
 					pDestBuf[BASE_CHANNEL_STRIDE * i + CHANNELS_TO_BLEND + c] = pBaseBuf[BASE_CHANNEL_STRIDE * i + CHANNELS_TO_BLEND + c];
 				}
