[mutable] Remove obsolete DGGUI functionality

#rnx

diff --git a/Engine/Plugins/Experimental/Mutable/Source/CustomizableObject/Private/MuCO/CustomizableObjectDGGUI.cpp b/Engine/Plugins/Experimental/Mutable/Source/CustomizableObject/Private/MuCO/CustomizableObjectDGGUI.cpp
--- a/Engine/Plugins/Experimental/Mutable/Source/CustomizableObject/Private/MuCO/CustomizableObjectDGGUI.cpp
+++ b/Engine/Plugins/Experimental/Mutable/Source/CustomizableObject/Private/MuCO/CustomizableObjectDGGUI.cpp
@@ -38,7 +38,6 @@ void UDGGUI::OpenDGGUI(const int32 SlotID, UCustomizableSkeletalComponent* Selec
 			if (WDGUI)
 			{
 				WDGUI->SetCustomizableSkeletalComponent(SelectedCustomizableSkeletalComponent);
-				SelectedCustomizableSkeletalComponent->PreUpdateDelegate.BindUObject(WDGUI, &UDGGUI::CustomizableSkeletalMeshPreUpdate);
 				WDGUI->AddToViewport();
 				Player->SetShowMouseCursor(true);
 			}
@@ -57,8 +56,6 @@ bool UDGGUI::CloseExistingDGGUI(const UWorld* CurrentWorld)
 		{
 			if (UCustomizableSkeletalComponent* CustomizableComponent = PreviousGUI->GetCustomizableSkeletalComponent())
 			{
-				CustomizableComponent->PreUpdateDelegate.Unbind();
-				CustomizableComponent->UpdatedDelegate.Unbind();
 				PreviousGUI->SetCustomizableSkeletalComponent(nullptr);
 				bClosing = true;
 			}
@@ -76,85 +73,3 @@ bool UDGGUI::CloseExistingDGGUI(const UWorld* CurrentWorld)
 #endif // WITH_EDITOR
 	return false;
 }
-
-
-// This is all based on UMeshCosmeticsVariance_ApplyParameters::PostCustomizationFixup( , but I think this should be set from within mutable, like how Physics assets are set up without user intervention
-void UDGGUI::CustomizableSkeletalMeshPreUpdate(UCustomizableSkeletalComponent* Component, USkeletalMesh* NextMesh)
-{
-#if WITH_EDITOR
-	if (Component && NextMesh)
-	{
-		USkeletalMeshComponent* Parent = Cast<USkeletalMeshComponent>(Component->GetAttachParent());
-		USkeleton* NextSkeleton = NextMesh->GetSkeleton();
-		if (Parent && NextSkeleton)
-		{
-			if (USkeletalMesh* ParentSkeletalMesh = Cast<USkeletalMesh>(Parent->GetSkinnedAsset()))
-			{
-				if (USkeleton* ParentSkeleton = ParentSkeletalMesh->GetSkeleton())
-				{
-					UAnimInstance* AnimInstance = Parent->GetAnimInstance();
-					UClass* AnimClass = Parent->GetAnimClass();
-					IAnimClassInterface* Animation = Cast<IAnimClassInterface>(AnimClass);
-					if (AnimInstance && Animation)
-					{
-						if (USkeleton* AnimationSkeleton = Animation->GetTargetSkeleton())
-						{
-							if (AnimationSkeleton != NextSkeleton && NextSkeleton->GetCompatibleSkeletons().Find(AnimationSkeleton) == INDEX_NONE)
-							{
-								NextSkeleton->AddCompatibleSkeleton(AnimationSkeleton);
-								LastAnimationClass = AnimClass;
-								Component->UpdatedDelegate.BindUObject(this, &UDGGUI::CustomizableSkeletalMeshUpdated);
-							}
-						}
-					}
-				}
-			}
-		}
-	}
-#endif // WITH_EDITOR
-}
-
-
-void UDGGUI::CustomizableSkeletalMeshUpdated()
-{
-#if WITH_EDITOR
-	if (UCustomizableSkeletalComponent* CustomizableComponent = GetCustomizableSkeletalComponent())
-	{
-		if (USkeletalMeshComponent* CustomizedComponent = Cast<USkeletalMeshComponent>(CustomizableComponent->GetAttachParent()))
-		{
-			CustomizedComponent->SetAnimInstanceClass(LastAnimationClass);
-
-			if (UAnimInstance* TargetAnimInstance = CustomizedComponent->GetAnimInstance())
-			{
-				const UCustomizableObjectInstance* Instance = CustomizableComponent->CustomizableObjectInstance;
-
-				const int32 ComponentIndex = CustomizableComponent->ComponentIndex;
-
-				IAnimClassInterface* AnimInstanceClassInterface = IAnimClassInterface::GetFromClass(LastAnimationClass);
-
-				USkeleton* ProportionalSkel = AnimInstanceClassInterface->GetTargetSkeleton();
-				check(ProportionalSkel);
-				const auto LinkAnimClassBySlot = [&TargetAnimInstance, ProportionalSkel](int32 SlotIndex, TSubclassOf<UAnimInstance> SlottedInstanceClass)
-				{
-					if (SlottedInstanceClass)
-					{
-						IAnimClassInterface* SlottedClassInterface = IAnimClassInterface::GetFromClass(SlottedInstanceClass);
-						USkeleton* SlottedSkeleton = SlottedClassInterface->GetTargetSkeleton();
-
-						if (SlottedSkeleton != ProportionalSkel && SlottedSkeleton->GetCompatibleSkeletons().Find(ProportionalSkel) == INDEX_NONE)
-						{
-							SlottedSkeleton->AddCompatibleSkeleton(ProportionalSkel);
-						}
-
-						TargetAnimInstance->LinkAnimGraphByTag(*FString::Printf(TEXT("Part%d"), SlotIndex), SlottedInstanceClass);
-					}
-				};
-
-				using ForEachDelType = FEachComponentAnimInstanceClassNativeDelegate;
-				Instance->ForEachAnimInstance(ComponentIndex, ForEachDelType::CreateLambda(LinkAnimClassBySlot));
-			}
-		}
-		//CustomizableComponent->UpdatedDelegate.Unbind();
-	}
-#endif // WITH_EDITOR
-}
diff --git a/Engine/Plugins/Experimental/Mutable/Source/CustomizableObject/Public/MuCO/CustomizableObjectDGGUI.h b/Engine/Plugins/Experimental/Mutable/Source/CustomizableObject/Public/MuCO/CustomizableObjectDGGUI.h
--- a/Engine/Plugins/Experimental/Mutable/Source/CustomizableObject/Public/MuCO/CustomizableObjectDGGUI.h
+++ b/Engine/Plugins/Experimental/Mutable/Source/CustomizableObject/Public/MuCO/CustomizableObjectDGGUI.h
@@ -28,8 +28,4 @@ public:
 
 	static void OpenDGGUI(const int32 SlotID, UCustomizableSkeletalComponent* SelectedCustomizableSkeletalComponent, const UWorld* CurrentWorld, const int32 PlayerIndex = 0);
 	static bool CloseExistingDGGUI(const UWorld* CurrentWorld);
-private:
-	void CustomizableSkeletalMeshPreUpdate(class UCustomizableSkeletalComponent* Component, class USkeletalMesh* NextMesh);
-	void CustomizableSkeletalMeshUpdated();
-	TSubclassOf<UAnimInstance> LastAnimationClass;
 };
