[mutable] Fix right click mouse button scroll on PinViewer
Fix usability of instance properties scrollbar
Fix right click mouse button scroll on Customizable Object Details states table
Refactor blocks selectors to setup on reconstruct instead of keeping them alive between refreshes

#rnx

diff --git a/Engine/Plugins/Experimental/Mutable/Source/CustomizableObjectEditor/Classes/MuCOE/Nodes/CustomizableObjectNodeEditMaterial.h b/Engine/Plugins/Experimental/Mutable/Source/CustomizableObjectEditor/Classes/MuCOE/Nodes/CustomizableObjectNodeEditMaterial.h
--- a/Engine/Plugins/Experimental/Mutable/Source/CustomizableObjectEditor/Classes/MuCOE/Nodes/CustomizableObjectNodeEditMaterial.h
+++ b/Engine/Plugins/Experimental/Mutable/Source/CustomizableObjectEditor/Classes/MuCOE/Nodes/CustomizableObjectNodeEditMaterial.h
@@ -11,6 +11,7 @@
 #include "Internationalization/Text.h"
 #include "Math/Color.h"
 #include "Misc/Guid.h"
+#include "MuCOE/CustomizableObjectEditor_Deprecated.h"
 #include "MuCOE/Nodes/CustomizableObjectNode.h"
 #include "MuCOE/Nodes/CustomizableObjectNodeEditLayoutBlocks.h"
 #include "MuCOE/Nodes/CustomizableObjectNodeUseMaterial.h"
@@ -23,7 +24,6 @@ class FCustomizableObjectNodeParentedMaterial;
 class UCustomizableObject;
 class UCustomizableObjectNodeRemapPins;
 class UObject;
-struct FCustomizableObjectNodeEditMaterialImage;
 
 
 /** Additional data for the Parent Texture Parameter to edit pin. */
diff --git a/Engine/Plugins/Experimental/Mutable/Source/CustomizableObjectEditor/Private/MuCOE/CustomizableObjectDetails.cpp b/Engine/Plugins/Experimental/Mutable/Source/CustomizableObjectEditor/Private/MuCOE/CustomizableObjectDetails.cpp
--- a/Engine/Plugins/Experimental/Mutable/Source/CustomizableObjectEditor/Private/MuCOE/CustomizableObjectDetails.cpp
+++ b/Engine/Plugins/Experimental/Mutable/Source/CustomizableObjectEditor/Private/MuCOE/CustomizableObjectDetails.cpp
@@ -131,7 +131,6 @@ void FCustomizableObjectDetails::CustomizeDetails( IDetailLayoutBuilder& DetailB
 			//SNew( SFilterableDetail, LOCTEXT( "States", "States" ), &StatesCategory )
 			//[
 				SAssignNew(StatesTree, STreeView<TSharedPtr< FStateDetailsNode > >)
-					//.Visibility(EVisibility::Collapsed)
 					.SelectionMode(ESelectionMode::Single)
 					.TreeItemsSource( &RootTreeItems )
 					// Called to child items for any given parent item
@@ -175,6 +174,8 @@ void FCustomizableObjectDetails::CustomizeDetails( IDetailLayoutBuilder& DetailB
 //		]
 	];
 
+	StatesTree->SetIsRightClickScrollingEnabled(false);
+
 	if ( CustomObject &&  CustomObject->GetModel() )
 	{
 		uint32 numElements = CustomObject->GetModel()->GetStateCount();
diff --git a/Engine/Plugins/Experimental/Mutable/Source/CustomizableObjectEditor/Private/MuCOE/CustomizableObjectEditor.cpp b/Engine/Plugins/Experimental/Mutable/Source/CustomizableObjectEditor/Private/MuCOE/CustomizableObjectEditor.cpp
--- a/Engine/Plugins/Experimental/Mutable/Source/CustomizableObjectEditor/Private/MuCOE/CustomizableObjectEditor.cpp
+++ b/Engine/Plugins/Experimental/Mutable/Source/CustomizableObjectEditor/Private/MuCOE/CustomizableObjectEditor.cpp
@@ -301,19 +301,13 @@ void FCustomizableObjectEditor::InitCustomizableObjectEditor(const EToolkitMode:
 	DetailsViewArgs.bAllowSearch = false;
 	//DetailsViewArgs.bShowActorLabel = false;
 	DetailsViewArgs.bShowObjectLabel = false;
+	DetailsViewArgs.bShowScrollBar = false;
 
 	CustomizableObjectDetailsView = PropPlugin.CreateDetailView( DetailsViewArgs );
 	CustomizableObjectDetailsView->SetObject(CustomizableObject);
 
-	CustomizableInstanceDetailsView = PropPlugin.CreateDetailView( DetailsViewArgs );	
-
-	FDetailsViewArgs GraphDetailsViewArgs;
-	GraphDetailsViewArgs.NotifyHook = this;
-	GraphDetailsViewArgs.NameAreaSettings = FDetailsViewArgs::ENameAreaSettings::HideNameArea;
-	GraphDetailsViewArgs.bAllowSearch = false;
-	//GraphDetailsViewArgs.bShowActorLabel = false;
-	GraphDetailsViewArgs.bShowObjectLabel = false;
-	GraphNodeDetailsView = PropPlugin.CreateDetailView( GraphDetailsViewArgs );
+	CustomizableInstanceDetailsView = PropPlugin.CreateDetailView( DetailsViewArgs );
+	GraphNodeDetailsView = PropPlugin.CreateDetailView(DetailsViewArgs);
 
 	Viewport = SNew(SCustomizableObjectEditorViewportTabBody)
 		.CustomizableObjectEditor(SharedThis(this));
@@ -596,11 +590,18 @@ TSharedRef<SDockTab> FCustomizableObjectEditor::SpawnTab_ObjectProperties( const
 {
 	check( Args.GetTabId() == ObjectPropertiesTabId );
 
+	TSharedRef<SScrollBox> ScrollBox = SNew(SScrollBox)
+		+ SScrollBox::Slot()
+		[
+			CustomizableObjectDetailsView.ToSharedRef()
+		];
+
+	ScrollBox->SetScrollBarRightClickDragAllowed(true);
+
 	TSharedRef<SDockTab> DockTab = SNew(SDockTab)
 		.Label( FText::FromString( GetTabPrefix() + LOCTEXT( "CustomizableObjectProperties_TabTitle", "Object Properties" ).ToString() ) )
 		[
-			CustomizableObjectDetailsView.ToSharedRef()
-			//CustomizableObjectProperties.ToSharedRef()
+			ScrollBox
 		];
 
 	DockTab->SetTabIcon(FCustomizableObjectEditorStyle::Get().GetBrush("CustomizableObjectEditor.Tabs.CustomizableObjectProperties"));
@@ -613,10 +614,18 @@ TSharedRef<SDockTab> FCustomizableObjectEditor::SpawnTab_InstanceProperties( con
 {
 	check( Args.GetTabId() == InstancePropertiesTabId );
 
+	TSharedRef<SScrollBox> ScrollBox = SNew(SScrollBox)
+		+ SScrollBox::Slot()
+		[
+			CustomizableInstanceDetailsView.ToSharedRef()
+		];
+
+	ScrollBox->SetScrollBarRightClickDragAllowed(true);
+
 	TSharedRef<SDockTab> DockTab = SNew(SDockTab)
 		.Label( FText::FromString( GetTabPrefix() + LOCTEXT( "CustomizableInstanceProperties_TabTitle", "Preview Inst. Properties" ).ToString() ) )
 		[
-			CustomizableInstanceDetailsView.ToSharedRef()
+			ScrollBox
 		];
 
 	DockTab->SetTabIcon(FCustomizableObjectEditorStyle::Get().GetBrush("CustomizableObjectEditor.Tabs.CustomizableInstanceProperties"));
@@ -746,44 +755,19 @@ TSharedRef<SDockTab> FCustomizableObjectEditor::SpawnTab_GraphNodeProperties( co
 {
 	check( Args.GetTabId().TabType == GraphNodePropertiesTabId );
 
-	TSharedRef<SScrollBar> ScrollBar = SNew(SScrollBar);
+	TSharedRef<SScrollBox> ScrollBox = SNew(SScrollBox)
+		+ SScrollBox::Slot()
+		[
+			GraphNodeDetailsView.ToSharedRef()
+		];
 
-	NodeDetailsSplitter = SNew( SSplitter )
-			.Orientation(Orient_Vertical)
-			+SSplitter::Slot()
-			.Value(0.5f)
-			[
-				SNew(SHorizontalBox)
-				+ SHorizontalBox::Slot()
-				.FillWidth(1)
-				[
-					SNew(SScrollBox)
-					.ExternalScrollbar(ScrollBar)
-					+ SScrollBox::Slot()
-					[
-						SNew(SVerticalBox)
-						+ SVerticalBox::Slot().AutoHeight()
-						[
-							GraphNodeDetailsView.ToSharedRef()
-						]
-					]
-				]
-				+ SHorizontalBox::Slot()
-				.AutoWidth()
-				[
-					SNew(SBox)
-					.WidthOverride(FOptionalSize(16))
-					[
-						ScrollBar
-					]
-				]
-			];
+	ScrollBox->SetScrollBarRightClickDragAllowed(true);
 
 	TSharedRef<SDockTab> DockTab = SNew(SDockTab)
 		.Label(FText::FromString(GetTabPrefix() + LOCTEXT("Graph Node Properties", "Node Properties").ToString()))
 		.TabColorScale(GetTabColorScale())
 		[
-			NodeDetailsSplitter.ToSharedRef()
+			ScrollBox
 		];
 
 	DockTab->SetTabIcon(FCustomizableObjectEditorStyle::Get().GetBrush("CustomizableObjectEditor.Tabs.NodeProperties"));
@@ -1694,78 +1678,6 @@ void FCustomizableObjectEditor::OnSelectedGraphNodesChanged(const FGraphPanelSel
 		SelectedProjectorParameterNode = nullptr;
 	}
 
-	// Special details. Usually widgets that you want to be able to scale
-	if ( NodeDetailsSplitter.IsValid() )
-	{
-		UObject* FirstNode = 0;		
-		if ( Objects.Num() )
-		{
-			FirstNode = Objects[0];
-		}
-
-		TSharedPtr<SWidget> ChildWidget;
-
-		// Set the size of the child widget to show properly the node widgets
-		FVector2D SplitWeights = { 0.15f,0.85f };
-		
-		if ( UCustomizableObjectNodeEditMaterial* CustomizableObjectNodeEditMaterial = Cast<UCustomizableObjectNodeEditMaterial>(FirstNode) )
-		{
-			if (!LayoutBlocksSelector.IsValid())
-			{
-				LayoutBlocksSelector = SNew(SCustomizableObjectNodeLayoutBlocksSelector)
-					.CustomizableObjectEditor(SharedThis(this));
-			}
-
-			LayoutBlocksSelector->SetSelectedNode(CustomizableObjectNodeEditMaterial);
-			SplitWeights = { 0.25f,0.75f };
-
-			ChildWidget = LayoutBlocksSelector;
-		}
-		else if (UCustomizableObjectNodeRemoveMeshBlocks* CustomizableObjectNodeRemoveMeshBlocks = Cast< UCustomizableObjectNodeRemoveMeshBlocks>(FirstNode))
-		{
-			if (!LayoutBlocksSelector.IsValid())
-			{
-				LayoutBlocksSelector = SNew(SCustomizableObjectNodeLayoutBlocksSelector)
-					.CustomizableObjectEditor(SharedThis(this));
-			}
-
-			LayoutBlocksSelector->SetSelectedNode(CustomizableObjectNodeRemoveMeshBlocks);
-			SplitWeights = { 0.25f,0.75f };
-
-			ChildWidget = LayoutBlocksSelector;
-		}
-
-		if ( ChildWidget.IsValid() )
-		{
-			if ( NodeDetailsSplitter->GetChildren()->Num()<=1 )
-			{
-				NodeDetailsSplitter->AddSlot();
-			}
-
-			NodeDetailsSplitter->SlotAt(0).SetSizeValue(SplitWeights.X);
-			NodeDetailsSplitter->SlotAt(1).SetSizeValue(SplitWeights.Y);
-
-			NodeDetailsSplitter->SlotAt(1)
-				[
-					SNew(SBorder)
-					.BorderImage( FAppStyle::GetBrush( TEXT("Graph.TitleBackground") ) )
-					.HAlign(HAlign_Fill)
-					[
-						ChildWidget.ToSharedRef()
-					]
-				];
-		}
-		else
-		{
-			NodeDetailsSplitter->SlotAt(0).SetSizeValue(0.5);
-
-			if ( NodeDetailsSplitter->GetChildren()->Num()>1 )
-			{
-				NodeDetailsSplitter->RemoveAt(1);
-			}
-		}
-	}
-
 	if (!ManagingProjector)
 	{
 		SelectedGraphNodesChanged = true;
diff --git a/Engine/Plugins/Experimental/Mutable/Source/CustomizableObjectEditor/Private/MuCOE/CustomizableObjectEditor.h b/Engine/Plugins/Experimental/Mutable/Source/CustomizableObjectEditor/Private/MuCOE/CustomizableObjectEditor.h
--- a/Engine/Plugins/Experimental/Mutable/Source/CustomizableObjectEditor/Private/MuCOE/CustomizableObjectEditor.h
+++ b/Engine/Plugins/Experimental/Mutable/Source/CustomizableObjectEditor/Private/MuCOE/CustomizableObjectEditor.h
@@ -433,9 +433,6 @@ private:
 	TSharedPtr<SGraphEditor> GraphEditor;
 	TSharedPtr<class IDetailsView> GraphNodeDetailsView;
 
-	/** This splitter can be used to add more widgets outside the property editor of the currently selected node. */
-	TSharedPtr<class SSplitter> NodeDetailsSplitter;
-
 	/** This widget shows the layout blocks and they can be edited. It is kept alive for some problems re-registering toolbar commands.*/
 	TSharedPtr<class SCustomizableObjectNodeLayoutBlocksEditor> LayoutBlocksEditor;
 	TSharedPtr<class SCustomizableObjectNodeLayoutBlocksSelector> LayoutBlocksSelector;
diff --git a/Engine/Plugins/Experimental/Mutable/Source/CustomizableObjectEditor/Private/MuCOE/CustomizableObjectEditorModule.cpp b/Engine/Plugins/Experimental/Mutable/Source/CustomizableObjectEditor/Private/MuCOE/CustomizableObjectEditorModule.cpp
--- a/Engine/Plugins/Experimental/Mutable/Source/CustomizableObjectEditor/Private/MuCOE/CustomizableObjectEditorModule.cpp
+++ b/Engine/Plugins/Experimental/Mutable/Source/CustomizableObjectEditor/Private/MuCOE/CustomizableObjectEditorModule.cpp
@@ -51,6 +51,7 @@
 #include "MuCOE/Nodes/CustomizableObjectNodeCopyMaterial.h"
 #include "MuCOE/Nodes/CustomizableObjectNodeCopyMaterialDetails.h"
 #include "MuCOE/Nodes/CustomizableObjectNodeEditMaterial.h"
+#include "MuCOE/Nodes/CustomizableObjectNodeEditMaterialDetails.h"
 #include "MuCOE/Nodes/CustomizableObjectNodeEditMaterialBase.h"
 #include "MuCOE/Nodes/CustomizableObjectNodeEditMaterialBaseDetails.h"
 #include "MuCOE/Nodes/CustomizableObjectNodeExtendMaterial.h"
@@ -78,6 +79,8 @@
 #include "MuCOE/Nodes/CustomizableObjectNodeProjectorParameter.h"
 #include "MuCOE/Nodes/CustomizableObjectNodeProjectorParameterDetails.h"
 #include "MuCOE/Nodes/CustomizableObjectNodeGroupProjectorParameterDetails.h"
+#include "MuCOE/Nodes/CustomizableObjectNodeRemoveMeshBlocks.h"
+#include "MuCOE/Nodes/CustomizableObjectNodeRemoveMeshBlocksDetails.h"
 #include "MuCOE/Nodes/CustomizableObjectNodeSkeletalMesh.h"
 #include "MuCOE/Nodes/CustomizableObjectNodeSkeletalMeshDetails.h"
 #include "MuCOE/Nodes/CustomizableObjectNodeStaticMesh.h"
@@ -197,9 +200,10 @@ void FCustomizableObjectEditorModule::StartupModule()
 	RegisterCustomDetails(PropertyModule, UCustomizableObject::StaticClass(), FOnGetDetailCustomizationInstance::CreateStatic(&FCustomizableObjectDetails::MakeInstance));
 	RegisterCustomDetails(PropertyModule, UCustomizableObjectInstance::StaticClass(), FOnGetDetailCustomizationInstance::CreateStatic(&FCustomizableInstanceDetails::MakeInstance));
 	RegisterCustomDetails(PropertyModule, UCustomizableObjectNodeLayoutBlocks::StaticClass(), FOnGetDetailCustomizationInstance::CreateStatic(&FCustomizableObjectNodeLayoutBlocksDetails::MakeInstance));
+	RegisterCustomDetails(PropertyModule, UCustomizableObjectNodeEditMaterial::StaticClass(), FOnGetDetailCustomizationInstance::CreateStatic(&FCustomizableObjectNodeEditMaterialDetails::MakeInstance));
+	RegisterCustomDetails(PropertyModule, UCustomizableObjectNodeRemoveMeshBlocks::StaticClass(), FOnGetDetailCustomizationInstance::CreateStatic(&FCustomizableObjectNodeRemoveMeshBlocksDetails::MakeInstance));
 	RegisterCustomDetails(PropertyModule, UCustomizableObjectNodeExtendMaterial::StaticClass(), FOnGetDetailCustomizationInstance::CreateStatic(&FCustomizableObjectNodeParentedMaterialDetails::MakeInstance));
 	RegisterCustomDetails(PropertyModule, UCustomizableObjectNodeEditMaterialBase::StaticClass(), FOnGetDetailCustomizationInstance::CreateStatic(&FCustomizableObjectNodeEditMaterialBaseDetails::MakeInstance));
-	RegisterCustomDetails(PropertyModule, UCustomizableObjectNodeEditMaterial::StaticClass(), FOnGetDetailCustomizationInstance::CreateStatic(&FCustomizableObjectNodePinViewerDetails::MakeInstance));
 	RegisterCustomDetails(PropertyModule, UCustomizableObjectNodeMorphMaterial::StaticClass(), FOnGetDetailCustomizationInstance::CreateStatic(&FCustomizableObjectNodeMorphMaterialDetails::MakeInstance));
 	RegisterCustomDetails(PropertyModule, UCustomizableObjectNodeObject::StaticClass(), FOnGetDetailCustomizationInstance::CreateStatic(&FCustomizableObjectNodeObjectDetails::MakeInstance));
 	RegisterCustomDetails(PropertyModule, UCustomizableObjectNodeObjectGroup::StaticClass(), FOnGetDetailCustomizationInstance::CreateStatic(&FCustomizableObjectNodeObjectGroupDetails::MakeInstance));
diff --git a/Engine/Plugins/Experimental/Mutable/Source/CustomizableObjectEditor/Private/MuCOE/Nodes/CustomizableObjectNodeEditMaterial.cpp b/Engine/Plugins/Experimental/Mutable/Source/CustomizableObjectEditor/Private/MuCOE/Nodes/CustomizableObjectNodeEditMaterial.cpp
--- a/Engine/Plugins/Experimental/Mutable/Source/CustomizableObjectEditor/Private/MuCOE/Nodes/CustomizableObjectNodeEditMaterial.cpp
+++ b/Engine/Plugins/Experimental/Mutable/Source/CustomizableObjectEditor/Private/MuCOE/Nodes/CustomizableObjectNodeEditMaterial.cpp
@@ -9,7 +9,6 @@
 #include "MaterialTypes.h"
 #include "MuCO/CustomizableObject.h"
 #include "MuCO/CustomizableObjectCustomVersion.h"
-#include "MuCOE/CustomizableObjectEditor_Deprecated.h"
 #include "MuCOE/CustomizableObjectLayout.h"
 #include "MuCOE/EdGraphSchema_CustomizableObject.h"
 #include "MuCOE/Nodes/CustomizableObjectNodeMaterial.h"
diff --git a/Engine/Plugins/Experimental/Mutable/Source/CustomizableObjectEditor/Private/MuCOE/Nodes/CustomizableObjectNodeEditMaterialDetails.cpp b/Engine/Plugins/Experimental/Mutable/Source/CustomizableObjectEditor/Private/MuCOE/Nodes/CustomizableObjectNodeEditMaterialDetails.cpp
new file mode 100644
--- /dev/null
+++ b/Engine/Plugins/Experimental/Mutable/Source/CustomizableObjectEditor/Private/MuCOE/Nodes/CustomizableObjectNodeEditMaterialDetails.cpp
@@ -0,0 +1,80 @@
+// Copyright Epic Games, Inc. All Rights Reserved.
+
+#include "MuCOE/Nodes/CustomizableObjectNodeEditMaterialDetails.h"
+
+#include "DetailCategoryBuilder.h"
+#include "DetailLayoutBuilder.h"
+#include "DetailWidgetRow.h"
+#include "IDetailsView.h"
+#include "Internationalization/Internationalization.h"
+#include "Math/IntPoint.h"
+#include "Math/IntRect.h"
+#include "MuCOE/CustomizableObjectLayout.h"
+#include "MuCOE/Nodes/CustomizableObjectNodeEditMaterial.h"
+#include "MuCOE/SCustomizableObjectNodeLayoutBlocksEditor.h"
+#include "MuCOE/SCustomizableObjectNodeLayoutBlocksSelector.h"
+#include "MuCOE/PinViewer/SPinViewer.h"
+#include "Templates/Casts.h"
+#include "UObject/ObjectPtr.h"
+#include "UObject/WeakObjectPtr.h"
+#include "UObject/WeakObjectPtrTemplates.h"
+#include "Widgets/DeclarativeSyntaxSupport.h"
+#include "Widgets/Layout/SBox.h"
+#include "Widgets/Text/STextBlock.h"
+
+class FString;
+
+#define LOCTEXT_NAMESPACE "CustomizableObjectDetails"
+
+
+TSharedRef<IDetailCustomization> FCustomizableObjectNodeEditMaterialDetails::MakeInstance()
+{
+	return MakeShareable( new FCustomizableObjectNodeEditMaterialDetails);
+}
+
+
+void FCustomizableObjectNodeEditMaterialDetails::CustomizeDetails( IDetailLayoutBuilder& DetailBuilder )
+{
+	Node = nullptr;
+	const IDetailsView* DetailsView = DetailBuilder.GetDetailsView();
+	if (DetailsView->GetSelectedObjects().Num())
+	{
+		Node = Cast<UCustomizableObjectNodeEditMaterial>(DetailsView->GetSelectedObjects()[0].Get());
+	}
+
+	IDetailCategoryBuilder& BlocksCategory = DetailBuilder.EditCategory( "Blocks" );
+
+	if (Node)
+	{
+		UCustomizableObjectNodeEditLayoutBlocks* NodeEditLayout = Cast<UCustomizableObjectNodeEditLayoutBlocks>(Node);
+		if (NodeEditLayout)
+		{
+			// Add blocks selector
+			LayoutBlocksSelector = SNew(SCustomizableObjectNodeLayoutBlocksSelector);
+
+			BlocksCategory.AddCustomRow(LOCTEXT("BlocksDetails_BlockInstructions", "BlockInstructions"))
+			[
+				SNew(SBox)
+				.HeightOverride(700.0f)
+				.WidthOverride(700.0f)
+				[
+					LayoutBlocksSelector.ToSharedRef()
+				]
+			];
+
+			LayoutBlocksSelector->SetSelectedNode(NodeEditLayout);
+		}
+	}
+	else
+	{
+		BlocksCategory.AddCustomRow( LOCTEXT("BlocksDetails_NodeNotFound", "NodeNotFound") )
+		[
+			SNew( STextBlock )
+			.Text( LOCTEXT( "Node not found", "Node not found" ) )
+		];
+	}
+
+	PinViewerAttachToDetailCustomization(DetailBuilder);
+}
+
+#undef LOCTEXT_NAMESPACE
diff --git a/Engine/Plugins/Experimental/Mutable/Source/CustomizableObjectEditor/Private/MuCOE/Nodes/CustomizableObjectNodeEditMaterialDetails.h b/Engine/Plugins/Experimental/Mutable/Source/CustomizableObjectEditor/Private/MuCOE/Nodes/CustomizableObjectNodeEditMaterialDetails.h
new file mode 100644
--- /dev/null
+++ b/Engine/Plugins/Experimental/Mutable/Source/CustomizableObjectEditor/Private/MuCOE/Nodes/CustomizableObjectNodeEditMaterialDetails.h
@@ -0,0 +1,35 @@
+// Copyright Epic Games, Inc. All Rights Reserved.
+
+#pragma once
+
+
+#include "Containers/Array.h"
+#include "IDetailCustomization.h"
+#include "Math/MathFwd.h"
+#include "Templates/SharedPointer.h"
+#include "Types/SlateEnums.h"
+
+class FString;
+class IDetailLayoutBuilder;
+class UCustomizableObjectNodeEditMaterial;
+class SCustomizableObjectNodeLayoutBlocksSelector;
+
+
+class FCustomizableObjectNodeEditMaterialDetails : public IDetailCustomization
+{
+public:
+	/** Makes a new instance of this detail layout class for a specific detail view requesting it */
+	static TSharedRef<IDetailCustomization> MakeInstance();
+
+	/** ILayoutDetails interface */
+	virtual void CustomizeDetails( IDetailLayoutBuilder& DetailBuilder ) override;
+
+
+private:
+
+	UCustomizableObjectNodeEditMaterial* Node;
+
+	// Layout block editor widget
+	TSharedPtr<SCustomizableObjectNodeLayoutBlocksSelector> LayoutBlocksSelector;
+
+};
\ No newline at end of file
diff --git a/Engine/Plugins/Experimental/Mutable/Source/CustomizableObjectEditor/Private/MuCOE/Nodes/CustomizableObjectNodeLayoutBlocksDetails.cpp b/Engine/Plugins/Experimental/Mutable/Source/CustomizableObjectEditor/Private/MuCOE/Nodes/CustomizableObjectNodeLayoutBlocksDetails.cpp
--- a/Engine/Plugins/Experimental/Mutable/Source/CustomizableObjectEditor/Private/MuCOE/Nodes/CustomizableObjectNodeLayoutBlocksDetails.cpp
+++ b/Engine/Plugins/Experimental/Mutable/Source/CustomizableObjectEditor/Private/MuCOE/Nodes/CustomizableObjectNodeLayoutBlocksDetails.cpp
@@ -34,7 +34,7 @@ TSharedRef<IDetailCustomization> FCustomizableObjectNodeLayoutBlocksDetails::Mak
 
 void FCustomizableObjectNodeLayoutBlocksDetails::CustomizeDetails( IDetailLayoutBuilder& DetailBuilder )
 {
-	Node = 0;
+	Node = nullptr;
 	const IDetailsView* DetailsView = DetailBuilder.GetDetailsView();
 	if (DetailsView->GetSelectedObjects().Num())
 	{
diff --git a/Engine/Plugins/Experimental/Mutable/Source/CustomizableObjectEditor/Private/MuCOE/Nodes/CustomizableObjectNodeRemoveMeshBlocksDetails.cpp b/Engine/Plugins/Experimental/Mutable/Source/CustomizableObjectEditor/Private/MuCOE/Nodes/CustomizableObjectNodeRemoveMeshBlocksDetails.cpp
new file mode 100644
--- /dev/null
+++ b/Engine/Plugins/Experimental/Mutable/Source/CustomizableObjectEditor/Private/MuCOE/Nodes/CustomizableObjectNodeRemoveMeshBlocksDetails.cpp
@@ -0,0 +1,79 @@
+// Copyright Epic Games, Inc. All Rights Reserved.
+
+#include "MuCOE/Nodes/CustomizableObjectNodeRemoveMeshBlocksDetails.h"
+
+#include "DetailCategoryBuilder.h"
+#include "DetailLayoutBuilder.h"
+#include "DetailWidgetRow.h"
+#include "IDetailsView.h"
+#include "Internationalization/Internationalization.h"
+#include "Math/IntPoint.h"
+#include "Math/IntRect.h"
+#include "MuCOE/CustomizableObjectLayout.h"
+#include "MuCOE/Nodes/CustomizableObjectNodeRemoveMeshBlocks.h"
+#include "MuCOE/SCustomizableObjectNodeLayoutBlocksEditor.h"
+#include "MuCOE/SCustomizableObjectNodeLayoutBlocksSelector.h"
+#include "Templates/Casts.h"
+#include "UObject/ObjectPtr.h"
+#include "UObject/WeakObjectPtr.h"
+#include "UObject/WeakObjectPtrTemplates.h"
+#include "Widgets/DeclarativeSyntaxSupport.h"
+#include "Widgets/Layout/SBox.h"
+#include "Widgets/Text/STextBlock.h"
+
+class FString;
+
+
+#define LOCTEXT_NAMESPACE "CustomizableObjectDetails"
+
+
+TSharedRef<IDetailCustomization> FCustomizableObjectNodeRemoveMeshBlocksDetails::MakeInstance()
+{
+	return MakeShareable( new FCustomizableObjectNodeRemoveMeshBlocksDetails );
+}
+
+
+void FCustomizableObjectNodeRemoveMeshBlocksDetails::CustomizeDetails( IDetailLayoutBuilder& DetailBuilder )
+{
+	Node = nullptr;
+	const IDetailsView* DetailsView = DetailBuilder.GetDetailsView();
+	if (DetailsView->GetSelectedObjects().Num())
+	{
+		Node = Cast<UCustomizableObjectNodeRemoveMeshBlocks>(DetailsView->GetSelectedObjects()[0].Get());
+	}
+
+	IDetailCategoryBuilder& BlocksCategory = DetailBuilder.EditCategory( "Blocks" );
+
+	if (Node)
+	{
+		UCustomizableObjectNodeEditLayoutBlocks* NodeEditLayout = Cast<UCustomizableObjectNodeEditLayoutBlocks>(Node);
+		if (NodeEditLayout)
+		{
+			// Add blocks selector
+			LayoutBlocksSelector = SNew(SCustomizableObjectNodeLayoutBlocksSelector);
+
+			BlocksCategory.AddCustomRow(LOCTEXT("BlocksDetails_BlockInstructions", "BlockInstructions"))
+			[
+				SNew(SBox)
+				.HeightOverride(700.0f)
+				.WidthOverride(700.0f)
+				[
+					LayoutBlocksSelector.ToSharedRef()
+				]
+			];
+
+			LayoutBlocksSelector->SetSelectedNode(NodeEditLayout);
+		}
+	}
+	else
+	{
+		BlocksCategory.AddCustomRow( LOCTEXT("BlocksDetails_NodeNotFound", "NodeNotFound") )
+		[
+			SNew( STextBlock )
+			.Text( LOCTEXT( "Node not found", "Node not found" ) )
+		];
+	}
+
+}
+
+#undef LOCTEXT_NAMESPACE
diff --git a/Engine/Plugins/Experimental/Mutable/Source/CustomizableObjectEditor/Private/MuCOE/Nodes/CustomizableObjectNodeRemoveMeshBlocksDetails.h b/Engine/Plugins/Experimental/Mutable/Source/CustomizableObjectEditor/Private/MuCOE/Nodes/CustomizableObjectNodeRemoveMeshBlocksDetails.h
new file mode 100644
--- /dev/null
+++ b/Engine/Plugins/Experimental/Mutable/Source/CustomizableObjectEditor/Private/MuCOE/Nodes/CustomizableObjectNodeRemoveMeshBlocksDetails.h
@@ -0,0 +1,35 @@
+// Copyright Epic Games, Inc. All Rights Reserved.
+
+#pragma once
+
+
+#include "Containers/Array.h"
+#include "IDetailCustomization.h"
+#include "Math/MathFwd.h"
+#include "Templates/SharedPointer.h"
+#include "Types/SlateEnums.h"
+
+class FString;
+class IDetailLayoutBuilder;
+class UCustomizableObjectNodeRemoveMeshBlocks;
+class SCustomizableObjectNodeLayoutBlocksSelector;
+
+
+class FCustomizableObjectNodeRemoveMeshBlocksDetails : public IDetailCustomization
+{
+public:
+	/** Makes a new instance of this detail layout class for a specific detail view requesting it */
+	static TSharedRef<IDetailCustomization> MakeInstance();
+
+	/** ILayoutDetails interface */
+	virtual void CustomizeDetails( IDetailLayoutBuilder& DetailBuilder ) override;
+
+
+private:
+
+	UCustomizableObjectNodeRemoveMeshBlocks* Node;
+
+	// Layout block editor widget
+	TSharedPtr<SCustomizableObjectNodeLayoutBlocksSelector> LayoutBlocksSelector;
+
+};
\ No newline at end of file
diff --git a/Engine/Plugins/Experimental/Mutable/Source/CustomizableObjectEditor/Private/MuCOE/PinViewer/SPinViewer.cpp b/Engine/Plugins/Experimental/Mutable/Source/CustomizableObjectEditor/Private/MuCOE/PinViewer/SPinViewer.cpp
--- a/Engine/Plugins/Experimental/Mutable/Source/CustomizableObjectEditor/Private/MuCOE/PinViewer/SPinViewer.cpp
+++ b/Engine/Plugins/Experimental/Mutable/Source/CustomizableObjectEditor/Private/MuCOE/PinViewer/SPinViewer.cpp
@@ -150,7 +150,6 @@ void SPinViewer::Construct(const FArguments& InArgs)
 	Node->PostReconstructNodeDelegate.AddSP(this, &SPinViewer::UpdateWidget);
 	Node->NodeConnectionListChangedDelegate.AddSP(this, &SPinViewer::UpdateWidget);
 	Node->RemapPinsDelegate.AddSP(this, &SPinViewer::PinsRemapped);
-
 	GeneratePinInfoList();
 	
 	ChildSlot
@@ -192,36 +191,34 @@ void SPinViewer::Construct(const FArguments& InArgs)
 
 		+SVerticalBox::Slot()
 		[
-			SNew(SScrollBox)
-			+ SScrollBox::Slot()
-			[
-				SAssignNew(ListView, SListView<TSharedPtr<FEdGraphPinReference>>)
-				.ListItemsSource(&PinReferences)
-				.OnGenerateRow(this, &SPinViewer::GenerateNodePinRow)
-				.ItemHeight(22.0f)
-				.SelectionMode(ESelectionMode::Single)
-				.IsFocusable(true)
-				.HeaderRow
-				(
-					SNew(SHeaderRow)
-					+ SHeaderRow::Column(COLUMN_NAME)
-					.DefaultLabel(LOCTEXT("PinNameLabel", "Name"))
-					.OnSort(this, &SPinViewer::SortListView)
-					.FillWidth(0.5)
-
-					+ SHeaderRow::Column(COLUMN_TYPE)
-					.DefaultLabel(LOCTEXT("TypeLabel", "Type"))
-					.OnSort(this, &SPinViewer::SortListView)
-					.FillWidth(0.25)
-
-					+ SHeaderRow::Column(COLUMN_VISIBILITY)
-					.DefaultLabel(LOCTEXT("VisibilityLabel", "Visibility"))
-					.OnSort(this, &SPinViewer::SortListView)
-					.FillWidth(0.25)
-				)
-			]
+			SAssignNew(ListView, SListView<TSharedPtr<FEdGraphPinReference>>)
+			.ListItemsSource(&PinReferences)
+			.OnGenerateRow(this, &SPinViewer::GenerateNodePinRow)
+			.ItemHeight(22.0f)
+			.SelectionMode(ESelectionMode::None)
+			.IsFocusable(false)
+			.HeaderRow
+			(
+				SNew(SHeaderRow)
+				+ SHeaderRow::Column(COLUMN_NAME)
+				.DefaultLabel(LOCTEXT("PinNameLabel", "Name"))
+				.OnSort(this, &SPinViewer::SortListView)
+				.FillWidth(0.5)
+
+				+ SHeaderRow::Column(COLUMN_TYPE)
+				.DefaultLabel(LOCTEXT("TypeLabel", "Type"))
+				.OnSort(this, &SPinViewer::SortListView)
+				.FillWidth(0.25)
+
+				+ SHeaderRow::Column(COLUMN_VISIBILITY)
+				.DefaultLabel(LOCTEXT("VisibilityLabel", "Visibility"))
+				.OnSort(this, &SPinViewer::SortListView)
+				.FillWidth(0.25)
+			)
 		]
 	];
+
+	ListView->SetIsRightClickScrollingEnabled(false);
 }
 
 
@@ -288,12 +285,11 @@ void PinViewerAttachToDetailCustomization(IDetailLayoutBuilder& DetailBuilder)
 		UCustomizableObjectNode* Node = Cast<UCustomizableObjectNode>(SelectedObjects[0].Get());
 
 		IDetailCategoryBuilder& PinViewerCategoryBuilder = DetailBuilder.EditCategory("PinViewer", FText::GetEmpty(), ECategoryPriority::Uncommon);
-		PinViewerCategoryBuilder.AddCustomRow(LOCTEXT("PinViewerDetailsCategory", "PinViwer"))
+		PinViewerCategoryBuilder.AddCustomRow(LOCTEXT("PinViewerDetailsCategory", "PinViwer")).ShouldAutoExpand(true)
 		[
 			SNew(SPinViewer).Node(Node)
 		];
 	}
 }
 
-
 #undef LOCTEXT_NAMESPACE
diff --git a/Engine/Plugins/Experimental/Mutable/Source/CustomizableObjectEditor/Private/MuCOE/SCustomizableObjectNodeLayoutBlocksSelector.cpp b/Engine/Plugins/Experimental/Mutable/Source/CustomizableObjectEditor/Private/MuCOE/SCustomizableObjectNodeLayoutBlocksSelector.cpp
--- a/Engine/Plugins/Experimental/Mutable/Source/CustomizableObjectEditor/Private/MuCOE/SCustomizableObjectNodeLayoutBlocksSelector.cpp
+++ b/Engine/Plugins/Experimental/Mutable/Source/CustomizableObjectEditor/Private/MuCOE/SCustomizableObjectNodeLayoutBlocksSelector.cpp
@@ -79,9 +79,13 @@ public:
 };
 
 
+SCustomizableObjectNodeLayoutBlocksSelector::SCustomizableObjectNodeLayoutBlocksSelector() : UICommandList(new FUICommandList())
+{
+}
+
+
 void SCustomizableObjectNodeLayoutBlocksSelector::Construct(const FArguments& InArgs)
 {
-	CustomizableObjectEditorPtr = InArgs._CustomizableObjectEditor;
 	CurrentNode = nullptr;
 	
 	BindCommands();
@@ -163,8 +167,7 @@ void SCustomizableObjectNodeLayoutBlocksSelector::AddReferencedObjects( FReferen
 
 TSharedRef<SWidget> SCustomizableObjectNodeLayoutBlocksSelector::BuildLayoutToolBar()
 {
-	TSharedPtr<FUICommandList> CommandList = CustomizableObjectEditorPtr.Pin()->GetToolkitCommands();
-	FToolBarBuilder LayoutToolbarBuilder(CommandList, FMultiBoxCustomization::None);
+	FToolBarBuilder LayoutToolbarBuilder(UICommandList, FMultiBoxCustomization::None);
 
 	// Build toolbar widgets
 
@@ -312,8 +315,6 @@ void SCustomizableObjectNodeLayoutBlocksSelector::BindCommands()
 
 	const FLayoutBlockSelectorCommands& Commands = FLayoutBlockSelectorCommands::Get();
 
-	const TSharedRef<FUICommandList>& UICommandList = CustomizableObjectEditorPtr.Pin()->GetToolkitCommands();
-
 	UICommandList->MapAction(
 		Commands.SelectAll,
 		FExecuteAction::CreateSP( this, &SCustomizableObjectNodeLayoutBlocksSelector::OnSelectAll ),
diff --git a/Engine/Plugins/Experimental/Mutable/Source/CustomizableObjectEditor/Private/MuCOE/SCustomizableObjectNodeLayoutBlocksSelector.h b/Engine/Plugins/Experimental/Mutable/Source/CustomizableObjectEditor/Private/MuCOE/SCustomizableObjectNodeLayoutBlocksSelector.h
--- a/Engine/Plugins/Experimental/Mutable/Source/CustomizableObjectEditor/Private/MuCOE/SCustomizableObjectNodeLayoutBlocksSelector.h
+++ b/Engine/Plugins/Experimental/Mutable/Source/CustomizableObjectEditor/Private/MuCOE/SCustomizableObjectNodeLayoutBlocksSelector.h
@@ -11,6 +11,7 @@
 #include "Widgets/SCompoundWidget.h"
 
 class FReferenceCollector;
+class FUICommandList;
 class ICustomizableObjectInstanceEditor;
 class SWidget;
 struct FCustomizableObjectLayoutBlock;
@@ -24,11 +25,12 @@ class SCustomizableObjectNodeLayoutBlocksSelector : public SCompoundWidget, publ
 {
 public:
 	SLATE_BEGIN_ARGS( SCustomizableObjectNodeLayoutBlocksSelector ){}
-		SLATE_ARGUMENT(TWeakPtr<ICustomizableObjectInstanceEditor>, CustomizableObjectEditor)
 	SLATE_END_ARGS()
 
+	SCustomizableObjectNodeLayoutBlocksSelector();
+	virtual ~SCustomizableObjectNodeLayoutBlocksSelector();
+
 	void Construct(const FArguments& InArgs);
-	~SCustomizableObjectNodeLayoutBlocksSelector();
 	
 	// FSerializableObject interface
 	void AddReferencedObjects( FReferenceCollector& Collector ) override;
@@ -47,9 +49,6 @@ public:
 
 private:
 
-	/** Pointer back to the editor tool that owns us */
-	TWeakPtr<ICustomizableObjectInstanceEditor> CustomizableObjectEditorPtr;
-
 	// 
 	class UCustomizableObjectNodeEditLayoutBlocks* CurrentNode = nullptr;
 
@@ -59,6 +58,9 @@ private:
 	/** */
 	TSharedPtr<class STextBlock> BlocksLabel;
 
+	/** The list of UI Commands executable */
+	TSharedRef<FUICommandList> UICommandList;
+
 	/** */
 	FIntPoint GetGridSize() const;
 	void OnSelectionChanged( const TArray<FGuid>& selected );
