[mutable] Now table columns are compiled separately and not all at the same time.
+ Added mesh name to the warning that indicates if a mesh is not closed.

#rnx

diff --git a/Engine/Plugins/Experimental/Mutable/Source/CustomizableObjectEditor/Private/MuCOE/GenerateMutableSource/GenerateMutableSourceColor.cpp b/Engine/Plugins/Experimental/Mutable/Source/CustomizableObjectEditor/Private/MuCOE/GenerateMutableSource/GenerateMutableSourceColor.cpp
--- a/Engine/Plugins/Experimental/Mutable/Source/CustomizableObjectEditor/Private/MuCOE/GenerateMutableSource/GenerateMutableSourceColor.cpp
+++ b/Engine/Plugins/Experimental/Mutable/Source/CustomizableObjectEditor/Private/MuCOE/GenerateMutableSource/GenerateMutableSourceColor.cpp
@@ -309,15 +309,22 @@ mu::NodeColourPtr GenerateMutableSourceColor(const UEdGraphPin* Pin, FMutableGra
 				ColumnName = GenerationContext.CurrentMaterialTableParameterId;
 			}
 
+			// Generating a new data table if not exists
 			Table = GenerateMutableSourceTable(TypedNodeTable->Table->GetName(), Pin, GenerationContext);
 
+			// Generating a new Color column if not exists
+			if (Table && Table->FindColumn(TCHAR_TO_ANSI(*ColumnName)) == INDEX_NONE)
+			{
+				GenerateTableColumn(TypedNodeTable, Pin, Table, ColumnName, GenerationContext.CurrentLOD, GenerationContext);
+			}
+
 			ColorTableNode->SetTable(Table);
 			ColorTableNode->SetColumn(TCHAR_TO_ANSI(*ColumnName));
 			ColorTableNode->SetParameterName(TCHAR_TO_ANSI(*TypedNodeTable->ParameterName));
 
 			GenerationContext.AddParameterNameUnique(Node, TypedNodeTable->ParameterName);
 
-			if (Table->FindColumn(TCHAR_TO_ANSI(*ColumnName)) == -1)
+			if (Table->FindColumn(TCHAR_TO_ANSI(*ColumnName)) == INDEX_NONE)
 			{
 				FString Msg = FString::Printf(TEXT("Couldn't find pin column with name %s"), *ColumnName);
 				GenerationContext.Compiler->CompilerLog(FText::FromString(Msg), Node);
diff --git a/Engine/Plugins/Experimental/Mutable/Source/CustomizableObjectEditor/Private/MuCOE/GenerateMutableSource/GenerateMutableSourceFloat.cpp b/Engine/Plugins/Experimental/Mutable/Source/CustomizableObjectEditor/Private/MuCOE/GenerateMutableSource/GenerateMutableSourceFloat.cpp
--- a/Engine/Plugins/Experimental/Mutable/Source/CustomizableObjectEditor/Private/MuCOE/GenerateMutableSource/GenerateMutableSourceFloat.cpp
+++ b/Engine/Plugins/Experimental/Mutable/Source/CustomizableObjectEditor/Private/MuCOE/GenerateMutableSource/GenerateMutableSourceFloat.cpp
@@ -359,7 +359,14 @@ mu::NodeScalarPtr GenerateMutableSourceFloat(const UEdGraphPin* Pin, FMutableGra
 				ColumnName = GenerationContext.CurrentMaterialTableParameterId;
 			}
 			
+			// Generating a new data table if not exists
 			Table = GenerateMutableSourceTable(TypedNodeTable->Table->GetName(), Pin, GenerationContext);
+
+			// Generating a new Float column if not exists
+			if (Table && Table->FindColumn(TCHAR_TO_ANSI(*ColumnName)) == INDEX_NONE)
+			{
+				GenerateTableColumn(TypedNodeTable, Pin, Table, ColumnName, GenerationContext.CurrentLOD, GenerationContext);
+			}
 			
 			ScalarTableNode->SetTable(Table);
 			ScalarTableNode->SetColumn(TCHAR_TO_ANSI(*ColumnName));
@@ -367,7 +374,7 @@ mu::NodeScalarPtr GenerateMutableSourceFloat(const UEdGraphPin* Pin, FMutableGra
 
 			GenerationContext.AddParameterNameUnique(Node, TypedNodeTable->ParameterName);
 			
-			if (Table->FindColumn(TCHAR_TO_ANSI(*ColumnName)) == -1)
+			if (Table->FindColumn(TCHAR_TO_ANSI(*ColumnName)) == INDEX_NONE)
 			{
 				Table = new mu::Table();
 
diff --git a/Engine/Plugins/Experimental/Mutable/Source/CustomizableObjectEditor/Private/MuCOE/GenerateMutableSource/GenerateMutableSourceImage.cpp b/Engine/Plugins/Experimental/Mutable/Source/CustomizableObjectEditor/Private/MuCOE/GenerateMutableSource/GenerateMutableSourceImage.cpp
--- a/Engine/Plugins/Experimental/Mutable/Source/CustomizableObjectEditor/Private/MuCOE/GenerateMutableSource/GenerateMutableSourceImage.cpp
+++ b/Engine/Plugins/Experimental/Mutable/Source/CustomizableObjectEditor/Private/MuCOE/GenerateMutableSource/GenerateMutableSourceImage.cpp
@@ -241,21 +241,28 @@ mu::NodeImagePtr GenerateMutableSourceImage(const UEdGraphPin* Pin, FMutableGrap
 		}
 		FormatNode->SetSource(TextureNode);
 
+		
+
+		//TODO(Uncoment this)
+		// //Result = FormatNode;
+		//if (TypedNodeParam->bForceFixedSize)
+		//{
 		mu::NodeImageResizePtr ResizeNode = new mu::NodeImageResize();
 		ResizeNode->SetBase(FormatNode);
 		ResizeNode->SetRelative(false);
 
 		if (TypedNodeParam->DefaultValue)
 		{
-			ResizeNode->SetSize(TypedNodeParam->DefaultValue->GetResource()->GetSizeX(), TypedNodeParam->DefaultValue->GetResource()->GetSizeY() );
+			ResizeNode->SetSize(TypedNodeParam->DefaultValue->GetResource()->GetSizeX(), TypedNodeParam->DefaultValue->GetResource()->GetSizeY());
 		}
 		else
 		{
 			// \TODO: Let the user specify this in the node?
-			ResizeNode->SetSize(1024,1024);
+			ResizeNode->SetSize(1024, 1024);
 		}
 
 		Result = ResizeNode;
+		//}
 	}
 
 	else if (const UCustomizableObjectNodeMesh* TypedNodeMesh = Cast<UCustomizableObjectNodeMesh>(Node))
@@ -818,7 +825,14 @@ mu::NodeImagePtr GenerateMutableSourceImage(const UEdGraphPin* Pin, FMutableGrap
 					ColumnName = GenerationContext.CurrentMaterialTableParameterId;
 				}
 
+				// Generating a new data table if not exists
 				Table = GenerateMutableSourceTable(TypedNodeTable->Table->GetName(), Pin, GenerationContext);
+
+				// Generating a new Texture column if not exists
+				if (Table && Table->FindColumn(TCHAR_TO_ANSI(*ColumnName)) == INDEX_NONE)
+				{
+					GenerateTableColumn(TypedNodeTable, Pin, Table, ColumnName, GenerationContext.CurrentLOD, GenerationContext);
+				}
 				
 				ImageTableNode->SetTable(Table);
 				ImageTableNode->SetColumn(TCHAR_TO_ANSI(*ColumnName));
@@ -826,7 +840,7 @@ mu::NodeImagePtr GenerateMutableSourceImage(const UEdGraphPin* Pin, FMutableGrap
 
 				GenerationContext.AddParameterNameUnique(Node, TypedNodeTable->ParameterName);
 				
-				if (Table->FindColumn(TCHAR_TO_ANSI(*ColumnName)) == -1)
+				if (Table->FindColumn(TCHAR_TO_ANSI(*ColumnName)) == INDEX_NONE)
 				{
 					FString Msg = FString::Printf(TEXT("Couldn't find pin column with name %s"), *ColumnName);
 					GenerationContext.Compiler->CompilerLog(FText::FromString(Msg), Node);
diff --git a/Engine/Plugins/Experimental/Mutable/Source/CustomizableObjectEditor/Private/MuCOE/GenerateMutableSource/GenerateMutableSourceMesh.cpp b/Engine/Plugins/Experimental/Mutable/Source/CustomizableObjectEditor/Private/MuCOE/GenerateMutableSource/GenerateMutableSourceMesh.cpp
--- a/Engine/Plugins/Experimental/Mutable/Source/CustomizableObjectEditor/Private/MuCOE/GenerateMutableSource/GenerateMutableSourceMesh.cpp
+++ b/Engine/Plugins/Experimental/Mutable/Source/CustomizableObjectEditor/Private/MuCOE/GenerateMutableSource/GenerateMutableSourceMesh.cpp
@@ -2452,7 +2452,7 @@ mu::NodeMeshPtr GenerateMorphMesh(const UEdGraphPin* Pin,
 		mu::TablePtr Table = GenerationContext.GeneratedTables[TypedNodeTable->Table->GetName()];
 
 		FString ColumnName = TableColumnName + TypedNodeMorphs[MorphIndex].MorphTargetName;
-		int32 ColumnIndex = -1;
+		int32 ColumnIndex = INDEX_NONE;
 
 		for (int32 RowIndex = 0; RowIndex < NumRows; ++RowIndex)
 		{
@@ -2460,7 +2460,7 @@ mu::NodeMeshPtr GenerateMorphMesh(const UEdGraphPin* Pin,
 
 			ColumnIndex = Table->FindColumn(TCHAR_TO_ANSI(*ColumnName));
 
-			if (ColumnIndex == -1)
+			if (ColumnIndex == INDEX_NONE)
 			{
 				ColumnIndex = Table->AddColumn(TCHAR_TO_ANSI(*ColumnName), mu::TABLE_COLUMN_TYPE::TCT_MESH);
 			}
@@ -2469,7 +2469,7 @@ mu::NodeMeshPtr GenerateMorphMesh(const UEdGraphPin* Pin,
 			Table->SetCell(ColumnIndex, RowIndex, MorphedSourceTableMesh.get());
 		}
 
-		if (ColumnIndex > -1)
+		if (ColumnIndex > INDEX_NONE)
 		{
 			bSuccess = true;
 
@@ -3359,9 +3359,7 @@ mu::NodeMeshPtr GenerateMutableSourceMesh(const UEdGraphPin * Pin,
 		mu::TablePtr Table = nullptr;
 
 		if (TypedNodeTable->Table)
-		{
-			Table = GenerateMutableSourceTable(TypedNodeTable->Table->GetName(), Pin, GenerationContext);
-	
+		{			
 			USkeletalMesh* SkeletalMesh = TypedNodeTable->GetColumnDefaultAssetByType<USkeletalMesh>(Pin);
 			
 			int32 CurrentLOD = 0;
@@ -3371,12 +3369,12 @@ mu::NodeMeshPtr GenerateMutableSourceMesh(const UEdGraphPin * Pin,
 
 			if (GenerationContext.CurrentAutoLODStrategy == ECustomizableObjectAutomaticLODStrategy::AutomaticFromMesh)
 			{
-				CurrentLOD += GenerationContext.CurrentLOD;
-
-				if (SkeletalMesh)
+				if (SkeletalMesh && Helper_GetImportedModel(SkeletalMesh))
 				{
 					FSkeletalMeshModel* ImportedModel = Helper_GetImportedModel(SkeletalMesh);
 
+					CurrentLOD += GenerationContext.CurrentLOD;
+
 					// Checking if the current LOD is valid
 					if (ImportedModel->LODModels.Num() <= CurrentLOD
 						||
@@ -3395,10 +3393,21 @@ mu::NodeMeshPtr GenerateMutableSourceMesh(const UEdGraphPin * Pin,
 				}
 			}
 
-			FString ColumnName = TypedNodeTable->GetMutableColumnName(Pin, CurrentLOD);
+			// Generating a new data table if not exists
+			Table = GenerateMutableSourceTable(TypedNodeTable->Table->GetName(), Pin, GenerationContext);
+
+			// Getting mutable and data table column names
+			FString MutableColumnName = TypedNodeTable->GetMutableColumnName(Pin, CurrentLOD);
+			FString DataTableColumnName = TypedNodeTable->GetColumnNameByPin(Pin);
+
+			// Generating a new Mesh column if not exists
+			if (Table && Table->FindColumn(TCHAR_TO_ANSI(*MutableColumnName)) == INDEX_NONE)
+			{
+				GenerateTableColumn(TypedNodeTable, Pin, Table, DataTableColumnName, CurrentLOD, GenerationContext);
+			}
 
 			MeshTableNode->SetTable(Table);
-			MeshTableNode->SetColumn(TCHAR_TO_ANSI(*ColumnName));
+			MeshTableNode->SetColumn(TCHAR_TO_ANSI(*MutableColumnName));
 			MeshTableNode->SetParameterName(TCHAR_TO_ANSI(*TypedNodeTable->ParameterName));
 
 			GenerationContext.AddParameterNameUnique(Node, TypedNodeTable->ParameterName);
@@ -3463,7 +3472,7 @@ mu::NodeMeshPtr GenerateMutableSourceMesh(const UEdGraphPin * Pin,
 					}
 					else
 					{
-						FString msg = "Mesh Column [" + ColumnName + "] Layout doesn't has any block. A grid sized block will be used instead.";
+						FString msg = "Mesh Column [" + MutableColumnName + "] Layout doesn't has any block. A grid sized block will be used instead.";
 						GenerationContext.Compiler->CompilerLog(FText::FromString(msg), Node, EMessageSeverity::Warning);
 
 						LayoutNode->SetBlock(0, 0, 0, Layouts[i]->GetGridSize().X, Layouts[i]->GetGridSize().Y);
@@ -3477,12 +3486,12 @@ mu::NodeMeshPtr GenerateMutableSourceMesh(const UEdGraphPin * Pin,
 			// Applying Mesh Morph Nodes
 			if (GenerationContext.MeshMorphStack.Num())
 			{
-				MorphResult = GenerateMorphMesh(Pin, GenerationContext.MeshMorphStack, 0, Result, GenerationContext, MeshData, ColumnName);
+				MorphResult = GenerateMorphMesh(Pin, GenerationContext.MeshMorphStack, 0, Result, GenerationContext, MeshData, MutableColumnName);
 			}
 
-			if (Table->FindColumn(TCHAR_TO_ANSI(*ColumnName)) == -1)
+			if (Table->FindColumn(TCHAR_TO_ANSI(*MutableColumnName)) == INDEX_NONE)
 			{
-				FString Msg = FString::Printf(TEXT("Couldn't find pin column with name %s"), *ColumnName);
+				FString Msg = FString::Printf(TEXT("Couldn't find pin column with name %s"), *MutableColumnName);
 				GenerationContext.Compiler->CompilerLog(FText::FromString(Msg), Node);
 			}
 		}
diff --git a/Engine/Plugins/Experimental/Mutable/Source/CustomizableObjectEditor/Private/MuCOE/GenerateMutableSource/GenerateMutableSourceModifier.cpp b/Engine/Plugins/Experimental/Mutable/Source/CustomizableObjectEditor/Private/MuCOE/GenerateMutableSource/GenerateMutableSourceModifier.cpp
--- a/Engine/Plugins/Experimental/Mutable/Source/CustomizableObjectEditor/Private/MuCOE/GenerateMutableSource/GenerateMutableSourceModifier.cpp
+++ b/Engine/Plugins/Experimental/Mutable/Source/CustomizableObjectEditor/Private/MuCOE/GenerateMutableSource/GenerateMutableSourceModifier.cpp
@@ -153,7 +153,8 @@ mu::NodeModifierPtr GenerateMutableSourceModifier(const UEdGraphPin * Pin, FMuta
 
 				if (!bClosed)
 				{
-					GenerationContext.Compiler->CompilerLog(LOCTEXT("Clipping mesh", "Clipping mesh not closed (i.e., it does not enclose a volume)."), MeshData.Node, EMessageSeverity::Warning);
+					FText ErrorMsg = FText::Format(LOCTEXT("Clipping mesh", "Clipping mesh [{0}] not closed (i.e., it does not enclose a volume)."), FText::FromName(MeshData.Mesh->GetFName()));
+					GenerationContext.Compiler->CompilerLog(ErrorMsg, MeshData.Node, EMessageSeverity::Warning);
 				}
 			}
 
diff --git a/Engine/Plugins/Experimental/Mutable/Source/CustomizableObjectEditor/Private/MuCOE/GenerateMutableSource/GenerateMutableSourceSurface.cpp b/Engine/Plugins/Experimental/Mutable/Source/CustomizableObjectEditor/Private/MuCOE/GenerateMutableSource/GenerateMutableSourceSurface.cpp
--- a/Engine/Plugins/Experimental/Mutable/Source/CustomizableObjectEditor/Private/MuCOE/GenerateMutableSource/GenerateMutableSourceSurface.cpp
+++ b/Engine/Plugins/Experimental/Mutable/Source/CustomizableObjectEditor/Private/MuCOE/GenerateMutableSource/GenerateMutableSourceSurface.cpp
@@ -426,7 +426,16 @@ mu::NodeSurfacePtr GenerateMutableSourceSurface(const UEdGraphPin * Pin, FMutabl
 		{
 			if (const UCustomizableObjectNodeTable* TypedNodeTable = Cast< UCustomizableObjectNodeTable >(ConnectedPin->GetOwningNode()))
 			{
-				GenerateMutableSourceTable(TypedNodeTable->Table->GetName(), ConnectedPin, GenerationContext);
+				FString ColumnName = ConnectedPin->PinFriendlyName.ToString();
+
+				// Generating a new data table if not exists
+				mu::TablePtr Table = GenerateMutableSourceTable(TypedNodeTable->Table->GetName(), ConnectedPin, GenerationContext);
+				
+				// Generating a column for each modified parameter(texture, color & float) if not exists
+				if (Table && Table->FindColumn(TCHAR_TO_ANSI(*ColumnName)) == INDEX_NONE)
+				{
+					GenerateTableColumn(TypedNodeTable, ConnectedPin, Table, ColumnName, GenerationContext.CurrentLOD, GenerationContext);
+				}
 
 				// Checking if this material has some parameters modified by the table node linked to it
 				if (GenerationContext.GeneratedParametersInTables.Contains(TypedNodeTable))
diff --git a/Engine/Plugins/Experimental/Mutable/Source/CustomizableObjectEditor/Private/MuCOE/GenerateMutableSource/GenerateMutableSourceTable.cpp b/Engine/Plugins/Experimental/Mutable/Source/CustomizableObjectEditor/Private/MuCOE/GenerateMutableSource/GenerateMutableSourceTable.cpp
--- a/Engine/Plugins/Experimental/Mutable/Source/CustomizableObjectEditor/Private/MuCOE/GenerateMutableSource/GenerateMutableSourceTable.cpp
+++ b/Engine/Plugins/Experimental/Mutable/Source/CustomizableObjectEditor/Private/MuCOE/GenerateMutableSource/GenerateMutableSourceTable.cpp
@@ -50,9 +50,9 @@
 
 #define LOCTEXT_NAMESPACE "CustomizableObjectEditor"
 
-void FillTableColumn(const UEdGraphPin* Pin, mu::TablePtr MutableTable, FString ColumnName, FString RowName, int32 RowIdx, uint8* CellData, FProperty* Property, FMutableGraphGenerationContext& GenerationContext)
+
+void FillTableColumn(const UCustomizableObjectNodeTable* TableNode,	mu::TablePtr MutableTable,	const FString& ColumnName,	const FString& RowName,	const int32 RowIdx,	uint8* CellData, FProperty* Property, const int32 LOD, FMutableGraphGenerationContext& GenerationContext)
 {
-	UCustomizableObjectNodeTable* CustomNodeTable = Cast<UCustomizableObjectNodeTable>(Pin->GetOwningNode());
 	int32 CurrentColumn;
 
 	// Getting property type
@@ -71,22 +71,22 @@ void FillTableColumn(const UEdGraphPin* Pin, mu::TablePtr MutableTable, FString
 			FString AnimBP, AnimSlot, GameplayTag, AnimBPAssetTag;
 			TArray<FGameplayTag> GameplayTags;
 
-			CustomNodeTable->GetAnimationColumns(ColumnName, AnimBP, AnimSlot, GameplayTag);
+			TableNode->GetAnimationColumns(ColumnName, AnimBP, AnimSlot, GameplayTag);
 
 			if (!AnimBP.IsEmpty())
 			{
 				if (!AnimSlot.IsEmpty())
 				{
-					if (CustomNodeTable->Table)
+					if (TableNode->Table)
 					{
-						uint8* AnimRowData = CustomNodeTable->Table->FindRowUnchecked(FName(*RowName));
+						uint8* AnimRowData = TableNode->Table->FindRowUnchecked(FName(*RowName));
 
 						if (AnimRowData)
 						{
 							int32 SlotIndex = -1;
 
 							// Getting animation slot row value from data table
-							if (FProperty* AnimSlotProperty = CustomNodeTable->Table->FindTableProperty(FName(*AnimSlot)))
+							if (FProperty* AnimSlotProperty = TableNode->Table->FindTableProperty(FName(*AnimSlot)))
 							{
 								uint8* AnimSlotData = AnimSlotProperty->ContainerPtrToValuePtr<uint8>(AnimRowData, 0);
 
@@ -102,7 +102,7 @@ void FillTableColumn(const UEdGraphPin* Pin, mu::TablePtr MutableTable, FString
 							if (SlotIndex > -1)
 							{
 								// Getting animation instance soft class from data table
-								if (FProperty* AnimBPProperty = CustomNodeTable->Table->FindTableProperty(FName(*AnimBP)))
+								if (FProperty* AnimBPProperty = TableNode->Table->FindTableProperty(FName(*AnimBP)))
 								{
 									uint8* AnimBPData = AnimBPProperty->ContainerPtrToValuePtr<uint8>(AnimRowData, 0);
 
@@ -125,7 +125,7 @@ void FillTableColumn(const UEdGraphPin* Pin, mu::TablePtr MutableTable, FString
 							else
 							{
 								FString msg = FString::Printf(TEXT("Could not found the Slot column of the animation blueprint column [%s] for the mesh column [%s] row [%s]."), *AnimBP, *ColumnName, *RowName);
-								GenerationContext.Compiler->CompilerLog(FText::FromString(msg), CustomNodeTable);
+								GenerationContext.Compiler->CompilerLog(FText::FromString(msg), TableNode);
 							}
 						}
 					}
@@ -133,21 +133,21 @@ void FillTableColumn(const UEdGraphPin* Pin, mu::TablePtr MutableTable, FString
 				else
 				{
 					FString msg = FString::Printf(TEXT("Could not found the Slot column of the animation blueprint column [%s] for the mesh column [%s]."), *AnimBP, *ColumnName);
-					GenerationContext.Compiler->CompilerLog(FText::FromString(msg), CustomNodeTable);
+					GenerationContext.Compiler->CompilerLog(FText::FromString(msg), TableNode);
 				}
 			}
 
 			// Getting Gameplay tags
 			if (!GameplayTag.IsEmpty())
 			{
-				if (CustomNodeTable->Table)
+				if (TableNode->Table)
 				{
-					uint8* GameplayRowData = CustomNodeTable->Table->FindRowUnchecked(FName(*RowName));
+					uint8* GameplayRowData = TableNode->Table->FindRowUnchecked(FName(*RowName));
 
 					if (GameplayRowData)
 					{
 						// Getting animation tag row value from data table
-						if (FProperty* GameplayTagProperty = CustomNodeTable->Table->FindTableProperty(FName(*GameplayTag)))
+						if (FProperty* GameplayTagProperty = TableNode->Table->FindTableProperty(FName(*GameplayTag)))
 						{
 							uint8* GameplayTagData = GameplayTagProperty->ContainerPtrToValuePtr<uint8>(GameplayRowData, 0);
 
@@ -168,189 +168,180 @@ void FillTableColumn(const UEdGraphPin* Pin, mu::TablePtr MutableTable, FString
 			}
 
 			// Getting reference Mesh
-			USkeletalMesh* ReferenceSkeletalMesh = CustomNodeTable->GetColumnDefaultAssetByType<USkeletalMesh>(ColumnName);
+			USkeletalMesh* ReferenceSkeletalMesh = TableNode->GetColumnDefaultAssetByType<USkeletalMesh>(ColumnName);
 
 			if (!ReferenceSkeletalMesh)
 			{
 				FString msg = FString::Printf(TEXT("Reference Skeletal Mesh not found for column [%s]."), *ColumnName);
-				GenerationContext.Compiler->CompilerLog(FText::FromString(msg), CustomNodeTable);
+				GenerationContext.Compiler->CompilerLog(FText::FromString(msg), TableNode);
 
 				return;
 			}
 
-			int ReferenceNumLODs = 0;
-			ReferenceNumLODs = Helper_GetLODInfoArray(ReferenceSkeletalMesh).Num();
+			// Parameter used for LOD differences
+			int32 CurrentLOD = LOD;
 
 			int NumLODs = Helper_GetLODInfoArray(SkeletalMesh).Num();
 
-			if (NumLODs != ReferenceNumLODs)
+			if (NumLODs <= CurrentLOD)
 			{
-				FString Dif_1 = NumLODs > ReferenceNumLODs ? "more" : "less";
-				FString Dif_2 = NumLODs > ReferenceNumLODs ? "Some will be ignored" : "LOD " + FString::FromInt(NumLODs - 1) + " will be used instead. This can cause some performance penalties.";
+				CurrentLOD = NumLODs - 1;
 
-				FString msg = FString::Printf(TEXT("Mesh from column [%s] row [%s] has %s LODs than the reference mesh. %s"), *ColumnName, *RowName, *Dif_1, *Dif_2);
-				GenerationContext.Compiler->CompilerLog(FText::FromString(msg), CustomNodeTable);
+				FString msg = FString::Printf(TEXT("Mesh from column [%s] row [%s] needs LOD %d but has less LODs than the reference mesh. LOD %d will be used instead. This can cause some performance penalties."),
+					*ColumnName, *RowName, LOD, CurrentLOD);
+				GenerationContext.Compiler->CompilerLog(FText::FromString(msg), TableNode);
 			}
 
-			// Parameter used for LOD differences
-			int32 CurrentLOD = 0;
+			int32 NumMaterials = Helper_GetImportedModel(SkeletalMesh)->LODModels[CurrentLOD].Sections.Num();
+			int32 ReferenceNumMaterials = Helper_GetImportedModel(ReferenceSkeletalMesh)->LODModels[CurrentLOD].Sections.Num();
+
+			if (NumMaterials != ReferenceNumMaterials)
+			{
+				FString Dif_1 = NumMaterials > ReferenceNumMaterials ? "more" : "less";
+				FString Dif_2 = NumMaterials > ReferenceNumMaterials ? "Some will be ignored" : "This can cause some compilation errors.";
+
+				FString msg = FString::Printf(TEXT("Mesh from column [%s] row [%s] has %s Sections than the reference mesh. %s"), *ColumnName, *RowName, *Dif_1, *Dif_2);
+				GenerationContext.Compiler->CompilerLog(FText::FromString(msg), TableNode);
+			}
 
-			// Generating skeletal mesh columns
-			for (int32 LODIndex = 0; LODIndex < ReferenceNumLODs; ++LODIndex)
+			int32 Materials = NumMaterials <= ReferenceNumMaterials ? NumMaterials : ReferenceNumMaterials;
+
+			for (int32 MatIndex = 0; MatIndex < Materials; ++MatIndex)
 			{
-				int32 NumMaterials = Helper_GetImportedModel(SkeletalMesh)->LODModels[CurrentLOD].Sections.Num();
-				int32 ReferenceNumMaterials = Helper_GetImportedModel(ReferenceSkeletalMesh)->LODModels[CurrentLOD].Sections.Num();
+				FString MutableColumnName = ColumnName + FString::Printf(TEXT(" LOD_%d "), LOD) + FString::Printf(TEXT("Mat_%d"), MatIndex);
+
+				CurrentColumn = MutableTable.get()->FindColumn(TCHAR_TO_ANSI(*MutableColumnName));
 
-				if (NumMaterials > ReferenceNumMaterials)
+				if (CurrentColumn == -1)
 				{
-					FString msg = FString::Printf(TEXT("Mesh from column [%s] row [%s] has more Sections than the reference mesh. Some will be ignored"), *ColumnName, *RowName);
-					GenerationContext.Compiler->CompilerLog(FText::FromString(msg), CustomNodeTable);
+					CurrentColumn = MutableTable->AddColumn(TCHAR_TO_ANSI(*MutableColumnName), mu::TABLE_COLUMN_TYPE::TCT_MESH);
 				}
 
-				int32 Materials = NumMaterials <= ReferenceNumMaterials ? NumMaterials : ReferenceNumMaterials;
+				// First process the mesh tags that are going to make the mesh unique and affect whether it's repeated in 
+				// the mesh cache or not
+				FString MeshUniqueTags;
 
-				for (int32 MatIndex = 0; MatIndex < Materials; ++MatIndex)
+				if (!AnimBPAssetTag.IsEmpty())
 				{
-					FString MutableColumnName = ColumnName + FString::Printf(TEXT(" LOD_%d "), LODIndex) + FString::Printf(TEXT("Mat_%d"), MatIndex);
+					MeshUniqueTags += AnimBPAssetTag;
+				}
 
-					CurrentColumn = MutableTable.get()->FindColumn(TCHAR_TO_ANSI(*MutableColumnName));
+				TArray<FString> ArrayAnimBPTags;
 
-					if (CurrentColumn == -1)
-					{
-						CurrentColumn = MutableTable->AddColumn(TCHAR_TO_ANSI(*MutableColumnName), mu::TABLE_COLUMN_TYPE::TCT_MESH);
-					}
+				for (const FGameplayTag& Tag : GameplayTags)
+				{
+					MeshUniqueTags += GenerateGameplayTag(Tag.ToString());
+				}
 
-					// First process the mesh tags that are going to make the mesh unique and affect whether it's repeated in 
-					// the mesh cache or not
-					FString MeshUniqueTags;
+				mu::MeshPtr MutableMesh = GenerateMutableMesh(SkeletalMesh, CurrentLOD, MatIndex, MeshUniqueTags, GenerationContext, TableNode);
+
+				if (MutableMesh)
+				{
+ 					if (SkeletalMesh->GetPhysicsAsset() && MutableMesh->GetPhysicsBody() && MutableMesh->GetPhysicsBody()->GetBodyCount())
+					{	
+						TSoftObjectPtr<UPhysicsAsset> PhysicsAsset = SkeletalMesh->GetPhysicsAsset();
+						GenerationContext.PhysicsAssetMap.Add(PhysicsAsset.ToString(), PhysicsAsset);
+						FString PhysicsAssetTag = FString("__PhysicsAsset:") + PhysicsAsset.ToString();
+
+						AddTagToMutableMeshUnique(*MutableMesh, PhysicsAssetTag);
+					}
 
 					if (!AnimBPAssetTag.IsEmpty())
 					{
-						MeshUniqueTags += AnimBPAssetTag;
+						AddTagToMutableMeshUnique(*MutableMesh, AnimBPAssetTag);
 					}
 
-					TArray<FString> ArrayAnimBPTags;
-
 					for (const FGameplayTag& Tag : GameplayTags)
 					{
-						MeshUniqueTags += GenerateGameplayTag(Tag.ToString());
+						AddTagToMutableMeshUnique(*MutableMesh, GenerateGameplayTag(Tag.ToString()));
 					}
 
-					mu::MeshPtr MutableMesh = GenerateMutableMesh(SkeletalMesh, CurrentLOD, MatIndex, MeshUniqueTags, GenerationContext, CustomNodeTable);
-
-					if (MutableMesh)
-					{
- 						if (SkeletalMesh->GetPhysicsAsset() && MutableMesh->GetPhysicsBody() && MutableMesh->GetPhysicsBody()->GetBodyCount())
-						{	
-							TSoftObjectPtr<UPhysicsAsset> PhysicsAsset = SkeletalMesh->GetPhysicsAsset();
-							GenerationContext.PhysicsAssetMap.Add(PhysicsAsset.ToString(), PhysicsAsset);
-							FString PhysicsAssetTag = FString("__PhysicsAsset:") + PhysicsAsset.ToString();
-
-							AddTagToMutableMeshUnique(*MutableMesh, PhysicsAssetTag);
-						}
-
-						if (!AnimBPAssetTag.IsEmpty())
-						{
-							AddTagToMutableMeshUnique(*MutableMesh, AnimBPAssetTag);
-						}
-
-						for (const FGameplayTag& Tag : GameplayTags)
-						{
-							AddTagToMutableMeshUnique(*MutableMesh, GenerateGameplayTag(Tag.ToString()));
-						}
-
-						AddSocketTagsToMesh(SkeletalMesh, MutableMesh, GenerationContext);
+					AddSocketTagsToMesh(SkeletalMesh, MutableMesh, GenerationContext);
 
-						if (UCustomizableObjectSystem::GetInstance()->IsMutableAnimInfoDebuggingEnabled())
-						{
-							FString MeshPath;
-							SkeletalMesh->GetOuter()->GetPathName(nullptr, MeshPath);
-							FString MeshTag = FString("__MeshPath:") + MeshPath;
-							AddTagToMutableMeshUnique(*MutableMesh, MeshTag);
-						}
-
-						MutableTable->SetCell(CurrentColumn, RowIdx, MutableMesh.get());
-					}
-					else
+					if (UCustomizableObjectSystem::GetInstance()->IsMutableAnimInfoDebuggingEnabled())
 					{
-						FString msg = FString::Printf(TEXT("Error converting skeletal mesh LOD %d, Material %d from column [%s] row [%s] to mutable."),
-							LODIndex, MatIndex, *ColumnName, *RowName);
-						GenerationContext.Compiler->CompilerLog(FText::FromString(msg), CustomNodeTable);
+						FString MeshPath;
+						SkeletalMesh->GetOuter()->GetPathName(nullptr, MeshPath);
+						FString MeshTag = FString("__MeshPath:") + MeshPath;
+						AddTagToMutableMeshUnique(*MutableMesh, MeshTag);
 					}
-				}
 
-				CurrentLOD = (CurrentLOD + 1) >= NumLODs ? CurrentLOD : (CurrentLOD + 1);
+					MutableTable->SetCell(CurrentColumn, RowIdx, MutableMesh.get());
+				}
+				else
+				{
+					FString msg = FString::Printf(TEXT("Error converting skeletal mesh LOD %d, Material %d from column [%s] row [%s] to mutable."),
+						LOD, MatIndex, *ColumnName, *RowName);
+					GenerationContext.Compiler->CompilerLog(FText::FromString(msg), TableNode);
+				}
 			}
 		}
 
 		else if (UStaticMesh* StaticMesh = Cast<UStaticMesh>(Object))
 		{
 			// Getting reference Mesh
-			UStaticMesh* ReferenceStaticMesh = CustomNodeTable->GetColumnDefaultAssetByType<UStaticMesh>(ColumnName);
+			UStaticMesh* ReferenceStaticMesh = TableNode->GetColumnDefaultAssetByType<UStaticMesh>(ColumnName);
 
 			if (!ReferenceStaticMesh)
 			{
 				FString msg = FString::Printf(TEXT("Reference Static Mesh not found for column [%s]."), *ColumnName);
-				GenerationContext.Compiler->CompilerLog(FText::FromString(msg), CustomNodeTable);
+				GenerationContext.Compiler->CompilerLog(FText::FromString(msg), TableNode);
 
 				return;
 			}
 
+			// Parameter used for LOD differences
+			int32 CurrentLOD = LOD;
+
 			int NumLODs = StaticMesh->GetRenderData()->LODResources.Num();
-			int ReferenceNumLODs = ReferenceStaticMesh->GetRenderData()->LODResources.Num();
 
-			if (NumLODs != ReferenceNumLODs)
+			if (NumLODs <= CurrentLOD)
 			{
-				FString Dif_1 = NumLODs > ReferenceNumLODs ? "more" : "less";
-				FString Dif_2 = NumLODs > ReferenceNumLODs ? "Some will be ignored" : "LOD " + FString::FromInt(NumLODs - 1) + " will be used instead. This can cause some performance penalties.";
+				CurrentLOD = NumLODs - 1;
 
-				FString msg = FString::Printf(TEXT("Mesh from column [%s] row [%s] has %s LODs than the reference mesh. %s"), *ColumnName, *RowName, *Dif_1, *Dif_2);
-				GenerationContext.Compiler->CompilerLog(FText::FromString(msg), CustomNodeTable);
+				FString msg = FString::Printf(TEXT("Mesh from column [%s] row [%s] needs LOD %d but has less LODs than the reference mesh. LOD %d will be used instead. This can cause some performance penalties."),
+					*ColumnName, *RowName, LOD, CurrentLOD);
+				GenerationContext.Compiler->CompilerLog(FText::FromString(msg), TableNode);
 			}
 
-			// Parameter used for LOD differences
-			int32 CurrentLOD = 0;
+			int32 NumMaterials = StaticMesh->GetRenderData()->LODResources[CurrentLOD].Sections.Num();
+			int32 ReferenceNumMaterials = ReferenceStaticMesh->GetRenderData()->LODResources[CurrentLOD].Sections.Num();
 
-			for (int32 LODIndex = 0; LODIndex < ReferenceNumLODs; ++LODIndex)
+			if (NumMaterials != ReferenceNumMaterials)
 			{
-				int32 NumMaterials = StaticMesh->GetRenderData()->LODResources[CurrentLOD].Sections.Num();
-				int32 ReferenceNumMaterials = ReferenceStaticMesh->GetRenderData()->LODResources[CurrentLOD].Sections.Num();
+				FString FirstTextOption = NumMaterials > ReferenceNumMaterials ? "more" : "less";
+				FString SecondTextOption = NumMaterials > ReferenceNumMaterials ? "Some will be ignored" : "This can cause some compilation errors.";
 
-				if (NumMaterials > ReferenceNumMaterials)
-				{
-					FString msg = FString::Printf(TEXT("Mesh from column [%s] row [%s] has more Sections than the reference mesh. Some will be ignored"), *ColumnName, *RowName);
-					GenerationContext.Compiler->CompilerLog(FText::FromString(msg), CustomNodeTable);
-				}
+				FString msg = FString::Printf(TEXT("Mesh from column [%s] row [%s] has %s Sections than the reference mesh. %s"), *ColumnName, *RowName, *FirstTextOption, *SecondTextOption);
+				GenerationContext.Compiler->CompilerLog(FText::FromString(msg), TableNode);
+			}
 
-				int32 Materials = NumMaterials <= ReferenceNumMaterials ? NumMaterials : ReferenceNumMaterials;
+			int32 Materials = NumMaterials <= ReferenceNumMaterials ? NumMaterials : ReferenceNumMaterials;
 
-				for (int32 MatIndex = 0; MatIndex < Materials; ++MatIndex)
-				{
-					FString MutableColumnName = ColumnName + FString::Printf(TEXT(" LOD_%d "), LODIndex) + FString::Printf(TEXT("Mat_%d"), MatIndex);
-
-					CurrentColumn = MutableTable.get()->FindColumn(TCHAR_TO_ANSI(*MutableColumnName));
+			for (int32 MatIndex = 0; MatIndex < Materials; ++MatIndex)
+			{
+				FString MutableColumnName = ColumnName + FString::Printf(TEXT(" LOD_%d "), LOD) + FString::Printf(TEXT("Mat_%d"), MatIndex);
 
-					if (CurrentColumn == -1)
-					{
-						CurrentColumn = MutableTable->AddColumn(TCHAR_TO_ANSI(*MutableColumnName), mu::TABLE_COLUMN_TYPE::TCT_MESH);
-					}
+				CurrentColumn = MutableTable.get()->FindColumn(TCHAR_TO_ANSI(*MutableColumnName));
 
-					mu::MeshPtr MutableMesh = GenerateMutableMesh(StaticMesh, CurrentLOD, MatIndex, FString(), GenerationContext, CustomNodeTable);
+				if (CurrentColumn == -1)
+				{
+					CurrentColumn = MutableTable->AddColumn(TCHAR_TO_ANSI(*MutableColumnName), mu::TABLE_COLUMN_TYPE::TCT_MESH);
+				}
 
-					if (MutableMesh)
-					{
-						MutableTable->SetCell(CurrentColumn, RowIdx, MutableMesh.get());
-					}
-					else
-					{
-						FString msg = FString::Printf(TEXT("Error converting skeletal mesh LOD %d, Material %d from column [%s] row [%s] to mutable."),
-							LODIndex, MatIndex, *ColumnName, *RowName);
+				mu::MeshPtr MutableMesh = GenerateMutableMesh(StaticMesh, CurrentLOD, MatIndex, FString(), GenerationContext, TableNode);
 
-						GenerationContext.Compiler->CompilerLog(FText::FromString(msg), CustomNodeTable);
-					}
+				if (MutableMesh)
+				{
+					MutableTable->SetCell(CurrentColumn, RowIdx, MutableMesh.get());
 				}
+				else
+				{
+					FString msg = FString::Printf(TEXT("Error converting skeletal mesh LOD %d, Material %d from column [%s] row [%s] to mutable."),
+						LOD, MatIndex, *ColumnName, *RowName);
 
-				CurrentLOD = (CurrentLOD + 1) >= NumLODs ? CurrentLOD : (CurrentLOD + 1);
+					GenerationContext.Compiler->CompilerLog(FText::FromString(msg), TableNode);
+				}
 			}
 		}
 
@@ -359,29 +350,35 @@ void FillTableColumn(const UEdGraphPin* Pin, mu::TablePtr MutableTable, FString
 			// Getting column index from column name
 			CurrentColumn = MutableTable->FindColumn(TCHAR_TO_ANSI(*ColumnName));
 
-			if (CurrentColumn == -1)
+			if (CurrentColumn == INDEX_NONE)
 			{
 				CurrentColumn = MutableTable->AddColumn(TCHAR_TO_ANSI(*ColumnName), mu::TABLE_COLUMN_TYPE::TCT_IMAGE);
 			}
 
-			GenerationContext.ArrayTextureUnrealToMutableTask.Add(FTextureUnrealToMutableTask(MutableTable, Texture, CustomNodeTable, CurrentColumn, RowIdx));
+			GenerationContext.ArrayTextureUnrealToMutableTask.Add(FTextureUnrealToMutableTask(MutableTable, Texture, TableNode, CurrentColumn, RowIdx));
 		}
 
 		else if (UMaterialInstance* Material = Cast<UMaterialInstance>(Object))
 		{
-			UMaterialInstance* ReferenceMaterial = CustomNodeTable->GetColumnDefaultAssetByType<UMaterialInstance>(ColumnName);
+			//Adding an empty column for searching purposes
+			if (MutableTable.get()->FindColumn(TCHAR_TO_ANSI(*ColumnName)) == -1)
+			{
+				CurrentColumn = MutableTable->AddColumn(TCHAR_TO_ANSI(*ColumnName), mu::TABLE_COLUMN_TYPE::TCT_NONE);
+			}
+
+			UMaterialInstance* ReferenceMaterial = TableNode->GetColumnDefaultAssetByType<UMaterialInstance>(ColumnName);
 			
-			if (!GenerationContext.GeneratedParametersInTables.Contains(CustomNodeTable))
+			if (!GenerationContext.GeneratedParametersInTables.Contains(TableNode))
 			{
-				GenerationContext.GeneratedParametersInTables.Add(CustomNodeTable);
+				GenerationContext.GeneratedParametersInTables.Add(TableNode);
 			}
 
-			TArray<FGuid>& Parameters = GenerationContext.GeneratedParametersInTables[CustomNodeTable];
+			TArray<FGuid>& Parameters = GenerationContext.GeneratedParametersInTables[TableNode];
 
 			if (!ReferenceMaterial)
 			{
 				FString msg = FString::Printf(TEXT("Reference Material not found for column [%s]."), *ColumnName);
-				GenerationContext.Compiler->CompilerLog(FText::FromString(msg), CustomNodeTable);
+				GenerationContext.Compiler->CompilerLog(FText::FromString(msg), TableNode);
 
 				return;
 			}
@@ -420,7 +417,7 @@ void FillTableColumn(const UEdGraphPin* Pin, mu::TablePtr MutableTable, FString
 						// Getting column index from parameter name
 						ColumnIndex = MutableTable->FindColumn(TCHAR_TO_ANSI(*ParameterGuid));
 
-						if (ColumnIndex == -1)
+						if (ColumnIndex == INDEX_NONE)
 						{
 							// If there is no column with the parameters name, we generate a new one
 							ColumnIndex = MutableTable->AddColumn(TCHAR_TO_ANSI(*ParameterGuid), mu::TABLE_COLUMN_TYPE::TCT_IMAGE);
@@ -449,12 +446,12 @@ void FillTableColumn(const UEdGraphPin* Pin, mu::TablePtr MutableTable, FString
 
 						if (TextureToConvert)
 						{
-							GenerationContext.ArrayTextureUnrealToMutableTask.Add(FTextureUnrealToMutableTask(MutableTable, TextureToConvert, CustomNodeTable, ColumnIndex, RowIdx));
+							GenerationContext.ArrayTextureUnrealToMutableTask.Add(FTextureUnrealToMutableTask(MutableTable, TextureToConvert, TableNode, ColumnIndex, RowIdx));
 						}
 						else
 						{
 							FString msg = FString::Printf(TEXT("Didn't find Material Texture Parameter [%s] from parent Material of column [%s] row [%s]."), *TextureParameterName, *ColumnName, *RowName);
-							GenerationContext.Compiler->CompilerLog(FText::FromString(msg), CustomNodeTable);
+							GenerationContext.Compiler->CompilerLog(FText::FromString(msg), TableNode);
 						}
 					}
 				}
@@ -466,7 +463,7 @@ void FillTableColumn(const UEdGraphPin* Pin, mu::TablePtr MutableTable, FString
 					FString msg =
 						FString::Printf(TEXT("Material Instance [%s] from column [%s] row [%s] has %d modifiable Textures that will not be modified, they are non-modifiable parameters in the Default Material Instance"),
 						*Material->GetName(), *ColumnName, *RowName, ParametersDiff);
-					GenerationContext.Compiler->CompilerLog(FText::FromString(msg), CustomNodeTable);
+					GenerationContext.Compiler->CompilerLog(FText::FromString(msg), TableNode);
 				}
 
 				ParameterInfos.Empty();
@@ -497,7 +494,7 @@ void FillTableColumn(const UEdGraphPin* Pin, mu::TablePtr MutableTable, FString
 					// Getting column index from parameter name
 					ColumnIndex = MutableTable->FindColumn(TCHAR_TO_ANSI(*ParameterGuid));
 
-					if (ColumnIndex == -1)
+					if (ColumnIndex == INDEX_NONE)
 					{
 						// If there is no column with the parameters name, we generate a new one
 						ColumnIndex = MutableTable->AddColumn(TCHAR_TO_ANSI(*ParameterGuid), mu::TABLE_COLUMN_TYPE::TCT_COLOUR);
@@ -532,7 +529,7 @@ void FillTableColumn(const UEdGraphPin* Pin, mu::TablePtr MutableTable, FString
 					FString msg =
 						FString::Printf(TEXT("Material Instance [%s] from column [%s] row [%s] has %d modifiable Vectors that will not be modified, they are non-modifiable parameters in the Default Material Instance"),
 							*Material->GetName(), *ColumnName, *RowName, ParametersDiff);
-					GenerationContext.Compiler->CompilerLog(FText::FromString(msg), CustomNodeTable);
+					GenerationContext.Compiler->CompilerLog(FText::FromString(msg), TableNode);
 				}
 
 				// Number of modified Float parameters in the Material Instance
@@ -563,7 +560,7 @@ void FillTableColumn(const UEdGraphPin* Pin, mu::TablePtr MutableTable, FString
 					// Getting column index from parameter name
 					ColumnIndex = MutableTable->FindColumn(TCHAR_TO_ANSI(*ParameterGuid));
 
-					if (ColumnIndex == -1)
+					if (ColumnIndex == INDEX_NONE)
 					{
 						// If there is no column with the parameters name, we generate a new one
 						ColumnIndex = MutableTable->AddColumn(TCHAR_TO_ANSI(*ParameterGuid), mu::TABLE_COLUMN_TYPE::TCT_SCALAR);
@@ -597,13 +594,13 @@ void FillTableColumn(const UEdGraphPin* Pin, mu::TablePtr MutableTable, FString
 					FString msg =
 						FString::Printf(TEXT("Material Instance [%s] from column [%s] row [%s] has %d modifiable Scalars that will not be modified, they are non-modifiable parameters in the Default Material Instance"),
 							*Material->GetName(), *ColumnName, *RowName, ParametersDiff);
-					GenerationContext.Compiler->CompilerLog(FText::FromString(msg), CustomNodeTable);
+					GenerationContext.Compiler->CompilerLog(FText::FromString(msg), TableNode);
 				}
 			}
 			else
 			{
 				FString msg = FString::Printf(TEXT("Material from column [%s] row [%s] is a diferent instance than the Reference Material of the table."), *ColumnName, *RowName);
-				GenerationContext.Compiler->CompilerLog(FText::FromString(msg), CustomNodeTable);
+				GenerationContext.Compiler->CompilerLog(FText::FromString(msg), TableNode);
 			}
 		}
 	}
@@ -614,7 +611,7 @@ void FillTableColumn(const UEdGraphPin* Pin, mu::TablePtr MutableTable, FString
 		{
 			CurrentColumn = MutableTable->FindColumn(TCHAR_TO_ANSI(*ColumnName));
 
-			if (CurrentColumn == -1)
+			if (CurrentColumn == INDEX_NONE)
 			{
 				CurrentColumn = MutableTable->AddColumn(TCHAR_TO_ANSI(*ColumnName), mu::TABLE_COLUMN_TYPE::TCT_COLOUR);
 			}
@@ -625,26 +622,92 @@ void FillTableColumn(const UEdGraphPin* Pin, mu::TablePtr MutableTable, FString
 		}
 	}
 
-	else if (const FNumericProperty* NumProperty = CastField<FFloatProperty>(Property))
+	else if (const FNumericProperty* FloatNumProperty = CastField<FFloatProperty>(Property))
 	{
 		CurrentColumn = MutableTable->FindColumn(TCHAR_TO_ANSI(*ColumnName));
 
-		if (CurrentColumn == -1)
+		if (CurrentColumn == INDEX_NONE)
 		{
 			CurrentColumn = MutableTable->AddColumn(TCHAR_TO_ANSI(*ColumnName), mu::TABLE_COLUMN_TYPE::TCT_SCALAR);
 		}
 
 		// Setting cell value
-		float Value = NumProperty->GetFloatingPointPropertyValue(CellData);
+		float Value = FloatNumProperty->GetFloatingPointPropertyValue(CellData);
 		MutableTable->SetCell(CurrentColumn, RowIdx, Value);
 	}
+
+	//TODO
+	//else if (const FNumericProperty* DoubleNumProperty = CastField<FDoubleProperty>(Property))
+	//{
+	//	CurrentColumn = MutableTable->FindColumn(TCHAR_TO_ANSI(*ColumnName));
+	//
+	//	if (CurrentColumn == INDEX_NONE)
+	//	{
+	//		CurrentColumn = MutableTable->AddColumn(TCHAR_TO_ANSI(*ColumnName), mu::TABLE_COLUMN_TYPE::TCT_SCALAR);
+	//	}
+	//
+	//	// Setting cell value
+	//	float Value = DoubleNumProperty->GetFloatingPointPropertyValue(CellData);
+	//	MutableTable->SetCell(CurrentColumn, RowIdx, Value);
+	//}
 }
 
 
-mu::TablePtr GenerateMutableSourceTable(const FString& TableName, const UEdGraphPin* Pin, FMutableGraphGenerationContext& GenerationContext)
+void GenerateTableColumn(const UCustomizableObjectNodeTable* TableNode, const UEdGraphPin* Pin, mu::TablePtr MutableTable, const FString& DataTableColumnName, const int32 LOD, FMutableGraphGenerationContext& GenerationContext)
 {
 	SCOPED_PIN_DATA(GenerationContext, Pin)
-	
+
+	if (!TableNode)
+	{
+		return;
+	}
+
+	UDataTable* DataTable = TableNode->Table;
+
+	if (!DataTable)
+	{
+		return;
+	}
+
+	const UScriptStruct* TableStruct = DataTable->GetRowStruct();
+
+	if (TableStruct)
+	{
+		// Getting names of the rows to access the information
+		TArray<FName> RowNames = TableNode->GetRowNames();
+
+		for (TFieldIterator<FProperty> It(TableStruct); It; ++It)
+		{
+			FProperty* ColumnProperty = *It;
+
+			if (!ColumnProperty || DataTableColumnName != DataTableUtils::GetPropertyExportName(ColumnProperty))
+			{
+				continue;
+			}
+
+			for (int32 RowIndex = 0; RowIndex < RowNames.Num(); ++RowIndex)
+			{
+				// Getting Row Data
+				uint8* RowData = DataTable->FindRowUnchecked(RowNames[RowIndex]);
+
+				if (RowData)
+				{
+					// Getting Cell Data
+					uint8* CellData = ColumnProperty->ContainerPtrToValuePtr<uint8>(RowData, 0);
+
+					if (CellData)
+					{
+						FillTableColumn(TableNode, MutableTable, DataTableColumnName, RowNames[RowIndex].ToString(), RowIndex, CellData, ColumnProperty, LOD, GenerationContext);
+					}
+				}
+			}
+		}
+	}
+}
+
+
+mu::TablePtr GenerateMutableSourceTable(const FString& TableName, const UEdGraphPin* Pin, FMutableGraphGenerationContext& GenerationContext)
+{	
 	if (mu::TablePtr* Result = GenerationContext.GeneratedTables.Find(TableName))
 	{
 		return *Result;
@@ -674,15 +737,9 @@ mu::TablePtr GenerateMutableSourceTable(const FString& TableName, const UEdGraph
 
 		if (TableStruct)
 		{
-			// Getting names of the column headers
-			TArray<FString> ColumnHeaders = Table->GetColumnTitles();
-
 			// Getting names of the rows to access the information
 			TArray<FName> RowNames = TypedTable->GetRowNames();
 
-			// index to iterate the column headers
-			int32 ColumnHeaderIndex = 1;
-
 			// Adding and filling Name Column
 			MutableTable->AddColumn("Name", mu::TABLE_COLUMN_TYPE::TCT_STRING);
 
@@ -707,30 +764,6 @@ mu::TablePtr GenerateMutableSourceTable(const FString& TableName, const UEdGraph
 					FMutableParamUIMetadata()));
 			}
 
-			for (TFieldIterator<FProperty> It(TableStruct); It; ++It)
-			{
-				FProperty* Property = *It;
-
-				for (int32 RowIndex = 0; RowIndex < RowNames.Num(); ++RowIndex)
-				{
-					// Getting Row Data
-					uint8* RowData = Table->FindRowUnchecked(RowNames[RowIndex]);
-
-					if (RowData)
-					{
-						// Getting Cell Data
-						uint8* CellData = Property->ContainerPtrToValuePtr<uint8>(RowData, 0);
-
-						if (CellData)
-						{
-							FillTableColumn(Pin, MutableTable, ColumnHeaders[ColumnHeaderIndex], RowNames[RowIndex].ToString(), RowIndex, CellData, Property, GenerationContext);
-						}
-					}
-				}
-
-				ColumnHeaderIndex++;
-			}
-
 			GenerationContext.ParameterUIDataMap.Add(TypedTable->ParameterName, ParameterUIData);
 		}
 	}
diff --git a/Engine/Plugins/Experimental/Mutable/Source/CustomizableObjectEditor/Private/MuCOE/GenerateMutableSource/GenerateMutableSourceTable.h b/Engine/Plugins/Experimental/Mutable/Source/CustomizableObjectEditor/Private/MuCOE/GenerateMutableSource/GenerateMutableSourceTable.h
--- a/Engine/Plugins/Experimental/Mutable/Source/CustomizableObjectEditor/Private/MuCOE/GenerateMutableSource/GenerateMutableSourceTable.h
+++ b/Engine/Plugins/Experimental/Mutable/Source/CustomizableObjectEditor/Private/MuCOE/GenerateMutableSource/GenerateMutableSourceTable.h
@@ -5,7 +5,13 @@
 #include "Containers/UnrealString.h"
 #include "MuT/Table.h"
 
+class UCustomizableObjectNodeTable;
 class UEdGraphPin;
 struct FMutableGraphGenerationContext;
 
+void GenerateTableColumn(const UCustomizableObjectNodeTable* TableNode, const UEdGraphPin* Pin, mu::TablePtr MutableTable, const FString& DataTableColumnName, const int32 LOD, FMutableGraphGenerationContext& GenerationContext);
+
+void FillTableColumn(const UCustomizableObjectNodeTable* TableNode, mu::TablePtr MutableTable, const FString& ColumnName,
+	const FString& RowName, const int32 RowIdx, uint8* CellData, FProperty* Property, const int32 LOD, FMutableGraphGenerationContext& GenerationContext);
+
 mu::TablePtr GenerateMutableSourceTable(const FString& TableName, const UEdGraphPin* Pin, FMutableGraphGenerationContext& GenerationContext);
