[Mutable] Fix Instances not being updated if they where previously discarted.

diff --git a/Engine/Plugins/Experimental/Mutable/Source/CustomizableObject/Private/MuCO/CustomizableObjectInstance.cpp b/Engine/Plugins/Experimental/Mutable/Source/CustomizableObject/Private/MuCO/CustomizableObjectInstance.cpp
--- a/Engine/Plugins/Experimental/Mutable/Source/CustomizableObject/Private/MuCO/CustomizableObjectInstance.cpp
+++ b/Engine/Plugins/Experimental/Mutable/Source/CustomizableObject/Private/MuCO/CustomizableObjectInstance.cpp
@@ -2526,6 +2526,7 @@ void UCustomizableInstancePrivateData::DiscardResourcesAndSetReferenceSkeletalMe
 	ClearAllLastMeshIds();
 	
 	Public->SkeletalMeshes.Reset();
+	Public->DescriptorRuntimeHash = FDescriptorRuntimeHash();
 
 	for (TObjectIterator<UCustomizableSkeletalComponent> It; It; ++It)
 	{
