[Mutable] Fix Instance LOD being set to early in the update process.

- Fixed Min Max LOD being commited to late during the Instance update process.
- Fixed Min Max LOD not being specified when creating an Instance.

diff --git a/Engine/Plugins/Experimental/Mutable/Source/CustomizableObject/Classes/MuCO/CustomizableObjectInstance.h b/Engine/Plugins/Experimental/Mutable/Source/CustomizableObject/Classes/MuCO/CustomizableObjectInstance.h
--- a/Engine/Plugins/Experimental/Mutable/Source/CustomizableObject/Classes/MuCO/CustomizableObjectInstance.h
+++ b/Engine/Plugins/Experimental/Mutable/Source/CustomizableObject/Classes/MuCO/CustomizableObjectInstance.h
@@ -601,11 +601,20 @@ public:
 	void SetIsPlayerOrNearIt(bool NewValue);
 	float GetMinSquareDistToPlayer() const;
 	void SetMinSquareDistToPlayer(float NewValue);
-	void SetMinMaxLODToLoad(int32 NewMinLOD, int32 NewMaxLOD, bool LimitLODUpgrades = true);
+	void SetMinMaxLODToLoad(int32 NewMinLOD = 0, int32 NewMaxLOD = INT32_MAX, bool LimitLODUpgrades = true);
 	int32 GetMinLODToLoad() const;
 	int32 GetMaxLODToLoad() const;
 	int32 GetNumLODsAvailable() const;
 
+	/** Return the Min LOD this Instance is using (from the beginning of an update. If an update fails this value will be incorrect). */
+	int32 GetCurrentMinLOD() const;
+
+	/** Return the Max LOD this Instance is using (from the beginning of an update. If an update fails this value will be incorrect). */
+	int32 GetCurrentMaxLOD() const;
+
+	/** Save the Min and Max LOD that will be used for the update. */
+	void CommitMinMaxLOD();
+	
 	/* If enabled, CurrentMinLOD will be the first LOD of the generated SkeletalMesh. Otherwise the number of LODs will remain constant and LODs [0 .. CurrentMinLOD] will share the same RenderData. 
 	 * Enabled by default */
 	void SetUseCurrentMinLODAsBaseLOD(bool bIsBaseLOD);
@@ -688,6 +697,10 @@ private:
 	/** Hash of the UCustomizableObjectInstance::Descriptor on the last successful update. */
 	FDescriptorRuntimeHash DescriptorRuntimeHash;
 
+	/** LODs applied on the beginning of the last update. Represent the actual LODs the Instance is using (not strictly true since an update can fail). */
+	int32 CurrentMinLOD = -1;
+	int32 CurrentMaxLOD = -1;
+
 #if WITH_EDITORONLY_DATA
 	/** Textures used in the Texture Parameters. */
 	UPROPERTY(EditAnywhere, Category = TextureParameter)
diff --git a/Engine/Plugins/Experimental/Mutable/Source/CustomizableObject/Private/MuCO/CustomizableInstanceLODManagement.cpp b/Engine/Plugins/Experimental/Mutable/Source/CustomizableObject/Private/MuCO/CustomizableInstanceLODManagement.cpp
--- a/Engine/Plugins/Experimental/Mutable/Source/CustomizableObject/Private/MuCO/CustomizableInstanceLODManagement.cpp
+++ b/Engine/Plugins/Experimental/Mutable/Source/CustomizableObject/Private/MuCO/CustomizableInstanceLODManagement.cpp
@@ -217,7 +217,7 @@ void UCustomizableInstanceLODManagement::UpdateInstanceDistsAndLODs()
 			{
 				if (!bAlreadyReachedFixedLOD && SortedInstances[i]->GetMinSquareDistToPlayer() < DistanceForFixedLODSquared)
 				{
-					SortedInstances[i]->SetMinMaxLODToLoad(0, INT32_MAX);
+					SortedInstances[i]->SetMinMaxLODToLoad();
 				}
 				else
 				{
@@ -254,7 +254,7 @@ void UCustomizableInstanceLODManagement::UpdateInstanceDistsAndLODs()
 			// No limit surpassed, set all instances to have all LODs, there will be an UpdateSkeletalMesh only if there's a change in LOD state
 			for (int32 i = 0; i < SortedInstances.Num(); ++i)
 			{
-				SortedInstances[i]->SetMinMaxLODToLoad(0, INT32_MAX);
+				SortedInstances[i]->SetMinMaxLODToLoad();
 			}
 		}
 	}
diff --git a/Engine/Plugins/Experimental/Mutable/Source/CustomizableObject/Private/MuCO/CustomizableInstancePrivateData.h b/Engine/Plugins/Experimental/Mutable/Source/CustomizableObject/Private/MuCO/CustomizableInstancePrivateData.h
--- a/Engine/Plugins/Experimental/Mutable/Source/CustomizableObject/Private/MuCO/CustomizableInstancePrivateData.h
+++ b/Engine/Plugins/Experimental/Mutable/Source/CustomizableObject/Private/MuCO/CustomizableInstancePrivateData.h
@@ -279,10 +279,6 @@ public:
 												
 	double LastUpdateTime;
 
-	// LODs applied on the last update. Represent the actual LODs the Instance is using.
-	int32 LastUpdateMinLOD = -1;
-	int32 LastUpdateMaxLOD = -1;
-
 	/** Saves the LODs requested on the update. */
 	void SaveMinMaxLODToLoad(const UCustomizableObjectInstance* Public);
 	
diff --git a/Engine/Plugins/Experimental/Mutable/Source/CustomizableObject/Private/MuCO/CustomizableObjectInstance.cpp b/Engine/Plugins/Experimental/Mutable/Source/CustomizableObject/Private/MuCO/CustomizableObjectInstance.cpp
--- a/Engine/Plugins/Experimental/Mutable/Source/CustomizableObject/Private/MuCO/CustomizableObjectInstance.cpp
+++ b/Engine/Plugins/Experimental/Mutable/Source/CustomizableObject/Private/MuCO/CustomizableObjectInstance.cpp
@@ -145,13 +145,6 @@ UTexture2D* UCustomizableInstancePrivateData::CreateTexture()
 }
 
 
-void UCustomizableInstancePrivateData::SaveMinMaxLODToLoad(const UCustomizableObjectInstance* Public)
-{
-	LastUpdateMinLOD = Public->Descriptor.MinLODToLoad;
-	LastUpdateMaxLOD = Public->Descriptor.MaxLODToLoad;
-}
-
-
 int32 UCustomizableInstancePrivateData::GetLastMeshId(int32 ComponentIndex, int32 LODIndex) const
 {
 	const FCustomizableInstanceComponentData* ComponentData = GetComponentData(ComponentIndex);
@@ -232,7 +225,7 @@ void UCustomizableInstancePrivateData::SetMinMaxLODToLoad(UCustomizableObjectIns
 	// Min LOD
 	NewMinLOD = FMath::Min(FMath::Max(NewMinLOD, static_cast<int32>(FirstLODAvailable)), FirstLODAvailable + NumMaxLODsToStream);
 
-	const bool bIsDowngradeLODUpdate = LastUpdateMinLOD >= 0 && NewMinLOD > LastUpdateMinLOD;
+	const bool bIsDowngradeLODUpdate = Public->GetCurrentMinLOD() >= 0 && NewMinLOD > Public->GetCurrentMinLOD();
 	SetCOInstanceFlags(bIsDowngradeLODUpdate ? PendingLODsDowngrade : ECONone);
 
 	if (!bIsDowngradeLODUpdate && bLimitLODUpgrades)
@@ -246,9 +239,9 @@ void UCustomizableInstancePrivateData::SetMinMaxLODToLoad(UCustomizableObjectIns
 	// Max LOD
 	if (NumLODsAvailable != INT32_MAX)
 	{
-		if (Public->Descriptor.MaxLODToLoad == INT32_MAX)
+		if (Public->Descriptor.MaxLOD == INT32_MAX)
 		{
-			Public->Descriptor.MaxLODToLoad = NumLODsAvailable - 1;
+			Public->Descriptor.MaxLOD = NumLODsAvailable - 1;
 		}
 
 		NewMaxLOD = FMath::Min(NewMaxLOD, NumLODsAvailable - 1);
@@ -257,10 +250,20 @@ void UCustomizableInstancePrivateData::SetMinMaxLODToLoad(UCustomizableObjectIns
 	NewMaxLOD = FMath::Max(NewMaxLOD, NewMinLOD);
 
 	// Save the new LODs
-	SetCOInstanceFlags(Public->Descriptor.MinLODToLoad != NewMinLOD || Public->Descriptor.MaxLODToLoad != NewMaxLOD ? PendingLODsUpdate : ECONone);
+	const bool bLODChanged = Public->Descriptor.MinLOD != NewMinLOD || Public->Descriptor.MaxLOD != NewMaxLOD;
 	
-	Public->Descriptor.MinLODToLoad = NewMinLOD;
-	Public->Descriptor.MaxLODToLoad = NewMaxLOD;
+	Public->Descriptor.MinLOD = NewMinLOD;
+	Public->Descriptor.MaxLOD = NewMaxLOD;
+
+	SetCOInstanceFlags(bLODChanged ? PendingLODsUpdate : ECONone);
+
+	if (bLODChanged)
+	{
+		if (const FMutableQueueElem* QueueElem = UCustomizableObjectSystem::GetInstance()->GetPrivate()->MutableOperationQueue.Get(Public))
+		{
+			QueueElem->Operation->InstanceDescriptorRuntimeHash.UpdateMinMaxLOD(NewMinLOD, NewMaxLOD);
+		}
+	}
 }
 
 
@@ -551,6 +554,7 @@ void UCustomizableObjectInstance::PostLoad()
 	Super::PostLoad();
 
 	Descriptor.ReloadParameters();
+	SetMinMaxLODToLoad();
 }
 
 
@@ -687,6 +691,7 @@ void UCustomizableObjectInstance::SetObject(UCustomizableObject* InObject)
 {
 	Descriptor.SetCustomizableObject(*InObject);
 	PrivateData->ReloadParameters(this);
+	SetMinMaxLODToLoad();
 }
 
 
@@ -902,7 +907,7 @@ void UCustomizableInstancePrivateData::DoUpdateSkeletalMesh(UCustomizableObjectI
 		if (bShouldUpdateLODs)
 		{
 			SetCOInstanceFlags(PendingLODsUpdateSecondStage);
-			UE_LOG(LogMutable, Verbose, TEXT("LOD change: %d, %d -> %d, %d"), LastUpdateMinLOD, LastUpdateMaxLOD, CustomizableInstance.Descriptor.MinLODToLoad, CustomizableInstance.Descriptor.MaxLODToLoad);
+			UE_LOG(LogMutable, Verbose, TEXT("LOD change: %d, %d -> %d, %d"), CustomizableInstance.GetCurrentMinLOD(), CustomizableInstance.GetCurrentMaxLOD(), CustomizableInstance.GetMinLODToLoad(), CustomizableInstance.GetMinLODToLoad());
 		}
 
 		if ((!(bIsGenerated && bOnlyUpdateIfNotGenerated) && !(bIsCloseDistTick && bIsGenerated)) || bShouldUpdateLODs)
@@ -2628,6 +2633,13 @@ FDescriptorRuntimeHash UCustomizableObjectInstance::GetUpdateDescriptorRuntimeHa
 }
 
 
+void UCustomizableObjectInstance::CommitMinMaxLOD()
+{
+	CurrentMinLOD = Descriptor.MinLOD;
+	CurrentMaxLOD = Descriptor.MaxLOD;	
+}
+
+
 void CopyTextureProperties(UTexture2D* Texture, const UTexture2D* SourceTexture)
 {
 	MUTABLE_CPUPROFILER_SCOPE(CopyTextureProperties)
@@ -5406,13 +5418,13 @@ void UCustomizableObjectInstance::SetMinMaxLODToLoad(int32 NewMinLOD, int32 NewM
 
 int32 UCustomizableObjectInstance::GetMinLODToLoad() const
 {
-	return Descriptor.MinLODToLoad;	
+	return Descriptor.MinLOD;	
 }
 
 
 int32 UCustomizableObjectInstance::GetMaxLODToLoad() const
 {
-	return Descriptor.MaxLODToLoad;	
+	return Descriptor.MaxLOD;	
 }
 
 
@@ -5434,6 +5446,18 @@ int32 UCustomizableObjectInstance::GetNumLODsAvailable() const
 }
 
 
+int32 UCustomizableObjectInstance::GetCurrentMinLOD() const
+{
+	return CurrentMinLOD;
+}
+
+
+int32 UCustomizableObjectInstance::GetCurrentMaxLOD() const
+{
+	return CurrentMaxLOD;
+}
+
+
 USkeletalMesh* UCustomizableObjectInstance::GetSkeletalMesh(int32 ComponentIndex) const
 {
 	return SkeletalMeshes.IsValidIndex(ComponentIndex) ? SkeletalMeshes[ComponentIndex] : nullptr;
diff --git a/Engine/Plugins/Experimental/Mutable/Source/CustomizableObject/Private/MuCO/CustomizableObjectInstanceDescriptor.cpp b/Engine/Plugins/Experimental/Mutable/Source/CustomizableObject/Private/MuCO/CustomizableObjectInstanceDescriptor.cpp
--- a/Engine/Plugins/Experimental/Mutable/Source/CustomizableObject/Private/MuCO/CustomizableObjectInstanceDescriptor.cpp
+++ b/Engine/Plugins/Experimental/Mutable/Source/CustomizableObject/Private/MuCO/CustomizableObjectInstanceDescriptor.cpp
@@ -2310,8 +2310,8 @@ FDescriptorHash::FDescriptorHash(const FCustomizableObjectInstanceDescriptor& De
 	
 	Hash = HashCombine(Hash, GetTypeHash(Descriptor.State));
 	Hash = HashCombine(Hash, GetTypeHash(Descriptor.bBuildParameterDecorations));
-	Hash = HashCombine(Hash, GetTypeHash(Descriptor.MinLODToLoad));
-	Hash = HashCombine(Hash, GetTypeHash(Descriptor.MaxLODToLoad));
+	Hash = HashCombine(Hash, GetTypeHash(Descriptor.MinLOD));
+	Hash = HashCombine(Hash, GetTypeHash(Descriptor.MaxLOD));
 
 	for (const TTuple<FName, FMultilayerProjector>& Pair : Descriptor.MultilayerProjectors)
 	{
@@ -2348,14 +2348,33 @@ FString FDescriptorHash::ToString() const
 FDescriptorRuntimeHash::FDescriptorRuntimeHash(const FCustomizableObjectInstanceDescriptor& Descriptor) :
 	FDescriptorHash(Descriptor)
 {
-	MinLODToLoad = Descriptor.MinLODToLoad;
-	MaxLODToLoad = Descriptor.MaxLODToLoad;	
+	MinLOD = Descriptor.MinLOD;
+	MaxLOD = Descriptor.MaxLOD;	
 }
 
 
 bool FDescriptorRuntimeHash::IsSubset(const FDescriptorRuntimeHash& Other) const
 {
 	return FDescriptorHash::operator==(Other) &&
-		MinLODToLoad >= Other.MinLODToLoad &&
-		MaxLODToLoad <= Other.MaxLODToLoad;
+		MinLOD >= Other.MinLOD &&
+		MaxLOD <= Other.MaxLOD;
+}
+
+
+void FDescriptorRuntimeHash::UpdateMinMaxLOD(const int32 InMinLOD, const int32 InMaxLOD)
+{
+	MinLOD = InMinLOD;
+	MaxLOD = InMaxLOD;	
+}
+
+
+int32 FDescriptorRuntimeHash::GetMinLOD() const
+{
+	return MinLOD;
+}
+
+
+int32 FDescriptorRuntimeHash::GetMaxLOD() const
+{
+	return MaxLOD;
 }
diff --git a/Engine/Plugins/Experimental/Mutable/Source/CustomizableObject/Private/MuCO/CustomizableObjectSystem.cpp b/Engine/Plugins/Experimental/Mutable/Source/CustomizableObject/Private/MuCO/CustomizableObjectSystem.cpp
--- a/Engine/Plugins/Experimental/Mutable/Source/CustomizableObject/Private/MuCO/CustomizableObjectSystem.cpp
+++ b/Engine/Plugins/Experimental/Mutable/Source/CustomizableObject/Private/MuCO/CustomizableObjectSystem.cpp
@@ -524,11 +524,11 @@ void FCustomizableObjectSystemPrivate::UpdateStats()
 					{
 						CountAllocatedSkeletalMesh++;
 
-						if (CustomizableObjectInstance->GetPrivate()->LastUpdateMinLOD < 1)
+						if (CustomizableObjectInstance->GetCurrentMinLOD() < 1)
 						{
 							++CountLOD0;
 						}
-						else if (CustomizableObjectInstance->GetPrivate()->LastUpdateMaxLOD < 2)
+						else if (CustomizableObjectInstance->GetCurrentMaxLOD() < 2)
 						{
 							++CountLOD1;
 						}
@@ -813,8 +813,6 @@ void FCustomizableObjectSystemPrivate::InitUpdateSkeletalMesh(UCustomizableObjec
 	
 	const TSharedPtr<FMutableOperation> Operation = MakeShared<FMutableOperation>(FMutableOperation::CreateInstanceUpdate(&Instance, bNeverStream, MipsToSkip));
 
-	Instance.GetPrivate()->SaveMinMaxLODToLoad(&Instance);
-
 	const FDescriptorRuntimeHash UpdateDescriptorHash = Instance.GetUpdateDescriptorRuntimeHash();
 
 	if (const FMutableQueueElem* QueueElem = MutableOperationQueue.Get(&Instance))
@@ -2170,20 +2168,15 @@ namespace impl
 		// Prepare streaming for the current customizable object
 		System->GetPrivate()->Streamer->PrepareStreamingForObject(CustomizableObject);
 
-		// Ensure we're not requesting invalid LODs
-		PrivateInstance->LastUpdateMinLOD = FMath::Max(PrivateInstance->LastUpdateMinLOD, (int32)PrivateInstance->FirstLODAvailable);
-		PrivateInstance->LastUpdateMinLOD = FMath::Min(PrivateInstance->LastUpdateMinLOD, PrivateInstance->FirstLODAvailable + PrivateInstance->NumMaxLODsToStream);
-
-		PrivateInstance->LastUpdateMaxLOD = FMath::Max(PrivateInstance->LastUpdateMaxLOD, PrivateInstance->LastUpdateMinLOD);
-
-
+		CandidateInstance->CommitMinMaxLOD();
+		
 		// Task: Mutable Update and GetMesh
 		//-------------------------------------------------------------
 		TSharedPtr<FMutableOperationData> CurrentOperationData = MakeShared<FMutableOperationData>();
 		CurrentOperationData->TextureCoverageQueries_MutableThreadParams = PrivateInstance->TextureCoverageQueries;
 		CurrentOperationData->TextureCoverageQueries_MutableThreadResults.Empty();
-		CurrentOperationData->CurrentMinLOD = PrivateInstance->LastUpdateMinLOD;
-		CurrentOperationData->CurrentMaxLOD = PrivateInstance->LastUpdateMaxLOD;
+		CurrentOperationData->CurrentMinLOD = Operation->InstanceDescriptorRuntimeHash.GetMinLOD();
+		CurrentOperationData->CurrentMaxLOD = Operation->InstanceDescriptorRuntimeHash.GetMaxLOD();
 		CurrentOperationData->bNeverStream = Operation->bNeverStream;
 		CurrentOperationData->MipsToSkip = Operation->MipsToSkip;
 		CurrentOperationData->MutableParameters = Parameters;
@@ -2361,7 +2354,11 @@ bool UCustomizableObjectSystem::Tick(float DeltaTime)
 
 	MUTABLE_CPUPROFILER_SCOPE(TickCustomizableObjectSystem)
 
-
+	// Reset the instance relevancy
+	// \TODO: Review
+	// TODO: This should be done only when requiring a new job. And for the current operation instance, in case it needs to be cancelled.
+	CurrentInstanceLODManagement->UpdateInstanceDistsAndLODs();
+	
 	// Get a new operation if we aren't working on one
 	if (!Private->CurrentMutableOperation)
 	{
@@ -2380,19 +2377,13 @@ bool UCustomizableObjectSystem::Tick(float DeltaTime)
 			Private->CurrentMutableOperation = nullptr;
 		}
 	}
-
+	
 	// Advance the current operation
 	if (Private->CurrentMutableOperation)
 	{
 		AdvanceCurrentOperation();
 	}
 
-
-	// Reset the instance relevancy
-	// \TODO: Review
-	// TODO: This should be done only when requiring a new job. And for the current operation instance, in case it needs to be cancelled.
-	CurrentInstanceLODManagement->UpdateInstanceDistsAndLODs();
-
 	for (TObjectIterator<UCustomizableObjectInstance> CustomizableObjectInstance; CustomizableObjectInstance; ++CustomizableObjectInstance)
 	{
 		if (IsValidChecked(*CustomizableObjectInstance) && CustomizableObjectInstance->GetPrivate())
diff --git a/Engine/Plugins/Experimental/Mutable/Source/CustomizableObject/Private/MuCO/LogInformationUtil.cpp b/Engine/Plugins/Experimental/Mutable/Source/CustomizableObject/Private/MuCO/LogInformationUtil.cpp
--- a/Engine/Plugins/Experimental/Mutable/Source/CustomizableObject/Private/MuCO/LogInformationUtil.cpp
+++ b/Engine/Plugins/Experimental/Mutable/Source/CustomizableObject/Private/MuCO/LogInformationUtil.cpp
@@ -125,8 +125,8 @@ void LogInformationUtil::LogShowInstanceData(const UCustomizableObjectInstance*
 	LogData += FString::Printf(TEXT("bIsUpdating=%d "), InstanceFlags & ECOInstanceFlags::Updating);
 	LogData += FString::Printf(TEXT("bShouldUpdateLODs=%d "), InstanceFlags & ECOInstanceFlags::PendingLODsUpdate);
 
-	LogData += FString::Printf(TEXT("CurrentMinLOD=%d "), CustomizableObjectInstance->GetPrivate()->LastUpdateMinLOD);
-	LogData += FString::Printf(TEXT("CurrentMaxLOD=%d "), CustomizableObjectInstance->GetPrivate()->LastUpdateMaxLOD);
+	LogData += FString::Printf(TEXT("CurrentMinLOD=%d "), CustomizableObjectInstance->GetCurrentMinLOD());
+	LogData += FString::Printf(TEXT("CurrentMaxLOD=%d "), CustomizableObjectInstance->GetCurrentMaxLOD());
 	LogData += FString::Printf(TEXT("MinLODToLoad=%d "), CustomizableObjectInstance->GetMinLODToLoad());
 	LogData += FString::Printf(TEXT("MaxLODToLoad=%d\n"), CustomizableObjectInstance->GetMaxLODToLoad());
 }
@@ -174,8 +174,8 @@ void LogInformationUtil::LogShowInstanceDataFull(const UCustomizableObjectInstan
 	MessageChunk += FString::Printf(TEXT("        LastMinSquareDistFromComponentToPlayer = %.2f\n"), CustomizableObjectInstance->GetPrivate()->LastMinSquareDistFromComponentToPlayer);
 	LogData += MessageChunk;
 
-	MessageChunk = FString::Printf(TEXT("        CurrentMinLOD = %d\n"), CustomizableObjectInstance->GetPrivate()->LastUpdateMinLOD);
-	MessageChunk += FString::Printf(TEXT("        CurrentMaxLOD = %d\n"), CustomizableObjectInstance->GetPrivate()->LastUpdateMaxLOD);
+	MessageChunk = FString::Printf(TEXT("        CurrentMinLOD = %d\n"), CustomizableObjectInstance->GetCurrentMinLOD());
+	MessageChunk += FString::Printf(TEXT("        CurrentMaxLOD = %d\n"), CustomizableObjectInstance->GetCurrentMaxLOD());
 	MessageChunk += FString::Printf(TEXT("        MinLODToLoad = %d\n"), CustomizableObjectInstance->GetMinLODToLoad());
 	MessageChunk += FString::Printf(TEXT("        MaxLODToLoad = %d\n"), CustomizableObjectInstance->GetMaxLODToLoad());
 	MessageChunk += FString::Printf(TEXT("        bShouldUpdateLODs = %d\n"), InstanceFlags & ECOInstanceFlags::PendingLODsUpdate);
@@ -192,11 +192,11 @@ void LogInformationUtil::LogShowInstanceDataFull(const UCustomizableObjectInstan
 	{
 		if ((CustomizableObjectInstance->GetSkeletalMesh(ComponentIndex) != nullptr) && (CustomizableObjectInstance->GetSkeletalMesh(ComponentIndex)->GetResourceForRendering()))
 		{
-			if (CustomizableObjectInstance->GetPrivate()->LastUpdateMinLOD < 1)
+			if (CustomizableObjectInstance->GetCurrentMinLOD() < 1)
 			{
 				CountLOD0++;
 			}
-			else if (CustomizableObjectInstance->GetPrivate()->LastUpdateMinLOD < 2)
+			else if (CustomizableObjectInstance->GetCurrentMinLOD() < 2)
 			{
 				CountLOD1++;
 			}
diff --git a/Engine/Plugins/Experimental/Mutable/Source/CustomizableObject/Public/MuCO/CustomizableObjectInstanceDescriptor.h b/Engine/Plugins/Experimental/Mutable/Source/CustomizableObject/Public/MuCO/CustomizableObjectInstanceDescriptor.h
--- a/Engine/Plugins/Experimental/Mutable/Source/CustomizableObject/Public/MuCO/CustomizableObjectInstanceDescriptor.h
+++ b/Engine/Plugins/Experimental/Mutable/Source/CustomizableObject/Public/MuCO/CustomizableObjectInstanceDescriptor.h
@@ -349,8 +349,8 @@ private:
 	bool bBuildParameterDecorations = false;
 
 	/** These are the LODs we want to have, they MUST NOT be used in an update (Mutable thread). */
-	int32 MinLODToLoad = 0;
-	int32 MaxLODToLoad = INT32_MAX;
+	int32 MinLOD = 0;
+	int32 MaxLOD = INT32_MAX;
 
 	/** Lookup of UCustomizableObjectDescriptor::IntParameters. */
 	TMap<FString, int32> IntParametersLookupTable;
@@ -404,8 +404,14 @@ public:
 
 	/** Return true if this Hash is a subset of the other Hash (i.e., this Descriptor is a subset of the other Descriptor). */
 	bool IsSubset(const FDescriptorRuntimeHash& Other) const;
+
+	void UpdateMinMaxLOD(int32 InMinLOD, int32 InMaxLOD);
+
+	int32 GetMinLOD() const;
+
+	int32 GetMaxLOD() const;
 	
 private:
-	int32 MinLODToLoad = 0;
-	int32 MaxLODToLoad = INT32_MAX;
+	int32 MinLOD = 0;
+	int32 MaxLOD = INT32_MAX;
 };
\ No newline at end of file
