[mutable] Add both way sort in PinViewer table

#rnx

diff --git a/Engine/Plugins/Experimental/Mutable/Source/CustomizableObjectEditor/Private/MuCOE/PinViewer/SPinViewer.cpp b/Engine/Plugins/Experimental/Mutable/Source/CustomizableObjectEditor/Private/MuCOE/PinViewer/SPinViewer.cpp
--- a/Engine/Plugins/Experimental/Mutable/Source/CustomizableObjectEditor/Private/MuCOE/PinViewer/SPinViewer.cpp
+++ b/Engine/Plugins/Experimental/Mutable/Source/CustomizableObjectEditor/Private/MuCOE/PinViewer/SPinViewer.cpp
@@ -58,7 +58,7 @@ FText SPinViewer::GetPinName(const UEdGraphPin& Pin)
 
 
 /** Sort pin references by name. The pin references to be valid. */
-bool SortPinsByName(const TSharedPtr<FEdGraphPinReference>& PinReferenceA, const TSharedPtr<FEdGraphPinReference>& PinReferenceB)
+bool SortPinsByNameAsc(const TSharedPtr<FEdGraphPinReference>& PinReferenceA, const TSharedPtr<FEdGraphPinReference>& PinReferenceB)
 {
 	const FText PinNameA = SPinViewer::GetPinName(*PinReferenceA->Get());
 	const FText PinNameB = SPinViewer::GetPinName(*PinReferenceB->Get());
@@ -66,49 +66,106 @@ bool SortPinsByName(const TSharedPtr<FEdGraphPinReference>& PinReferenceA, const
 	return PinNameA.CompareTo(PinNameB, ETextComparisonLevel::Default) < 0;
 }
 
+bool SortPinsByNameDesc(const TSharedPtr<FEdGraphPinReference>& PinReferenceA, const TSharedPtr<FEdGraphPinReference>& PinReferenceB)
+{
+	const FText PinNameA = SPinViewer::GetPinName(*PinReferenceA->Get());
+	const FText PinNameB = SPinViewer::GetPinName(*PinReferenceB->Get());
+
+	return PinNameA.CompareTo(PinNameB, ETextComparisonLevel::Default) >= 0;
+}
+
 
 /** Sort pin references by the provided method. The pin references have to be valid. */
-void Sort(const FName& Method, TArray<TSharedPtr<FEdGraphPinReference>>& Pins)
+void Sort(const FName& Method, const EColumnSortMode::Type& SortMode, TArray<TSharedPtr<FEdGraphPinReference>>& Pins)
 {
-	if (Method == SPinViewer::COLUMN_NAME)
+	if (Method == SPinViewer::COLUMN_NAME || SortMode == EColumnSortMode::None )
 	{
-		Pins.Sort(&SortPinsByName);
+		if (SortMode == EColumnSortMode::Descending)
+		{
+			Pins.Sort(&SortPinsByNameDesc);
+		}
+		else
+		{
+			Pins.Sort(&SortPinsByNameAsc);
+		}
 	}
 	else if (Method == SPinViewer::COLUMN_TYPE)
 	{
-		Pins.Sort([](const TSharedPtr<FEdGraphPinReference>& PinReferenceA, const TSharedPtr<FEdGraphPinReference>& PinReferenceB)
-			{
-				if (PinReferenceA->Get()->PinType.PinCategory == PinReferenceB->Get()->PinType.PinCategory)
-				{
-					return SortPinsByName(PinReferenceA, PinReferenceB);
-				}
-				else if (PinReferenceA->Get()->PinType.PinCategory.FastLess(PinReferenceB->Get()->PinType.PinCategory))
+		if (SortMode == EColumnSortMode::Descending)
+		{
+			Pins.Sort([](const TSharedPtr<FEdGraphPinReference>& PinReferenceA, const TSharedPtr<FEdGraphPinReference>& PinReferenceB)
 				{
-					return true;
-				}
-				else
+					if (PinReferenceA->Get()->PinType.PinCategory == PinReferenceB->Get()->PinType.PinCategory)
+					{
+						return SortPinsByNameDesc(PinReferenceA, PinReferenceB);
+					}
+					else if (PinReferenceA->Get()->PinType.PinCategory.FastLess(PinReferenceB->Get()->PinType.PinCategory))
+					{
+						return false;
+					}
+					else
+					{
+						return true;
+					}
+				});
+		}
+		else
+		{
+			Pins.Sort([](const TSharedPtr<FEdGraphPinReference>& PinReferenceA, const TSharedPtr<FEdGraphPinReference>& PinReferenceB)
 				{
-					return false;
-				}
-			});
+					if (PinReferenceA->Get()->PinType.PinCategory == PinReferenceB->Get()->PinType.PinCategory)
+					{
+						return SortPinsByNameAsc(PinReferenceA, PinReferenceB);
+					}
+					else if (PinReferenceA->Get()->PinType.PinCategory.FastLess(PinReferenceB->Get()->PinType.PinCategory))
+					{
+						return true;
+					}
+					else
+					{
+						return false;
+					}
+				});
+		}
 	}
 	else if (Method == SPinViewer::COLUMN_VISIBILITY)
 	{
-		Pins.Sort([](const TSharedPtr<FEdGraphPinReference>& PinReferenceA, const TSharedPtr<FEdGraphPinReference>& PinReferenceB)
-			{
-				if (PinReferenceA->Get()->bHidden == PinReferenceB->Get()->bHidden)
-				{
-					return SortPinsByName(PinReferenceA, PinReferenceB);
-				}
-				else if (PinReferenceA->Get()->bHidden < PinReferenceB->Get()->bHidden)
+		if (SortMode == EColumnSortMode::Descending)
+		{
+			Pins.Sort([](const TSharedPtr<FEdGraphPinReference>& PinReferenceA, const TSharedPtr<FEdGraphPinReference>& PinReferenceB)
 				{
-					return false;
-				}
-				else
+					if (PinReferenceA->Get()->bHidden == PinReferenceB->Get()->bHidden)
+					{
+						return SortPinsByNameDesc(PinReferenceA, PinReferenceB);
+					}
+					else if (PinReferenceA->Get()->bHidden < PinReferenceB->Get()->bHidden)
+					{
+						return true;
+					}
+					else
+					{
+						return false;
+					}
+				});
+		}
+		else
+		{
+			Pins.Sort([](const TSharedPtr<FEdGraphPinReference>& PinReferenceA, const TSharedPtr<FEdGraphPinReference>& PinReferenceB)
 				{
-					return true;
-				}
-			});
+					if (PinReferenceA->Get()->bHidden == PinReferenceB->Get()->bHidden)
+					{
+						return SortPinsByNameAsc(PinReferenceA, PinReferenceB);
+					}
+					else if (PinReferenceA->Get()->bHidden < PinReferenceB->Get()->bHidden)
+					{
+						return false;
+					}
+					else
+					{
+						return true;
+					}
+				});
+		}
 	}
 	else
 	{
@@ -128,7 +185,7 @@ void SPinViewer::GeneratePinInfoList()
 		}
 	}
 
-	Sort(CurrentSortColumn, PinReferences);
+	Sort(CurrentSortColumn, SortMode, PinReferences);
 }
 
 
@@ -202,16 +259,19 @@ void SPinViewer::Construct(const FArguments& InArgs)
 				SNew(SHeaderRow)
 				+ SHeaderRow::Column(COLUMN_NAME)
 				.DefaultLabel(LOCTEXT("PinNameLabel", "Name"))
+				.SortMode(this, &SPinViewer::GetColumnSortMode, COLUMN_NAME)
 				.OnSort(this, &SPinViewer::SortListView)
 				.FillWidth(0.5)
 
 				+ SHeaderRow::Column(COLUMN_TYPE)
 				.DefaultLabel(LOCTEXT("TypeLabel", "Type"))
+				.SortMode(this, &SPinViewer::GetColumnSortMode, COLUMN_TYPE)
 				.OnSort(this, &SPinViewer::SortListView)
 				.FillWidth(0.25)
 
 				+ SHeaderRow::Column(COLUMN_VISIBILITY)
 				.DefaultLabel(LOCTEXT("VisibilityLabel", "Visibility"))
+				.SortMode(this, &SPinViewer::GetColumnSortMode, COLUMN_VISIBILITY)
 				.OnSort(this, &SPinViewer::SortListView)
 				.FillWidth(0.25)
 			)
@@ -273,10 +333,22 @@ FReply SPinViewer::OnHideAllPressed() const
 void SPinViewer::SortListView(const EColumnSortPriority::Type SortPriority, const FName& ColumnId, const EColumnSortMode::Type NewSortMode)
 {
 	CurrentSortColumn = ColumnId;
+	SortMode = NewSortMode;
 	UpdateWidget();
 }
 
 
+EColumnSortMode::Type SPinViewer::GetColumnSortMode(const FName ColumnId) const
+{
+	if (CurrentSortColumn != ColumnId)
+	{
+		return EColumnSortMode::None;
+	}
+
+	return SortMode;
+}
+
+
 void PinViewerAttachToDetailCustomization(IDetailLayoutBuilder& DetailBuilder)
 {
 	const TArray<TWeakObjectPtr<UObject>>& SelectedObjects = DetailBuilder.GetDetailsView()->GetSelectedObjects();
diff --git a/Engine/Plugins/Experimental/Mutable/Source/CustomizableObjectEditor/Private/MuCOE/PinViewer/SPinViewer.h b/Engine/Plugins/Experimental/Mutable/Source/CustomizableObjectEditor/Private/MuCOE/PinViewer/SPinViewer.h
--- a/Engine/Plugins/Experimental/Mutable/Source/CustomizableObjectEditor/Private/MuCOE/PinViewer/SPinViewer.h
+++ b/Engine/Plugins/Experimental/Mutable/Source/CustomizableObjectEditor/Private/MuCOE/PinViewer/SPinViewer.h
@@ -75,6 +75,7 @@ public:
 	// Columns sorting methods
 	/** Sets the sorting method to be applied */
 	void SortListView(const EColumnSortPriority::Type SortPriority, const FName& ColumnId, const EColumnSortMode::Type NewSortMode);
+	EColumnSortMode::Type GetColumnSortMode(const FName ColumnId) const;
 
 	/** Get the pin name which is actually displayed. */
 	static FText GetPinName(const UEdGraphPin& Pin);
@@ -97,6 +98,7 @@ private:
 	
 	/** Current column. */
 	FName CurrentSortColumn = COLUMN_TYPE;
+	EColumnSortMode::Type SortMode = EColumnSortMode::Type::None;
 
 	FString CurrentFilter;
 
