[Mutable] Fix clip moprh modifiers not working for some components somtimes.

- Fix mesh generations flags beed disabled.
- Added mesh generation flags to modifiers meshes.
- Now all modifiers affect all components.

diff --git a/Engine/Plugins/Experimental/Mutable/Source/CustomizableObjectEditor/Private/MuCOE/CustomizableObjectLayout.cpp b/Engine/Plugins/Experimental/Mutable/Source/CustomizableObjectEditor/Private/MuCOE/CustomizableObjectLayout.cpp
--- a/Engine/Plugins/Experimental/Mutable/Source/CustomizableObjectEditor/Private/MuCOE/CustomizableObjectLayout.cpp
+++ b/Engine/Plugins/Experimental/Mutable/Source/CustomizableObjectEditor/Private/MuCOE/CustomizableObjectLayout.cpp
@@ -96,8 +96,9 @@ void UCustomizableObjectLayout::GenerateBlocksFromUVs()
 		if (USkeletalMesh* SkeletalMesh = Cast<USkeletalMesh>(Mesh))
 		{
 			// We don't need all the data to generate the blocks
-			const uint32 ShapeFlags = uint32(EMutableMeshConversionFlags::IgnoreSkinning) &
-				uint32(EMutableMeshConversionFlags::IgnorePhysics);
+			const EMutableMeshConversionFlags ShapeFlags = 
+					EMutableMeshConversionFlags::IgnoreSkinning | 
+					EMutableMeshConversionFlags::IgnorePhysics;
 
 			GenerationContext.MeshGenerationFlags.Push(ShapeFlags);
 
diff --git a/Engine/Plugins/Experimental/Mutable/Source/CustomizableObjectEditor/Private/MuCOE/GenerateMutableSource/GenerateMutableSource.cpp b/Engine/Plugins/Experimental/Mutable/Source/CustomizableObjectEditor/Private/MuCOE/GenerateMutableSource/GenerateMutableSource.cpp
--- a/Engine/Plugins/Experimental/Mutable/Source/CustomizableObjectEditor/Private/MuCOE/GenerateMutableSource/GenerateMutableSource.cpp
+++ b/Engine/Plugins/Experimental/Mutable/Source/CustomizableObjectEditor/Private/MuCOE/GenerateMutableSource/GenerateMutableSource.cpp
@@ -159,7 +159,7 @@ FMutableGraphGenerationContext::FMutableGraphGenerationContext(UCustomizableObje
 	: Object(InObject), Compiler(InCompiler), Options(InOptions)
 {
 	// Default flags for mesh generation nodes.
-	MeshGenerationFlags.Push(0);
+	MeshGenerationFlags.Push(EMutableMeshConversionFlags::None);
 }
 
 
@@ -608,7 +608,6 @@ mu::NodeMeshApplyPosePtr CreateNodeMeshApplyPose(mu::NodeMeshPtr InputMeshNode,
 	return NodeMeshApplyPose;
 }
 
-
 // Convert a CustomizableObject Source Graph into a mutable source graph  
 mu::NodeObjectPtr GenerateMutableSource(const UEdGraphPin * Pin, FMutableGraphGenerationContext & GenerationContext, bool bPartialCompilation)
 {
@@ -645,7 +644,7 @@ mu::NodeObjectPtr GenerateMutableSource(const UEdGraphPin * Pin, FMutableGraphGe
 		ObjectNode->SetUid(TCHAR_TO_ANSI(*GenerationContext.GetNodeIdUnique(TypedNodeObj).ToString()));
 
 		// LOD
-		int NumLODs = TypedNodeObj->GetNumLODPins();
+		const int32 NumLODs = TypedNodeObj->GetNumLODPins();
 
 		// Fill the basic LOD Settings
 		if (!GenerationContext.NumLODsInRoot)
@@ -767,116 +766,130 @@ mu::NodeObjectPtr GenerateMutableSource(const UEdGraphPin * Pin, FMutableGraphGe
 			GenerationContext.NumMeshComponentsInRoot = TypedNodeObj->NumMeshComponents;
 		}
 
-		while (GenerationContext.ComponentNewNode.Last().ComponentsPerLOD.Num() < GenerationContext.NumLODsInRoot)
+		const int32 NumLODsInRoot = GenerationContext.NumLODsInRoot;
+		const int32 NumMeshComponentsInRoot = GenerationContext.NumMeshComponentsInRoot;
+		while (GenerationContext.ComponentNewNode.Last().ComponentsPerLOD.Num() < NumLODsInRoot)
 		{
-			TArray<mu::NodeComponentNewPtr> AuxArray;
-			AuxArray.AddZeroed(GenerationContext.NumMeshComponentsInRoot);
-			GenerationContext.ComponentNewNode.Last().ComponentsPerLOD.Add(AuxArray);
+			GenerationContext.ComponentNewNode.Last().ComponentsPerLOD.Emplace_GetRef()
+				.Init(nullptr, NumMeshComponentsInRoot);
 		}
 
 		int32 LastDefinedLOD = -1;
-		ObjectNode->SetLODCount(GenerationContext.NumLODsInRoot);
-		for (int32 LODIndex = 0; LODIndex < GenerationContext.NumLODsInRoot; ++LODIndex)
+		ObjectNode->SetLODCount(NumLODsInRoot);
+		for (int32 LODIndex = 0; LODIndex < NumLODsInRoot; ++LODIndex)
 		{
 			GenerationContext.CurrentLOD = LODIndex;
 
 			mu::NodeLODPtr LODNode = new mu::NodeLOD();
 			ObjectNode->SetLOD(LODIndex, LODNode);
+
 			LODNode->SetMessageContext(Node);
 			LODNode->SetComponentCount(GenerationContext.NumMeshComponentsInRoot);
+		
+			TArray<mu::NodeComponentPtr> ComponentNodes;
+			ComponentNodes.Init(nullptr, NumMeshComponentsInRoot);
 
-			for (int32 MeshComponentIndex = 0; MeshComponentIndex < GenerationContext.NumMeshComponentsInRoot; ++MeshComponentIndex)
+			// Generate a component node for every component in root regardless of if will be populated. 
+			// Not sure this is what we we want, but is what it was done before.
+			// TODO: Review if this is the behaviour we want.
+			for (int32 MeshComponentIndex = 0; MeshComponentIndex < NumMeshComponentsInRoot; ++MeshComponentIndex)
 			{
-				GenerationContext.CurrentMeshComponent = MeshComponentIndex;
-				mu::NodeComponentPtr ComponentNode;
+				mu::NodeComponentPtr& ComponentNode = ComponentNodes[MeshComponentIndex];
 
-				if (!GenerationContext.ComponentNewNode.Last().ComponentsPerLOD[LODIndex][MeshComponentIndex])
-				{
-					mu::NodeComponentNewPtr ComponentNewNode = new mu::NodeComponentNew();
-					ComponentNewNode->SetMessageContext(Node);
-					ComponentNewNode->SetId(MeshComponentIndex);
-					GenerationContext.ComponentNewNode.Last().ComponentsPerLOD[LODIndex][MeshComponentIndex] = ComponentNewNode;
-					ComponentNode = ComponentNewNode;
-				}
-				else
+				ComponentNode = Invoke([&GenerationContext, LODIndex, MeshComponentIndex]() -> mu::NodeComponentPtr
 				{
-					mu::NodeComponentEditPtr ComponentEditNode = new mu::NodeComponentEdit();
-					ComponentEditNode->SetMessageContext(Node);
-					ComponentEditNode->SetParent(GenerationContext.ComponentNewNode.Last().ComponentsPerLOD[LODIndex][MeshComponentIndex].get());
-					ComponentNode = ComponentEditNode;
-				}
+					mu::NodeComponentNewPtr& ParentComponent =
+							GenerationContext.ComponentNewNode.Last().ComponentsPerLOD[LODIndex][MeshComponentIndex];
+
+					if (!ParentComponent)
+					{
+						ParentComponent = new mu::NodeComponentNew();
+						ParentComponent->SetId(MeshComponentIndex);
+				
+						return ParentComponent;
+					}
+					
+					mu::NodeComponentEditPtr EditComponent = new mu::NodeComponentEdit();
+					EditComponent->SetParent(ParentComponent.get());
+
+					return EditComponent;
+				});
 
+				ComponentNode->SetMessageContext(Node);
 				LODNode->SetComponent(MeshComponentIndex, ComponentNode);
+			}
 
-				// Is the LOD defined in this object?
-				int32 LODToGenerate = -1;
-				if (LODIndex < NumLODs)
+			if (GenerationContext.CurrentLOD < GenerationContext.FirstLODAvailable)
+			{
+				continue;
+			} 
+
+			const bool bUseAutomaticLods = 
+					GenerationContext.CurrentAutoLODStrategy == ECustomizableObjectAutomaticLODStrategy::AutomaticFromMesh;
+			LastDefinedLOD = (LODIndex < NumLODs) && (LastDefinedLOD == -1 || !bUseAutomaticLods) ? LODIndex : LastDefinedLOD;
+
+			// It turns out LODToGenerate is always LastDefinedLOD.
+			const int32& LODToGenerate = LastDefinedLOD;
+			if (LODToGenerate < 0)
+			{
+				continue;
+			}
+			
+			TArray<UEdGraphPin*> ConnectedLODPins = FollowInputPinArray(*TypedNodeObj->LODPin(LODToGenerate));
+
+			// Proccess non modifier material nodes.
+			for (int32 MeshComponentIndex = 0; MeshComponentIndex < NumMeshComponentsInRoot; ++MeshComponentIndex)
+			{
+				GenerationContext.CurrentMeshComponent = MeshComponentIndex;
+
+				for (UEdGraphPin* const ChildNodePin : ConnectedLODPins)
 				{
-					// If it is manual LODs, we don't want to show anything for this node in this LOD
-					if (GenerationContext.CurrentAutoLODStrategy == ECustomizableObjectAutomaticLODStrategy::AutomaticFromMesh
-						&& LastDefinedLOD != INDEX_NONE)
+					// Modifiers are shared for all components and are processed per LOD and not component.
+					if (Cast<UCustomizableObjectNodeModifierBase>(ChildNodePin->GetOwningNode()))
 					{
-						LODToGenerate = LastDefinedLOD;
+						continue;
 					}
-					else
+
+					if (!AffectsCurrentComponent(ChildNodePin, GenerationContext))
 					{
-						// In automatic LODs, we follow the last LOD
-						LastDefinedLOD = LODIndex;
-						LODToGenerate = LODIndex;
+						continue;
 					}
+
+					FMutableGraphSurfaceGenerationData DummySurfaceData;
+					mu::NodeSurfacePtr SurfaceNode = GenerateMutableSourceSurface(ChildNodePin, GenerationContext, DummySurfaceData);
+
+					mu::NodeComponentPtr& ComponentNode = ComponentNodes[MeshComponentIndex];
+
+					const int32 SurfaceCount = ComponentNode->GetSurfaceCount();
+					ComponentNode->SetSurfaceCount(SurfaceCount + 1);
+					ComponentNode->SetSurface(SurfaceCount, SurfaceNode.get());
+					ComponentNode->SetMessageContext(Node);
 				}
-				else
-				{
-					// This LOD is not defined in this object, so we assume the last LOD has to be used instead
-					LODToGenerate = LastDefinedLOD;
-				}
+			}
 
-				if (GenerationContext.CurrentLOD < GenerationContext.FirstLODAvailable)
+			// Process modfiers. Those are shared between different components in a lod.	
+			for (UEdGraphPin* const ChildNodePin : ConnectedLODPins)
+			{
+				// Set it to -1 to indicate we don't care about component id.
+				GenerationContext.CurrentMeshComponent = -1;
+
+				if (!Cast<UCustomizableObjectNodeModifierBase>(ChildNodePin->GetOwningNode()))
 				{
 					continue;
 				}
 
-				if (LODToGenerate >= 0)
-				{
-					TArray<UEdGraphPin*> ConnectedLODPins = FollowInputPinArray(*TypedNodeObj->LODPin(LODToGenerate));
-					int32 LODConnections = ConnectedLODPins.Num();
-					int32 NumMaterials = 0;
-					int32 NumModifiers = 0;
-					ComponentNode->SetSurfaceCount(NumMaterials);
-					for (int32 MatIndex = 0; MatIndex < LODConnections; ++MatIndex)
-					{
-						const UEdGraphPin* ChildNodePin = ConnectedLODPins[MatIndex];
-						if (!AffectsCurrentComponent(ChildNodePin, GenerationContext))
-						{
-							continue;
-						}
+				mu::NodeModifierPtr ModifierNode = GenerateMutableSourceModifier(ChildNodePin, GenerationContext);
 
-						if (Cast<UCustomizableObjectNodeModifierBase>(ChildNodePin->GetOwningNode()))
-						{
-							mu::NodeModifierPtr MatNode = GenerateMutableSourceModifier(ChildNodePin, GenerationContext);
-							LODNode->SetModifierCount(NumModifiers + 1);
-							LODNode->SetModifier(NumModifiers, MatNode.get());
-							LODNode->SetMessageContext(Node);
-							++NumModifiers;
-						}
-						else
-						{
-							FMutableGraphSurfaceGenerationData DummySurfaceData;
-							mu::NodeSurfacePtr MatNode = GenerateMutableSourceSurface(ChildNodePin, GenerationContext, DummySurfaceData);
-							ComponentNode->SetSurfaceCount(NumMaterials + 1);
-							ComponentNode->SetSurface(NumMaterials, MatNode.get());
-							ComponentNode->SetMessageContext(Node);
-							++NumMaterials;
-						}
-					}
-				}
+				const int32 ModifierCount = LODNode->GetModifierCount();
+				LODNode->SetModifierCount(ModifierCount + 1);
+				LODNode->SetModifier(ModifierCount, ModifierNode.get());
 			}
 		}
 
 		// Children
 		TArray<UEdGraphPin*> ConnectedChildrenPins = FollowInputPinArray(*TypedNodeObj->ChildrenPin());
 		ObjectNode->SetChildCount(ConnectedChildrenPins.Num());
-		for (int ChildIndex = 0; ChildIndex < ConnectedChildrenPins.Num(); ++ChildIndex)
+		for (int32 ChildIndex = 0; ChildIndex < ConnectedChildrenPins.Num(); ++ChildIndex)
 		{
 			mu::NodeObjectPtr ChildNode = GenerateMutableSource(ConnectedChildrenPins[ChildIndex], GenerationContext, bPartialCompilation);
 			ObjectNode->SetChild(ChildIndex, ChildNode.get());
@@ -1194,19 +1207,14 @@ bool AffectsCurrentComponent(const UEdGraphPin* Pin, FMutableGraphGenerationCont
 	}
 	else if (const UCustomizableObjectNodeModifierBase* TypedNodeModifier = Cast<UCustomizableObjectNodeModifierBase>(Node))
 	{
-		// Modifiers affect all components with the exception of the UCustomizableObjectNodeMeshClipMorph
-		if (const UCustomizableObjectNodeMeshClipMorph* TypedNodeMeshClipMorph = Cast<UCustomizableObjectNodeMeshClipMorph>(Node))
-		{
-			ComponentIndex = TypedNodeMeshClipMorph->ReferenceSkeletonIndex;
-		}
-		else
-		{
-			return true;
-		}
+		// Modifiers affect all compoenents at lod level. This branch should never be reached.
+		check(false);
+		return false;
 	}
 	else
 	{
-		unimplemented()
+		unimplemented();
+		return false;
 	}
 
 	return ComponentIndex == GenerationContext.CurrentMeshComponent;
diff --git a/Engine/Plugins/Experimental/Mutable/Source/CustomizableObjectEditor/Private/MuCOE/GenerateMutableSource/GenerateMutableSource.h b/Engine/Plugins/Experimental/Mutable/Source/CustomizableObjectEditor/Private/MuCOE/GenerateMutableSource/GenerateMutableSource.h
--- a/Engine/Plugins/Experimental/Mutable/Source/CustomizableObjectEditor/Private/MuCOE/GenerateMutableSource/GenerateMutableSource.h
+++ b/Engine/Plugins/Experimental/Mutable/Source/CustomizableObjectEditor/Private/MuCOE/GenerateMutableSource/GenerateMutableSource.h
@@ -96,6 +96,7 @@ enum class EMutableMeshConversionFlags : uint32
 	IgnorePhysics = 1 << 1
 };
 
+ENUM_CLASS_FLAGS(EMutableMeshConversionFlags)
 
 /** 
 	Struct to store the necessary data to generate the morphs of a skeletal mesh 
@@ -142,7 +143,7 @@ private:
 	int32 LOD;
 
 	/** Flag used to generate this mesh. Bit mask of EMutableMeshConversionFlags */
-	uint32 Flags = 0;
+	EMutableMeshConversionFlags Flags = EMutableMeshConversionFlags::None;
 
 	/** Active morphs at the time of mesh generation. */
 	TArray<FMorphNodeData> MeshMorphStack;
@@ -576,7 +577,7 @@ struct FMutableGraphGenerationContext
 			int32 MaterialIndex = 0;
 
 			/** Flag used to generate this mesh. Bit mask of EMutableMeshConversionFlags */
-			uint32 Flags = 0;
+			EMutableMeshConversionFlags Flags = EMutableMeshConversionFlags::None;
 
 			/** Tags added at the UE level that go through the Mutable core and are merged in the generated mesh.
 			 *  Only add the tags that make the mesh unique and require it not to be cached together with the 
@@ -600,7 +601,7 @@ struct FMutableGraphGenerationContext
 
 	// Stack of mesh generation flags. The last one is the currently valid.
 	// The value is a bit mask of EMutableMeshConversionFlags
-	TArray<uint32> MeshGenerationFlags;
+	TArray<EMutableMeshConversionFlags> MeshGenerationFlags;
 
 	/** Find a mesh if already generated for a given source and flags. */
 	mu::MeshPtr FindGeneratedMesh(const FGeneratedMeshData::FKey& Key);
diff --git a/Engine/Plugins/Experimental/Mutable/Source/CustomizableObjectEditor/Private/MuCOE/GenerateMutableSource/GenerateMutableSourceMesh.cpp b/Engine/Plugins/Experimental/Mutable/Source/CustomizableObjectEditor/Private/MuCOE/GenerateMutableSource/GenerateMutableSourceMesh.cpp
--- a/Engine/Plugins/Experimental/Mutable/Source/CustomizableObjectEditor/Private/MuCOE/GenerateMutableSource/GenerateMutableSourceMesh.cpp
+++ b/Engine/Plugins/Experimental/Mutable/Source/CustomizableObjectEditor/Private/MuCOE/GenerateMutableSource/GenerateMutableSourceMesh.cpp
@@ -465,14 +465,13 @@ void SetAndPropagatePoseBoneUsage(
 
 }
 
-
 mu::MeshPtr ConvertSkeletalMeshToMutable(USkeletalMesh* InSkeletalMesh, int LOD, int MaterialIndex, FMutableGraphGenerationContext& GenerationContext, const UCustomizableObjectNode* CurrentNode)
 {
 	// Get the mesh generation flags to use
-	uint32 CurrentFlags = GenerationContext.MeshGenerationFlags.Last();
-	bool bIgnoreSkeleton = CurrentFlags & uint32(EMutableMeshConversionFlags::IgnoreSkinning);
-	bool bIgnorePhysics = (CurrentFlags & uint32(EMutableMeshConversionFlags::IgnorePhysics)) || 
-						  !GenerationContext.Options.bPhysicsAssetMergeEnebled;
+	const EMutableMeshConversionFlags CurrentFlags = GenerationContext.MeshGenerationFlags.Last();
+	const bool bIgnoreSkeleton = EnumHasAnyFlags(CurrentFlags, EMutableMeshConversionFlags::IgnoreSkinning);
+	const bool bIgnorePhysics = EnumHasAnyFlags(CurrentFlags, EMutableMeshConversionFlags::IgnorePhysics) || 
+								!GenerationContext.Options.bPhysicsAssetMergeEnebled;
 
 	mu::MeshPtr MutableMesh = new mu::Mesh();
 		
@@ -487,12 +486,10 @@ mu::MeshPtr ConvertSkeletalMeshToMutable(USkeletalMesh* InSkeletalMesh, int LOD,
 
 	TObjectPtr<const USkeletalMesh> SkeletalMesh = InSkeletalMesh;
 
-	FMutableComponentInfo& MutComponentInfo = GenerationContext.GetCurrentComponentInfo();
-	USkeletalMesh* ComponentRefSkeletalMesh = MutComponentInfo.RefSkeletalMesh;
-	USkeleton* ComponentRefSkeleton = MutComponentInfo.RefSkeleton;
-	check(ComponentRefSkeletalMesh);
-	check(ComponentRefSkeleton);
-	
+	// CurrentMeshComponent < 0 implies IgnoreSkeleton flag.  
+	// CurrentMeshComponent < 0 will only happen with modifiers and, for now, any mesh generated from a modifier
+	// should ignore skinning.
+	check(!(GenerationContext.CurrentMeshComponent < 0) || bIgnoreSkeleton);
 
 	// Apply removed bones
 	TMap<int32, int32> RemovedBonesActiveParentIndices;
@@ -542,6 +539,14 @@ mu::MeshPtr ConvertSkeletalMeshToMutable(USkeletalMesh* InSkeletalMesh, int LOD,
 	// Check for a matching skeleton to the reference mesh
 	if (!bIgnoreSkeleton)
 	{
+		check(GenerationContext.CurrentMeshComponent >= 0);
+
+		FMutableComponentInfo& MutComponentInfo = GenerationContext.GetCurrentComponentInfo();
+		USkeletalMesh* ComponentRefSkeletalMesh = MutComponentInfo.RefSkeletalMesh;
+		USkeleton* ComponentRefSkeleton = MutComponentInfo.RefSkeleton;
+		check(ComponentRefSkeletalMesh);
+		check(ComponentRefSkeleton);
+
 		// Add the RefSkeleton ID to the mesh.
 		const int32 RefSkeletonID = GenerationContext.ReferencedSkeletons.AddUnique(ComponentRefSkeleton);
 		MutableMesh->AddSkeletonID(RefSkeletonID);
@@ -792,6 +797,8 @@ mu::MeshPtr ConvertSkeletalMeshToMutable(USkeletalMesh* InSkeletalMesh, int LOD,
 		FMemory::Memcpy(MutableMesh->GetVertexBuffers().GetBufferData(0), Vertices.GetData() + VertexStart, VertexCount * ElementSize);
 	}
 
+
+	// TODO: Add Mesh generation flags to not include RT Morph and clothing if not needed.
 	int32 nextBufferIndex = 1;
 	if (GenerationContext.Options.bRealTimeMorphTargetsEnabled)
 	{
@@ -832,7 +839,7 @@ mu::MeshPtr ConvertSkeletalMeshToMutable(USkeletalMesh* InSkeletalMesh, int LOD,
 
         // Add SkeletalMesh node used defined realtime morph targets to a temporal array where
         // the actual to be used real-time morphs names will be placed.        
-        TArray<FName> UsedMorphTargetsNames = [&]()
+        TArray<FName> UsedMorphTargetsNames = Invoke([&]()
         {
             TArray<FName> MorphTargetsNames;
             MorphTargetsNames.Reserve(SkeletalMeshMorphTargets.Num());
@@ -854,30 +861,31 @@ mu::MeshPtr ConvertSkeletalMeshToMutable(USkeletalMesh* InSkeletalMesh, int LOD,
             }
 
             return MorphTargetsNames;
-        }(); // lambda is invoked.
+        });
 
         //Apply global morph targets overrides to the SkeletalMesh user defined RT morph targets. 
         for (FRealTimeMorphSelectionOverride& MorphTargetOverride : RealTimeMorphTargetOverrides)
         {
-            const ECustomizableObjectSelectionOverride OverrideValue = [&]() -> ECustomizableObjectSelectionOverride
-                {
-                    const ECustomizableObjectSelectionOverride GlobalOverrideValue = MorphTargetOverride.SelectionOverride;
-                    
-                    if (GlobalOverrideValue != ECustomizableObjectSelectionOverride::NoOverride)
-                    {
-                        return GlobalOverrideValue;
-                    } 
-                
-                    const int32 FoundIdx = 
-                            MorphTargetOverride.SkeletalMeshesNames.Find(SkeletalMesh->GetFName());
-
-                    if (FoundIdx != INDEX_NONE)
-                    {
-                        return MorphTargetOverride.Override[FoundIdx];
-                    }
-
-                    return ECustomizableObjectSelectionOverride::NoOverride;
-                }(); // lambda is invoked.
+            const ECustomizableObjectSelectionOverride OverrideValue = 
+				Invoke([&]() -> ECustomizableObjectSelectionOverride
+				{
+					const ECustomizableObjectSelectionOverride GlobalOverrideValue = MorphTargetOverride.SelectionOverride;
+					
+					if (GlobalOverrideValue != ECustomizableObjectSelectionOverride::NoOverride)
+					{
+						return GlobalOverrideValue;
+					} 
+				
+					const int32 FoundIdx = 
+							MorphTargetOverride.SkeletalMeshesNames.Find(SkeletalMesh->GetFName());
+
+					if (FoundIdx != INDEX_NONE)
+					{
+						return MorphTargetOverride.Override[FoundIdx];
+					}
+
+					return ECustomizableObjectSelectionOverride::NoOverride;
+				});
 
             if (OverrideValue == ECustomizableObjectSelectionOverride::Enable)
             {
@@ -1033,7 +1041,7 @@ mu::MeshPtr ConvertSkeletalMeshToMutable(USkeletalMesh* InSkeletalMesh, int LOD,
 
 		// Create new asset or find an already created one if the section has clothing assets.
 		// clothing assets are shared among all LODs in a section
-		const int32 ClothingAssetIndex = [&]() -> int32
+		const int32 ClothingAssetIndex = Invoke([&]() -> int32
 		{
 			const UClothingAssetBase* ClothingAssetBase = SkeletalMesh->GetSectionClothingAsset(LOD, MaterialIndex);
 
@@ -1079,7 +1087,7 @@ mu::MeshPtr ConvertSkeletalMeshToMutable(USkeletalMesh* InSkeletalMesh, int LOD,
 			}
 
 			return NewAssetIndex;
-		}(); // lambda is invoked
+		});
 
 		if (ClothingAssetIndex != INDEX_NONE)
 		{
@@ -1826,7 +1834,7 @@ mu::MeshPtr ConvertStaticMeshToMutable(UStaticMesh* StaticMesh, int LOD, int Mat
 mu::MeshPtr GenerateMutableMesh(UObject * Mesh, int32 LOD, int32 MaterialIndex, const FString& UniqueTags, FMutableGraphGenerationContext & GenerationContext, const UCustomizableObjectNode* CurrentNode)
 {
 	// Get the mesh generation flags to use
-	uint32 CurrentFlags = GenerationContext.MeshGenerationFlags.Last();
+	EMutableMeshConversionFlags CurrentFlags = GenerationContext.MeshGenerationFlags.Last();
 
 	FMutableGraphGenerationContext::FGeneratedMeshData::FKey Key = { Mesh, LOD, MaterialIndex, CurrentFlags, UniqueTags };
 	mu::MeshPtr MutableMesh = GenerationContext.FindGeneratedMesh(Key);
@@ -1851,10 +1859,10 @@ mu::MeshPtr GenerateMutableMesh(UObject * Mesh, int32 LOD, int32 MaterialIndex,
 			GenerationContext.GeneratedMeshes.Push({ Key, MutableMesh });
 		}
 	}
-
+	
 	if (MutableMesh)
 	{
-		FMeshData MeshData; 
+		FMeshData MeshData;
 		MeshData.Mesh = Mesh;
 		MeshData.LOD = LOD;
 		MeshData.MaterialIndex = MaterialIndex;
@@ -3266,8 +3274,9 @@ mu::NodeMeshPtr GenerateMutableSourceMesh(const UEdGraphPin * Pin,
 			}		
 		}
 		// We don't need all the data for the shape meshes
-		const uint32 ShapeFlags = uint32(EMutableMeshConversionFlags::IgnoreSkinning) & 
-								  uint32(EMutableMeshConversionFlags::IgnorePhysics);
+		const EMutableMeshConversionFlags ShapeFlags = 
+				EMutableMeshConversionFlags::IgnoreSkinning |
+				EMutableMeshConversionFlags::IgnorePhysics;
 
 		GenerationContext.MeshGenerationFlags.Push( ShapeFlags );
 			
diff --git a/Engine/Plugins/Experimental/Mutable/Source/CustomizableObjectEditor/Private/MuCOE/GenerateMutableSource/GenerateMutableSourceModifier.cpp b/Engine/Plugins/Experimental/Mutable/Source/CustomizableObjectEditor/Private/MuCOE/GenerateMutableSource/GenerateMutableSourceModifier.cpp
--- a/Engine/Plugins/Experimental/Mutable/Source/CustomizableObjectEditor/Private/MuCOE/GenerateMutableSource/GenerateMutableSourceModifier.cpp
+++ b/Engine/Plugins/Experimental/Mutable/Source/CustomizableObjectEditor/Private/MuCOE/GenerateMutableSource/GenerateMutableSourceModifier.cpp
@@ -62,6 +62,14 @@ mu::NodeModifierPtr GenerateMutableSourceModifier(const UEdGraphPin * Pin, FMuta
 
 	bool bDoNotAddToGeneratedCache = false; // TODO Remove on MTBL-829 
 
+	
+	// We don't need all the data for the modifiers meshes
+	const EMutableMeshConversionFlags ModifiersMeshFlags = 
+			EMutableMeshConversionFlags::IgnoreSkinning |
+			EMutableMeshConversionFlags::IgnorePhysics;
+
+	GenerationContext.MeshGenerationFlags.Push(ModifiersMeshFlags);
+
 	if (const UCustomizableObjectNodeMeshClipMorph* TypedNodeClip = Cast<UCustomizableObjectNodeMeshClipMorph>(Node))
 	{
 		mu::NodeModifierMeshClipMorphPlanePtr ClipNode = new mu::NodeModifierMeshClipMorphPlane();
@@ -132,6 +140,7 @@ mu::NodeModifierPtr GenerateMutableSourceModifier(const UEdGraphPin * Pin, FMuta
 		if (const UEdGraphPin* ConnectedPin = FollowInputPin(*TypedNodeClipMesh->ClipMeshPin()))
 		{
 			FMutableGraphMeshGenerationData DummyMeshData;
+
 			mu::NodeMeshPtr ClipMesh = GenerateMutableSourceMesh(ConnectedPin, GenerationContext, DummyMeshData);
 
 			FPinDataValue* PinData = GenerationContext.PinData.Find(ConnectedPin);
@@ -206,12 +215,15 @@ mu::NodeModifierPtr GenerateMutableSourceModifier(const UEdGraphPin * Pin, FMuta
 		GenerationContext.Compiler->CompilerLog(LOCTEXT("UnimplementedNode", "Node type not implemented yet."), Node);
 	}
 
+	GenerationContext.MeshGenerationFlags.Pop();
+
 	if (!bDoNotAddToGeneratedCache)
 	{
 		GenerationContext.Generated.Add(Key, FGeneratedData(Node, Result));
 	}
 	GenerationContext.GeneratedNodes.Add(Node);
 	
+
 	return Result;
 }
 
