#rnx

diff --git a/Engine/Plugins/Experimental/Mutable/Source/CustomizableObjectEditor/Private/MuCOE/GenerateMutableSource/GenerateMutableSourceSurface.cpp b/Engine/Plugins/Experimental/Mutable/Source/CustomizableObjectEditor/Private/MuCOE/GenerateMutableSource/GenerateMutableSourceSurface.cpp
--- a/Engine/Plugins/Experimental/Mutable/Source/CustomizableObjectEditor/Private/MuCOE/GenerateMutableSource/GenerateMutableSourceSurface.cpp
+++ b/Engine/Plugins/Experimental/Mutable/Source/CustomizableObjectEditor/Private/MuCOE/GenerateMutableSource/GenerateMutableSourceSurface.cpp
@@ -1300,10 +1300,15 @@ mu::NodeSurfacePtr GenerateMutableSourceSurface(const UEdGraphPin * Pin, FMutabl
 					mu::NodeImagePtr ImageNode = GenerateMutableSourceImage(ConnectedImagePin, GenerationContext, 0.f);
 					ImagePatchNode->SetImage(ImageNode);
 
-					if (const UEdGraphPin* ConnectedMaskPin = FollowInputPin(*TypedNodeEdit->GetUsedImageMaskPin(ImageId)))
+					const UEdGraphPin* ImageMaskPin = TypedNodeEdit->GetUsedImageMaskPin(ImageId);
+
+					if (ImageMaskPin)
 					{
-						mu::NodeImagePtr MaskNode = GenerateMutableSourceImage(ConnectedMaskPin, GenerationContext, 0.f);
-						ImagePatchNode->SetMask(MaskNode);
+						if (const UEdGraphPin* ConnectedMaskPin = FollowInputPin(*ImageMaskPin))
+						{
+							mu::NodeImagePtr MaskNode = GenerateMutableSourceImage(ConnectedMaskPin, GenerationContext, 0.f);
+							ImagePatchNode->SetMask(MaskNode);
+						}
 					}
 
 					// Add the block indices
diff --git a/Engine/Plugins/Experimental/Mutable/Source/CustomizableObjectEditor/Private/MuCOE/Nodes/CustomizableObjectNodeEditMaterial.cpp b/Engine/Plugins/Experimental/Mutable/Source/CustomizableObjectEditor/Private/MuCOE/Nodes/CustomizableObjectNodeEditMaterial.cpp
--- a/Engine/Plugins/Experimental/Mutable/Source/CustomizableObjectEditor/Private/MuCOE/Nodes/CustomizableObjectNodeEditMaterial.cpp
+++ b/Engine/Plugins/Experimental/Mutable/Source/CustomizableObjectEditor/Private/MuCOE/Nodes/CustomizableObjectNodeEditMaterial.cpp
@@ -73,6 +73,32 @@ void UCustomizableObjectNodeEditMaterial::RemapPins(const TMap<UEdGraphPin*, UEd
 			{
 				PinData->PinMask = FEdGraphPinReference(*Result);
 			}
+
+			FName PinMaskName = FName(Pin->PinName.ToString() + FString(" Mask"));
+
+			if (!PinData->PinMask.Get())
+			{
+				// Something went wrong, an Edit Material Image pin should have an associated mask pin, try to find it
+				for (const UEdGraphPin* PotentialMaskPin : GetAllPins())
+				{
+					if (PotentialMaskPin->PinName == PinMaskName)
+					{
+						PinData->PinMask = PotentialMaskPin;
+						break;
+					}
+				}
+			}
+
+			if (!PinData->PinMask.Get())
+			{
+				// The mask pin didn't actually exist, so create it
+				const UEdGraphSchema_CustomizableObject* Schema = GetDefault<UEdGraphSchema_CustomizableObject>();
+				UEdGraphPin* PinMask = CustomCreatePin(EGPD_Input, Schema->PC_Image, PinMaskName);
+				PinMask->bHidden = Pin->bHidden;
+				PinMask->bDefaultValueIsIgnored = true;
+
+				PinData->PinMask = FEdGraphPinReference(PinMask);
+			}
 		}
 	}
 }
@@ -82,7 +108,8 @@ const UEdGraphPin* UCustomizableObjectNodeEditMaterial::GetUsedImageMaskPin(cons
 {
 	if (const UEdGraphPin* Pin = GetUsedImagePin(ImageId))
 	{
-		return GetPinData<UCustomizableObjectNodeEditMaterialPinEditImageData>(*Pin).PinMask.Get();
+		UCustomizableObjectNodeEditMaterialPinEditImageData& PinData = GetPinData<UCustomizableObjectNodeEditMaterialPinEditImageData>(*Pin);
+		return PinData.PinMask.Get();
 	}
 
 	return nullptr;
