[mutable] Additional optimization of projector and image layer-related code.
Refactored Image layer to their own ASTOp class.
Additional debugging data in code tree.
Removed old tasking support in compilation.
#rnx

diff --git a/Engine/Plugins/Experimental/Mutable/Source/CustomizableObject/Classes/MuCO/CustomizableObject.h b/Engine/Plugins/Experimental/Mutable/Source/CustomizableObject/Classes/MuCO/CustomizableObject.h
--- a/Engine/Plugins/Experimental/Mutable/Source/CustomizableObject/Classes/MuCO/CustomizableObject.h
+++ b/Engine/Plugins/Experimental/Mutable/Source/CustomizableObject/Classes/MuCO/CustomizableObject.h
@@ -1220,7 +1220,7 @@ private:
 	// This is a manual version number for the binary blobs in this asset.
 	// Increasing it invalidates all the previously compiled models.
 	// Warning: If while merging code both versions have changed, take the highest+1.
-	static const int32 CurrentSupportedVersion = 358;
+	static const int32 CurrentSupportedVersion = 359;
 
 public:
 
diff --git a/Engine/Plugins/Experimental/Mutable/Source/CustomizableObject/Public/MuCO/CustomizableObjectMipDataProvider.h b/Engine/Plugins/Experimental/Mutable/Source/CustomizableObject/Public/MuCO/CustomizableObjectMipDataProvider.h
--- a/Engine/Plugins/Experimental/Mutable/Source/CustomizableObject/Public/MuCO/CustomizableObjectMipDataProvider.h
+++ b/Engine/Plugins/Experimental/Mutable/Source/CustomizableObject/Public/MuCO/CustomizableObjectMipDataProvider.h
@@ -30,12 +30,12 @@ class UTexture;
 */
 struct FMutableUpdateContext
 {
-	mu::SystemPtr System;
-	mu::ModelPtr Model;
-	mu::ParametersPtrConst Parameters;
+	mu::Ptr<mu::System> System;
+	mu::Ptr<mu::Model> Model;
+	mu::Ptr<const mu::Parameters> Parameters;
 	int32 State = -1;
 
-	TArray<mu::ImagePtrConst> ImageParameterValues;
+	TArray<mu::Ptr<const mu::Image>> ImageParameterValues;
 };
 
 
diff --git a/Engine/Plugins/Experimental/Mutable/Source/CustomizableObjectEditor/Private/MuCOE/CustomizableObjectCompileRunnable.cpp b/Engine/Plugins/Experimental/Mutable/Source/CustomizableObjectEditor/Private/MuCOE/CustomizableObjectCompileRunnable.cpp
--- a/Engine/Plugins/Experimental/Mutable/Source/CustomizableObjectEditor/Private/MuCOE/CustomizableObjectCompileRunnable.cpp
+++ b/Engine/Plugins/Experimental/Mutable/Source/CustomizableObjectEditor/Private/MuCOE/CustomizableObjectCompileRunnable.cpp
@@ -103,10 +103,6 @@ uint32 FCustomizableObjectCompileRunnable::Run()
 		? mu::CompilerOptions::TextureLayoutStrategy::None
 		: mu::CompilerOptions::TextureLayoutStrategy::Pack);
 
-	// \TODO For now force it to be disabled.
-	//CompilerOptions->SetEnableConcurrency(Options.bUseParallelCompilation);
-	CompilerOptions->SetEnableConcurrency(false);
-
 	mu::CompilerPtr Compiler = new mu::Compiler(CompilerOptions);
 
 	UE_LOG(LogMutable, Verbose, TEXT("PROFILE: [ %16.8f ] FCustomizableObjectCompileRunnable Compile start."), FPlatformTime::Seconds());
diff --git a/Engine/Plugins/Experimental/Mutable/Source/CustomizableObjectEditor/Private/MuCOE/SMutableCodeViewer.cpp b/Engine/Plugins/Experimental/Mutable/Source/CustomizableObjectEditor/Private/MuCOE/SMutableCodeViewer.cpp
--- a/Engine/Plugins/Experimental/Mutable/Source/CustomizableObjectEditor/Private/MuCOE/SMutableCodeViewer.cpp
+++ b/Engine/Plugins/Experimental/Mutable/Source/CustomizableObjectEditor/Private/MuCOE/SMutableCodeViewer.cpp
@@ -153,6 +153,27 @@ public:
 				break;
 			}
 
+			case mu::OP_TYPE::IM_MIPMAP:
+			{
+				mu::OP::ImageMipmapArgs Args = Program.GetOpArgs<mu::OP::ImageMipmapArgs>(RowItem->MutableOperation);
+				OpName += FString::Printf(TEXT(" levels: %d-%d tail: %d"), Args.levels, Args.blockLevels, int32(Args.onlyTail));
+				break;
+			}
+
+			case mu::OP_TYPE::IM_RESIZE:
+			{
+				mu::OP::ImageResizeArgs Args = Program.GetOpArgs<mu::OP::ImageResizeArgs>(RowItem->MutableOperation);
+				OpName += FString::Printf(TEXT(" %d x %d"), int32(Args.size[0]), int32(Args.size[1]));
+				break;
+			}
+
+			case mu::OP_TYPE::IM_RESIZEREL:
+			{
+				mu::OP::ImageResizeRelArgs Args = Program.GetOpArgs<mu::OP::ImageResizeRelArgs>(RowItem->MutableOperation);
+				OpName += FString::Printf(TEXT(" %.3f x %.3f"), int32(Args.factor[0]), int32(Args.factor[1]));
+				break;
+			}
+
 			case mu::OP_TYPE::IM_MULTILAYER:
 			{
 				mu::OP::ImageMultiLayerArgs Args = Program.GetOpArgs<mu::OP::ImageMultiLayerArgs>(RowItem->MutableOperation);
@@ -171,8 +192,13 @@ public:
 			case mu::OP_TYPE::IM_LAYER:
 			{
 				mu::OP::ImageLayerArgs Args = Program.GetOpArgs<mu::OP::ImageLayerArgs>(RowItem->MutableOperation);
-				OpName += TEXT(" ");
+				OpName += TEXT(" rgb: ");
 				OpName += mu::TypeInfo::s_blendModeName[int32(Args.blendType)];
+				if (Args.blendTypeAlpha != int8(mu::EBlendType::BT_NONE))
+				{
+					OpName += TEXT(", a: ");
+					OpName += mu::TypeInfo::s_blendModeName[int32(Args.blendTypeAlpha)];
+				}
 				OpName += FString::Printf(TEXT(" flags %d"),Args.flags);
 				break;
 			}
@@ -192,7 +218,7 @@ public:
 			}
 
 			// Prepare the text shown on the UI side of the operation tree
-			FString MainLabel = FString::Printf(TEXT("%s (%d) : %s"), *RowItem->Caption, int32(RowItem->MutableOperation), *OpName);
+			FString MainLabel = FString::Printf(TEXT("%s %d : %s"), *RowItem->Caption, int32(RowItem->MutableOperation), *OpName);
 
 #if UE_BUILD_DEBUG
 			FString IndexOnTree = FString::FromInt(RowItem->IndexOnTree);
@@ -2151,6 +2177,48 @@ bool SMutableCodeViewer::IsConstantResourceUsedByOperation(const int32 IndexOnCo
 #pragma endregion 
 
 
+namespace
+{
+	/** Test implementation to provide image parameters. It will generate some images of a fixed size and format. */
+	class TestImageProvider : public mu::ImageParameterGenerator
+	{
+	public:
+		mu::Ptr<mu::Image> GetImage(mu::EXTERNAL_IMAGE_ID id) override
+		{
+			MUTABLE_CPUPROFILER_SCOPE(TestImageProvider_GetImage);
+
+			const int32 Size = 1024;
+			const int32 LODCount = 1;
+			const mu::EImageFormat Format = mu::EImageFormat::IF_RGBA_UBYTE;
+			mu::Ptr<mu::Image> Image = new mu::Image(Size, Size, LODCount, Format);
+
+			// Generate an alphatested circle with an horizontal gradient color.
+			uint8* Data = Image->GetData();
+			int32 CircleRadius = (Size * 2) / 5;
+			int32 CircleRadius2 = CircleRadius * CircleRadius;
+			int32 Color[3] = { 255,128,0 };
+
+			for (int y = 0; y < Size; ++y)
+			{
+				for (int x = 0; x < Size; ++x)
+				{
+					float R2 = (x - Size / 2) * (x - Size / 2) + (y - Size / 2) * (y - Size / 2);
+					int32 Opacity = FMath::Clamp( (CircleRadius2-R2)/CircleRadius2 * 512 - 64, 0, 255 );
+					Data[0] = (Color[0] * x) / Size;
+					Data[1] = (Color[1] * x) / Size;
+					Data[2] = (Color[2] * x) / Size;
+					Data[3] = uint8(Opacity);
+					Data += 4;
+				}
+			}
+
+			return Image;
+		}
+
+	};
+}
+
+
 void SMutableCodeViewer::Tick(const FGeometry& AllottedGeometry, const double InCurrentTime, const float InDeltaTime)
 {
 	SCompoundWidget::Tick(AllottedGeometry, InCurrentTime, InDeltaTime);
@@ -2178,8 +2246,12 @@ void SMutableCodeViewer::Tick(const FGeometry& AllottedGeometry, const double In
 	const mu::OP_TYPE OperationType = MutableModel->GetPrivate()->m_program.GetOpType(SelectedOperationAddress);
 	const mu::DATATYPE OperationDataType = mu::GetOpDataType(OperationType);
 
-	const mu::SettingsPtr Settings = new mu::Settings();
-	const mu::SystemPtr System = new mu::System(Settings);
+	const mu::Ptr<mu::Settings> Settings = new mu::Settings();
+	const mu::Ptr<mu::System> System = new mu::System(Settings);
+
+	TestImageProvider* ImageProvider = new TestImageProvider;
+	System->SetImageParameterGenerator(ImageProvider);
+
 	System->GetPrivate()->BeginBuild(MutableModel);
 
 	switch (OperationDataType)
diff --git a/Engine/Plugins/Experimental/Mutable/Source/CustomizableObjectEditor/Private/MuCOE/SMutableCodeViewer.h b/Engine/Plugins/Experimental/Mutable/Source/CustomizableObjectEditor/Private/MuCOE/SMutableCodeViewer.h
--- a/Engine/Plugins/Experimental/Mutable/Source/CustomizableObjectEditor/Private/MuCOE/SMutableCodeViewer.h
+++ b/Engine/Plugins/Experimental/Mutable/Source/CustomizableObjectEditor/Private/MuCOE/SMutableCodeViewer.h
@@ -235,6 +235,7 @@ private:
 	const TArray<mu::OP_TYPE> ExpensiveOperationTypes
 	{
 		mu::OP_TYPE::IM_PIXELFORMAT,
+		mu::OP_TYPE::IM_SWIZZLE,
 		mu::OP_TYPE::ME_PROJECT,
 	};
 	
diff --git a/Engine/Plugins/Experimental/Mutable/Source/MutableRuntime/Private/MuR/CodeRunnerTasked.cpp b/Engine/Plugins/Experimental/Mutable/Source/MutableRuntime/Private/MuR/CodeRunnerTasked.cpp
--- a/Engine/Plugins/Experimental/Mutable/Source/MutableRuntime/Private/MuR/CodeRunnerTasked.cpp
+++ b/Engine/Plugins/Experimental/Mutable/Source/MutableRuntime/Private/MuR/CodeRunnerTasked.cpp
@@ -395,10 +395,13 @@ namespace mu
 		m_imageCompressionQuality = imageCompressionQuality;
 		Args = InArgs;
 		m_base = Memory.GetImage({ Args.base, InOp.ExecutionIndex, InOp.ExecutionOptions });
+		check(!m_base || m_base->GetFormat() < EImageFormat::IF_COUNT);
+
 		m_blended = Memory.GetImage({ Args.blended, InOp.ExecutionIndex, InOp.ExecutionOptions });
 		if (Args.mask)
 		{
 			m_mask = Memory.GetImage({ Args.mask, InOp.ExecutionIndex, InOp.ExecutionOptions });
+			check(!m_mask || m_mask->GetFormat() < EImageFormat::IF_COUNT);
 		}
 	}
 
@@ -416,27 +419,24 @@ namespace mu
 			return;
 		}
 
-		// TODO: a bug may cause this (issue T97)
-		EImageFormat finalFormat = EImageFormat::IF_NONE;
-		if (IsCompressedFormat(m_base->GetFormat()))
+		EImageFormat InitialFormat = m_base->GetFormat();
+		check(InitialFormat < EImageFormat::IF_COUNT);
+
+		if (IsCompressedFormat(InitialFormat))
 		{
 			UE_LOG(LogMutableCore, Error, TEXT("Image layer on compressed image is not supported. Ignored."));
-			finalFormat = m_base->GetFormat();
-			auto uncompressedFormat = GetUncompressedFormat(finalFormat);
-			m_base = ImagePixelFormat(m_imageCompressionQuality, m_base.get(), uncompressedFormat);
+			EImageFormat UncompressedFormat = GetUncompressedFormat(InitialFormat);
+			m_base = ImagePixelFormat(m_imageCompressionQuality, m_base.get(), UncompressedFormat);
 		}
 
-		if (m_blended && m_base->GetFormat() != m_blended->GetFormat())
+		if (m_blended && InitialFormat != m_blended->GetFormat())
 		{
-			m_blended =
-				ImagePixelFormat(m_imageCompressionQuality, m_blended.get(), m_base->GetFormat());
+			m_blended = ImagePixelFormat(m_imageCompressionQuality, m_blended.get(), m_base->GetFormat());
 		}
 
-		bool bApplyToAlpha = (Args.flags & OP::ImageLayerArgs::F_APPLY_TO_ALPHA) != 0;
-
-		ImagePtr pNew = new Image(m_base->GetSizeX(), m_base->GetSizeY(),
-			m_base->GetLODCount(), m_base->GetFormat());
+		bool bApplyColorBlendToAlpha = (Args.flags & OP::ImageLayerArgs::F_APPLY_TO_ALPHA) != 0;
 
+		ImagePtr pNew = mu::CloneOrTakeOver<Image>(m_base.get());
 		check(pNew->GetDataSize() == m_base->GetDataSize());
 
 		bool bValid = pNew->GetSizeX() > 0 && pNew->GetSizeY() > 0;
@@ -456,9 +456,7 @@ namespace mu
 				MUTABLE_CPUPROFILER_SCOPE(ImageLayer_MipBlendedEmergencyFix);
 
 				int32 levelCount = m_base->GetLODCount();
-				ImagePtr pDest = new Image(m_blended->GetSizeX(), m_blended->GetSizeY(),
-					levelCount,
-					m_blended->GetFormat());
+				ImagePtr pDest = new Image(m_blended->GetSizeX(), m_blended->GetSizeY(), levelCount, m_blended->GetFormat());
 
 				SCRATCH_IMAGE_MIPMAP scratch;
 				FMipmapGenerationSettings settings{};
@@ -470,9 +468,26 @@ namespace mu
 				m_blended = pDest;
 			}
 
+			bool bDone = false;
+
 			bool bOnlyOneMip = false;
 
-			if (Args.mask)
+			bool bUseMaskFromBlendAlpha = (Args.flags & OP::ImageLayerArgs::F_USE_MASK_FROM_BLENDED);
+
+			if (!m_mask && bUseMaskFromBlendAlpha
+				&&
+				Args.blendType == uint8(EBlendType::BT_BLEND)
+				&&
+				Args.blendTypeAlpha == uint8(EBlendType::BT_LIGHTEN))
+			{
+				// This is a frequent critical-path case because of multilayer projectors.
+				bDone = true;
+
+				BufferLayerComposite<BlendChannelMasked, LightenChannel, false>(pNew.get(), m_blended.get(), bOnlyOneMip);
+			}
+
+
+			if (!bDone && Args.mask)
 			{
 				if (m_base->GetSize() != m_mask->GetSize())
 				{
@@ -503,41 +518,79 @@ namespace mu
 				switch (EBlendType(Args.blendType))
 				{
 				case EBlendType::BT_NORMAL_COMBINE: ImageNormalCombine(pNew.get(), m_base.get(), m_mask.get(), m_blended.get(), bOnlyOneMip); break;
-				case EBlendType::BT_SOFTLIGHT: BufferLayer<SoftLightChannelMasked, SoftLightChannel, true>(pNew->GetData(), m_base.get(), m_mask.get(), m_blended.get(), bApplyToAlpha, bOnlyOneMip); break;
-				case EBlendType::BT_HARDLIGHT: BufferLayer<HardLightChannelMasked, HardLightChannel, true>(pNew->GetData(), m_base.get(), m_mask.get(), m_blended.get(), bApplyToAlpha, bOnlyOneMip); break;
-				case EBlendType::BT_BURN: BufferLayer<BurnChannelMasked, BurnChannel, true>(pNew->GetData(), m_base.get(), m_mask.get(), m_blended.get(), bApplyToAlpha, bOnlyOneMip); break;
-				case EBlendType::BT_DODGE: BufferLayer<DodgeChannelMasked, DodgeChannel, true>(pNew->GetData(), m_base.get(), m_mask.get(), m_blended.get(), bApplyToAlpha, bOnlyOneMip); break;
-				case EBlendType::BT_SCREEN: BufferLayer<ScreenChannelMasked, ScreenChannel, true>(pNew->GetData(), m_base.get(), m_mask.get(), m_blended.get(), bApplyToAlpha, bOnlyOneMip); break;
-				case EBlendType::BT_OVERLAY: BufferLayer<OverlayChannelMasked, OverlayChannel, true>(pNew->GetData(), m_base.get(), m_mask.get(), m_blended.get(), bApplyToAlpha, bOnlyOneMip); break;
-				case EBlendType::BT_LIGHTEN: BufferLayer<LightenChannelMasked, LightenChannel, true>(pNew->GetData(), m_base.get(), m_mask.get(), m_blended.get(), bApplyToAlpha, bOnlyOneMip); break;
-				case EBlendType::BT_MULTIPLY: BufferLayer<MultiplyChannelMasked, MultiplyChannel, true>(pNew->GetData(), m_base.get(), m_mask.get(), m_blended.get(), bApplyToAlpha, bOnlyOneMip); break;
-				case EBlendType::BT_BLEND: BufferLayer<BlendChannelMasked, BlendChannel, true>(pNew->GetData(), m_base.get(), m_mask.get(), m_blended.get(), bApplyToAlpha, bOnlyOneMip); break;
+				case EBlendType::BT_SOFTLIGHT: BufferLayer<SoftLightChannelMasked, SoftLightChannel, true>(pNew->GetData(), m_base.get(), m_mask.get(), m_blended.get(), bApplyColorBlendToAlpha, bOnlyOneMip); break;
+				case EBlendType::BT_HARDLIGHT: BufferLayer<HardLightChannelMasked, HardLightChannel, true>(pNew->GetData(), m_base.get(), m_mask.get(), m_blended.get(), bApplyColorBlendToAlpha, bOnlyOneMip); break;
+				case EBlendType::BT_BURN: BufferLayer<BurnChannelMasked, BurnChannel, true>(pNew->GetData(), m_base.get(), m_mask.get(), m_blended.get(), bApplyColorBlendToAlpha, bOnlyOneMip); break;
+				case EBlendType::BT_DODGE: BufferLayer<DodgeChannelMasked, DodgeChannel, true>(pNew->GetData(), m_base.get(), m_mask.get(), m_blended.get(), bApplyColorBlendToAlpha, bOnlyOneMip); break;
+				case EBlendType::BT_SCREEN: BufferLayer<ScreenChannelMasked, ScreenChannel, true>(pNew->GetData(), m_base.get(), m_mask.get(), m_blended.get(), bApplyColorBlendToAlpha, bOnlyOneMip); break;
+				case EBlendType::BT_OVERLAY: BufferLayer<OverlayChannelMasked, OverlayChannel, true>(pNew->GetData(), m_base.get(), m_mask.get(), m_blended.get(), bApplyColorBlendToAlpha, bOnlyOneMip); break;
+				case EBlendType::BT_LIGHTEN: BufferLayer<LightenChannelMasked, LightenChannel, true>(pNew->GetData(), m_base.get(), m_mask.get(), m_blended.get(), bApplyColorBlendToAlpha, bOnlyOneMip); break;
+				case EBlendType::BT_MULTIPLY: BufferLayer<MultiplyChannelMasked, MultiplyChannel, true>(pNew->GetData(), m_base.get(), m_mask.get(), m_blended.get(), bApplyColorBlendToAlpha, bOnlyOneMip); break;
+				case EBlendType::BT_BLEND: BufferLayer<BlendChannelMasked, BlendChannel, true>(pNew->GetData(), m_base.get(), m_mask.get(), m_blended.get(), bApplyColorBlendToAlpha, bOnlyOneMip); break;
 				default: check(false);
 				}
 
 			}
-			else
+			else if (!bDone && bUseMaskFromBlendAlpha)
+			{
+				// Apply blend without to RGB using mask in blended alpha
+				switch (EBlendType(Args.blendType))
+				{
+				case EBlendType::BT_NORMAL_COMBINE: check(false); break;
+				case EBlendType::BT_SOFTLIGHT: BufferLayerEmbeddedMask<SoftLightChannelMasked, SoftLightChannel, false>(pNew->GetData(), pNew.get(), m_blended.get(), bApplyColorBlendToAlpha, bOnlyOneMip); break;
+				case EBlendType::BT_HARDLIGHT: BufferLayerEmbeddedMask<HardLightChannelMasked, HardLightChannel, false>(pNew->GetData(), pNew.get(), m_blended.get(), bApplyColorBlendToAlpha, bOnlyOneMip); break;
+				case EBlendType::BT_BURN: BufferLayerEmbeddedMask<BurnChannelMasked, BurnChannel, false>(pNew->GetData(), pNew.get(), m_blended.get(), bApplyColorBlendToAlpha, bOnlyOneMip); break;
+				case EBlendType::BT_DODGE: BufferLayerEmbeddedMask<DodgeChannelMasked, DodgeChannel, false>(pNew->GetData(), pNew.get(), m_blended.get(), bApplyColorBlendToAlpha, bOnlyOneMip); break;
+				case EBlendType::BT_SCREEN: BufferLayerEmbeddedMask<ScreenChannelMasked, ScreenChannel, false>(pNew->GetData(), pNew.get(), m_blended.get(), bApplyColorBlendToAlpha, bOnlyOneMip); break;
+				case EBlendType::BT_OVERLAY: BufferLayerEmbeddedMask<OverlayChannelMasked, OverlayChannel, false>(pNew->GetData(), pNew.get(), m_blended.get(), bApplyColorBlendToAlpha, bOnlyOneMip); break;
+				case EBlendType::BT_LIGHTEN: BufferLayerEmbeddedMask<LightenChannelMasked, LightenChannel, false>(pNew->GetData(), pNew.get(), m_blended.get(), bApplyColorBlendToAlpha, bOnlyOneMip); break;
+				case EBlendType::BT_MULTIPLY: BufferLayerEmbeddedMask<MultiplyChannelMasked, MultiplyChannel, false>(pNew->GetData(), pNew.get(), m_blended.get(), bApplyColorBlendToAlpha, bOnlyOneMip); break;
+				case EBlendType::BT_BLEND: BufferLayerEmbeddedMask<BlendChannelMasked, BlendChannel, false>(pNew->GetData(), pNew.get(), m_blended.get(), bApplyColorBlendToAlpha, bOnlyOneMip); break;
+				default: check(false);
+				}
+			}
+			else if (!bDone)
 			{
+				// Apply blend without mask to RGB
 				switch (EBlendType(Args.blendType))
 				{
 				case EBlendType::BT_NORMAL_COMBINE: ImageNormalCombine(pNew.get(), m_base.get(), m_blended.get(), bOnlyOneMip); break;
-				case EBlendType::BT_SOFTLIGHT: BufferLayer<SoftLightChannel, true>(pNew.get(), m_base.get(), m_blended.get(), bApplyToAlpha, bOnlyOneMip); break;
-				case EBlendType::BT_HARDLIGHT: BufferLayer<HardLightChannel, true>(pNew.get(), m_base.get(), m_blended.get(), bApplyToAlpha, bOnlyOneMip); break;
-				case EBlendType::BT_BURN: BufferLayer<BurnChannel, true>(pNew.get(), m_base.get(), m_blended.get(), bApplyToAlpha, bOnlyOneMip); break;
-				case EBlendType::BT_DODGE: BufferLayer<DodgeChannel, true>(pNew.get(), m_base.get(), m_blended.get(), bApplyToAlpha, bOnlyOneMip); break;
-				case EBlendType::BT_SCREEN: BufferLayer<ScreenChannel, true>(pNew.get(), m_base.get(), m_blended.get(), bApplyToAlpha, bOnlyOneMip); break;
-				case EBlendType::BT_OVERLAY: BufferLayer<OverlayChannel, true>(pNew.get(), m_base.get(), m_blended.get(), bApplyToAlpha, bOnlyOneMip); break;
-				case EBlendType::BT_LIGHTEN: BufferLayer<LightenChannel, true>(pNew.get(), m_base.get(), m_blended.get(), bApplyToAlpha, bOnlyOneMip); break;
-				case EBlendType::BT_MULTIPLY: BufferLayer<MultiplyChannel, true>(pNew.get(), m_base.get(), m_blended.get(), bApplyToAlpha, bOnlyOneMip); break;
-				case EBlendType::BT_BLEND: BufferLayer<BlendChannel, true>(pNew.get(), m_base.get(), m_blended.get(), bApplyToAlpha, bOnlyOneMip); break;
+				case EBlendType::BT_SOFTLIGHT: BufferLayer<SoftLightChannel, true>(pNew.get(), m_base.get(), m_blended.get(), bApplyColorBlendToAlpha, bOnlyOneMip); break;
+				case EBlendType::BT_HARDLIGHT: BufferLayer<HardLightChannel, true>(pNew.get(), m_base.get(), m_blended.get(), bApplyColorBlendToAlpha, bOnlyOneMip); break;
+				case EBlendType::BT_BURN: BufferLayer<BurnChannel, true>(pNew.get(), m_base.get(), m_blended.get(), bApplyColorBlendToAlpha, bOnlyOneMip); break;
+				case EBlendType::BT_DODGE: BufferLayer<DodgeChannel, true>(pNew.get(), m_base.get(), m_blended.get(), bApplyColorBlendToAlpha, bOnlyOneMip); break;
+				case EBlendType::BT_SCREEN: BufferLayer<ScreenChannel, true>(pNew.get(), m_base.get(), m_blended.get(), bApplyColorBlendToAlpha, bOnlyOneMip); break;
+				case EBlendType::BT_OVERLAY: BufferLayer<OverlayChannel, true>(pNew.get(), m_base.get(), m_blended.get(), bApplyColorBlendToAlpha, bOnlyOneMip); break;
+				case EBlendType::BT_LIGHTEN: BufferLayer<LightenChannel, true>(pNew.get(), m_base.get(), m_blended.get(), bApplyColorBlendToAlpha, bOnlyOneMip); break;
+				case EBlendType::BT_MULTIPLY: BufferLayer<MultiplyChannel, true>(pNew.get(), m_base.get(), m_blended.get(), bApplyColorBlendToAlpha, bOnlyOneMip); break;
+				case EBlendType::BT_BLEND: BufferLayer<BlendChannel, true>(pNew.get(), m_base.get(), m_blended.get(), bApplyColorBlendToAlpha, bOnlyOneMip); break;
+				default: check(false);
+				}
+			}
+
+			// Apply the separate blend operation for alpha
+			if (!bDone && !bApplyColorBlendToAlpha && Args.blendTypeAlpha != uint8(EBlendType::BT_NONE))
+			{
+				// Separate alpha operation ignores the mask.
+				switch (EBlendType(Args.blendTypeAlpha))
+				{
+				case EBlendType::BT_SOFTLIGHT: BufferLayerInPlace<SoftLightChannel, false, 1>(pNew.get(), m_blended.get(), bOnlyOneMip, 3, 3); break;
+				case EBlendType::BT_HARDLIGHT: BufferLayerInPlace<HardLightChannel, false, 1>(pNew.get(), m_blended.get(), bOnlyOneMip, 3, 3); break;
+				case EBlendType::BT_BURN: BufferLayerInPlace<BurnChannel, false, 1>(pNew.get(), m_blended.get(), bOnlyOneMip, 3, 3); break;
+				case EBlendType::BT_DODGE: BufferLayerInPlace<DodgeChannel, false, 1>(pNew.get(), m_blended.get(), bOnlyOneMip, 3, 3); break;
+				case EBlendType::BT_SCREEN: BufferLayerInPlace<ScreenChannel, false, 1>(pNew.get(), m_blended.get(), bOnlyOneMip, 3, 3); break;
+				case EBlendType::BT_OVERLAY: BufferLayerInPlace<OverlayChannel, false, 1>(pNew.get(), m_blended.get(), bOnlyOneMip, 3, 3); break;
+				case EBlendType::BT_LIGHTEN: BufferLayerInPlace<LightenChannel, false, 1>(pNew.get(), m_blended.get(), bOnlyOneMip, 3, 3); break;
+				case EBlendType::BT_MULTIPLY: BufferLayerInPlace<MultiplyChannel, false, 1>(pNew.get(), m_blended.get(), bOnlyOneMip, 3, 3); break;
+				case EBlendType::BT_BLEND: BufferLayerInPlace<BlendChannel, false, 1>(pNew.get(), m_blended.get(), bOnlyOneMip, 3, 3); break;
 				default: check(false);
 				}
 			}
 		}
 
-		if (finalFormat != EImageFormat::IF_NONE)
+		if (InitialFormat != pNew->GetFormat())
 		{
-			pNew = ImagePixelFormat(m_imageCompressionQuality, pNew.get(), finalFormat);
+			pNew = ImagePixelFormat(m_imageCompressionQuality, pNew.get(), InitialFormat);
 		}
 
 		// Free resources
@@ -585,10 +638,13 @@ namespace mu
 		m_imageCompressionQuality = imageCompressionQuality;
 		Args = InArgs;
 		m_base = Memory.GetImage({ Args.base, InOp.ExecutionIndex, InOp.ExecutionOptions });
+		check(!m_base || m_base->GetFormat() < EImageFormat::IF_COUNT);
+
 		m_col = Memory.GetColour({ Args.colour, InOp.ExecutionIndex, 0 });
 		if (Args.mask)
 		{
 			m_mask = Memory.GetImage({ Args.mask, InOp.ExecutionIndex, InOp.ExecutionOptions });
+			check(!m_mask || m_mask->GetFormat() < EImageFormat::IF_COUNT);
 		}
 	}
 
@@ -604,7 +660,10 @@ namespace mu
 		{
 			return;
 		}
-		
+
+		EImageFormat InitialFormat = m_base->GetFormat();
+		check(InitialFormat < EImageFormat::IF_COUNT);
+
 		if (Args.mask)
 		{
 			if (m_base->GetSize() != m_mask->GetSize())
@@ -635,7 +694,7 @@ namespace mu
 		// Does it apply to colour?
 		if (EBlendType(Args.blendType) != EBlendType::BT_NONE)
 		{
-			pNew = new Image(m_base->GetSizeX(), m_base->GetSizeY(), m_base->GetLODCount(), m_base->GetFormat());
+			pNew = new Image(m_base->GetSizeX(), m_base->GetSizeY(), m_base->GetLODCount(), InitialFormat);
 
 			if (Args.mask)
 			{			
diff --git a/Engine/Plugins/Experimental/Mutable/Source/MutableRuntime/Private/MuR/Image.cpp b/Engine/Plugins/Experimental/Mutable/Source/MutableRuntime/Private/MuR/Image.cpp
--- a/Engine/Plugins/Experimental/Mutable/Source/MutableRuntime/Private/MuR/Image.cpp
+++ b/Engine/Plugins/Experimental/Mutable/Source/MutableRuntime/Private/MuR/Image.cpp
@@ -50,7 +50,7 @@ namespace mu
     //---------------------------------------------------------------------------------------------
     const FImageFormatData& GetImageFormatData( EImageFormat format )
     {
-        check( uint8(format) < int8(EImageFormat::IF_COUNT) );
+        check( format < EImageFormat::IF_COUNT );
         return s_imageFormatData[ uint8(format) ];
     }
 
@@ -65,9 +65,10 @@ namespace mu
     Image::Image( uint32_t sizeX, uint32_t sizeY, uint32_t lods, EImageFormat format )
     {
 		MUTABLE_CPUPROFILER_SCOPE(NewImage)
+			LLM_SCOPE_BYNAME(TEXT("MutableRuntime"));
 
-        check( format!= EImageFormat::IF_NONE );
-		LLM_SCOPE_BYNAME(TEXT("MutableRuntime"));
+        check(format != EImageFormat::IF_NONE);
+		check(format < EImageFormat::IF_COUNT);
 
         m_format = format;
         m_size = FImageSize( (uint16)sizeX, (uint16)sizeY );
@@ -75,8 +76,8 @@ namespace mu
         m_internalId = 0;
 
         const FImageFormatData& fdata = GetImageFormatData( format );
-        int pixelsPerBlock = fdata.m_pixelsPerBlockX*fdata.m_pixelsPerBlockY;
-        if ( pixelsPerBlock )
+        int32 PixelsPerBlock = fdata.m_pixelsPerBlockX*fdata.m_pixelsPerBlockY;
+        if (PixelsPerBlock)
         {
 			int32 DataSize = CalculateDataSize();
             m_data.SetNum( DataSize );
diff --git a/Engine/Plugins/Experimental/Mutable/Source/MutableRuntime/Private/MuR/ImagePrivate.h b/Engine/Plugins/Experimental/Mutable/Source/MutableRuntime/Private/MuR/ImagePrivate.h
--- a/Engine/Plugins/Experimental/Mutable/Source/MutableRuntime/Private/MuR/ImagePrivate.h
+++ b/Engine/Plugins/Experimental/Mutable/Source/MutableRuntime/Private/MuR/ImagePrivate.h
@@ -87,6 +87,8 @@ namespace mu
 	//---------------------------------------------------------------------------------------------
 	MUTABLERUNTIME_API inline EImageFormat GetUncompressedFormat(EImageFormat f )
 	{
+		check(f < EImageFormat::IF_COUNT);
+
 		EImageFormat r = f;
 
 		switch ( r )
diff --git a/Engine/Plugins/Experimental/Mutable/Source/MutableRuntime/Private/MuR/OpMeshMerge.h b/Engine/Plugins/Experimental/Mutable/Source/MutableRuntime/Private/MuR/OpMeshMerge.h
--- a/Engine/Plugins/Experimental/Mutable/Source/MutableRuntime/Private/MuR/OpMeshMerge.h
+++ b/Engine/Plugins/Experimental/Mutable/Source/MutableRuntime/Private/MuR/OpMeshMerge.h
@@ -663,10 +663,8 @@ namespace mu
 
 			if (SharedResultPhysicsBody.Get<1>())
 			{
-				check(!pFirst->GetPhysicsBody() || !pSecond->GetPhysicsBody())
-
-					// Only one or non of the meshes has physics, share the result.
-					pResult->SetPhysicsBody(SharedResultPhysicsBody.Get<0>());
+				// Only one or non of the meshes has physics, share the result.
+				pResult->SetPhysicsBody(SharedResultPhysicsBody.Get<0>());
 			}
 			else
 			{
diff --git a/Engine/Plugins/Experimental/Mutable/Source/MutableRuntime/Private/MuR/Operations.cpp b/Engine/Plugins/Experimental/Mutable/Source/MutableRuntime/Private/MuR/Operations.cpp
--- a/Engine/Plugins/Experimental/Mutable/Source/MutableRuntime/Private/MuR/Operations.cpp
+++ b/Engine/Plugins/Experimental/Mutable/Source/MutableRuntime/Private/MuR/Operations.cpp
@@ -247,24 +247,6 @@ namespace mu
             break;
 
         //-------------------------------------------------------------------------------------
-        case OP_TYPE::IM_LAYER:
-            f(&op.args.ImageLayer.base );
-            if ( op.args.ImageLayer.mask )
-            {
-                f(&op.args.ImageLayer.mask );
-            }
-            f(&op.args.ImageLayer.blended );
-            break;
-
-        case OP_TYPE::IM_LAYERCOLOUR:
-            f(&op.args.ImageLayerColour.base );
-            if ( op.args.ImageLayerColour.mask )
-            {
-                f(&op.args.ImageLayerColour.mask );
-            }
-            f(&op.args.ImageLayerColour.colour );
-            break;
-
         case OP_TYPE::IM_RESIZE:
             f(&op.args.ImageResize.source );
             break;
diff --git a/Engine/Plugins/Experimental/Mutable/Source/MutableRuntime/Private/MuR/Operations.h b/Engine/Plugins/Experimental/Mutable/Source/MutableRuntime/Private/MuR/Operations.h
--- a/Engine/Plugins/Experimental/Mutable/Source/MutableRuntime/Private/MuR/Operations.h
+++ b/Engine/Plugins/Experimental/Mutable/Source/MutableRuntime/Private/MuR/Operations.h
@@ -527,16 +527,19 @@ namespace mu
             ADDRESS base;
             ADDRESS mask;
             ADDRESS blended;
-			uint8 blendType;	// One of BLEND_TYPE
+			uint8 blendType;		// One of EBlendType
+			uint8 blendTypeAlpha;	// One of EBlendType
 			uint8 flags;
 			typedef enum
             {
-                F_NONE              =	0,
-                //! The mask is considered binary: 0 means 0% and any other value means 100%
-                F_BINARY_MASK       =	1 << 0,
-                //! If the image has 4 channels, apply to the fourth channel as well.
-                F_APPLY_TO_ALPHA    =	1 << 1
-            } FLAGS;
+                F_NONE           = 0,
+                /** The mask is considered binary : 0 means 0% and any other value means 100% */
+                F_BINARY_MASK    = 1 << 0,
+				/** If the image has 4 channels, apply to the fourth channel as well. */
+				F_APPLY_TO_ALPHA = 1 << 1,
+				/** Use the alpha channel of the blended image as mask. Mask should be null.*/
+				F_USE_MASK_FROM_BLENDED = 1 << 2
+			} FLAGS;
         };
 
         struct ImageMultiLayerArgs
@@ -546,8 +549,8 @@ namespace mu
             ADDRESS blended;
             ADDRESS rangeSize;
             uint16 rangeId;
-			uint8 blendType;		// One of BLEND_TYPE
-			uint8 blendTypeAlpha;	// One of BLEND_TYPE
+			uint8 blendType;		// One of EBlendType
+			uint8 blendTypeAlpha;	// One of EBlendType
 			uint8 bUseMaskFromBlended;	
 		};
 
@@ -556,8 +559,8 @@ namespace mu
             ADDRESS base;
             ADDRESS mask;
             ADDRESS colour;
-			uint8 blendType;		// One of BLEND_TYPE
-			uint8 blendTypeAlpha;	// One of BLEND_TYPE
+			uint8 blendType;		// One of EBlendType
+			uint8 blendTypeAlpha;	// One of EBlendType
 		};
 
         struct ImagePixelFormatArgs
@@ -1075,8 +1078,6 @@ namespace mu
             ArithmeticArgs ColourArithmetic;
 
             //-------------------------------------------------------------------------------------
-            ImageLayerArgs ImageLayer;
-            ImageLayerColourArgs ImageLayerColour;
             ImageResizeArgs ImageResize;
             ImageResizeLikeArgs ImageResizeLike;
             ImageResizeVarArgs ImageResizeVar;
diff --git a/Engine/Plugins/Experimental/Mutable/Source/MutableRuntime/Private/MuR/System.cpp b/Engine/Plugins/Experimental/Mutable/Source/MutableRuntime/Private/MuR/System.cpp
--- a/Engine/Plugins/Experimental/Mutable/Source/MutableRuntime/Private/MuR/System.cpp
+++ b/Engine/Plugins/Experimental/Mutable/Source/MutableRuntime/Private/MuR/System.cpp
@@ -462,10 +462,10 @@ namespace mu
 
 
     //---------------------------------------------------------------------------------------------
-    ImagePtrConst System::BuildParameterAdditionalImage( const ModelPtrConst& pModel,
-                                                         const ParametersPtrConst& pParams,
-                                                         int parameter,
-                                                         int image )
+    ImagePtrConst System::BuildParameterAdditionalImage( const Ptr<const Model>& pModel,
+                                                         const Ptr<const Parameters>& pParams,
+                                                         int32 parameter,
+                                                         int32 ImageIndex )
     {
 		LLM_SCOPE_BYNAME(TEXT("MutableRuntime"));
 		MUTABLE_CPUPROFILER_SCOPE(BuildParameterAdditionalImage);
@@ -474,11 +474,10 @@ namespace mu
 			&&
 			parameter < (int)pModel->GetPrivate()->m_program.m_parameters.Num());
 
-		const FParameterDesc& param =
-			pModel->GetPrivate()->m_program.m_parameters[parameter];
-		check(image >= 0 && image < (int)param.m_descImages.Num());
+		const FParameterDesc& param = pModel->GetPrivate()->m_program.m_parameters[parameter];
+		check(ImageIndex >= 0 && ImageIndex < (int)param.m_descImages.Num());
 
-		OP::ADDRESS imageAddress = param.m_descImages[image];
+		OP::ADDRESS imageAddress = param.m_descImages[ImageIndex];
 
 		ImagePtrConst pResult;
 
diff --git a/Engine/Plugins/Experimental/Mutable/Source/MutableRuntime/Public/MuR/Model.h b/Engine/Plugins/Experimental/Mutable/Source/MutableRuntime/Public/MuR/Model.h
--- a/Engine/Plugins/Experimental/Mutable/Source/MutableRuntime/Public/MuR/Model.h
+++ b/Engine/Plugins/Experimental/Mutable/Source/MutableRuntime/Public/MuR/Model.h
@@ -12,7 +12,7 @@
 //! data structures. Compiled models are not necessarily compatible when the runtime is updated,
 //! so this version number can be used externally to verify this. It is not used internally, and
 //! serializing models from different versions than this runtime will probably result in a crash.
-#define MUTABLE_COMPILED_MODEL_CODE_VERSION		uint32( 51 )
+#define MUTABLE_COMPILED_MODEL_CODE_VERSION		uint32( 52 )
 #define MUTABLE_PARAMETERS_VERSION              uint32( 1 )
 
 
diff --git a/Engine/Plugins/Experimental/Mutable/Source/MutableRuntime/Public/MuR/System.h b/Engine/Plugins/Experimental/Mutable/Source/MutableRuntime/Public/MuR/System.h
--- a/Engine/Plugins/Experimental/Mutable/Source/MutableRuntime/Public/MuR/System.h
+++ b/Engine/Plugins/Experimental/Mutable/Source/MutableRuntime/Public/MuR/System.h
@@ -21,22 +21,10 @@ namespace mu
 	// Forward references
 	class Model;
 	class ModelStreamer;
-
-    using ModelPtr=Ptr<Model>;
-    using ModelPtrConst=Ptr<const Model>;
-
 	class Parameters;
-
-    using ParametersPtr=Ptr<Parameters>;
-    using ParametersPtrConst=Ptr<const Parameters>;
-
     class Mesh;
 
-    using MeshPtr=Ptr<Mesh>;
-    using MeshPtrConst=Ptr<const Mesh>;
-
     class System;
-
     using SystemPtr=Ptr<System>;
     using SystemPtrConst=Ptr<const System>;
 
@@ -82,7 +70,7 @@ namespace mu
         virtual ~ImageParameterGenerator() = default;
 
         //!
-        virtual ImagePtr GetImage( EXTERNAL_IMAGE_ID id ) = 0;
+        virtual Ptr<Image> GetImage( EXTERNAL_IMAGE_ID id ) = 0;
     };
 
 
@@ -98,9 +86,9 @@ namespace mu
     public:
 
 		//! Constructor of a system object to build data.
-        //! \param settings Optional class with the settings to use in this system. The default
+        //! \param Settings Optional class with the settings to use in this system. The default
         //! value configures a production-ready system.
-        System( const SettingsPtr& settings = nullptr );
+        System( const Ptr<Settings>& Settings = nullptr );
 
         //! Set a new provider for model data. The provider will become owned by this instance and
         //! destroyed when necessary.
@@ -112,6 +100,7 @@ namespace mu
 
         //! Set a new provider for external image data. This is only necessary if image parameters
         //! are used in the models.
+		//! The provided pointer becomes owned by the system.
         void SetImageParameterGenerator( ImageParameterGenerator* );
 
         //! Set the maximum memory that this system can use. This memory is used for built data,
@@ -126,7 +115,7 @@ namespace mu
         //! no longer needed.
         //! \param pModel Model to build an instance of
         //! \return An identifier that is always bigger than 0.
-        Instance::ID NewInstance( const ModelPtrConst& pModel );
+        Instance::ID NewInstance( const Ptr<const Model>& Model );
 
         //! \brief Update an instance with a new parameter set and/or state.
         //!
@@ -140,41 +129,41 @@ namespace mu
         //! \return the instance data with all the LOD, components, and ids to generate the meshes 
         //! and images.  The returned Instance is only valid until the next call to EndUpdate with 
         //! the same instanceID parameter.
-        const Instance* BeginUpdate( Instance::ID instanceID,
-                                     const ParametersPtrConst& pParams,
-                                     int32 stateIndex,
-                                     uint32 lodMask );
+        const Instance* BeginUpdate( Instance::ID InstanceID,
+                                     const Ptr<const Parameters>& Params,
+                                     int32 StateIndex,
+                                     uint32 LodMask );
 
 		//! Only valid between BeginUpdate and EndUpdate
 		//! Calculate the description of an image, without generating it.
-		void GetImageDesc(Instance::ID instanceID, RESOURCE_ID imageId, FImageDesc& OutDesc );
+		void GetImageDesc(Instance::ID InstanceID, RESOURCE_ID ImageId, FImageDesc& OutDesc );
 
 		//! Only valid between BeginUpdate and EndUpdate
 		//! \param MipsToSkip Number of mips to skip compared from the full image.
 		//! If 0, all mip levels will be generated. If more levels than possible to discard are specified, 
 		//! the image will still contain a minimum number of mips specified at model compile time.
-		ImagePtrConst GetImage(Instance::ID instanceID, RESOURCE_ID imageId, int32 MipsToSkip = 0);
+		Ptr<const Image> GetImage(Instance::ID InstanceID, RESOURCE_ID ImageId, int32 MipsToSkip = 0);
 
         //! Only valid between BeginUpdate and EndUpdate
-        MeshPtrConst GetMesh(Instance::ID instanceID, RESOURCE_ID meshId);
+        Ptr<const Mesh> GetMesh(Instance::ID InstanceID, RESOURCE_ID MeshId);
 
         //! Invalidate and free the last Instance data returned by a call to BeginUpdate with
         //! the same instance index. After a call to this method, that Instance cannot be used any
         //! more and its content is undefined.
         //! \param instance The index of the instance whose last data will be invalidated.
-        void EndUpdate( Instance::ID instance );
+        void EndUpdate( Instance::ID InstanceID);
 
         //! Completely destroy an instance. After a call to this method the given instance cannot be
         //! updated any more, and its resources may have been freed.
         //! \param instance The id of the instance to destroy.
-        void ReleaseInstance( Instance::ID instance );
+        void ReleaseInstance( Instance::ID InstanceID);
 
 		//! Build one of the images defined in a model parameter as additional description. These
 		//! images can be used for colour bars, icons, etc...
-        ImagePtrConst BuildParameterAdditionalImage( const ModelPtrConst& pModel,
-                                                     const ParametersPtrConst& pParams,
-                                                     int parameter,
-                                                     int image );
+		Ptr<const Image> BuildParameterAdditionalImage( const Ptr<const Model>& Model,
+														const Ptr<const Parameters>& Params,
+														int32 Parameter,
+														int32 ImageIndex );
 
 		//! Calculate the relevancy of every parameter. Some parameters may be unused depending on
 		//! the values of other parameters. This method will set to true the flags for parameters
@@ -184,9 +173,9 @@ namespace mu
         //! \param pParameters Parameter set that we want to find the relevancy of.
         //! \param pFlags is a pointer to a preallocated array of booleans that contains at least
 		//! pParameters->GetCount() elements.
-        void GetParameterRelevancy( const ModelPtrConst& pModel,
-									const ParametersPtrConst& pParameters,
-									bool* pFlags );
+        void GetParameterRelevancy( const Ptr<const Model>& Model,
+									const Ptr<const Parameters>& Parameters,
+									bool* Flags );
 
         //! Free memory used in internal runtime caches. This is useful for long-running processes
         //! that keep models loaded. This could be called when a game finishes, or a change of
diff --git a/Engine/Plugins/Experimental/Mutable/Source/MutableTools/Private/MuT/AST.cpp b/Engine/Plugins/Experimental/Mutable/Source/MutableTools/Private/MuT/AST.cpp
--- a/Engine/Plugins/Experimental/Mutable/Source/MutableTools/Private/MuT/AST.cpp
+++ b/Engine/Plugins/Experimental/Mutable/Source/MutableTools/Private/MuT/AST.cpp
@@ -1122,22 +1122,6 @@ FImageDesc ASTOpFixed::GetImageDesc( bool returnBestOption, FGetImageDescContext
     case OP_TYPE::NONE:
         break;
 
-    case OP_TYPE::IM_LAYER:
-        res = GetImageDesc( op.args.ImageLayer.base, returnBestOption, context );
-        if ( res.m_format==EImageFormat::IF_L_UBYTE )
-        {
-            res.m_format = EImageFormat::IF_RGB_UBYTE;
-        }
-        break;
-
-    case OP_TYPE::IM_LAYERCOLOUR:
-        res = GetImageDesc( op.args.ImageLayerColour.base, returnBestOption, context );
-        if ( res.m_format==EImageFormat::IF_L_UBYTE )
-        {
-            res.m_format = EImageFormat::IF_RGB_UBYTE;
-        }
-        break;
-
     case OP_TYPE::IM_SATURATE:
         res = GetImageDesc( op.args.ImageSaturate.base, returnBestOption, context );
         break;
@@ -1303,18 +1287,6 @@ void ASTOpFixed::GetLayoutBlockSize( int* pBlockX, int* pBlockY )
 		break;
 	}
 
-    case OP_TYPE::IM_LAYER:
-    {
-        GetLayoutBlockSize(op.args.ImageLayer.base, pBlockX, pBlockY);
-        break;
-    }
-
-    case OP_TYPE::IM_LAYERCOLOUR:
-    {
-        GetLayoutBlockSize(op.args.ImageLayerColour.base, pBlockX, pBlockY);
-        break;
-    }
-
 	case OP_TYPE::IM_BLANKLAYOUT:
 	{
 		*pBlockX = int(op.args.ImageBlankLayout.blockSize[0]);
@@ -1628,20 +1600,6 @@ mu::Ptr<ImageSizeExpression> ASTOpFixed::GetImageSizeExpression() const
         break;
     }
 
-    case OP_TYPE::IM_LAYER:
-        if ( children[op.args.ImageLayer.base] )
-        {
-            pRes = children[op.args.ImageLayer.base].child()->GetImageSizeExpression();
-        }
-        break;
-
-    case OP_TYPE::IM_LAYERCOLOUR:
-        if ( children[op.args.ImageLayerColour.base] )
-        {
-            pRes = children[op.args.ImageLayerColour.base].child()->GetImageSizeExpression();
-        }
-        break;
-
     case OP_TYPE::IM_RESIZE:
         pRes->type = ImageSizeExpression::ISET_CONSTANT;
         pRes->size[0] = op.args.ImageResize.size[0];
diff --git a/Engine/Plugins/Experimental/Mutable/Source/MutableTools/Private/MuT/ASTOpImageLayer.cpp b/Engine/Plugins/Experimental/Mutable/Source/MutableTools/Private/MuT/ASTOpImageLayer.cpp
new file mode 100644
--- /dev/null
+++ b/Engine/Plugins/Experimental/Mutable/Source/MutableTools/Private/MuT/ASTOpImageLayer.cpp
@@ -0,0 +1,407 @@
+// Copyright Epic Games, Inc. All Rights Reserved.
+
+#include "MuT/ASTOpImageLayer.h"
+
+#include "Containers/Map.h"
+#include "HAL/PlatformMath.h"
+#include "MuR/MemoryPrivate.h"
+#include "MuR/ModelPrivate.h"
+#include "MuR/RefCounted.h"
+#include "MuR/Types.h"
+#include "MuT/StreamsPrivate.h"
+#include "MuT/ASTOpImagePatch.h"
+
+#include <memory>
+#include <utility>
+
+
+namespace mu
+{
+
+
+	ASTOpImageLayer::ASTOpImageLayer()
+		: base(this)
+		, blend(this)
+		, mask(this)
+	{
+	}
+
+
+	//-------------------------------------------------------------------------------------------------
+	ASTOpImageLayer::~ASTOpImageLayer()
+	{
+		// Explicit call needed to avoid recursive destruction
+		ASTOp::RemoveChildren();
+	}
+
+
+	//-------------------------------------------------------------------------------------------------
+	bool ASTOpImageLayer::IsEqual(const ASTOp& InOtherUntyped) const
+	{
+		if (const ASTOpImageLayer* Other = dynamic_cast<const ASTOpImageLayer*>(&InOtherUntyped))
+		{
+			return base == Other->base &&
+				blend == Other->blend &&
+				mask == Other->mask &&
+				blendType == Other->blendType &&
+				blendTypeAlpha == Other->blendTypeAlpha &&
+				Flags == Other->Flags;
+		}
+		return false;
+	}
+
+
+	//-------------------------------------------------------------------------------------------------
+	uint64 ASTOpImageLayer::Hash() const
+	{
+		uint64 res = std::hash<OP_TYPE>()(GetOpType());
+		hash_combine(res, base.child().get());
+		hash_combine(res, blend.child().get());
+		hash_combine(res, mask.child().get());
+		return res;
+	}
+
+
+	//-------------------------------------------------------------------------------------------------
+	mu::Ptr<ASTOp> ASTOpImageLayer::Clone(MapChildFuncRef mapChild) const
+	{
+		Ptr<ASTOpImageLayer> n = new ASTOpImageLayer();
+		n->base = mapChild(base.child());
+		n->blend = mapChild(blend.child());
+		n->mask = mapChild(mask.child());
+		n->blendType = blendType;
+		n->blendTypeAlpha = blendTypeAlpha;
+		n->Flags = Flags;
+		return n;
+	}
+
+
+	//-------------------------------------------------------------------------------------------------
+	void ASTOpImageLayer::ForEachChild(const TFunctionRef<void(ASTChild&)> f)
+	{
+		f(base);
+		f(blend);
+		f(mask);
+	}
+
+
+	//-------------------------------------------------------------------------------------------------
+	void ASTOpImageLayer::Link(FProgram& program, const FLinkerOptions*)
+	{
+		// Already linked?
+		if (!linkedAddress)
+		{
+			OP::ImageLayerArgs args;
+			FMemory::Memzero(&args, sizeof(args));
+
+			args.blendType = (uint8)blendType;
+			args.blendTypeAlpha = (uint8)blendTypeAlpha;
+			args.flags = Flags;
+
+			if (base) args.base = base->linkedAddress;
+			if (blend) args.blended = blend->linkedAddress;
+			if (mask) args.mask = mask->linkedAddress;
+
+			linkedAddress = (OP::ADDRESS)program.m_opAddress.Num();
+			program.m_opAddress.Add((uint32)program.m_byteCode.Num());
+			AppendCode(program.m_byteCode, GetOpType());
+			AppendCode(program.m_byteCode, args);
+		}
+	}
+
+
+	//-------------------------------------------------------------------------------------------------
+	FImageDesc ASTOpImageLayer::GetImageDesc(bool returnBestOption, FGetImageDescContext* context) const 
+	{
+		FImageDesc res;
+
+		// Local context in case it is necessary
+		FGetImageDescContext localContext;
+		if (!context)
+		{
+			context = &localContext;
+		}
+		else
+		{
+			// Cached result?
+			FImageDesc* PtrValue = context->m_results.Find(this);
+			if (PtrValue)
+			{
+				return *PtrValue;
+			}
+		}
+
+		// Actual work
+		if (base)
+		{
+			res = base->GetImageDesc(returnBestOption, context);
+		}
+
+
+		// Cache the result
+		if (context)
+		{
+			context->m_results.Add(this, res);
+		}
+
+		return res;
+	}
+
+
+	//-------------------------------------------------------------------------------------------------
+	void ASTOpImageLayer::GetLayoutBlockSize(int* pBlockX, int* pBlockY)
+	{
+		if (base)
+		{
+			base->GetLayoutBlockSize(pBlockX, pBlockY);
+		}
+	}
+
+
+	//-------------------------------------------------------------------------------------------------
+	mu::Ptr<ImageSizeExpression> ASTOpImageLayer::GetImageSizeExpression() const
+	{
+		if (base)
+		{
+			return base->GetImageSizeExpression();
+		}
+
+		return nullptr;
+	}
+
+
+
+	//---------------------------------------------------------------------------------------------
+	Ptr<ASTOp> ASTOpImageLayer::OptimiseSemantic(const FModelOptimizationOptions& options) const
+	{
+		Ptr<ASTOp> at;
+
+		auto baseAt = base.child();
+
+		// Plain masks optimization
+		auto maskAt = mask.child();
+		if (maskAt)
+		{
+			FVector4f colour;
+			if (maskAt->IsImagePlainConstant(colour))
+			{
+				if (colour.IsNearlyZero3(UE_SMALL_NUMBER))
+				{
+					// If the mask is black, we can skip the entire operation
+					at = base.child();
+				}
+				else if (colour.Equals(FVector4f(1, 1, 1, 1), UE_SMALL_NUMBER))
+				{
+					// If the mask is white, we can remove it
+					Ptr<ASTOpImageLayer> nop = mu::Clone<ASTOpImageLayer>(this);
+					nop->mask = nullptr;
+					at = nop;
+				}
+			}
+		}
+
+		// Introduce crop if mask is constant and smaller than the base
+		FImageRect sourceMaskUsage;
+		FImageDesc maskDesc;
+		bool validUsageRect = false;
+		if (!at && maskAt)
+		{
+			//MUTABLE_CPUPROFILER_SCOPE(EvaluateAreasForCrop);
+
+			validUsageRect = maskAt->GetNonBlackRect(sourceMaskUsage);
+			if (validUsageRect)
+			{
+				check(sourceMaskUsage.size[0] > 0);
+				check(sourceMaskUsage.size[1] > 0);
+
+				FGetImageDescContext context;
+				maskDesc = maskAt->GetImageDesc(false, &context);
+			}
+		}
+
+		if (!at && maskAt && validUsageRect)
+		{
+			// Adjust for compressed blocks (4), and some extra mips (2 more mips, which is 4)
+			constexpr int blockSize = 4 * 4;
+
+			FImageRect maskUsage;
+			maskUsage.min[0] = (sourceMaskUsage.min[0] / blockSize) * blockSize;
+			maskUsage.min[1] = (sourceMaskUsage.min[1] / blockSize) * blockSize;
+			vec2<uint16> minOffset = sourceMaskUsage.min - maskUsage.min;
+			maskUsage.size[0] = ((sourceMaskUsage.size[0] + minOffset[0] + blockSize - 1) / blockSize) * blockSize;
+			maskUsage.size[1] = ((sourceMaskUsage.size[1] + minOffset[1] + blockSize - 1) / blockSize) * blockSize;
+
+			// Is it worth?
+			float ratio = float(maskUsage.size[0] * maskUsage.size[1])
+				/ float(maskDesc.m_size[0] * maskDesc.m_size[1]);
+			float acceptableCropRatio = options.m_acceptableCropRatio;
+			if (ratio < acceptableCropRatio)
+			{
+				check(maskUsage.size[0] > 0);
+				check(maskUsage.size[1] > 0);
+
+				Ptr<ASTOpFixed> cropMask = new ASTOpFixed();
+				cropMask->op.type = OP_TYPE::IM_CROP;
+				cropMask->SetChild(cropMask->op.args.ImageCrop.source, mask.child());
+				cropMask->op.args.ImageCrop.minX = maskUsage.min[0];
+				cropMask->op.args.ImageCrop.minY = maskUsage.min[1];
+				cropMask->op.args.ImageCrop.sizeX = maskUsage.size[0];
+				cropMask->op.args.ImageCrop.sizeY = maskUsage.size[1];
+
+				Ptr<ASTOpFixed> cropBlended = new ASTOpFixed();
+				cropBlended->op.type = OP_TYPE::IM_CROP;
+				cropBlended->SetChild(cropBlended->op.args.ImageCrop.source, blend.child());
+				cropBlended->op.args.ImageCrop.minX = maskUsage.min[0];
+				cropBlended->op.args.ImageCrop.minY = maskUsage.min[1];
+				cropBlended->op.args.ImageCrop.sizeX = maskUsage.size[0];
+				cropBlended->op.args.ImageCrop.sizeY = maskUsage.size[1];
+
+				Ptr<ASTOpFixed> cropBase = new ASTOpFixed();
+				cropBase->op.type = OP_TYPE::IM_CROP;
+				cropBase->SetChild(cropBase->op.args.ImageCrop.source, base.child());
+				cropBase->op.args.ImageCrop.minX = maskUsage.min[0];
+				cropBase->op.args.ImageCrop.minY = maskUsage.min[1];
+				cropBase->op.args.ImageCrop.sizeX = maskUsage.size[0];
+				cropBase->op.args.ImageCrop.sizeY = maskUsage.size[1];
+
+				Ptr<ASTOpImageLayer> newLayer = mu::Clone<ASTOpImageLayer>(this);
+				newLayer->base = cropBase;
+				newLayer->blend = cropBlended;
+				newLayer->mask = cropMask;
+
+				Ptr<ASTOpImagePatch> patch = new ASTOpImagePatch();
+				patch->base = baseAt;
+				patch->patch = newLayer;
+				patch->location = maskUsage.min;
+				at = patch;
+			}
+		}
+
+		return at;
+	}
+
+
+	Ptr<ASTOp> ASTOpImageLayer::OptimiseSink(const FModelOptimizationOptions& options, FOptimizeSinkContext& context) const
+	{
+		Ptr<ASTOp> at;
+
+		// Layer effects may be worth sinking down switches and conditionals, to be able
+		// to apply extra optimisations
+		auto baseAt = base.child();
+		auto blendAt = blend.child();
+		auto maskAt = mask.child();
+
+		// Promote conditions from the base
+		OP_TYPE baseType = baseAt->GetOpType();
+		switch (baseType)
+		{
+			// Seems to cause operation explosion in optimizer in bandit model.
+			// moved to generic sink in the default.
+//            case OP_TYPE::IM_CONDITIONAL:
+//            {
+//                m_modified = true;
+
+//                OP op = program.m_code[baseAt];
+
+//                OP aOp = program.m_code[at];
+//                aOp.args.ImageLayer.base = program.m_code[baseAt].args.Conditional.yes;
+//                op.args.Conditional.yes = program.AddOp( aOp );
+
+//                OP bOp = program.m_code[at];
+//                bOp.args.ImageLayer.base = program.m_code[baseAt].args.Conditional.no;
+//                op.args.Conditional.no = program.AddOp( bOp );
+
+//                at = program.AddOp( op );
+//                break;
+//            }
+
+		case OP_TYPE::IM_SWITCH:
+		{
+			// Disabled:
+			// It seems to cause data explosion in optimizer in some models. Because
+			// all switch branches become unique constants
+
+//                // See if the blended has an identical switch, to optimise it too
+//                auto baseSwitch = dynamic_cast<const ASTOpSwitch*>( baseAt.get() );
+//                auto blendedSwitch = dynamic_cast<const ASTOpSwitch*>( blendAt.get() );
+//                auto maskSwitch = dynamic_cast<const ASTOpSwitch*>( maskAt.get() );
+
+//                bool blendedToo = baseSwitch->IsCompatibleWith( blendedSwitch );
+//                bool maskToo = baseSwitch->IsCompatibleWith( maskSwitch );
+
+//                // Move the layer operation down all the paths
+//                auto nop = mu::Clone<ASTOpSwitch>(baseSwitch);
+
+//                if (nop->def)
+//                {
+//                    auto defOp = mu::Clone<ASTOpFixed>(this);
+//                    defOp->SetChild( defOp->op.args.ImageLayer.base, nop->def );
+//                    if (blendedToo)
+//                    {
+//                        defOp->SetChild( defOp->op.args.ImageLayer.blended, blendedSwitch->def );
+//                    }
+//                    if (maskToo)
+//                    {
+//                        defOp->SetChild( defOp->op.args.ImageLayer.mask, maskSwitch->def );
+//                    }
+//                    nop->def = defOp;
+//                }
+
+//                for ( size_t v=0; v<nop->cases.Num(); ++v )
+//                {
+//                    if ( nop->cases[v].branch )
+//                    {
+//                        auto bOp = mu::Clone<ASTOpFixed>(this);
+//                        bOp->SetChild( bOp->op.args.ImageLayer.base, nop->cases[v].branch );
+
+//                        if (blendedToo)
+//                        {
+//                            bOp->SetChild( bOp->op.args.ImageLayer.blended, blendedSwitch->FindBranch( nop->cases[v].condition ) );
+//                        }
+
+//                        if (maskToo)
+//                        {
+//                            bOp->SetChild( bOp->op.args.ImageLayer.mask, maskSwitch->FindBranch( nop->cases[v].condition ) );
+//                        }
+
+//                        nop->cases[v].branch = bOp;
+//                    }
+//                }
+
+//                at = nop;
+			break;
+		}
+
+		default:
+		{
+			// Try generic base sink
+//                OP templateOp = program.m_code[at];
+
+//                CloneNeutralPartialTreeVisitor neutral;
+//                OP_TYPE supportedOps[] = {
+//                    OP_TYPE::IM_CONDITIONAL,
+//                    OP_TYPE::NONE
+//                };
+//                auto newRoot = neutral.Apply
+//                         (
+//                            program.m_code[at].args.ImageLayer.base, program,
+//                            DT_IMAGE, supportedOps,
+//                            &templateOp, &templateOp.args.ImageLayer.base);
+
+//                // If there is a change
+//                if (newRoot!=program.m_code[at].args.ImageLayer.base)
+//                {
+//                    m_modified = true;
+//                    at = newRoot;
+//                }
+
+			break;
+		}
+
+		}
+
+		return at;
+	}
+
+
+}
diff --git a/Engine/Plugins/Experimental/Mutable/Source/MutableTools/Private/MuT/ASTOpImageLayer.h b/Engine/Plugins/Experimental/Mutable/Source/MutableTools/Private/MuT/ASTOpImageLayer.h
new file mode 100644
--- /dev/null
+++ b/Engine/Plugins/Experimental/Mutable/Source/MutableTools/Private/MuT/ASTOpImageLayer.h
@@ -0,0 +1,58 @@
+// Copyright Epic Games, Inc. All Rights Reserved.
+
+#pragma once
+
+#include "HAL/Platform.h"
+#include "MuR/Image.h"
+#include "MuR/Operations.h"
+#include "MuR/Ptr.h"
+#include "MuT/AST.h"
+
+
+namespace mu
+{
+	struct FProgram;
+
+
+	//---------------------------------------------------------------------------------------------
+	//!
+	//---------------------------------------------------------------------------------------------
+	class ASTOpImageLayer : public ASTOp
+	{
+	public:
+
+		ASTChild base;
+		ASTChild blend;
+		ASTChild mask;
+
+		/** Blend type used for the colour channels. */
+		EBlendType blendType = EBlendType::BT_NONE;
+
+		/** Blend type used for the alpha channel if any. */
+		EBlendType blendTypeAlpha = EBlendType::BT_NONE;
+
+		/** See ImageLayerArgs::Flags .*/
+		uint8 Flags = 0;
+
+	public:
+
+		ASTOpImageLayer();
+		ASTOpImageLayer(const ASTOpImageLayer&) = delete;
+		~ASTOpImageLayer();
+
+		OP_TYPE GetOpType() const override { return OP_TYPE::IM_LAYER; }
+		uint64 Hash() const override;
+		bool IsEqual(const ASTOp& otherUntyped) const override;
+		Ptr<ASTOp> Clone(MapChildFuncRef mapChild) const override;
+		void ForEachChild(const TFunctionRef<void(ASTChild&)>) override;
+		void Link(FProgram& program, const FLinkerOptions*) override;
+		FImageDesc GetImageDesc(bool returnBestOption, FGetImageDescContext*) const override;
+		void GetLayoutBlockSize(int* pBlockX, int* pBlockY) override;
+		Ptr<ImageSizeExpression> GetImageSizeExpression() const override;
+		Ptr<ASTOp> OptimiseSemantic(const FModelOptimizationOptions&) const override;
+		Ptr<ASTOp> OptimiseSink(const FModelOptimizationOptions&, FOptimizeSinkContext&) const;
+
+	};
+
+}
+
diff --git a/Engine/Plugins/Experimental/Mutable/Source/MutableTools/Private/MuT/ASTOpImageLayerColor.cpp b/Engine/Plugins/Experimental/Mutable/Source/MutableTools/Private/MuT/ASTOpImageLayerColor.cpp
new file mode 100644
--- /dev/null
+++ b/Engine/Plugins/Experimental/Mutable/Source/MutableTools/Private/MuT/ASTOpImageLayerColor.cpp
@@ -0,0 +1,291 @@
+// Copyright Epic Games, Inc. All Rights Reserved.
+
+#include "MuT/ASTOpImageLayerColor.h"
+
+#include "MuT/ASTOpSwitch.h"
+#include "Containers/Map.h"
+#include "HAL/PlatformMath.h"
+#include "MuR/MemoryPrivate.h"
+#include "MuR/ModelPrivate.h"
+#include "MuR/RefCounted.h"
+#include "MuR/Types.h"
+#include "MuT/StreamsPrivate.h"
+
+#include <memory>
+#include <utility>
+
+
+namespace mu
+{
+
+
+	ASTOpImageLayerColor::ASTOpImageLayerColor()
+		: base(this)
+		, color(this)
+		, mask(this)
+	{
+	}
+
+
+	//-------------------------------------------------------------------------------------------------
+	ASTOpImageLayerColor::~ASTOpImageLayerColor()
+	{
+		// Explicit call needed to avoid recursive destruction
+		ASTOp::RemoveChildren();
+	}
+
+
+	//-------------------------------------------------------------------------------------------------
+	bool ASTOpImageLayerColor::IsEqual(const ASTOp& InOtherUntyped) const
+	{
+		if (const ASTOpImageLayerColor* Other = dynamic_cast<const ASTOpImageLayerColor*>(&InOtherUntyped))
+		{
+			return base == Other->base &&
+				color == Other->color &&
+				mask == Other->mask &&
+				blendType == Other->blendType &&
+				blendTypeAlpha == Other->blendTypeAlpha;
+		}
+		return false;
+	}
+
+
+	//-------------------------------------------------------------------------------------------------
+	uint64 ASTOpImageLayerColor::Hash() const
+	{
+		uint64 res = std::hash<OP_TYPE>()(GetOpType());
+		hash_combine(res, base.child().get());
+		hash_combine(res, color.child().get());
+		hash_combine(res, mask.child().get());
+		return res;
+	}
+
+
+	//-------------------------------------------------------------------------------------------------
+	mu::Ptr<ASTOp> ASTOpImageLayerColor::Clone(MapChildFuncRef mapChild) const
+	{
+		Ptr<ASTOpImageLayerColor> n = new ASTOpImageLayerColor();
+		n->base = mapChild(base.child());
+		n->color = mapChild(color.child());
+		n->mask = mapChild(mask.child());
+		n->blendType = blendType;
+		n->blendTypeAlpha = blendTypeAlpha;
+		return n;
+	}
+
+
+	//-------------------------------------------------------------------------------------------------
+	void ASTOpImageLayerColor::ForEachChild(const TFunctionRef<void(ASTChild&)> f)
+	{
+		f(base);
+		f(color);
+		f(mask);
+	}
+
+
+	//-------------------------------------------------------------------------------------------------
+	void ASTOpImageLayerColor::Link(FProgram& program, const FLinkerOptions*)
+	{
+		// Already linked?
+		if (!linkedAddress)
+		{
+			OP::ImageLayerColourArgs args;
+			FMemory::Memzero(&args, sizeof(args));
+
+			args.blendType = (uint8)blendType;
+			args.blendTypeAlpha = (uint8)blendTypeAlpha;
+
+			if (base) args.base = base->linkedAddress;
+			if (color) args.colour = color->linkedAddress;
+			if (mask) args.mask = mask->linkedAddress;
+
+			linkedAddress = (OP::ADDRESS)program.m_opAddress.Num();
+			program.m_opAddress.Add((uint32)program.m_byteCode.Num());
+			AppendCode(program.m_byteCode, GetOpType());
+			AppendCode(program.m_byteCode, args);
+		}
+	}
+
+
+	//-------------------------------------------------------------------------------------------------
+	FImageDesc ASTOpImageLayerColor::GetImageDesc(bool returnBestOption, FGetImageDescContext* context) const 
+	{
+		FImageDesc res;
+
+		// Local context in case it is necessary
+		FGetImageDescContext localContext;
+		if (!context)
+		{
+			context = &localContext;
+		}
+		else
+		{
+			// Cached result?
+			FImageDesc* PtrValue = context->m_results.Find(this);
+			if (PtrValue)
+			{
+				return *PtrValue;
+			}
+		}
+
+		// Actual work
+		if (base)
+		{
+			res = base->GetImageDesc(returnBestOption, context);
+		}
+
+
+		// Cache the result
+		if (context)
+		{
+			context->m_results.Add(this, res);
+		}
+
+		return res;
+	}
+
+
+	//-------------------------------------------------------------------------------------------------
+	void ASTOpImageLayerColor::GetLayoutBlockSize(int* pBlockX, int* pBlockY)
+	{
+		if (base)
+		{
+			base->GetLayoutBlockSize(pBlockX, pBlockY);
+		}
+	}
+
+
+	//-------------------------------------------------------------------------------------------------
+	mu::Ptr<ImageSizeExpression> ASTOpImageLayerColor::GetImageSizeExpression() const
+	{
+		if (base)
+		{
+			return base->GetImageSizeExpression();
+		}
+
+		return nullptr;
+	}
+
+
+	//-------------------------------------------------------------------------------------------------
+	Ptr<ASTOp> ASTOpImageLayerColor::OptimiseSink(const FModelOptimizationOptions& options, FOptimizeSinkContext& context) const
+	{
+		Ptr<ASTOp> at;
+
+		// Layer effects may be worth sinking down switches and conditionals, to be able
+		// to apply extra optimisations
+		auto baseAt = base.child();
+		auto maskAt = mask.child();
+
+		// Promote conditions from the base
+		OP_TYPE baseType = baseAt->GetOpType();
+		switch (baseType)
+		{
+			// Seems to cause operation explosion in optimizer in bandit model.
+			// moved to generic sink in the default.
+//            case OP_TYPE::IM_CONDITIONAL:
+//            {
+//                m_modified = true;
+
+//                OP op = program.m_code[baseAt];
+
+//                OP aOp = program.m_code[at];
+//                aOp.args.ImageLayerColour.base = program.m_code[baseAt].args.Conditional.yes;
+//                op.args.Conditional.yes = program.AddOp( aOp );
+
+//                OP bOp = program.m_code[at];
+//                bOp.args.ImageLayerColour.base = program.m_code[baseAt].args.Conditional.no;
+//                op.args.Conditional.no = program.AddOp( bOp );
+
+//                at = program.AddOp( op );
+//                break;
+//            }
+
+		case OP_TYPE::IM_SWITCH:
+		{
+			// Warning:
+			// It seems to cause data explosion in optimizer in some models. Because
+			// all switch branches become unique constants
+
+			// See if the blended has an identical switch, to optimise it too
+			const ASTOpSwitch* baseSwitch = dynamic_cast<const ASTOpSwitch*>(baseAt.get());
+
+			// Mask not supported yet
+			if (maskAt)
+			{
+				break;
+			}
+
+			// Move the layer operation down base paths
+			Ptr<ASTOpSwitch> nop = mu::Clone<ASTOpSwitch>(baseSwitch);
+
+			if (nop->def)
+			{
+				Ptr<ASTOpImageLayerColor> defOp = mu::Clone<ASTOpImageLayerColor>(this);
+				defOp->base = nop->def.child();
+				nop->def = defOp;
+			}
+
+			for (int32 v = 0; v < nop->cases.Num(); ++v)
+			{
+				if (nop->cases[v].branch)
+				{
+					Ptr<ASTOpImageLayerColor> bOp = mu::Clone<ASTOpImageLayerColor>(this);
+					bOp->base = nop->cases[v].branch.child();
+					nop->cases[v].branch = bOp;
+				}
+			}
+
+			at = nop;
+			break;
+		}
+
+		case OP_TYPE::IM_DISPLACE:
+		{
+			// Mask not supported yet. If there is a mask it wouldn't be correct to sink
+			// unless the mask was a similar displace.
+			if (maskAt)
+			{
+				break;
+			}
+
+			Ptr<ASTOpFixed> NewDisplace = mu::Clone<ASTOpFixed>(baseAt);
+
+			Ptr<ASTOp> sourceOp = NewDisplace->children[NewDisplace->op.args.ImageDisplace.source].child();
+			Ptr<ASTOpImageLayerColor> NewSource = mu::Clone<ASTOpImageLayerColor>(this);
+			NewSource->base = sourceOp;
+			NewDisplace->SetChild(NewDisplace->op.args.ImageDisplace.source, NewSource);
+
+			at = NewDisplace;
+			break;
+		}
+
+		case OP_TYPE::IM_RASTERMESH:
+		{
+			// Mask not supported yet. If there is a mask it wouldn't be correct to sink.				
+			if (maskAt)
+			{
+				break;
+			}
+
+			Ptr<ASTOpFixed> NewRaster = mu::Clone<ASTOpFixed>(baseAt);
+
+			Ptr<ASTOp> sourceOp = NewRaster->children[NewRaster->op.args.ImageRasterMesh.image].child();
+			Ptr<ASTOpImageLayerColor> NewSource = mu::Clone<ASTOpImageLayerColor>(this);
+			NewSource->base = sourceOp;
+			NewRaster->SetChild(NewRaster->op.args.ImageRasterMesh.image, NewSource);
+
+			at = NewRaster;
+			break;
+		}
+
+
+		default:
+			break;
+
+		}
+
+		return at;
+	}
+
+}
diff --git a/Engine/Plugins/Experimental/Mutable/Source/MutableTools/Private/MuT/ASTOpImageLayerColor.h b/Engine/Plugins/Experimental/Mutable/Source/MutableTools/Private/MuT/ASTOpImageLayerColor.h
new file mode 100644
--- /dev/null
+++ b/Engine/Plugins/Experimental/Mutable/Source/MutableTools/Private/MuT/ASTOpImageLayerColor.h
@@ -0,0 +1,54 @@
+// Copyright Epic Games, Inc. All Rights Reserved.
+
+#pragma once
+
+#include "HAL/Platform.h"
+#include "MuR/Image.h"
+#include "MuR/Operations.h"
+#include "MuR/Ptr.h"
+#include "MuT/AST.h"
+
+
+namespace mu
+{
+	struct FProgram;
+
+
+	//---------------------------------------------------------------------------------------------
+	//!
+	//---------------------------------------------------------------------------------------------
+	class ASTOpImageLayerColor : public ASTOp
+	{
+	public:
+
+		ASTChild base;
+		ASTChild color;
+		ASTChild mask;
+
+		/** Blend type used for the colour channels. */
+		EBlendType blendType = EBlendType::BT_NONE;
+
+		/** Blend type used for the alpha channel if any. */
+		EBlendType blendTypeAlpha = EBlendType::BT_NONE;
+
+	public:
+
+		ASTOpImageLayerColor();
+		ASTOpImageLayerColor(const ASTOpImageLayerColor&) = delete;
+		~ASTOpImageLayerColor();
+
+		OP_TYPE GetOpType() const override { return OP_TYPE::IM_LAYERCOLOUR; }
+		uint64 Hash() const override;
+		bool IsEqual(const ASTOp& otherUntyped) const override;
+		Ptr<ASTOp> Clone(MapChildFuncRef mapChild) const override;
+		void ForEachChild(const TFunctionRef<void(ASTChild&)>) override;
+		void Link(FProgram& program, const FLinkerOptions* Options) override;
+		FImageDesc GetImageDesc(bool returnBestOption, FGetImageDescContext* context) const override;
+		void GetLayoutBlockSize(int* pBlockX, int* pBlockY) override;
+		Ptr<ImageSizeExpression> GetImageSizeExpression() const override;
+		Ptr<ASTOp> OptimiseSink(const FModelOptimizationOptions&, FOptimizeSinkContext&) const;
+
+	};
+
+}
+
diff --git a/Engine/Plugins/Experimental/Mutable/Source/MutableTools/Private/MuT/ASTOpImageMultiLayer.cpp b/Engine/Plugins/Experimental/Mutable/Source/MutableTools/Private/MuT/ASTOpImageMultiLayer.cpp
--- a/Engine/Plugins/Experimental/Mutable/Source/MutableTools/Private/MuT/ASTOpImageMultiLayer.cpp
+++ b/Engine/Plugins/Experimental/Mutable/Source/MutableTools/Private/MuT/ASTOpImageMultiLayer.cpp
@@ -55,7 +55,7 @@ namespace mu
 	//-------------------------------------------------------------------------------------------------
 	uint64 ASTOpImageMultiLayer::Hash() const
 	{
-		uint64 res = std::hash<OP_TYPE>()(OP_TYPE::IM_MULTILAYER);
+		uint64 res = std::hash<OP_TYPE>()(GetOpType());
 		hash_combine(res, base.child().get());
 		hash_combine(res, blend.child().get());
 		hash_combine(res, mask.child().get());
@@ -97,7 +97,7 @@ namespace mu
 		if (!linkedAddress)
 		{
 			OP::ImageMultiLayerArgs args;
-			memset(&args, 0, sizeof(args));
+			FMemory::Memzero(&args, sizeof(args));
 
 			args.blendType = (uint8)blendType;
 			args.blendTypeAlpha = (uint8)blendTypeAlpha;
@@ -112,8 +112,8 @@ namespace mu
 			}
 
 			linkedAddress = (OP::ADDRESS)program.m_opAddress.Num();
-			program.m_opAddress.Add((uint32_t)program.m_byteCode.Num());
-			AppendCode(program.m_byteCode, OP_TYPE::IM_MULTILAYER);
+			program.m_opAddress.Add((uint32)program.m_byteCode.Num());
+			AppendCode(program.m_byteCode, GetOpType());
 			AppendCode(program.m_byteCode, args);
 		}
 	}
diff --git a/Engine/Plugins/Experimental/Mutable/Source/MutableTools/Private/MuT/ASTOpImagePixelFormat.cpp b/Engine/Plugins/Experimental/Mutable/Source/MutableTools/Private/MuT/ASTOpImagePixelFormat.cpp
--- a/Engine/Plugins/Experimental/Mutable/Source/MutableTools/Private/MuT/ASTOpImagePixelFormat.cpp
+++ b/Engine/Plugins/Experimental/Mutable/Source/MutableTools/Private/MuT/ASTOpImagePixelFormat.cpp
@@ -15,6 +15,8 @@
 #include "MuT/ASTOpImageCompose.h"
 #include "MuT/ASTOpImageMipmap.h"
 #include "MuT/ASTOpImagePatch.h"
+#include "MuT/ASTOpImageLayer.h"
+#include "MuT/ASTOpImageLayerColor.h"
 #include "MuT/ASTOpSwitch.h"
 #include "MuT/StreamsPrivate.h"
 
@@ -26,498 +28,493 @@ namespace mu
 
 	template <class SCALAR> class vec4;
 
-//-------------------------------------------------------------------------------------------------
-//-------------------------------------------------------------------------------------------------
-//-------------------------------------------------------------------------------------------------
-ASTOpImagePixelFormat::ASTOpImagePixelFormat()
-    : Source(this)
-{
-}
+	//-------------------------------------------------------------------------------------------------
+	//-------------------------------------------------------------------------------------------------
+	//-------------------------------------------------------------------------------------------------
+	ASTOpImagePixelFormat::ASTOpImagePixelFormat()
+		: Source(this)
+	{
+	}
 
 
-ASTOpImagePixelFormat::~ASTOpImagePixelFormat()
-{
-    // Explicit call needed to avoid recursive destruction
-    ASTOp::RemoveChildren();
-}
+	ASTOpImagePixelFormat::~ASTOpImagePixelFormat()
+	{
+		// Explicit call needed to avoid recursive destruction
+		ASTOp::RemoveChildren();
+	}
 
 
-bool ASTOpImagePixelFormat::IsEqual(const ASTOp& otherUntyped) const
-{
-    if (const ASTOpImagePixelFormat* other = dynamic_cast<const ASTOpImagePixelFormat*>(&otherUntyped) )
-    {
-        return Source==other->Source && Format == other->Format && FormatIfAlpha == other->FormatIfAlpha;
-    }
-    return false;
-}
+	bool ASTOpImagePixelFormat::IsEqual(const ASTOp& otherUntyped) const
+	{
+		if (const ASTOpImagePixelFormat* other = dynamic_cast<const ASTOpImagePixelFormat*>(&otherUntyped))
+		{
+			return Source == other->Source && Format == other->Format && FormatIfAlpha == other->FormatIfAlpha;
+		}
+		return false;
+	}
 
 
-uint64 ASTOpImagePixelFormat::Hash() const
-{
-	uint64 res = std::hash<void*>()(Source.child().get() );
-    hash_combine( res, Format );
-    return res;
-}
+	uint64 ASTOpImagePixelFormat::Hash() const
+	{
+		uint64 res = std::hash<void*>()(Source.child().get());
+		hash_combine(res, Format);
+		return res;
+	}
 
 
-mu::Ptr<ASTOp> ASTOpImagePixelFormat::Clone(MapChildFuncRef mapChild) const
-{
-	mu::Ptr<ASTOpImagePixelFormat> n = new ASTOpImagePixelFormat();
-    n->Source = mapChild(Source.child());
-	n->Format = Format;
-	n->FormatIfAlpha = FormatIfAlpha;
-    return n;
-}
+	mu::Ptr<ASTOp> ASTOpImagePixelFormat::Clone(MapChildFuncRef mapChild) const
+	{
+		mu::Ptr<ASTOpImagePixelFormat> n = new ASTOpImagePixelFormat();
+		n->Source = mapChild(Source.child());
+		n->Format = Format;
+		n->FormatIfAlpha = FormatIfAlpha;
+		return n;
+	}
 
 
-void ASTOpImagePixelFormat::ForEachChild(const TFunctionRef<void(ASTChild&)> f )
-{
-    f( Source );
-}
+	void ASTOpImagePixelFormat::ForEachChild(const TFunctionRef<void(ASTChild&)> f)
+	{
+		f(Source);
+	}
 
 
-void ASTOpImagePixelFormat::Link( FProgram& program, const FLinkerOptions* )
-{
-    // Already linked?
-    if (!linkedAddress)
-    {
-        OP::ImagePixelFormatArgs args;
-        FMemory::Memzero( &args, sizeof(args) );
-
-		args.format = Format;
-		args.formatIfAlpha = FormatIfAlpha;
-		if (Source) args.source = Source->linkedAddress;
-
-        linkedAddress = (OP::ADDRESS)program.m_opAddress.Num();
-        //program.m_code.push_back(op);
-        program.m_opAddress.Add((uint32_t)program.m_byteCode.Num());
-        AppendCode(program.m_byteCode,OP_TYPE::IM_PIXELFORMAT);
-        AppendCode(program.m_byteCode,args);
-    }
+	void ASTOpImagePixelFormat::Link(FProgram& program, const FLinkerOptions*)
+	{
+		// Already linked?
+		if (!linkedAddress)
+		{
+			OP::ImagePixelFormatArgs args;
+			FMemory::Memzero(&args, sizeof(args));
+
+			args.format = Format;
+			args.formatIfAlpha = FormatIfAlpha;
+			if (Source) args.source = Source->linkedAddress;
+
+			linkedAddress = (OP::ADDRESS)program.m_opAddress.Num();
+			//program.m_code.push_back(op);
+			program.m_opAddress.Add((uint32_t)program.m_byteCode.Num());
+			AppendCode(program.m_byteCode, OP_TYPE::IM_PIXELFORMAT);
+			AppendCode(program.m_byteCode, args);
+		}
 
-}
+	}
 
 
-mu::Ptr<ASTOp> ASTOpImagePixelFormat::OptimiseSink(const FModelOptimizationOptions& options, FOptimizeSinkContext& context) const
-{
-	mu::Ptr<ASTOp> at;
+	mu::Ptr<ASTOp> ASTOpImagePixelFormat::OptimiseSink(const FModelOptimizationOptions& options, FOptimizeSinkContext& context) const
+	{
+		mu::Ptr<ASTOp> at;
 
-	mu::Ptr<ASTOp> sourceAt = Source.child();
+		mu::Ptr<ASTOp> sourceAt = Source.child();
 
-	EImageFormat format = Format;
-	bool isCompressedFormat = IsCompressedFormat(Format);
-	//bool isBlockFormat = GetImageFormatData( format ).m_pixelsPerBlockX!=0;
+		EImageFormat format = Format;
+		bool isCompressedFormat = IsCompressedFormat(Format);
+		//bool isBlockFormat = GetImageFormatData( format ).m_pixelsPerBlockX!=0;
 
-	// The instruction can be sunk
-	OP_TYPE sourceType = sourceAt->GetOpType();
-	switch (sourceType)
-	{
-	case OP_TYPE::IM_PIXELFORMAT:
-	{
-		// Keep only the top pixel format
-		const ASTOpImagePixelFormat* typedSource = dynamic_cast<const ASTOpImagePixelFormat*>(sourceAt.get());
-		mu::Ptr<ASTOpImagePixelFormat> formatOp = mu::Clone<ASTOpImagePixelFormat>(this);
-		formatOp->Source = typedSource->Source.child();
-		at = formatOp;
-		break;
-	}
+		// The instruction can be sunk
+		OP_TYPE sourceType = sourceAt->GetOpType();
+		switch (sourceType)
+		{
+		case OP_TYPE::IM_PIXELFORMAT:
+		{
+			// Keep only the top pixel format
+			const ASTOpImagePixelFormat* typedSource = dynamic_cast<const ASTOpImagePixelFormat*>(sourceAt.get());
+			mu::Ptr<ASTOpImagePixelFormat> formatOp = mu::Clone<ASTOpImagePixelFormat>(this);
+			formatOp->Source = typedSource->Source.child();
+			at = formatOp;
+			break;
+		}
 
-	case OP_TYPE::IM_DISPLACE:
-	{
-		// This op doesn't support compressed formats
-		if (!isCompressedFormat)
+		case OP_TYPE::IM_DISPLACE:
 		{
-			mu::Ptr<ASTOpFixed> newOp = mu::Clone<ASTOpFixed>(sourceAt);
+			// This op doesn't support compressed formats
+			if (!isCompressedFormat)
+			{
+				mu::Ptr<ASTOpFixed> newOp = mu::Clone<ASTOpFixed>(sourceAt);
 
-			mu::Ptr<ASTOpImagePixelFormat> fop = mu::Clone<ASTOpImagePixelFormat>(this);
-			fop->Source = newOp->children[newOp->op.args.ImageDisplace.source].child();
-			newOp->SetChild(newOp->op.args.ImageDisplace.source, fop);
+				mu::Ptr<ASTOpImagePixelFormat> fop = mu::Clone<ASTOpImagePixelFormat>(this);
+				fop->Source = newOp->children[newOp->op.args.ImageDisplace.source].child();
+				newOp->SetChild(newOp->op.args.ImageDisplace.source, fop);
 
-			at = newOp;
+				at = newOp;
+			}
+			break;
 		}
-		break;
-	}
 
-	case OP_TYPE::IM_RASTERMESH:
-	{
-		// This op doesn't support compressed formats
-		if (!isCompressedFormat
-			&&
-			dynamic_cast<const ASTOpFixed*>(sourceAt.get())->op.args.ImageRasterMesh.image)
+		case OP_TYPE::IM_RASTERMESH:
 		{
-			auto newOp = mu::Clone<ASTOpFixed>(sourceAt);
+			// This op doesn't support compressed formats
+			if (!isCompressedFormat
+				&&
+				dynamic_cast<const ASTOpFixed*>(sourceAt.get())->op.args.ImageRasterMesh.image)
+			{
+				auto newOp = mu::Clone<ASTOpFixed>(sourceAt);
 
-			mu::Ptr<ASTOpImagePixelFormat> fop = mu::Clone<ASTOpImagePixelFormat>(this);
-			fop->Source = newOp->children[newOp->op.args.ImageRasterMesh.image].child();
-			newOp->SetChild(newOp->op.args.ImageRasterMesh.image, fop);
+				mu::Ptr<ASTOpImagePixelFormat> fop = mu::Clone<ASTOpImagePixelFormat>(this);
+				fop->Source = newOp->children[newOp->op.args.ImageRasterMesh.image].child();
+				newOp->SetChild(newOp->op.args.ImageRasterMesh.image, fop);
 
-			at = newOp;
+				at = newOp;
+			}
+			break;
 		}
-		break;
-	}
-
 
-	case OP_TYPE::IM_BLANKLAYOUT:
-	{
-		// Just make sure the layout format is the right one and forget the op
-		auto nop = mu::Clone<ASTOpFixed>(sourceAt);
 
-		EImageFormat layoutFormat = (EImageFormat)nop->op.args.ImageBlankLayout.format;
-		if (FormatIfAlpha != EImageFormat::IF_NONE
-			&&
-			GetImageFormatData(layoutFormat).m_channels > 3)
+		case OP_TYPE::IM_BLANKLAYOUT:
 		{
-			format = FormatIfAlpha;
-		}
+			// Just make sure the layout format is the right one and forget the op
+			auto nop = mu::Clone<ASTOpFixed>(sourceAt);
 
-		nop->op.args.ImageBlankLayout.format = format;
-		at = nop;
-		break;
-	}
+			EImageFormat layoutFormat = (EImageFormat)nop->op.args.ImageBlankLayout.format;
+			if (FormatIfAlpha != EImageFormat::IF_NONE
+				&&
+				GetImageFormatData(layoutFormat).m_channels > 3)
+			{
+				format = FormatIfAlpha;
+			}
 
-	default:
-	{
+			nop->op.args.ImageBlankLayout.format = format;
+			at = nop;
+			break;
+		}
 
-		at = context.ImagePixelFormatSinker.Apply(this);
+		default:
+		{
 
-		break;
-	} // pixelformat source default
+			at = context.ImagePixelFormatSinker.Apply(this);
 
-	}
+			break;
+		} // pixelformat source default
 
-	return at;
-}
+		}
 
+		return at;
+	}
 
-//!
-FImageDesc ASTOpImagePixelFormat::GetImageDesc(bool returnBestOption, FGetImageDescContext* context) const
-{
-	FImageDesc res;
 
-	// Local context in case it is necessary
-	FGetImageDescContext localContext;
-	if (!context)
-	{
-		context = &localContext;
-	}
-	else
+	//!
+	FImageDesc ASTOpImagePixelFormat::GetImageDesc(bool returnBestOption, FGetImageDescContext* context) const
 	{
-		// Cached result?
-		FImageDesc* PtrValue = context->m_results.Find(this);
-		if (PtrValue)
+		FImageDesc res;
+
+		// Local context in case it is necessary
+		FGetImageDescContext localContext;
+		if (!context)
 		{
-			return *PtrValue;
+			context = &localContext;
+		}
+		else
+		{
+			// Cached result?
+			FImageDesc* PtrValue = context->m_results.Find(this);
+			if (PtrValue)
+			{
+				return *PtrValue;
+			}
 		}
-	}
-
-	res = Source.child()->GetImageDesc( returnBestOption, context);
-
-	if (FormatIfAlpha != EImageFormat::IF_NONE
-		&&
-		GetImageFormatData(res.m_format).m_channels > 3)
-	{
-		res.m_format = FormatIfAlpha;
-	}
-	else
-	{
-		res.m_format = Format;
-	}
-	check(res.m_format != EImageFormat::IF_NONE);
 
+		res = Source.child()->GetImageDesc(returnBestOption, context);
 
-	// Cache the result
-	if (context)
-	{
-		context->m_results.Add(this, res);
-	}
+		if (FormatIfAlpha != EImageFormat::IF_NONE
+			&&
+			GetImageFormatData(res.m_format).m_channels > 3)
+		{
+			res.m_format = FormatIfAlpha;
+		}
+		else
+		{
+			res.m_format = Format;
+		}
+		check(res.m_format != EImageFormat::IF_NONE);
 
-	return res;
-}
 
+		// Cache the result
+		if (context)
+		{
+			context->m_results.Add(this, res);
+		}
 
-void ASTOpImagePixelFormat::GetLayoutBlockSize(int* pBlockX, int* pBlockY)
-{
-	if (Source.child())
-	{
-		Source.child()->GetLayoutBlockSize(pBlockX, pBlockY);
+		return res;
 	}
-}
 
 
-bool ASTOpImagePixelFormat::IsImagePlainConstant(FVector4f& colour) const
-{
-	bool res = false;
-	if (Source.child())
+	void ASTOpImagePixelFormat::GetLayoutBlockSize(int* pBlockX, int* pBlockY)
 	{
-		Source.child()->IsImagePlainConstant(colour);
+		if (Source.child())
+		{
+			Source.child()->GetLayoutBlockSize(pBlockX, pBlockY);
+		}
 	}
-	return res;
-}
-
 
-mu::Ptr<ImageSizeExpression> ASTOpImagePixelFormat::GetImageSizeExpression() const
-{
-	mu::Ptr<ImageSizeExpression> pRes;
 
-	if (Source.child())
-	{
-		pRes = Source.child()->GetImageSizeExpression();
-	}
-	else
+	bool ASTOpImagePixelFormat::IsImagePlainConstant(FVector4f& colour) const
 	{
-		pRes = new ImageSizeExpression;
+		bool res = false;
+		if (Source.child())
+		{
+			Source.child()->IsImagePlainConstant(colour);
+		}
+		return res;
 	}
 
-	return pRes;
-}
-
 
-//---------------------------------------------------------------------------------------------
-//---------------------------------------------------------------------------------------------
-//---------------------------------------------------------------------------------------------
-mu::Ptr<ASTOp> Sink_ImagePixelFormatAST::Apply(const ASTOp* root)
-{
-	m_root = dynamic_cast<const ASTOpImagePixelFormat*>(root);
-	OldToNew.Empty();
+	mu::Ptr<ImageSizeExpression> ASTOpImagePixelFormat::GetImageSizeExpression() const
+	{
+		mu::Ptr<ImageSizeExpression> pRes;
 
-	check(root->GetOpType() == OP_TYPE::IM_PIXELFORMAT);
+		if (Source.child())
+		{
+			pRes = Source.child()->GetImageSizeExpression();
+		}
+		else
+		{
+			pRes = new ImageSizeExpression;
+		}
 
-	m_initialSource = m_root->Source.child();
-	mu::Ptr<ASTOp> newSource = Visit(m_initialSource, m_root);
+		return pRes;
+	}
 
-	m_root = nullptr;
 
-	// If there is any change, it is the new root.
-	if (newSource != m_initialSource)
+	//---------------------------------------------------------------------------------------------
+	//---------------------------------------------------------------------------------------------
+	//---------------------------------------------------------------------------------------------
+	mu::Ptr<ASTOp> Sink_ImagePixelFormatAST::Apply(const ASTOp* root)
 	{
-		return newSource;
-	}
+		m_root = dynamic_cast<const ASTOpImagePixelFormat*>(root);
+		OldToNew.Empty();
 
-	return nullptr;
-}
+		check(root->GetOpType() == OP_TYPE::IM_PIXELFORMAT);
 
+		m_initialSource = m_root->Source.child();
+		mu::Ptr<ASTOp> newSource = Visit(m_initialSource, m_root);
 
-//---------------------------------------------------------------------------------------------
-mu::Ptr<ASTOp> Sink_ImagePixelFormatAST::Visit(mu::Ptr<ASTOp> at, const ASTOpImagePixelFormat* currentFormatOp)
-{
-	if (!at) return nullptr;
+		m_root = nullptr;
 
-	EImageFormat format = currentFormatOp->Format;
-	bool isCompressedFormat = IsCompressedFormat(format);
-	bool isBlockFormat = GetImageFormatData(format).m_pixelsPerBlockX != 0;
+		// If there is any change, it is the new root.
+		if (newSource != m_initialSource)
+		{
+			return newSource;
+		}
 
-	// Already visited?
-	const Ptr<ASTOp>* Cached = OldToNew.Find({ at,currentFormatOp });
-	if (Cached)
-	{
-		return *Cached;
+		return nullptr;
 	}
 
-	mu::Ptr<ASTOp> newAt = at;
-	switch (at->GetOpType())
-	{
 
-	case OP_TYPE::IM_CONDITIONAL:
+	//---------------------------------------------------------------------------------------------
+	mu::Ptr<ASTOp> Sink_ImagePixelFormatAST::Visit(mu::Ptr<ASTOp> at, const ASTOpImagePixelFormat* currentFormatOp)
 	{
-		// We move the mask creation down the two paths
-		auto newOp = mu::Clone<ASTOpConditional>(at);
-		newOp->yes = Visit(newOp->yes.child(), currentFormatOp);
-		newOp->no = Visit(newOp->no.child(), currentFormatOp);
-		newAt = newOp;
-		break;
-	}
+		if (!at) return nullptr;
 
-	case OP_TYPE::IM_SWITCH:
-	{
-		// We move the mask creation down all the paths
-		auto newOp = mu::Clone<ASTOpSwitch>(at);
-		newOp->def = Visit(newOp->def.child(), currentFormatOp);
-		for (auto& c : newOp->cases)
+		EImageFormat format = currentFormatOp->Format;
+		bool isCompressedFormat = IsCompressedFormat(format);
+		bool isBlockFormat = GetImageFormatData(format).m_pixelsPerBlockX != 0;
+
+		// Already visited?
+		const Ptr<ASTOp>* Cached = OldToNew.Find({ at,currentFormatOp });
+		if (Cached)
 		{
-			c.branch = Visit(c.branch.child(), currentFormatOp);
+			return *Cached;
 		}
-		newAt = newOp;
-		break;
-	}
 
-	case OP_TYPE::IM_COMPOSE:
-	{
-		if (isBlockFormat)
+		mu::Ptr<ASTOp> newAt = at;
+		switch (at->GetOpType())
+		{
+
+		case OP_TYPE::IM_CONDITIONAL:
 		{
-			// We can only optimise if the layout grid blocks size in pixels is
-			// a multiple of the image format block size.
-			int imageFormatBlockSizeX = GetImageFormatData(format).m_pixelsPerBlockX;
-			int imageFormatBlockSizeY = GetImageFormatData(format).m_pixelsPerBlockY;
-			bool acceptable = imageFormatBlockSizeX == 1 && imageFormatBlockSizeY == 1;
+			// We move the mask creation down the two paths
+			auto newOp = mu::Clone<ASTOpConditional>(at);
+			newOp->yes = Visit(newOp->yes.child(), currentFormatOp);
+			newOp->no = Visit(newOp->no.child(), currentFormatOp);
+			newAt = newOp;
+			break;
+		}
 
-			if (!acceptable)
+		case OP_TYPE::IM_SWITCH:
+		{
+			// We move the mask creation down all the paths
+			auto newOp = mu::Clone<ASTOpSwitch>(at);
+			newOp->def = Visit(newOp->def.child(), currentFormatOp);
+			for (auto& c : newOp->cases)
 			{
-				const ASTOpImageCompose* typedAt = dynamic_cast<const ASTOpImageCompose*>(at.get());
-				auto originalBaseOp = typedAt->Base.child();
-
-				int layoutBlockPixelsX = 0;
-				int layoutBlockPixelsY = 0;
-				originalBaseOp->GetLayoutBlockSize(&layoutBlockPixelsX, &layoutBlockPixelsY);
-
-				acceptable =
-					(layoutBlockPixelsX != 0 && layoutBlockPixelsY != 0)
-					&&
-					(layoutBlockPixelsX % imageFormatBlockSizeX) == 0
-					&&
-					(layoutBlockPixelsY % imageFormatBlockSizeY) == 0;
+				c.branch = Visit(c.branch.child(), currentFormatOp);
 			}
+			newAt = newOp;
+			break;
+		}
 
-			if (acceptable)
+		case OP_TYPE::IM_COMPOSE:
+		{
+			if (isBlockFormat)
 			{
-				// We move the format down the two paths
-				auto newOp = mu::Clone<ASTOpImageCompose>(at);
+				// We can only optimise if the layout grid blocks size in pixels is
+				// a multiple of the image format block size.
+				int imageFormatBlockSizeX = GetImageFormatData(format).m_pixelsPerBlockX;
+				int imageFormatBlockSizeY = GetImageFormatData(format).m_pixelsPerBlockY;
+				bool acceptable = imageFormatBlockSizeX == 1 && imageFormatBlockSizeY == 1;
+
+				if (!acceptable)
+				{
+					const ASTOpImageCompose* typedAt = dynamic_cast<const ASTOpImageCompose*>(at.get());
+					auto originalBaseOp = typedAt->Base.child();
+
+					int layoutBlockPixelsX = 0;
+					int layoutBlockPixelsY = 0;
+					originalBaseOp->GetLayoutBlockSize(&layoutBlockPixelsX, &layoutBlockPixelsY);
+
+					acceptable =
+						(layoutBlockPixelsX != 0 && layoutBlockPixelsY != 0)
+						&&
+						(layoutBlockPixelsX % imageFormatBlockSizeX) == 0
+						&&
+						(layoutBlockPixelsY % imageFormatBlockSizeY) == 0;
+				}
+
+				if (acceptable)
+				{
+					// We move the format down the two paths
+					auto newOp = mu::Clone<ASTOpImageCompose>(at);
+
+					// TODO: We have to make sure we don't end up with two different formats if
+					// there is an formatIfAlpha
+
+					auto baseOp = newOp->Base.child();
+					newOp->Base = Visit(baseOp, currentFormatOp);
+
+					auto blockOp = newOp->BlockImage.child();
+					newOp->BlockImage = Visit(blockOp, currentFormatOp);
+
+					newAt = newOp;
+				}
+			}
 
-				// TODO: We have to make sure we don't end up with two different formats if
-				// there is an formatIfAlpha
+			break;
+		}
 
-				auto baseOp = newOp->Base.child();
-				newOp->Base = Visit(baseOp, currentFormatOp);
+		case OP_TYPE::IM_PATCH:
+		{
+			if (isBlockFormat)
+			{
+				// We move the format down the two paths
+				auto newOp = mu::Clone<ASTOpImagePatch>(at);
 
-				auto blockOp = newOp->BlockImage.child();
-				newOp->BlockImage = Visit(blockOp, currentFormatOp);
+				newOp->base = Visit(newOp->base.child(), currentFormatOp);
+				newOp->patch = Visit(newOp->patch.child(), currentFormatOp);
 
 				newAt = newOp;
 			}
-		}
-
-		break;
-	}
 
-	case OP_TYPE::IM_PATCH:
-	{
-		if (isBlockFormat)
-		{
-			// We move the format down the two paths
-			auto newOp = mu::Clone<ASTOpImagePatch>(at);
-
-			newOp->base = Visit(newOp->base.child(), currentFormatOp);
-			newOp->patch = Visit(newOp->patch.child(), currentFormatOp);
-
-			newAt = newOp;
+			break;
 		}
 
-		break;
-	}
-
-	case OP_TYPE::IM_MIPMAP:
-	{
-		auto typedSource = dynamic_cast<const ASTOpImageMipmap*>(at.get());
-
-		// If its a compressed format, only sink formats on mipmap operations that
-		// generate the tail. To avoid optimization loop.
-		if (!isCompressedFormat || typedSource->bOnlyTail)
+		case OP_TYPE::IM_MIPMAP:
 		{
-			auto newOp = mu::Clone<ASTOpImageMipmap>(typedSource);
+			auto typedSource = dynamic_cast<const ASTOpImageMipmap*>(at.get());
 
-			auto baseOp = newOp->Source.child();
-			newOp->Source = Visit(baseOp, currentFormatOp);
+			// If its a compressed format, only sink formats on mipmap operations that
+			// generate the tail. To avoid optimization loop.
+			if (!isCompressedFormat || typedSource->bOnlyTail)
+			{
+				auto newOp = mu::Clone<ASTOpImageMipmap>(typedSource);
 
-			newAt = newOp;
+				auto baseOp = newOp->Source.child();
+				newOp->Source = Visit(baseOp, currentFormatOp);
+
+				newAt = newOp;
+			}
+			break;
 		}
-		break;
-	}
 
-	case OP_TYPE::IM_INTERPOLATE:
-	{
-		// This op doesn't support compressed formats
-		if (!isCompressedFormat)
+		case OP_TYPE::IM_INTERPOLATE:
 		{
-			// Move the format down all the paths
-			auto newOp = mu::Clone<ASTOpFixed>(at);
-
-			for (int v = 0; v < MUTABLE_OP_MAX_INTERPOLATE_COUNT; ++v)
+			// This op doesn't support compressed formats
+			if (!isCompressedFormat)
 			{
-				auto child = newOp->children[newOp->op.args.ImageInterpolate.targets[v]].child();
-				auto bOp = Visit(child, currentFormatOp);
-				newOp->SetChild(newOp->op.args.ImageInterpolate.targets[v], bOp);
-			}
+				// Move the format down all the paths
+				auto newOp = mu::Clone<ASTOpFixed>(at);
 
-			newAt = newOp;
+				for (int v = 0; v < MUTABLE_OP_MAX_INTERPOLATE_COUNT; ++v)
+				{
+					auto child = newOp->children[newOp->op.args.ImageInterpolate.targets[v]].child();
+					auto bOp = Visit(child, currentFormatOp);
+					newOp->SetChild(newOp->op.args.ImageInterpolate.targets[v], bOp);
+				}
+
+				newAt = newOp;
+			}
+			break;
 		}
-		break;
-	}
 
 
-	case OP_TYPE::IM_INTERPOLATE3:
-	{
-		// This op doesn't support compressed formats
-		if (!isCompressedFormat)
+		case OP_TYPE::IM_INTERPOLATE3:
 		{
-			// We move the format down all the paths
-			auto newOp = mu::Clone<ASTOpFixed>(at);
+			// This op doesn't support compressed formats
+			if (!isCompressedFormat)
+			{
+				// We move the format down all the paths
+				auto newOp = mu::Clone<ASTOpFixed>(at);
 
-			auto top0 = newOp->children[newOp->op.args.ImageInterpolate3.target0].child();
-			newOp->SetChild(newOp->op.args.ImageInterpolate3.target0, Visit(top0, currentFormatOp));
+				auto top0 = newOp->children[newOp->op.args.ImageInterpolate3.target0].child();
+				newOp->SetChild(newOp->op.args.ImageInterpolate3.target0, Visit(top0, currentFormatOp));
 
-			auto top1 = newOp->children[newOp->op.args.ImageInterpolate3.target1].child();
-			newOp->SetChild(newOp->op.args.ImageInterpolate3.target1, Visit(top1, currentFormatOp));
+				auto top1 = newOp->children[newOp->op.args.ImageInterpolate3.target1].child();
+				newOp->SetChild(newOp->op.args.ImageInterpolate3.target1, Visit(top1, currentFormatOp));
 
-			auto top2 = newOp->children[newOp->op.args.ImageInterpolate3.target2].child();
-			newOp->SetChild(newOp->op.args.ImageInterpolate3.target2, Visit(top2, currentFormatOp));
+				auto top2 = newOp->children[newOp->op.args.ImageInterpolate3.target2].child();
+				newOp->SetChild(newOp->op.args.ImageInterpolate3.target2, Visit(top2, currentFormatOp));
 
-			newAt = newOp;
+				newAt = newOp;
+			}
+			break;
 		}
-		break;
-	}
 
-	case OP_TYPE::IM_LAYER:
-	{
-		if (GetOpDesc(at->GetOpType()).supportedBasePixelFormats[(size_t)format])
+		case OP_TYPE::IM_LAYER:
 		{
-			// We move the format down the two paths
-			auto nop = mu::Clone<ASTOpFixed>(at);
-
-			auto aOp = nop->children[nop->op.args.ImageLayer.base].child();
-			nop->SetChild(nop->op.args.ImageLayer.base, Visit(aOp, currentFormatOp));
+			if (GetOpDesc(at->GetOpType()).supportedBasePixelFormats[(size_t)format])
+			{
+				// We move the format down the two paths
+				Ptr<ASTOpImageLayer> nop = mu::Clone<ASTOpImageLayer>(at);
 
-			auto bOp = nop->children[nop->op.args.ImageLayer.blended].child();
-			nop->SetChild(nop->op.args.ImageLayer.blended, Visit(bOp, currentFormatOp));
+				nop->base = Visit(nop->base.child(), currentFormatOp);
+				nop->blend = Visit(nop->blend.child(), currentFormatOp);
 
-			newAt = nop;
+				newAt = nop;
+			}
+			break;
 		}
-		break;
-	}
 
-	case OP_TYPE::IM_LAYERCOLOUR:
-	{
-		if (GetOpDesc(at->GetOpType()).supportedBasePixelFormats[(size_t)format])
+		case OP_TYPE::IM_LAYERCOLOUR:
 		{
-			// We move the format down the base
-			auto nop = mu::Clone<ASTOpFixed>(at);
+			if (GetOpDesc(at->GetOpType()).supportedBasePixelFormats[(size_t)format])
+			{
+				// We move the format down the base
+				Ptr<ASTOpImageLayerColor> nop = mu::Clone<ASTOpImageLayerColor>(at);
 
-			auto aOp = nop->children[nop->op.args.ImageLayerColour.base].child();
-			nop->SetChild(nop->op.args.ImageLayerColour.base, Visit(aOp, currentFormatOp));
+				nop->base = Visit(nop->base.child(), currentFormatOp);
 
-			newAt = nop;
+				newAt = nop;
+			}
+			break;
 		}
-		break;
-	}
 
 
-	default:
-		break;
-	}
+		default:
+			break;
+		}
 
-	// end on tree branch, replace with format
-	if (at == newAt && at != m_initialSource)
-	{
-		mu::Ptr<ASTOpImagePixelFormat> newOp = mu::Clone<ASTOpImagePixelFormat>(currentFormatOp);
-		check(newOp->GetOpType() == OP_TYPE::IM_PIXELFORMAT);
+		// end on tree branch, replace with format
+		if (at == newAt && at != m_initialSource)
+		{
+			mu::Ptr<ASTOpImagePixelFormat> newOp = mu::Clone<ASTOpImagePixelFormat>(currentFormatOp);
+			check(newOp->GetOpType() == OP_TYPE::IM_PIXELFORMAT);
 
-		newOp->Source = at;
+			newOp->Source = at;
 
-		newAt = newOp;
-	}
+			newAt = newOp;
+		}
 
-	OldToNew.Add( { at, currentFormatOp }, newAt);
-
-	return newAt;
-}
+		OldToNew.Add({ at, currentFormatOp }, newAt);
 
+		return newAt;
+	}
 
 }
diff --git a/Engine/Plugins/Experimental/Mutable/Source/MutableTools/Private/MuT/ASTOpImageSwizzle.cpp b/Engine/Plugins/Experimental/Mutable/Source/MutableTools/Private/MuT/ASTOpImageSwizzle.cpp
--- a/Engine/Plugins/Experimental/Mutable/Source/MutableTools/Private/MuT/ASTOpImageSwizzle.cpp
+++ b/Engine/Plugins/Experimental/Mutable/Source/MutableTools/Private/MuT/ASTOpImageSwizzle.cpp
@@ -17,6 +17,8 @@
 #include "MuT/ASTOpImagePatch.h"
 #include "MuT/ASTOpImagePixelFormat.h"
 #include "MuT/ASTOpImageMultiLayer.h"
+#include "MuT/ASTOpImageLayer.h"
+#include "MuT/ASTOpImageLayerColor.h"
 #include "MuT/ASTOpSwitch.h"
 #include "MuT/StreamsPrivate.h"
 
@@ -119,7 +121,6 @@ namespace mu
 
 	mu::Ptr<ASTOp> ASTOpImageSwizzle::OptimiseSemantic(const FModelOptimizationOptions&) const
 	{
-		// If children channels are also swizzle ops, recurse them
 		Ptr<ASTOpImageSwizzle> sat;
 
 		for (int c = 0; c < MUTABLE_OP_MAX_SWIZZLE_CHANNELS; ++c)
@@ -230,6 +231,17 @@ namespace mu
 			}
 		}
 
+		// If we are not changing channel order, just remove the swizzle and adjust the format.
+		bool bSameChannelOrder = true;
+		int32 NumChannelsInFormat = GetImageFormatData(Format).m_channels;
+		for (int32 c = 0; c < NumChannelsInFormat; ++c)
+		{
+			if (Sources[c] && SourceChannels[c] != c)
+			{
+				bSameChannelOrder = false;
+			}
+		}
+
 		OP_TYPE sourceType = channelSourceAt->GetOpType();
 
 		if (bAllChannelsAreTheSame && channelSourceAt)
@@ -300,15 +312,15 @@ namespace mu
 			case OP_TYPE::IM_LAYER:
 			{
 				// We move the swizzle down the two paths
-				Ptr<ASTOpFixed> nop = mu::Clone<ASTOpFixed>(channelSourceAt);
+				Ptr<ASTOpImageLayer> nop = mu::Clone<ASTOpImageLayer>(channelSourceAt);
 
 				Ptr<ASTOpImageSwizzle> aOp = mu::Clone<ASTOpImageSwizzle>(this);
-				ReplaceAllSources(aOp, nop->children[nop->op.args.ImageLayer.base].child());
-				nop->SetChild(nop->op.args.ImageLayer.base, aOp);
+				ReplaceAllSources(aOp, nop->base.child());
+				nop->base = aOp;
 
 				Ptr<ASTOpImageSwizzle> bOp = mu::Clone<ASTOpImageSwizzle>(this);
-				ReplaceAllSources(bOp, nop->children[nop->op.args.ImageLayer.blended].child());
-				nop->SetChild(nop->op.args.ImageLayer.blended, bOp);
+				ReplaceAllSources(bOp, nop->blend.child());
+				nop->blend = bOp;
 
 				at = nop;
 				break;
@@ -317,21 +329,21 @@ namespace mu
 			case OP_TYPE::IM_LAYERCOLOUR:
 			{
 				// We move the swizzle down the base path
-				Ptr<ASTOpFixed> nop = mu::Clone<ASTOpFixed>(channelSourceAt);
+				Ptr<ASTOpImageLayerColor> nop = mu::Clone<ASTOpImageLayerColor>(channelSourceAt);
 
 				Ptr<ASTOpImageSwizzle> NewSwizzle = mu::Clone<ASTOpImageSwizzle>(this);
-				ReplaceAllSources(NewSwizzle, nop->children[nop->op.args.ImageLayerColour.base].child());
-				nop->SetChild(nop->op.args.ImageLayerColour.base, NewSwizzle);
+				ReplaceAllSources(NewSwizzle, nop->base.child());
+				nop->base = NewSwizzle;
 
 				// We need to swizzle the colour too
 				Ptr<ASTOpFixed> cOp = new ASTOpFixed;
 				cOp->op.type = OP_TYPE::CO_SWIZZLE;
 				for (int i = 0; i < MUTABLE_OP_MAX_SWIZZLE_CHANNELS; ++i)
 				{
-					cOp->SetChild(cOp->op.args.ColourSwizzle.sources[i], nop->children[nop->op.args.ImageLayerColour.colour]);
+					cOp->SetChild(cOp->op.args.ColourSwizzle.sources[i], nop->color);
 					cOp->op.args.ColourSwizzle.sourceChannels[i] = SourceChannels[i];
 				}
-				nop->SetChild(nop->op.args.ImageLayerColour.colour, cOp);
+				nop->color = cOp;
 
 				at = nop;
 				break;
@@ -357,6 +369,38 @@ namespace mu
 				break;
 			}
 
+			case OP_TYPE::IM_RESIZE:
+			{
+				Ptr<ASTOpFixed> NewResize = mu::Clone<ASTOpFixed>(channelSourceAt);
+				Ptr<ASTOpImageSwizzle> NewSwizzle = mu::Clone<ASTOpImageSwizzle>(this);
+				ReplaceAllSources(NewSwizzle, NewResize->children[NewResize->op.args.ImageResize.source].child());
+				NewResize->SetChild(NewResize->op.args.ImageResize.source, NewSwizzle);
+				at = NewResize;
+				break;
+			}
+
+			// This is not valid: binarize always forces format L8
+			//case OP_TYPE::IM_BINARISE:
+			//{
+			//	Ptr<ASTOpFixed> NewBinarise = mu::Clone<ASTOpFixed>(channelSourceAt);
+			//	Ptr<ASTOpImageSwizzle> NewSwizzle = mu::Clone<ASTOpImageSwizzle>(this);
+			//	ReplaceAllSources(NewSwizzle, NewBinarise->children[NewBinarise->op.args.ImageBinarise.base].child());
+			//	NewBinarise->SetChild(NewBinarise->op.args.ImageBinarise.base, NewSwizzle);
+			//	at = NewBinarise;
+			//	break;
+			//}
+
+			case OP_TYPE::IM_PIXELFORMAT:
+			{
+				if (bSameChannelOrder)
+				{
+					Ptr<ASTOpImagePixelFormat> NewFormat = mu::Clone<ASTOpImagePixelFormat>(channelSourceAt);
+					NewFormat->Format = Format;
+					at = NewFormat;
+				}
+				break;
+			}
+
 			default:
 				bAllChannelsAreTheSame = false;
 				break;
@@ -370,6 +414,8 @@ namespace mu
 			// If we have RGB being the same IM_MULTILAYER, and alpha a compatible IM_MULTILAYER we can optimize with
 			// a special multilayer blend mode. This happens often because of higher level group projector nodes.
 			if (!at
+				&&
+				Format == EImageFormat::IF_RGBA_UBYTE
 				&&
 				Sources[0] == Sources[1] && Sources[0] == Sources[2]
 				&&
@@ -424,6 +470,62 @@ namespace mu
 				}
 			}
 
+			// If we have RGB being the same IM_LAYER, and alpha a compatible IM_LAYER we can optimize with a special layer blend mode.
+			if (!at
+				&&
+				Format == EImageFormat::IF_RGBA_UBYTE
+				&&
+				Sources[0] == Sources[1] && (Sources[0] == Sources[2] || !Sources[2])
+				&&
+				Sources[0] && Sources[0]->GetOpType() == OP_TYPE::IM_LAYER
+				&&
+				Sources[3] && Sources[3]->GetOpType() == OP_TYPE::IM_LAYER
+				&&
+				SourceChannels[0] == 0 && SourceChannels[1] == 1 && (SourceChannels[2] == 2 || !Sources[2] ) && SourceChannels[3] == 0
+				)
+			{
+				const ASTOpImageLayer* ColorLayer = dynamic_cast<const ASTOpImageLayer*>(Sources[0].child().get());
+				check(ColorLayer);
+				const ASTOpImageLayer* AlphaLayer = dynamic_cast<const ASTOpImageLayer*>(Sources[3].child().get());
+				check(AlphaLayer);
+
+				bool bIsSpecialMultiLayer = !AlphaLayer->mask;
+
+				if (bIsSpecialMultiLayer)
+				{
+					// We can combine the 2 multilayers into the composite blend+lighten mode
+
+					Ptr<ASTOpImageSwizzle> NewBase = mu::Clone<ASTOpImageSwizzle>(this);
+					NewBase->Sources[0] = ColorLayer->base.child();
+					NewBase->Sources[1] = ColorLayer->base.child();
+					NewBase->Sources[2] = Sources[2] ? ColorLayer->base.child() : nullptr;
+					NewBase->Sources[3] = AlphaLayer->base.child();
+
+					Ptr<ASTOpImageSwizzle> NewBlended = mu::Clone<ASTOpImageSwizzle>(this);
+					NewBlended->Sources[0] = ColorLayer->blend.child();
+					NewBlended->Sources[1] = ColorLayer->blend.child();
+					NewBlended->Sources[2] = Sources[2] ? ColorLayer->blend.child() : nullptr;
+					NewBlended->Sources[3] = AlphaLayer->blend.child();
+
+					Ptr<ASTOpImageLayer> NewLayer = mu::Clone<ASTOpImageLayer>(ColorLayer);
+					NewLayer->blendTypeAlpha = AlphaLayer->blendType;
+					NewLayer->base = NewBase;
+					NewLayer->blend = NewBlended;
+
+					if (NewLayer->mask.child() == AlphaLayer->blend.child()
+						&&
+						NewBlended->Format == EImageFormat::IF_RGBA_UBYTE
+						)
+					{
+						// Additional optimization is possible here.
+						NewLayer->Flags |= OP::ImageLayerArgs::FLAGS::F_USE_MASK_FROM_BLENDED;
+						NewLayer->mask = nullptr;
+					}
+
+					at = NewLayer;
+				}
+			}
+
 
 			// If the channels are compatible switches, we can still sink the swizzle.
 			if (!at && sourceType == OP_TYPE::IM_SWITCH)
@@ -486,7 +588,7 @@ namespace mu
 				}
 			}
 
-			// Swizzle down a displace.
+			// Swizzle down compatible displaces.
 			if (!at && sourceType == OP_TYPE::IM_DISPLACE)
 			{
 				const ASTOpFixed* FirstDisplace = dynamic_cast<const ASTOpFixed*>(Sources[0].child().get());
@@ -529,13 +631,13 @@ namespace mu
 
 			}
 
-			// Swizzle down a raster mesh.
+			// Swizzle down compatible raster meshes.
 			if (!at && sourceType == OP_TYPE::IM_RASTERMESH)
 			{
 				const ASTOpFixed* FirstRasterMesh = dynamic_cast<const ASTOpFixed*>(Sources[0].child().get());
 				check(FirstRasterMesh);
 
-				bool bAreAllFirstRasterMeshesCompatible = true;
+				bool bAreAllRasterMeshesCompatible = true;
 				for (int32 c = 1; c < MUTABLE_OP_MAX_SWIZZLE_CHANNELS; ++c)
 				{
 					if (Sources[c])
@@ -548,13 +650,13 @@ namespace mu
 
 						if (FMemory::Memcmp(&ArgCopy, &Typed->op.args.ImageRasterMesh, sizeof(OP::ImageRasterMeshArgs)) != 0)
 						{
-							bAreAllFirstRasterMeshesCompatible = false;
+							bAreAllRasterMeshesCompatible = false;
 							break;
 						}
 					}
 				}
 
-				if (bAreAllFirstRasterMeshesCompatible)
+				if (bAreAllRasterMeshesCompatible)
 				{
 					// Move the swizzle down all the paths
 					Ptr<ASTOpFixed> NewRaster = mu::Clone<ASTOpFixed>(FirstRasterMesh);
@@ -575,10 +677,88 @@ namespace mu
 				}
 
 			}
+
+			// Swizzle down compatible resizes.
+			if (!at && sourceType == OP_TYPE::IM_RESIZE)
+			{
+				const ASTOpFixed* FirstResize = dynamic_cast<const ASTOpFixed*>(Sources[0].child().get());
+				check(FirstResize);
+
+				bool bAreAllResizesCompatible = true;
+				for (int32 c = 1; c < MUTABLE_OP_MAX_SWIZZLE_CHANNELS; ++c)
+				{
+					if (Sources[c])
+					{
+						const ASTOpFixed* Typed = dynamic_cast<const ASTOpFixed*>(Sources[c].child().get());
+						check(Typed);
+						// Compare all args but the source image
+						OP::ImageResizeArgs ArgCopy = FirstResize->op.args.ImageResize;
+						ArgCopy.source = Typed->op.args.ImageResize.source;
+
+						if (FMemory::Memcmp(&ArgCopy, &Typed->op.args.ImageResize, sizeof(OP::ImageResizeArgs)) != 0)
+						{
+							bAreAllResizesCompatible = false;
+							break;
+						}
+					}
+				}
+
+				if (bAreAllResizesCompatible)
+				{
+					// Move the swizzle down all the paths
+					Ptr<ASTOpFixed> NewResize = mu::Clone<ASTOpFixed>(FirstResize);
+
+					Ptr<ASTOpImageSwizzle> NewSwizzle = mu::Clone<ASTOpImageSwizzle>(this);
+					for (int c = 0; c < MUTABLE_OP_MAX_SWIZZLE_CHANNELS; ++c)
+					{
+						const ASTOpFixed* ChannelResize = dynamic_cast<const ASTOpFixed*>(Sources[c].child().get());
+						if (ChannelResize)
+						{
+							NewSwizzle->Sources[c] = ChannelResize->children[ChannelResize->op.args.ImageResize.source].child();
+						}
+					}
+
+					NewResize->SetChild(NewResize->op.args.ImageResize.source, NewSwizzle);
+
+					at = NewResize;
+				}
+
+			}
+
+			// Swizzle down compatible pixelformats.
+			if (!at && sourceType == OP_TYPE::IM_PIXELFORMAT && bSameChannelOrder)
+			{
+				const ASTOpImagePixelFormat* FirstFormat = dynamic_cast<const ASTOpImagePixelFormat*>(Sources[0].child().get());
+				check(FirstFormat);
+
+				bool bAreAllFormatsCompatible = true;
+				for (int32 c = 1; c < MUTABLE_OP_MAX_SWIZZLE_CHANNELS; ++c)
+				{
+					if (Sources[c])
+					{
+						const ASTOpImagePixelFormat* Typed = dynamic_cast<const ASTOpImagePixelFormat*>(Sources[c].child().get());
+						check(Typed);
+
+						if (Typed->Source.child() != FirstFormat->Source.child())
+						{
+							bAreAllFormatsCompatible = false;
+							break;
+						}
+					}
+				}
+
+				if (bAreAllFormatsCompatible)
+				{
+					// Move the swizzle down all the paths
+					Ptr<ASTOpImagePixelFormat> NewFormat = mu::Clone<ASTOpImagePixelFormat>(FirstFormat);
+					NewFormat->Format = Format;
+					at = NewFormat;
+				}
+			}
 		}
 
 		// Swizzle of RGB from a source + A from a layer colour
-		// This can be optimize to apply the layer colour on-base directly to the alpha channel to skip the swizzle
+		// This can be optimized to apply the layer colour on-base directly to the alpha channel to skip the swizzle
 		if ( !at 
 			&&
 			Sources[0] && Sources[0]==Sources[1] && Sources[0]==Sources[2]
@@ -586,361 +766,69 @@ namespace mu
 			Sources[3] && Sources[3]->GetOpType()==OP_TYPE::IM_LAYERCOLOUR )
 		{
 			// Move the swizzle down all the paths
-			Ptr<ASTOpFixed> NewLayerColour = mu::Clone<ASTOpFixed>(Sources[3].child());
+			Ptr<ASTOpImageLayerColor> NewLayerColour = mu::Clone<ASTOpImageLayerColor>(Sources[3].child());
 
 			Ptr<ASTOpImageSwizzle> NewSwizzle = mu::Clone<ASTOpImageSwizzle>(this);
-			NewSwizzle->Sources[3] = NewLayerColour->children[NewLayerColour->op.args.ImageLayerColour.base].child();
+			NewSwizzle->Sources[3] = NewLayerColour->base.child();
 
-			NewLayerColour->op.args.ImageLayerColour.blendTypeAlpha = NewLayerColour->op.args.ImageLayerColour.blendType;
-			NewLayerColour->op.args.ImageLayerColour.blendType = uint8(EBlendType::BT_NONE);
-			NewLayerColour->SetChild(NewLayerColour->op.args.ImageLayerColour.base, NewSwizzle);
+			NewLayerColour->blendTypeAlpha = NewLayerColour->blendType;
+			NewLayerColour->blendType = EBlendType::BT_NONE;
+			NewLayerColour->base = NewSwizzle;
 
 			at = NewLayerColour;
 
 		}
 
+		// Swizzle of RGB from a layer colour + A from a different source
+		// This can be optimized to apply the layer colour on-base directly to the rgb channel to skip the swizzle
+		if (!at
+			&&
+			Sources[0] 
+			&& 
+			(!Sources[1] || Sources[0] == Sources[1]) 
+			&& 
+			(!Sources[2] || Sources[0] == Sources[2])
+			&&
+			Sources[0]->GetOpType() == OP_TYPE::IM_LAYERCOLOUR)
+		{
+			// Move the swizzle down all the rgb path
+			Ptr<ASTOpImageLayerColor> NewLayerColour = mu::Clone<ASTOpImageLayerColor>(Sources[0].child());
+			check(NewLayerColour);
+
+			Ptr<ASTOpImageSwizzle> NewSwizzle = mu::Clone<ASTOpImageSwizzle>(this);
+			NewSwizzle->Sources[0] = NewLayerColour->base.child();
+			NewSwizzle->Sources[1] = Sources[1] ? NewLayerColour->base.child() : nullptr;
+			NewSwizzle->Sources[2] = Sources[2] ? NewLayerColour->base.child() : nullptr;
+
+			NewLayerColour->blendTypeAlpha = EBlendType::BT_NONE;
+			NewLayerColour->base = NewSwizzle;
+
+			at = NewLayerColour;
+		}
+
+		// Swizzle of RGB from a binarise + A from a different source
+		// This can be optimized to apply the layer colour on-base directly to the rgb channel to skip the swizzle
+		// This is not valid because binarise always forces L8 format.
+		//if (!at
+		//	&&
+		//	Sources[0] && (!Sources[0] ||Sources[0] == Sources[1]) && (!Sources[2] || Sources[0] == Sources[2])
+		//	&&
+		//	Sources[0]->GetOpType() == OP_TYPE::IM_BINARISE)
+		//{
+		//	// Move the swizzle down all the rgb path
+		//	Ptr<ASTOpFixed> NewBinarise = mu::Clone<ASTOpFixed>(Sources[0].child());
+		//	check(NewBinarise);
+
+		//	Ptr<ASTOpImageSwizzle> NewSwizzle = mu::Clone<ASTOpImageSwizzle>(this);
+		//	NewSwizzle->Sources[0] = NewBinarise->children[NewBinarise->op.args.ImageBinarise.base].child();
+		//	NewSwizzle->Sources[1] = Sources[1] ? NewSwizzle->Sources[0].child() : nullptr;
+		//	NewSwizzle->Sources[2] = Sources[2] ? NewSwizzle->Sources[0].child() : nullptr;
+
+		//	NewBinarise->SetChild(NewBinarise->op.args.ImageBinarise.base, NewSwizzle);
+
+		//	at = NewBinarise;
+		//}
 
-//            OP::ADDRESS sourceAt = program.m_code[at].args.ImageSwizzle.sources[0];
-//            OP_TYPE sourceType = (OP_TYPE)program.m_code[sourceAt].type;
-//            switch ( sourceType )
-//            {
-//            case OP_TYPE::IM_COMPOSE:
-//            {
-//                bool canSink = true;
-//                OP::ADDRESS layout = program.m_code[sourceAt].args.ImageCompose.layout;
-//                uint32_t block = program.m_code[sourceAt].args.ImageCompose.blockIndex;
-//                for ( int c=1; c<MUTABLE_OP_MAX_SWIZZLE_CHANNELS; ++c )
-//                {
-//                    OP::ADDRESS candidate = program.m_code[at].args.ImageSwizzle.sources[c];
-//                    if ( candidate )
-//                    {
-//                        if ( program.m_code[candidate].type==sourceType )
-//                        {
-//                            OP::ADDRESS clayout =
-//                                    program.m_code[candidate].args.ImageCompose.layout;
-//                            uint32_t cblock =
-//                                    program.m_code[candidate].args.ImageCompose.blockIndex;
-//                            canSink = canSink
-//                                    && ( layout == clayout )
-//                                    && ( block == cblock );
-//                        }
-//                        else
-//                        {
-//                            canSink = false;
-//                        }
-//                    }
-//                }
-
-//                if ( canSink )
-//                {
-//                    m_modified = true;
-
-//                    OP op = program.m_code[sourceAt];
-
-//                    OP swizzleBaseOp = program.m_code[at];
-//                    OP swizzleBlockOp = program.m_code[at];
-
-//                    for ( int c=0; c<MUTABLE_OP_MAX_SWIZZLE_CHANNELS; ++c )
-//                    {
-//                        OP::ADDRESS source = program.m_code[at].args.ImageSwizzle.sources[c];
-//                        if ( source )
-//                        {
-//                            swizzleBaseOp.args.ImageSwizzle.sources[c] =
-//                                    program.m_code[ source ].args.ImageCompose.base;
-
-//                            swizzleBlockOp.args.ImageSwizzle.sources[c] =
-//                                    program.m_code[ source ].args.ImageCompose.blockImage;
-//                        }
-//                    }
-//                    op.args.ImageCompose.base = program.AddOp( swizzleBaseOp );
-//                    op.args.ImageCompose.blockImage = program.AddOp( swizzleBlockOp );
-
-//                    at = program.AddOp( op );
-//                }
-//                break;
-//            }
-
-//            case OP_TYPE::IM_CONDITIONAL:
-//            {
-//                // We can always sink into conditionals
-//                m_modified = true;
-
-//                OP op = program.m_code[sourceAt];
-
-//                OP swizzleYesOp = program.m_code[at];
-//                OP swizzleNoOp = program.m_code[at];
-//                OP::ADDRESS source0 = 0;
-//                for ( int c=0; c<MUTABLE_OP_MAX_SWIZZLE_CHANNELS; ++c )
-//                {
-//                    OP::ADDRESS source = program.m_code[at].args.ImageSwizzle.sources[c];
-//                    if ( c==0 )
-//                    {
-//                        source0 = source;
-//                    }
-
-//                    if ( source
-//                        &&
-//                        ( (c==0) || ( program.m_code[source0].args.Conditional.condition == program.m_code[source].args.Conditional.condition ) )
-//                        )
-//                    {
-//                        swizzleYesOp.args.ImageSwizzle.sources[c] =
-//                                program.m_code[ source ].args.Conditional.yes;
-
-//                        swizzleNoOp.args.ImageSwizzle.sources[c] =
-//                                program.m_code[ source ].args.Conditional.no;
-//                    }
-//                }
-//                op.args.Conditional.yes = program.AddOp( swizzleYesOp );
-//                op.args.Conditional.no = program.AddOp( swizzleNoOp );
-
-//                at = program.AddOp( op );
-//                break;
-//            }
-
-
-//            case OP_TYPE::IM_SWITCH:
-//            {
-//                bool canSink = true;
-//                OP::ADDRESS variable = program.m_code[sourceAt].args.Switch.variable;
-
-//                // If at least 3 channels can be combined, do it. We may be duplicating some data
-//                // But it greatly optimises speed in some cases like the constant alpha masks of
-//                // the face colour interpolation.
-//                // TODO: Make this dependent on compilation options.
-//                //for ( int c=1; c<MUTABLE_OP_MAX_SWIZZLE_CHANNELS; ++c )
-//                for ( int c=1; c<3; ++c )
-//                {
-//                    OP::ADDRESS candidate = program.m_code[at].args.ImageSwizzle.sources[c];
-//                    if ( candidate )
-//                    {
-//                        if ( program.m_code[candidate].type==sourceType )
-//                        {
-//                            OP::ADDRESS cvariable =
-//                                    program.m_code[candidate].args.Switch.variable;
-//                            canSink = canSink && ( variable == cvariable );
-//                        }
-//                        else
-//                        {
-//                            canSink = false;
-//                        }
-//                    }
-//                }
-
-//                if ( canSink )
-//                {
-//                    // TODO since data-switch
-////                    m_modified = true;
-
-////                    OP op = program.m_code[sourceAt];
-
-////                    if ( op.args.Switch.def )
-////                    {
-////                        OP swizzleDefOp = program.m_code[at];
-////                        for ( int c=0; c<MUTABLE_OP_MAX_SWIZZLE_CHANNELS; ++c )
-////                        {
-////                            OP::ADDRESS source = program.m_code[at].args.ImageSwizzle.sources[c];
-////                            if ( source )
-////                            {
-////                                if ( program.m_code[source].type==sourceType )
-////                                {
-////                                    swizzleDefOp.args.ImageSwizzle.sources[c] =
-////                                            program.m_code[ source ].args.Switch.def;
-////                                }
-////                                else
-////                                {
-////                                    swizzleDefOp.args.ImageSwizzle.sources[c] = source;
-////                                }
-////                            }
-////                        }
-////                        op.args.Switch.def = program.AddOp( swizzleDefOp );
-////                    }
-
-////                    for ( int o=0; o<MUTABLE_OP_MAX_SWITCH_OPTIONS; ++o )
-////                    {
-////                        if ( op.args.Switch.values[o] )
-////                        {
-////                            OP swizzleOptOp = program.m_code[at];
-////                            for ( int c=0; c<MUTABLE_OP_MAX_SWIZZLE_CHANNELS; ++c )
-////                            {
-////                                OP::ADDRESS source = program.m_code[at].args.ImageSwizzle.sources[c];
-////                                if ( source )
-////                                {
-////                                    if ( program.m_code[source].type==sourceType )
-////                                    {
-////                                        swizzleOptOp.args.ImageSwizzle.sources[c] =
-////                                                program.m_code[ source ].args.Switch.values[o];
-////                                    }
-////                                    else
-////                                    {
-////                                        swizzleOptOp.args.ImageSwizzle.sources[c] = source;
-////                                    }
-////                                }
-////                            }
-////                            op.args.Switch.values[o] = program.AddOp( swizzleOptOp );
-////                        }
-////                    }
-
-////                    at = program.AddOp( op );
-//                }
-//                break;
-//            }
-
-
-//            case OP_TYPE::IM_INTERPOLATE:
-//            {
-//                bool canSink = true;
-//                OP::ADDRESS factor = program.m_code[sourceAt].args.ImageInterpolate.factor;
-
-//                // It is worth to sink 3 channels
-//                // TODO: Make it depend on compilation parameters
-//                //for ( int c=1; c<MUTABLE_OP_MAX_SWIZZLE_CHANNELS; ++c )
-//                for ( int c=1; c<3; ++c )
-//                {
-//                    OP::ADDRESS candidate = program.m_code[at].args.ImageSwizzle.sources[c];
-//                    if ( candidate )
-//                    {
-//                        if ( program.m_code[candidate].type==sourceType )
-//                        {
-//                            OP::ADDRESS cfactor =
-//                                    program.m_code[candidate].args.ImageInterpolate.factor;
-//                            canSink = canSink && ( factor == cfactor );
-//                        }
-//                        else
-//                        {
-//                            canSink = false;
-//                        }
-//                    }
-//                }
-
-//                if ( canSink )
-//                {
-//                    m_modified = true;
-
-//                    OP op = program.m_code[sourceAt];
-
-//                    for ( int o=0; o<MUTABLE_OP_MAX_INTERPOLATE_COUNT; ++o )
-//                    {
-//                        if ( op.args.ImageInterpolate.targets[o] )
-//                        {
-//                            OP swizzleOptOp = program.m_code[at];
-//                            for ( int c=0; c<MUTABLE_OP_MAX_SWIZZLE_CHANNELS; ++c )
-//                            {
-//                                OP::ADDRESS source = program.m_code[at].args.ImageSwizzle.sources[c];
-
-//                                if ( source
-//                                     && program.m_code[source].type==OP_TYPE::IM_INTERPOLATE
-//                                     && program.m_code[source].args.ImageInterpolate.factor==factor
-//                                     )
-//                                {
-//                                    swizzleOptOp.args.ImageSwizzle.sources[c] =
-//                                            program.m_code[ source ].args.ImageInterpolate.targets[o];
-//                                }
-//                            }
-//                            op.args.ImageInterpolate.targets[o] = program.AddOp( swizzleOptOp );
-//                        }
-//                    }
-
-//                    at = program.AddOp( op );
-//                }
-//                break;
-//            }
-
-
-//            case OP_TYPE::IM_INTERPOLATE3:
-//            {
-//                // We can sink if all channels are interpolates for the same factors
-//                OP::ADDRESS sourceAt0 = program.m_code[at].args.ImageSwizzle.sources[0];
-//                OP::ADDRESS factor10 = program.m_code[sourceAt0].args.ImageInterpolate3.factor1;
-//                OP::ADDRESS factor20 = program.m_code[sourceAt0].args.ImageInterpolate3.factor2;
-
-//                OP::ADDRESS sourceAt1 = program.m_code[at].args.ImageSwizzle.sources[1];
-//                OP::ADDRESS factor11 = program.m_code[sourceAt1].args.ImageInterpolate3.factor1;
-//                OP::ADDRESS factor21 = program.m_code[sourceAt1].args.ImageInterpolate3.factor2;
-
-//                OP::ADDRESS sourceAt2 = program.m_code[at].args.ImageSwizzle.sources[2];
-//                OP::ADDRESS factor12 = program.m_code[sourceAt2].args.ImageInterpolate3.factor1;
-//                OP::ADDRESS factor22 = program.m_code[sourceAt2].args.ImageInterpolate3.factor2;
-
-//                bool canSink = true;
-//                canSink &= program.m_code[sourceAt0].type == program.m_code[sourceAt1].type;
-//                canSink &= program.m_code[sourceAt0].type == program.m_code[sourceAt2].type;
-//                canSink &= factor10 == factor11;
-//                canSink &= factor10 == factor12;
-//                canSink &= factor20 == factor21;
-//                canSink &= factor20 == factor22;
-
-//                if ( canSink )
-//                {
-//                    m_modified = true;
-
-//                    OP op = program.m_code[sourceAt0];
-
-//                    // Target 0
-//                    {
-//                        OP swizzleOptOp = program.m_code[at];
-//                        for ( int c=0; c<MUTABLE_OP_MAX_SWIZZLE_CHANNELS; ++c )
-//                        {
-//                            OP::ADDRESS source = program.m_code[at].args.ImageSwizzle.sources[c];
-//                            if ( source
-//                                 && program.m_code[source].type==OP_TYPE::IM_INTERPOLATE3
-//                                 && program.m_code[source].args.ImageInterpolate3.factor1==factor10
-//                                 && program.m_code[source].args.ImageInterpolate3.factor2==factor20
-//                                 )
-//                            {
-//                                swizzleOptOp.args.ImageSwizzle.sources[c] =
-//                                        program.m_code[ source ].args.ImageInterpolate3.target0;
-//                            }
-//                        }
-//                        op.args.ImageInterpolate3.target0 = program.AddOp( swizzleOptOp );
-//                    }
-
-//                    // Target 1
-//                    {
-//                        OP swizzleOptOp = program.m_code[at];
-//                        for ( int c=0; c<MUTABLE_OP_MAX_SWIZZLE_CHANNELS; ++c )
-//                        {
-//                            OP::ADDRESS source = program.m_code[at].args.ImageSwizzle.sources[c];
-//                            if ( source
-//                                 && program.m_code[source].type==OP_TYPE::IM_INTERPOLATE3
-//                                 && program.m_code[source].args.ImageInterpolate3.factor1==factor10
-//                                 && program.m_code[source].args.ImageInterpolate3.factor2==factor20
-//                                 )
-//                            {
-//                                swizzleOptOp.args.ImageSwizzle.sources[c] =
-//                                        program.m_code[ source ].args.ImageInterpolate3.target1;
-//                            }
-//                        }
-//                        op.args.ImageInterpolate3.target1 = program.AddOp( swizzleOptOp );
-//                    }
-
-//                    // Target 2
-//                    {
-//                        OP swizzleOptOp = program.m_code[at];
-//                        for ( int c=0; c<MUTABLE_OP_MAX_SWIZZLE_CHANNELS; ++c )
-//                        {
-//                            OP::ADDRESS source = program.m_code[at].args.ImageSwizzle.sources[c];
-//                            if ( source
-//                                 && program.m_code[source].type==OP_TYPE::IM_INTERPOLATE3
-//                                 && program.m_code[source].args.ImageInterpolate3.factor1==factor10
-//                                 && program.m_code[source].args.ImageInterpolate3.factor2==factor20
-//                                 )
-//                            {
-//                                swizzleOptOp.args.ImageSwizzle.sources[c] =
-//                                        program.m_code[ source ].args.ImageInterpolate3.target2;
-//                            }
-//                        }
-//                        op.args.ImageInterpolate3.target2 = program.AddOp( swizzleOptOp );
-//                    }
-
-//                    at = program.AddOp( op );
-//                }
-//                break;
-//            }
-
-//            default:
-//                break;
-//            }
-//		}
 
 		return at;
 	}
diff --git a/Engine/Plugins/Experimental/Mutable/Source/MutableTools/Private/MuT/CodeGenerator.cpp b/Engine/Plugins/Experimental/Mutable/Source/MutableTools/Private/MuT/CodeGenerator.cpp
--- a/Engine/Plugins/Experimental/Mutable/Source/MutableTools/Private/MuT/CodeGenerator.cpp
+++ b/Engine/Plugins/Experimental/Mutable/Source/MutableTools/Private/MuT/CodeGenerator.cpp
@@ -20,6 +20,8 @@
 #include "MuT/ASTOpImageCompose.h"
 #include "MuT/ASTOpImageMipmap.h"
 #include "MuT/ASTOpImagePixelFormat.h"
+#include "MuT/ASTOpImageLayer.h"
+#include "MuT/ASTOpImageLayerColor.h"
 #include "MuT/ASTOpInstanceAdd.h"
 #include "MuT/ASTOpMeshBindShape.h"
 #include "MuT/ASTOpMeshClipDeform.h"
@@ -82,138 +84,12 @@
 #include "MuT/NodeSurfaceVariation.h"
 #include "MuT/NodeSurfaceVariationPrivate.h"
 #include "MuT/TablePrivate.h"
-#include "MuT/TaskManager.h"
 #include "Trace/Detail/Channel.h"
 
 
 namespace mu
 {
 
-    void PartialOptimise( Ptr<ASTOp>& lastCompOp, const CompilerOptions::Private* options )
-    {
-        if (!lastCompOp) return;
-
-        MUTABLE_CPUPROFILER_SCOPE(PartialOptimise);
-
-        //UE_LOG(LogMutableCore, Verbose, " Begin partial optimisation.");
-
-        ASTOpList roots;
-        roots.Add( lastCompOp );
-
-        //size_t initialOps = ASTOp::CountNodes(roots);
-        //(void)initialOps;
-
-        //AXE_INT_VALUE("Mutable", Verbose, "ast size", (int64_t)ASTOp::CountNodes(roots));
-        //UE_LOG(LogMutableCore, Verbose, " - duplicated data remover");
-        DuplicatedDataRemoverAST( roots );
-        //AXE_INT_VALUE("Mutable", Verbose, "ast size", (int64_t)ASTOp::CountNodes(roots));
-
-        //UE_LOG(LogMutableCore, Verbose, " - duplicated code remover");
-        DuplicatedCodeRemoverAST( roots );
-        //AXE_INT_VALUE("Mutable", Verbose, "ast size", (int64_t)ASTOp::CountNodes(roots));
-
-        //UE_LOG(LogMutableCore, Verbose, " - size optimiser");
-        SizeOptimiserAST( roots );
-        //AXE_INT_VALUE("Mutable", Verbose, "ast size", (int64_t)ASTOp::CountNodes(roots));
-
-        //UE_LOG(LogMutableCore, Verbose, " - semantic optimiser");
-        SemanticOptimiserAST( roots, options->m_optimisationOptions );
-        //AXE_INT_VALUE("Mutable", Verbose, "ast size", (int64_t)ASTOp::CountNodes(roots));
-
-        //UE_LOG(LogMutableCore, Verbose, " - sink optimiser");
-        SinkOptimiserAST( roots, options->m_optimisationOptions );
-        //AXE_INT_VALUE("Mutable", Verbose, "ast size", (int64_t)ASTOp::CountNodes(roots));
-
-        // Image size operations are treated separately
-        //UE_LOG(LogMutableCore, Verbose, " - size optimiser");
-        SizeOptimiserAST( roots );
-        //AXE_INT_VALUE("Mutable", Verbose, "ast size", (int64_t)ASTOp::CountNodes(roots));
-
-        //UE_LOG(LogMutableCore, Verbose, " - semantic optimiser");
-        SemanticOptimiserAST( roots, options->m_optimisationOptions );
-        //AXE_INT_VALUE("Mutable", Verbose, "ast size", (int64_t)ASTOp::CountNodes(roots));
-
-        //UE_LOG(LogMutableCore, Verbose, " - duplicated data remover");
-        DuplicatedDataRemoverAST( roots );
-        //AXE_INT_VALUE("Mutable", Verbose, "ast size", (int64_t)ASTOp::CountNodes(roots));
-
-        //UE_LOG(LogMutableCore, Verbose, " - duplicated code remover");
-        DuplicatedCodeRemoverAST( roots );
-        //AXE_INT_VALUE("Mutable", Verbose, "ast size", (int64_t)ASTOp::CountNodes(roots));
-
-        // Constant subtree generation
-        ConstantGeneratorAST( options, *roots.begin(), nullptr );
-
-        //UE_LOG(LogMutableCore, Verbose, " - duplicated data remover");
-        DuplicatedDataRemoverAST( roots );
-        //AXE_INT_VALUE("Mutable", Verbose, "ast size", (int64_t)ASTOp::CountNodes(roots));
-
-        //UE_LOG(LogMutableCore, Verbose, " - duplicated code remover");
-        DuplicatedCodeRemoverAST( roots );
-        //AXE_INT_VALUE("Mutable", Verbose, "ast size", (int64_t)ASTOp::CountNodes(roots));
-
-        //UE_LOG(LogMutableCore, Verbose, " - logic optimiser");
-        LocalLogicOptimiserAST( roots );
-        //AXE_INT_VALUE("Mutable", Verbose, "ast size", (int64_t)ASTOp::CountNodes(roots));
-
-        //size_t finalOps = ASTOp::CountNodes(roots);
-        //(void)finalOps;
-
-        //lastCompOp = *roots.begin();
-
-        //UE_LOG(LogMutableCore, Verbose, " End partial optimisation.");
-
-        //UE_LOG(LogMutableCore, Verbose,
-        //        "Partial optimisation from %6d to %6d", int64_t(initialOps), int64_t(finalOps) );
-
-        lastCompOp = roots[0];
-    }
-
-
-    //---------------------------------------------------------------------------------------------
-    class OptimizeTask : public TaskManager::Task
-    {
-    public:
-        // input
-        Ptr<ASTOp> m_source;
-        CompilerOptions::Private* m_options;
-
-        std::shared_timed_mutex* m_codeAccessMutex = nullptr;
-
-    private:
-
-        // output
-        Ptr<ASTOp> m_result;
-
-    public:
-
-        // mu::Task interface
-        void Run() override
-        {
-            // This runs in a random thread
-
-            // Clone the tree
-            m_codeAccessMutex->lock_shared();
-            m_result = ASTOp::DeepClone( m_source );
-            m_codeAccessMutex->unlock_shared();
-
-            // Optimise
-            PartialOptimise( m_result, m_options );
-        }
-
-
-        // TaskManager::Task interface
-        void Complete() override
-        {
-            // This runs in the main thread
-            m_codeAccessMutex->lock();
-            ASTOp::Replace( m_source, m_result );
-            m_codeAccessMutex->unlock();
-        }
-
-    };
-
-
     //---------------------------------------------------------------------------------------------
     //---------------------------------------------------------------------------------------------
     //---------------------------------------------------------------------------------------------
@@ -229,12 +105,10 @@ namespace mu
     }
 
     //---------------------------------------------------------------------------------------------
-    void CodeGenerator::GenerateRoot(const NodePtrConst pNode, TaskManager* taskManager )
+    void CodeGenerator::GenerateRoot(const NodePtrConst pNode )
     {
         MUTABLE_CPUPROFILER_SCOPE(Generate);
 
-        m_pTaskManager = taskManager;
-
 //        UE_LOG(LogMutableCore,Warning,"astop %9u#  %9u+  %9u- generate root start", size_t(ASTOp::s_count), size_t(ASTOp::s_created), size_t(ASTOp::s_destroyed) );
 
 		// First pass
@@ -286,22 +160,6 @@ namespace mu
         m_currentObject.Empty();
         m_additionalComponents.clear();
 
-
-//        UE_LOG(LogMutableCore,Warning,"astop %9u#  %9u+  %9u- after cache clear", size_t(ASTOp::s_count), size_t(ASTOp::s_created), size_t(ASTOp::s_destroyed) );
-
-        // Complete pending tasks
-        {
-            if (m_pTaskManager)
-            {
-                UE_LOG(LogMutableCore,Verbose,TEXT("Waiting for remaining code generation tasks..."));
-                MUTABLE_CPUPROFILER_SCOPE(PendingTasks);
-                m_pTaskManager->CompleteTasks();
-                UE_LOG(LogMutableCore,Verbose, TEXT("End code generation tasks"));
-            }
-        }
-
-        m_pTaskManager = nullptr;
-
 //        UE_LOG(LogMutableCore,Warning,"astop %9u#  %9u+  %9u- exiting generate root", size_t(ASTOp::s_count), size_t(ASTOp::s_created), size_t(ASTOp::s_destroyed) );
 	}
 
@@ -560,15 +418,14 @@ namespace mu
          {
              MUTABLE_CPUPROFILER_SCOPE(PatchBlend);
 
-             Ptr<ASTOpFixed> op = new ASTOpFixed();
-			 op->op.type = OP_TYPE::IM_LAYER;
-			 op->op.args.ImageLayer.blendType = uint8_t(pPatch->GetPrivate()->m_blendType);
-             op->SetChild(op->op.args.ImageLayer.base, blockAd );
+             Ptr<ASTOpImageLayer> op = new ASTOpImageLayer();
+			 op->blendType = pPatch->GetPrivate()->m_blendType;
+             op->base = blockAd;
 
              // When we patch from edit nodes, we want to apply it to all the channels.
              // \todo: since we can choose the patch function, maybe we want to be able to
              // select this as well.
-             op->op.args.ImageLayer.flags = pPatch->GetPrivate()->m_applyToAlpha
+             op->Flags = pPatch->GetPrivate()->m_applyToAlpha
                      ? OP::ImageLayerArgs::F_APPLY_TO_ALPHA
                      : 0;
 
@@ -590,7 +447,7 @@ namespace mu
                      FImageSize( (uint16)m_imageState.Last().m_imageRect.size[0],
                                  (uint16)m_imageState.Last().m_imageRect.size[1] )
                  );
-             op->SetChild(op->op.args.ImageLayer.blended, blend );
+             op->blend = blend;
 
              NodeImage* pMask = pPatch->GetPrivate()->m_pMask.get();
              Ptr<ASTOp> mask;
@@ -611,7 +468,7 @@ namespace mu
                      FImageSize( (uint16)m_imageState.Last().m_imageRect.size[0],
                                  (uint16)m_imageState.Last().m_imageRect.size[1] )
                  );
-             op->SetChild(op->op.args.ImageLayer.mask, mask );
+             op->mask = mask;
 
              blendAd = op;
          }
@@ -1401,22 +1258,6 @@ namespace mu
                         imageAd = fop;
                     }
 
-                    // if (imageAd && m_compilerOptions->m_enablePartialOptimise)
-                    // {
-                    //     if (m_pTaskManager->HasRunner())
-                    //     {
-                    //         OptimizeTask* ot = new OptimizeTask;
-                    //         ot->m_optimizationOptions = &m_compilerOptions->m_optimisationOptions;
-                    //         ot->m_source = imageAd;
-                    //         ot->m_codeAccessMutex = &m_codeAccessMutex;
-                    //         m_pTaskManager->AddTask(ot);
-                    //     }
-                    //     else
-                    //     {
-                    //         PartialOptimise(imageAd, m_compilerOptions->m_optimisationOptions);
-                    //     }
-                    // }
-
                     Ptr<ASTOpInstanceAdd> op = new ASTOpInstanceAdd();
                     op->type = OP_TYPE::IN_ADDIMAGE;
                     op->instance = lastSurfOp;
@@ -1751,22 +1592,6 @@ namespace mu
                             imageAd = fop;
                         }
 
-                        // if (imageAd && m_compilerOptions->m_enablePartialOptimise)
-                        // {
-                        //     if (m_pTaskManager->HasRunner())
-                        //     {                            
-                        //         OptimizeTask* ot = new OptimizeTask;
-                        //         ot->m_optimizationOptions = &m_compilerOptions->m_optimisationOptions;
-                        //         ot->m_source = imageAd;
-                        //         ot->m_codeAccessMutex = &m_codeAccessMutex;
-                        //         m_pTaskManager->AddTask(ot);
-                        //     }
-                        //     else
-                        //     {
-                        //         PartialOptimise(imageAd, m_compilerOptions->m_optimisationOptions);
-                        //     }
-                        // }
-
                         op->value = imageAd;
 
                         // Name
@@ -1850,25 +1675,6 @@ namespace mu
             }
         }
 
-
-        // Partial optimisation of the code generated so far
-        if (lastSurfOp && m_compilerOptions->m_enablePartialOptimise)
-        {
-            if (m_pTaskManager->IsConcurrencyEnabled())
-            {
-                OptimizeTask* t = new OptimizeTask;
-                t->m_options = m_compilerOptions;
-                t->m_source = lastSurfOp;
-                t->m_codeAccessMutex = &m_codeAccessMutex;
-                m_pTaskManager->AddTask(t);
-            }
-            else
-            {
-                PartialOptimise(lastSurfOp, m_compilerOptions);
-            }
-        }
-
-
         result.surfaceOp = lastSurfOp;
         targetSurface->resultSurfaceOp = lastSurfOp;
     }
@@ -1957,22 +1763,7 @@ namespace mu
                 // We add the merge op even for the first mesh, so that we set the surface id.
                 {
                     Ptr<ASTOp> added = its.resultMeshOp;
-                    // if (added && m_compilerOptions->m_enablePartialOptimise)
-                    // {
-                    //     if (m_pTaskManager->HasRunner())
-                    //     {
-                    //         OptimizeTask* t = new OptimizeTask;
-                    //         t->m_optimizationOptions = &m_compilerOptions->m_optimisationOptions;
-                    //         t->m_source = added;
-                    //         t->m_codeAccessMutex = &m_codeAccessMutex;
-                    //         m_pTaskManager->AddTask(t);
-                    //     }
-                    //     else
-                    //     {
-                    //         PartialOptimise( added, m_compilerOptions->m_optimisationOptions );
-                    //     }
-                    // }
-
+ 
                     Ptr<ASTOpFixed> mop = new ASTOpFixed();
                     mop->op.type = OP_TYPE::ME_MERGE;
                     mop->SetChild(mop->op.args.MeshMerge.base, lastMeshOp );
diff --git a/Engine/Plugins/Experimental/Mutable/Source/MutableTools/Private/MuT/CodeGenerator.h b/Engine/Plugins/Experimental/Mutable/Source/MutableTools/Private/MuT/CodeGenerator.h
--- a/Engine/Plugins/Experimental/Mutable/Source/MutableTools/Private/MuT/CodeGenerator.h
+++ b/Engine/Plugins/Experimental/Mutable/Source/MutableTools/Private/MuT/CodeGenerator.h
@@ -42,7 +42,6 @@
 #include "MuT/Visitor.h"
 #include "Templates/TypeHash.h"
 
-#include <shared_mutex>
 #include <utility>
 
 
@@ -109,7 +108,6 @@ namespace mu
 	class NodeScalarVariation;
 	class NodeStringConstant;
 	class NodeStringParameter;
-	class TaskManager;
 	struct OBJECT_STATE;
 	struct FProgram;
 
@@ -134,7 +132,7 @@ namespace mu
         CodeGenerator( CompilerOptions::Private* options );
 
         //! Data will be stored in m_states
-        void GenerateRoot( const NodePtrConst pNode, TaskManager* );
+        void GenerateRoot( const NodePtrConst pNode );
 
 	protected:
 
@@ -352,8 +350,6 @@ namespace mu
 
     private:
 
-        std::shared_timed_mutex m_codeAccessMutex;
-
 		//! Generate the key with all the relevant state that is used in generation of operations for a node.
 		FVisitedKeyMap GetCurrentCacheKey(const NodePtrConst& InNode) const
 		{
@@ -672,8 +668,6 @@ namespace mu
         void GenerateSurface( FSurfaceGenerationResult& result,
                               NodeSurfaceNewPtrConst node,
                               const TArray<FirstPassGenerator::SURFACE::EDIT>& edits );
-
-        TaskManager* m_pTaskManager = nullptr;
     };
 
 
diff --git a/Engine/Plugins/Experimental/Mutable/Source/MutableTools/Private/MuT/CodeGenerator_Image.cpp b/Engine/Plugins/Experimental/Mutable/Source/MutableTools/Private/MuT/CodeGenerator_Image.cpp
--- a/Engine/Plugins/Experimental/Mutable/Source/MutableTools/Private/MuT/CodeGenerator_Image.cpp
+++ b/Engine/Plugins/Experimental/Mutable/Source/MutableTools/Private/MuT/CodeGenerator_Image.cpp
@@ -22,6 +22,8 @@
 #include "MuT/ASTOpConditional.h"
 #include "MuT/ASTOpConstantResource.h"
 #include "MuT/ASTOpImageMipmap.h"
+#include "MuT/ASTOpImageLayer.h"
+#include "MuT/ASTOpImageLayerColor.h"
 #include "MuT/ASTOpImageMultiLayer.h"
 #include "MuT/ASTOpImageNormalComposite.h"
 #include "MuT/ASTOpImagePixelFormat.h"
@@ -299,10 +301,9 @@ namespace mu
 		
 		MUTABLE_CPUPROFILER_SCOPE(NodeImageLayer);
 
-        Ptr<ASTOpFixed> op = new ASTOpFixed();
+        Ptr<ASTOpImageLayer> op = new ASTOpImageLayer();
 
-		op->op.type = OP_TYPE::IM_LAYER;
-        op->op.args.ImageLayer.blendType = uint8(node.m_type);
+        op->blendType = node.m_type;
 
         // Base image
         Ptr<ASTOp> base;
@@ -322,7 +323,7 @@ namespace mu
         base = GenerateImageSize
                 ( base, FImageSize((uint16)m_imageState.Last().m_imageRect.size[0],
                                    (uint16)m_imageState.Last().m_imageRect.size[1]) );
-        op->SetChild( op->op.args.ImageLayer.base, base);
+        op->base = base;
 
         // Mask of the effect
         Ptr<ASTOp> mask = 0;
@@ -334,7 +335,7 @@ namespace mu
                     ( mask, FImageSize((uint16)m_imageState.Last().m_imageRect.size[0],
                                        (uint16)m_imageState.Last().m_imageRect.size[1]) );
         }
-        op->SetChild( op->op.args.ImageLayer.mask, mask);
+        op->mask = mask;
 
         // Image to apply
         Ptr<ASTOp> blended = 0;
@@ -352,7 +353,7 @@ namespace mu
         blended = GenerateImageSize
                 ( blended, FImageSize((uint16)m_imageState.Last().m_imageRect.size[0],
                                       (uint16)m_imageState.Last().m_imageRect.size[1]) );
-        op->SetChild( op->op.args.ImageLayer.blended, blended);
+        op->blend = blended;
 
         result.op = op;
     }
@@ -365,9 +366,8 @@ namespace mu
 		
 		MUTABLE_CPUPROFILER_SCOPE(NodeImageLayerColour);
 
-        Ptr<ASTOpFixed> op = new ASTOpFixed();
-		op->op.type = OP_TYPE::IM_LAYERCOLOUR;
-		op->op.args.ImageLayerColour.blendType = uint8(node.m_type);
+        Ptr<ASTOpImageLayerColor> op = new ASTOpImageLayerColor();
+		op->blendType = node.m_type;
 
         // Base image
         Ptr<ASTOp> base;
@@ -381,9 +381,8 @@ namespace mu
             base = GenerateMissingImageCode( "Layer base image", EImageFormat::IF_RGB_UBYTE, node.m_errorContext );
         }
         base = GenerateImageFormat( base, EImageFormat::IF_RGB_UBYTE );
-        base = GenerateImageSize
-                ( base, FImageSize((uint16)m_imageState.Last().m_imageRect.size[0],(uint16)m_imageState.Last().m_imageRect.size[1]) );
-        op->SetChild( op->op.args.ImageLayerColour.base, base);
+        base = GenerateImageSize( base, FImageSize((uint16)m_imageState.Last().m_imageRect.size[0],(uint16)m_imageState.Last().m_imageRect.size[1]) );
+        op->base = base;
 
         // Mask of the effect
         Ptr<ASTOp> mask = 0;
@@ -391,10 +390,9 @@ namespace mu
         {
             mask = Generate( pMask );
             mask = GenerateImageFormat( mask, EImageFormat::IF_L_UBYTE );
-            mask = GenerateImageSize
-                    ( mask, FImageSize((uint16)m_imageState.Last().m_imageRect.size[0],(uint16)m_imageState.Last().m_imageRect.size[1]) );
+            mask = GenerateImageSize( mask, FImageSize((uint16)m_imageState.Last().m_imageRect.size[0],(uint16)m_imageState.Last().m_imageRect.size[1]) );
         }
-        op->SetChild( op->op.args.ImageLayerColour.mask, mask);
+        op->mask = mask;
 
         // Colour to apply
         Ptr<ASTOp> colour = 0;
@@ -407,7 +405,7 @@ namespace mu
             // This argument is required
             colour = GenerateMissingColourCode( "Layer colour", node.m_errorContext );
         }
-        op->SetChild( op->op.args.ImageLayerColour.colour, colour);
+        op->color = colour;
 
         result.op = op;
     }
diff --git a/Engine/Plugins/Experimental/Mutable/Source/MutableTools/Private/MuT/CodeOptimiser.cpp b/Engine/Plugins/Experimental/Mutable/Source/MutableTools/Private/MuT/CodeOptimiser.cpp
--- a/Engine/Plugins/Experimental/Mutable/Source/MutableTools/Private/MuT/CodeOptimiser.cpp
+++ b/Engine/Plugins/Experimental/Mutable/Source/MutableTools/Private/MuT/CodeOptimiser.cpp
@@ -21,7 +21,6 @@
 #include "MuT/ASTOpMeshApplyPose.h"
 
 #include <unordered_set>
-#include <shared_mutex>
 
 namespace mu
 {
@@ -427,7 +426,7 @@ namespace mu
 	//---------------------------------------------------------------------------------------------
 	//---------------------------------------------------------------------------------------------
 	//---------------------------------------------------------------------------------------------
-	class ConstantTask : public TaskManager::Task
+	class ConstantTask 
 	{
 	public:
 		// input
@@ -435,8 +434,6 @@ namespace mu
 		bool m_useDiskCache = false;
 		int m_imageCompressionQuality = 0;
 
-		std::shared_timed_mutex* m_codeAccessMutex = nullptr;
-
 	private:
 
 		//
@@ -445,14 +442,12 @@ namespace mu
 	public:
 
 		// mu::Task interface
-		void Run() override
+		void Run()
 		{
 			MUTABLE_CPUPROFILER_SCOPE(ConstantTask_Run);
 
 			// We need the clone because linking modifies ASTOp state
-			m_codeAccessMutex->lock_shared();
 			mu::Ptr<ASTOp> cloned = ASTOp::DeepClone( m_source );
-			m_codeAccessMutex->unlock_shared();
 
 			OP_TYPE type = cloned->GetOpType();
 			DATATYPE dtype = GetOpDataType(type);
@@ -576,20 +571,16 @@ namespace mu
 
 		}
 
-		// TaskManager::Task interface
-		void Complete() override
+		void Complete()
 		{
 			// This runs in the managing thread
-			m_codeAccessMutex->lock();
 			ASTOp::Replace( m_source, m_result );
-			m_codeAccessMutex->unlock();
 		}
 	};
 
 
 	//---------------------------------------------------------------------------------------------
-	bool ConstantGeneratorAST( const CompilerOptions::Private* options, Ptr<ASTOp>& root,
-								   TaskManager* pTaskManager )
+	bool ConstantGeneratorAST( const CompilerOptions::Private* options, Ptr<ASTOp>& root )
 	{
 		MUTABLE_CPUPROFILER_SCOPE(ConstantGenerator);
 
@@ -666,8 +657,6 @@ namespace mu
 
 		}
 
-		std::shared_timed_mutex codeAccessMutex;
-
 		TArray< Ptr<ASTOp> > roots;
 		roots.Add(root);
 
@@ -755,23 +744,13 @@ namespace mu
 					modified = true;
 
 					ConstantTask* constantTask = new ConstantTask;
-					constantTask->m_useDiskCache =
-						options->m_optimisationOptions.m_useDiskCache;
+					constantTask->m_useDiskCache = options->m_optimisationOptions.m_useDiskCache;
 					constantTask->m_source = n;
-					constantTask->m_codeAccessMutex = &codeAccessMutex;
-					constantTask->m_imageCompressionQuality =
-						options->m_imageCompressionQuality;
+					constantTask->m_imageCompressionQuality = options->m_imageCompressionQuality;
 
-					if ( pTaskManager )
-					{
-						pTaskManager->AddTask( constantTask );
-					}
-					else
-					{
-						constantTask->Run();
-						constantTask->Complete();
-						delete constantTask;
-					}
+					constantTask->Run();
+					constantTask->Complete();
+					delete constantTask;
 
 					break;
 				}
@@ -796,11 +775,6 @@ namespace mu
 			return recurse;
 		});
 
-		if (pTaskManager)
-		{
-			pTaskManager->CompleteTasks();
-		}
-
 		return modified;
 	}
 
@@ -816,8 +790,7 @@ namespace mu
 
 
 	//-------------------------------------------------------------------------------------------------
-	void CodeOptimiser::FullOptimiseAST( ASTOpList& roots,
-										 TaskManager* pTaskManager )
+	void CodeOptimiser::FullOptimiseAST( ASTOpList& roots )
 	{
 		bool modified = true;
 		int numIterations = 0;
@@ -880,7 +853,7 @@ namespace mu
 			UE_LOG(LogMutableCore, Verbose, TEXT(" - constant generator"));
 
 			// Constant subtree generation
-			modified = ConstantGeneratorAST( m_options->GetPrivate(), r, pTaskManager );
+			modified = ConstantGeneratorAST( m_options->GetPrivate(), r );
 		}
 
 		UE_LOG(LogMutableCore, Verbose, TEXT("(int) %s : %ld"), TEXT("ast size"), int64(ASTOp::CountNodes(roots)));
@@ -1180,7 +1153,7 @@ namespace mu
 
 
 	//---------------------------------------------------------------------------------------------
-	void CodeOptimiser::OptimiseAST( TaskManager* pTaskManager )
+	void CodeOptimiser::OptimiseAST()
 	{
 		MUTABLE_CPUPROFILER_SCOPE(OptimiseAST);
 
@@ -1251,32 +1224,11 @@ namespace mu
 			// Main optimisation stage
 			{
 				MUTABLE_CPUPROFILER_SCOPE(MainStage);
-				FullOptimiseAST( roots, pTaskManager );
+				FullOptimiseAST( roots );
 				UE_LOG(LogMutableCore, Verbose, TEXT("(int) %s : %ld"), TEXT("ast size"), int64(ASTOp::CountNodes(roots)));
 				ASTOp::LogHistogram(roots);
 			}
 
-	//            // Third optimisation stage: introduce image compose conditions
-	//            {
-	//                // TODO: We have disable this since the new operation "layout remove block" is used.
-	//                // Proper support needs to be implemented for this case.
-	//    //                MUTABLE_CPUPROFILER_SCOPE(ImageComposeConditionGenerator);
-	//    //                AXE_INT_VALUE("Mutable", Verbose, "program size", (int64_t)program.m_opAddress.Num());
-	//    //                UE_LOG(LogMutableCore, Verbose, " - image compose conditions generator");
-	//    //                ImageComposeConditionGenerator composeGen( program );
-	//            }
-
-
-			//{
-			//	MUTABLE_CPUPROFILER_SCOPE(ConditionalsStage);
-
-			//	//program.Check();
-
-			//	// Optimise the conditionals added by the image compose condition generator
-			//	// also the constants generator at the end of previous fulloptimiseAST.
-			//	FullOptimiseAST( roots, pTaskManager );
-			//}
-
 			// Analyse mesh constants to see which of them are in optimised mesh formats, and set the flags.
 			ASTOp::Traverse_BottomUp_Unique_NonReentrant( roots, [&](Ptr<ASTOp>& n)
 			{
@@ -1302,7 +1254,7 @@ namespace mu
 				MUTABLE_CPUPROFILER_SCOPE(StatesStage);
 
 				// Optimise for every state
-				OptimiseStatesAST( pTaskManager );
+				OptimiseStatesAST( );
 
 				// Optimise the data formats (TODO)
 				//OperationFlagGenerator flagGen( pResult.get() );
@@ -1329,7 +1281,7 @@ namespace mu
 			for ( size_t s=0;  s<m_states.size(); ++s )
 			{
 				UE_LOG(LogMutableCore, Verbose, TEXT(" - constant generator"));
-				ConstantGeneratorAST( m_options->GetPrivate(), m_states[s].root, pTaskManager );
+				ConstantGeneratorAST( m_options->GetPrivate(), m_states[s].root );
 				//AXE_INT_VALUE("Mutable", Verbose, "ast size", (int64_t)ASTOp::CountNodes(roots));
 			}
 
diff --git a/Engine/Plugins/Experimental/Mutable/Source/MutableTools/Private/MuT/CodeOptimiser.h b/Engine/Plugins/Experimental/Mutable/Source/MutableTools/Private/MuT/CodeOptimiser.h
--- a/Engine/Plugins/Experimental/Mutable/Source/MutableTools/Private/MuT/CodeOptimiser.h
+++ b/Engine/Plugins/Experimental/Mutable/Source/MutableTools/Private/MuT/CodeOptimiser.h
@@ -8,7 +8,6 @@
 #include "MuR/CodeVisitor.h"
 #include "MuT/CompilerPrivate.h"
 #include "MuT/AST.h"
-#include "MuT/TaskManager.h"
 
 #include <memory>
 
@@ -24,7 +23,7 @@ namespace mu
         CodeOptimiser( CompilerOptionsPtr options, vector<STATE_COMPILATION_DATA>& states );
 
         //! Optimise the virtual machine code, using several transforms.
-        void OptimiseAST( TaskManager* );
+        void OptimiseAST();
 
     private:
 
@@ -39,12 +38,11 @@ namespace mu
         int m_optimizeIterationsLeft=0;
 
         //! Full optimisation pass
-        void FullOptimiseAST( ASTOpList& roots,
-                              TaskManager* pTaskManager );
+        void FullOptimiseAST( ASTOpList& roots );
 
         //! Optimise the code of a model for a specific state, generating new instructions and
         //! state information.
-        void OptimiseStatesAST( TaskManager* pTaskManager );
+        void OptimiseStatesAST();
 
     };
 
@@ -102,8 +100,7 @@ namespace mu
     //! ConstantGenerator replaces constant subtrees of operations with an equivalent single
 	//! constant value operation. 
     //---------------------------------------------------------------------------------------------
-    extern bool ConstantGeneratorAST( const CompilerOptions::Private* options, Ptr<ASTOp>& root,
-                                      TaskManager* pTaskManager );
+    extern bool ConstantGeneratorAST( const CompilerOptions::Private* options, Ptr<ASTOp>& root );
 
     //---------------------------------------------------------------------------------------------
     //! \TODO: shapes, projectors, others? but not switches (they must be unique)
diff --git a/Engine/Plugins/Experimental/Mutable/Source/MutableTools/Private/MuT/CodeOptimiser_Parameters.cpp b/Engine/Plugins/Experimental/Mutable/Source/MutableTools/Private/MuT/CodeOptimiser_Parameters.cpp
--- a/Engine/Plugins/Experimental/Mutable/Source/MutableTools/Private/MuT/CodeOptimiser_Parameters.cpp
+++ b/Engine/Plugins/Experimental/Mutable/Source/MutableTools/Private/MuT/CodeOptimiser_Parameters.cpp
@@ -20,6 +20,8 @@
 #include "MuT/ASTOpImageCompose.h"
 #include "MuT/ASTOpImageMipmap.h"
 #include "MuT/ASTOpImagePixelFormat.h"
+#include "MuT/ASTOpImageLayer.h"
+#include "MuT/ASTOpImageLayerColor.h"
 #include "MuT/ASTOpMeshMorph.h"
 #include "MuT/ASTOpInstanceAdd.h"
 #include "MuT/ASTOpLayoutFromMesh.h"
@@ -38,8 +40,6 @@
 
 namespace mu
 {
-class TaskManager;
-
 
     //---------------------------------------------------------------------------------------------
     //---------------------------------------------------------------------------------------------
@@ -461,7 +461,7 @@ class TaskManager;
                         {
                             optimised = true;
 
-                            auto typedYes = dynamic_cast<const ASTOpFixed*>(typedAt->yes.child().get());
+							const ASTOpImageLayerColor* typedYes = dynamic_cast<const ASTOpImageLayerColor*>(typedAt->yes.child().get());
 
                             Ptr<ASTOpFixed> blackOp = new ASTOpFixed;
                             blackOp->op.type = OP_TYPE::CO_CONSTANT;
@@ -479,24 +479,21 @@ class TaskManager;
                             Ptr<ASTOpFixed> resizeOp = new ASTOpFixed;
                             resizeOp->op.type = OP_TYPE::IM_RESIZELIKE;
                             resizeOp->SetChild( resizeOp->op.args.ImageResizeLike.source, plainOp );
-                            resizeOp->SetChild( resizeOp->op.args.ImageResizeLike.sizeSource,
-                                                typedYes->children[ typedYes->op.args.ImageLayerColour.base ] );
+                            resizeOp->SetChild( resizeOp->op.args.ImageResizeLike.sizeSource, typedYes->base );
 
                             auto maskOp = mu::Clone<ASTOpConditional>(typedAt);
                             maskOp->no = resizeOp;
 
                             // If there is no mask (because it is optional), we need to make a
                             // white plain image
-                            maskOp->yes = EnsureValidMask
-                                    ( typedYes->children[ typedYes->op.args.ImageLayerColour.mask ].child(),
-                                      typedYes->children[ typedYes->op.args.ImageLayerColour.base ].child() );
+                            maskOp->yes = EnsureValidMask( typedYes->mask.child(), typedYes->base.child() );
 
-                            auto baseOp = mu::Clone<ASTOpConditional>(typedAt);
-                            baseOp->yes = typedYes->children[ typedYes->op.args.ImageLayerColour.base ].child();
+							Ptr<ASTOpConditional> baseOp = mu::Clone<ASTOpConditional>(typedAt);
+                            baseOp->yes = typedYes->base.child();
 
-                            auto softOp = mu::Clone<ASTOpFixed>(typedYes);
-                            softOp->SetChild( softOp->op.args.ImageLayerColour.base, baseOp );
-                            softOp->SetChild( softOp->op.args.ImageLayerColour.mask, maskOp );
+                            Ptr<ASTOpImageLayerColor> softOp = mu::Clone<ASTOpImageLayerColor>(typedYes);
+                            softOp->base = baseOp;
+                            softOp->mask = maskOp;
 
                             at = softOp;
                             break;
@@ -566,7 +563,7 @@ class TaskManager;
                         {
                             optimised = true;
 
-                            auto typedNo = dynamic_cast<const ASTOpFixed*>(typedAt->no.child().get());
+                            const ASTOpImageLayerColor* typedNo = dynamic_cast<const ASTOpImageLayerColor*>(typedAt->no.child().get());
 
                             Ptr<ASTOpFixed> blackOp = new ASTOpFixed;
                             blackOp->op.type = OP_TYPE::CO_CONSTANT;
@@ -584,24 +581,21 @@ class TaskManager;
                             Ptr<ASTOpFixed> resizeOp = new ASTOpFixed;
                             resizeOp->op.type = OP_TYPE::IM_RESIZELIKE;
                             resizeOp->SetChild( resizeOp->op.args.ImageResizeLike.source, plainOp );
-                            resizeOp->SetChild( resizeOp->op.args.ImageResizeLike.sizeSource,
-                                                typedNo->children[ typedNo->op.args.ImageLayerColour.base ] );
+                            resizeOp->SetChild( resizeOp->op.args.ImageResizeLike.sizeSource, typedNo->base );
 
-                            auto maskOp = mu::Clone<ASTOpConditional>(typedAt);
+                            Ptr<ASTOpConditional> maskOp = mu::Clone<ASTOpConditional>(typedAt);
                             maskOp->no = resizeOp;
 
                             // If there is no mask (because it is optional), we need to make a
                             // white plain image
-                            maskOp->no = EnsureValidMask
-                                    ( typedNo->children[ typedNo->op.args.ImageLayerColour.mask ].child(),
-                                      typedNo->children[ typedNo->op.args.ImageLayerColour.base ].child() );
+                            maskOp->no = EnsureValidMask( typedNo->mask.child(), typedNo->base.child() );
 
-                            auto baseOp = mu::Clone<ASTOpConditional>(typedAt);
-                            baseOp->no = typedNo->children[ typedNo->op.args.ImageLayerColour.base ].child();
+                            Ptr<ASTOpConditional> baseOp = mu::Clone<ASTOpConditional>(typedAt);
+                            baseOp->no = typedNo->base.child();
 
-                            auto softOp = mu::Clone<ASTOpFixed>(typedNo);
-                            softOp->SetChild( softOp->op.args.ImageLayerColour.base, baseOp );
-                            softOp->SetChild( softOp->op.args.ImageLayerColour.mask, maskOp );
+                            Ptr<ASTOpImageLayerColor> softOp = mu::Clone<ASTOpImageLayerColor>(typedNo);
+                            softOp->base = baseOp;
+                            softOp->mask = maskOp;
 
                             at = softOp;
                             break;
@@ -885,23 +879,15 @@ class TaskManager;
                     {
                         optimised = true;
 
-                        auto typedBaseAt = dynamic_cast<const ASTOpFixed*>(baseAt.get());
-                        auto typedBlockAt = dynamic_cast<const ASTOpFixed*>(blockAt.get());
+						const ASTOpImageLayerColor* typedBaseAt = dynamic_cast<const ASTOpImageLayerColor*>(baseAt.get());
+						const ASTOpImageLayerColor* typedBlockAt = dynamic_cast<const ASTOpImageLayerColor*>(blockAt.get());
 
                         // The mask is a compose of the block mask on the base mask, but if none has
                         // a mask we don't need to make one.
-                        auto baseImage =
-                            typedBaseAt->children[typedBaseAt->op.args.ImageLayerColour.base]
-                                .child();
-                        auto baseMask =
-                            typedBaseAt->children[typedBaseAt->op.args.ImageLayerColour.mask]
-                                .child();
-                        auto blockImage =
-                            typedBlockAt->children[typedBlockAt->op.args.ImageLayerColour.base]
-                                .child();
-                        auto blockMask =
-                            typedBlockAt->children[typedBlockAt->op.args.ImageLayerColour.mask]
-                                .child();
+                        auto baseImage = typedBaseAt->base.child();
+                        auto baseMask = typedBaseAt->mask.child();
+                        auto blockImage = typedBlockAt->base.child();
+                        auto blockMask = typedBlockAt->mask.child();
 
                         Ptr<ASTOpImageCompose> maskOp;
                         if (baseMask || blockMask)
@@ -922,9 +908,9 @@ class TaskManager;
                         baseOp->Base = baseImage;
                         baseOp->BlockImage = blockImage;
 
-                        Ptr<ASTOpFixed> nop = mu::Clone<ASTOpFixed>(blockAt);
-                        nop->SetChild( nop->op.args.ImageLayerColour.mask, maskOp );
-                        nop->SetChild( nop->op.args.ImageLayerColour.base, baseOp );
+                        Ptr<ASTOpImageLayerColor> nop = mu::Clone<ASTOpImageLayerColor>(blockAt);
+                        nop->mask = maskOp;
+                        nop->base = baseOp;
 
                         // Done
                         at = nop;
@@ -935,29 +921,17 @@ class TaskManager;
                     {
                         optimised = true;
 
-                        auto typedBaseAt = dynamic_cast<const ASTOpFixed*>(baseAt.get());
-                        auto typedBlockAt = dynamic_cast<const ASTOpFixed*>(blockAt.get());
+						const ASTOpImageLayer* typedBaseAt = dynamic_cast<const ASTOpImageLayer*>(baseAt.get());
+						const ASTOpImageLayer* typedBlockAt = dynamic_cast<const ASTOpImageLayer*>(blockAt.get());
 
                         // The mask is a compose of the block mask on the base mask, but if none has
                         // a mask we don't need to make one.
-                        auto baseImage =
-                            typedBaseAt->children[typedBaseAt->op.args.ImageLayer.base]
-                                .child();
-                        auto baseBlended =
-                            typedBaseAt->children[typedBaseAt->op.args.ImageLayer.blended]
-                                .child();
-                        auto baseMask =
-                            typedBaseAt->children[typedBaseAt->op.args.ImageLayer.mask]
-                                .child();
-                        auto blockImage =
-                            typedBlockAt->children[typedBlockAt->op.args.ImageLayer.base]
-                                .child();
-                        auto blockBlended =
-                            typedBlockAt->children[typedBlockAt->op.args.ImageLayer.blended]
-                                .child();
-                        auto blockMask =
-                            typedBlockAt->children[typedBlockAt->op.args.ImageLayer.mask]
-                                .child();
+                        auto baseImage = typedBaseAt->base.child();
+                        auto baseBlended = typedBaseAt->blend.child();
+                        auto baseMask = typedBaseAt->mask.child();
+                        auto blockImage = typedBlockAt->base.child();
+                        auto blockBlended = typedBlockAt->blend.child();
+                        auto blockMask = typedBlockAt->mask.child();
 
                         Ptr<ASTOpImageCompose> maskOp;
                         if (baseMask || blockMask)
@@ -983,10 +957,10 @@ class TaskManager;
                         blendedOp->Base = baseBlended;
                         blendedOp->BlockImage = blockBlended;
 
-                        Ptr<ASTOpFixed> nop = mu::Clone<ASTOpFixed>(blockAt);
-                        nop->SetChild(nop->op.args.ImageLayer.mask, maskOp);
-                        nop->SetChild(nop->op.args.ImageLayer.base, baseOp);
-                        nop->SetChild(nop->op.args.ImageLayer.blended, blendedOp);
+                        Ptr<ASTOpImageLayer> nop = mu::Clone<ASTOpImageLayer>(blockAt);
+                        nop->mask = maskOp;
+                        nop->base = baseOp;
+                        nop->blend = blendedOp;
 
                         // Done
                         at = nop;
@@ -1042,12 +1016,10 @@ class TaskManager;
                 {
                     optimised = true;
 
-                    auto typedBlockAt = dynamic_cast<const ASTOpFixed*>(blockAt.get());
+					const ASTOpImageLayerColor* typedBlockAt = dynamic_cast<const ASTOpImageLayerColor*>(blockAt.get());
 
-                    auto blockImage =
-                        typedBlockAt->children[typedBlockAt->op.args.ImageLayerColour.base].child();
-                    auto blockMask =
-                        typedBlockAt->children[typedBlockAt->op.args.ImageLayerColour.mask].child();
+                    auto blockImage = typedBlockAt->base.child();
+                    auto blockMask = typedBlockAt->mask.child();
 
                     // The mask is a compose of the layer mask on a black image, however if there is
                     // no mask and the base of the layer opertation is a blanklayout, we can skip
@@ -1081,12 +1053,12 @@ class TaskManager;
                     }
 
                     // The base is composition of the layer base on the compose base
-                    auto baseOp = mu::Clone<ASTOpImageCompose>(at);
-                    baseOp->BlockImage = typedBlockAt->children[typedBlockAt->op.args.ImageLayerColour.base].child();
+                    Ptr<ASTOpImageCompose> baseOp = mu::Clone<ASTOpImageCompose>(at);
+                    baseOp->BlockImage = typedBlockAt->base.child();
 
-                    auto nop = mu::Clone<ASTOpFixed>(blockAt);
-                    nop->SetChild( nop->op.args.ImageLayerColour.mask, maskOp );
-                    nop->SetChild( nop->op.args.ImageLayerColour.base, baseOp );
+                    Ptr<ASTOpImageLayerColor> nop = mu::Clone<ASTOpImageLayerColor>(blockAt);
+                    nop->mask = maskOp;
+                    nop->base = baseOp;
 
                     // Done
                     at = nop;
@@ -1097,14 +1069,11 @@ class TaskManager;
                 {
                     optimised = true;
 
-                    auto typedBlockAt = dynamic_cast<const ASTOpFixed*>(blockAt.get());
+                    const ASTOpImageLayer* typedBlockAt = dynamic_cast<const ASTOpImageLayer*>(blockAt.get());
 
-                    auto blockImage =
-                        typedBlockAt->children[typedBlockAt->op.args.ImageLayer.base].child();
-                    auto blockBlended =
-                         typedBlockAt->children[typedBlockAt->op.args.ImageLayer.blended].child();
-                    auto blockMask =
-                        typedBlockAt->children[typedBlockAt->op.args.ImageLayer.mask].child();
+                    auto blockImage = typedBlockAt->base.child();
+                    auto blockBlended = typedBlockAt->blend.child();
+                    auto blockMask = typedBlockAt->mask.child();
 
                     // The mask is a compose of the layer mask on a black image, however if there is
                     // no mask and the base of the layer opertation is a blanklayout, we can skip
@@ -1166,13 +1135,13 @@ class TaskManager;
                     }
 
                     // The base is composition of the softlight base on the compose base
-                    auto baseOp = mu::Clone<ASTOpImageCompose>(at);
-                    baseOp->BlockImage = typedBlockAt->children[typedBlockAt->op.args.ImageLayer.base].child();
+                    Ptr<ASTOpImageCompose> baseOp = mu::Clone<ASTOpImageCompose>(at);
+                    baseOp->BlockImage = typedBlockAt->base.child();
 
-                    auto nop = mu::Clone<ASTOpFixed>(blockAt);
-                    nop->SetChild( nop->op.args.ImageLayer.base, baseOp );
-                    nop->SetChild( nop->op.args.ImageLayer.mask, maskOp );
-                    nop->SetChild( nop->op.args.ImageLayer.blended, blendedOp );
+                    Ptr<ASTOpImageLayer> nop = mu::Clone<ASTOpImageLayer>(blockAt);
+                    nop->base = baseOp;
+                    nop->mask = maskOp;
+                    nop->blend = blendedOp;
 
                     // Done
                     at = nop;
@@ -1222,7 +1191,7 @@ class TaskManager;
                 {
                     optimised = true;
 
-                    auto typedBaseAt = dynamic_cast<const ASTOpFixed*>(baseAt.get());
+                    const ASTOpImageLayerColor* typedBaseAt = dynamic_cast<const ASTOpImageLayerColor*>(baseAt.get());
 
                     auto maskOp = mu::Clone<ASTOpImageCompose>(at);
                     {
@@ -1245,20 +1214,18 @@ class TaskManager;
                         blockResizeOp->SetChild( blockResizeOp->op.args.ImageResizeLike.source, plainOp );
 
                         // Blank out the block from the mask
-                        auto newMaskBase = EnsureValidMask(
-                                    typedBaseAt->children[ typedBaseAt->op.args.ImageLayerColour.mask].child(),
-                                    baseAt );
+                        auto newMaskBase = EnsureValidMask( typedBaseAt->mask.child(), baseAt );
                         maskOp->Base = newMaskBase;
                         maskOp->BlockImage = blockResizeOp;
                     }
 
                     // The base is composition of the softlight base on the compose base
                     auto baseOp = mu::Clone<ASTOpImageCompose>(at);
-                    baseOp->Base = typedBaseAt->children[typedBaseAt->op.args.ImageLayerColour.base].child();
+                    baseOp->Base = typedBaseAt->base.child();
 
-                    auto nop = mu::Clone<ASTOpFixed>(baseAt);
-                    nop->SetChild( nop->op.args.ImageLayerColour.base, baseOp );
-                    nop->SetChild( nop->op.args.ImageLayerColour.mask, maskOp );
+                    Ptr<ASTOpImageLayerColor> nop = mu::Clone<ASTOpImageLayerColor>(baseAt);
+                    nop->base = baseOp;
+                    nop->mask = maskOp;
 
                     // Done
                     at = nop;
@@ -1269,7 +1236,7 @@ class TaskManager;
                 {
                     optimised = true;
 
-                    auto typedBaseAt = dynamic_cast<const ASTOpFixed*>(baseAt.get());
+                    const ASTOpImageLayer* typedBaseAt = dynamic_cast<const ASTOpImageLayer*>(baseAt.get());
 
                     auto maskOp = mu::Clone<ASTOpImageCompose>(at);
                     {
@@ -1292,20 +1259,18 @@ class TaskManager;
                         blockResizeOp->SetChild( blockResizeOp->op.args.ImageResizeLike.source, plainOp );
 
                         // Blank out the block from the mask
-                        auto newMaskBase = EnsureValidMask(
-                                    typedBaseAt->children[ typedBaseAt->op.args.ImageLayer.mask ].child(),
-                                    baseAt );
+                        Ptr<ASTOp> newMaskBase = EnsureValidMask( typedBaseAt->mask.child(), baseAt );
                         maskOp->Base = newMaskBase;
                         maskOp->BlockImage = blockResizeOp;
                     }
 
                     // The base is composition of the effect base on the compose base
-                    auto baseOp = mu::Clone<ASTOpImageCompose>(at);
-                    baseOp->Base = typedBaseAt->children[typedBaseAt->op.args.ImageLayer.base].child();
+                    Ptr<ASTOpImageCompose> baseOp = mu::Clone<ASTOpImageCompose>(at);
+                    baseOp->Base = typedBaseAt->base.child();
 
-                    auto nop = mu::Clone<ASTOpFixed>(baseAt);
-                    nop->SetChild( nop->op.args.ImageLayer.base, baseOp );
-                    nop->SetChild( nop->op.args.ImageLayer.mask, maskOp );
+                    Ptr<ASTOpImageLayer> nop = mu::Clone<ASTOpImageLayer>(baseAt);
+                    nop->base = baseOp;
+                    nop->mask = maskOp;
 
                     // Done
                     at = nop;
@@ -1451,27 +1416,26 @@ class TaskManager;
             {
             case OP_TYPE::IM_LAYERCOLOUR:
             {
-                auto typedSource = dynamic_cast<const ASTOpFixed*>(sourceOp.get());
+                const ASTOpImageLayerColor* typedSource = dynamic_cast<const ASTOpImageLayerColor*>(sourceOp.get());
 
-                bool colourHasRuntime = m_hasRuntimeParamVisitor.HasAny
-                        ( typedSource->children[ typedSource->op.args.ImageLayerColour.colour].child() );
+                bool colourHasRuntime = m_hasRuntimeParamVisitor.HasAny( typedSource->color.child() );
 
                 if (colourHasRuntime)
                 {
                     m_modified = true;
 
-                    auto top = mu::Clone<ASTOpFixed>(sourceOp);
+                    auto top = mu::Clone<ASTOpImageLayerColor>(sourceOp);
 
                     auto baseOp = mu::Clone<ASTOpImageMipmap>(at);
-                    baseOp->Source = typedSource->children[ typedSource->op.args.ImageLayerColour.base ].child();
-                    top->SetChild( top->op.args.ImageLayerColour.base, baseOp );
+                    baseOp->Source = typedSource->base.child();
+                    top->base = baseOp;
 
-                    auto sourceMaskOp = typedSource->children[typedSource->op.args.ImageLayerColour.mask].child();
+                    auto sourceMaskOp = typedSource->mask.child();
                     if (sourceMaskOp)
                     {
                         auto maskOp = mu::Clone<ASTOpImageMipmap>(at);
                         maskOp->Source = sourceMaskOp;
-                        top->SetChild( top->op.args.ImageLayerColour.mask, maskOp );
+                        top->mask = maskOp;
                     }
 
                     at = top;
@@ -1545,29 +1509,29 @@ class TaskManager;
             // TODO: Code shared with the constant data format optimisation visitor
             case OP_TYPE::IM_LAYERCOLOUR:
             {
-                auto typedAt = dynamic_cast<const ASTOpFixed*>(at.get());
+                const ASTOpImageLayerColor* typedAt = dynamic_cast<const ASTOpImageLayerColor*>(at.get());
 
-                RecurseWithCurrentState( typedAt->children[typedAt->op.args.ImageLayerColour.base].child() );
-                RecurseWithCurrentState( typedAt->children[typedAt->op.args.ImageLayerColour.colour].child() );
+                RecurseWithCurrentState( typedAt->base.child() );
+                RecurseWithCurrentState( typedAt->color.child() );
 
-                if ( typedAt->children[typedAt->op.args.ImageLayerColour.mask] )
+                if ( typedAt->mask )
                 {
                     std::array<uint8_t, (size_t)EImageFormat::IF_COUNT> newState;
                     newState.fill(0);
                     newState[(size_t)EImageFormat::IF_L_UBYTE ] = 1;
                     newState[(size_t)EImageFormat::IF_L_UBYTE_RLE ] = 1;
 
-                    RecurseWithState( typedAt->children[typedAt->op.args.ImageLayerColour.mask].child(), newState );
+                    RecurseWithState( typedAt->mask.child(), newState );
                 }
                 break;
             }
 
             case OP_TYPE::IM_LAYER:
             {
-                auto typedAt = dynamic_cast<const ASTOpFixed*>(at.get());
+                const ASTOpImageLayer* typedAt = dynamic_cast<const ASTOpImageLayer*>(at.get());
 
-                RecurseWithCurrentState( typedAt->children[typedAt->op.args.ImageLayer.base].child() );
-                RecurseWithCurrentState( typedAt->children[typedAt->op.args.ImageLayer.blended].child() );
+                RecurseWithCurrentState( typedAt->base.child() );
+                RecurseWithCurrentState( typedAt->blend.child() );
 
                 std::array<uint8_t, (size_t)EImageFormat::IF_COUNT> newState;
                 newState.fill(0);
@@ -1575,9 +1539,9 @@ class TaskManager;
                 //newState[ IF_L_UBYTE ] = 1;
                 //newState[ IF_L_UBYTE_RLE ] = 1;
 
-                if ( typedAt->op.args.ImageLayer.mask )
+                if ( typedAt->mask )
                 {
-                    RecurseWithState( typedAt->children[typedAt->op.args.ImageLayer.mask].child(), newState );
+                    RecurseWithState( typedAt->mask.child(), newState );
                 }
                 break;
             }
@@ -2019,7 +1983,7 @@ class TaskManager;
     //---------------------------------------------------------------------------------------------
     //---------------------------------------------------------------------------------------------
     //---------------------------------------------------------------------------------------------
-    void CodeOptimiser::OptimiseStatesAST( TaskManager* pTaskManager )
+    void CodeOptimiser::OptimiseStatesAST()
     {
         MUTABLE_CPUPROFILER_SCOPE(OptimiseStatesAST);
 
@@ -2185,7 +2149,7 @@ class TaskManager;
             for(auto& s:m_states)
             {
                 UE_LOG(LogMutableCore, Verbose, TEXT(" - constant optimiser"));
-                ConstantGeneratorAST( m_options->GetPrivate(), s.root, pTaskManager );
+                ConstantGeneratorAST( m_options->GetPrivate(), s.root );
             }
 
             TArray<Ptr<ASTOp>> roots;
diff --git a/Engine/Plugins/Experimental/Mutable/Source/MutableTools/Private/MuT/CodeOptimiser_Semantic.cpp b/Engine/Plugins/Experimental/Mutable/Source/MutableTools/Private/MuT/CodeOptimiser_Semantic.cpp
--- a/Engine/Plugins/Experimental/Mutable/Source/MutableTools/Private/MuT/CodeOptimiser_Semantic.cpp
+++ b/Engine/Plugins/Experimental/Mutable/Source/MutableTools/Private/MuT/CodeOptimiser_Semantic.cpp
@@ -16,6 +16,8 @@
 #include "MuT/ASTOpConstantBool.h"
 #include "MuT/ASTOpImageCompose.h"
 #include "MuT/ASTOpImageMipmap.h"
+#include "MuT/ASTOpImageLayer.h"
+#include "MuT/ASTOpImageLayerColor.h"
 #include "MuT/ASTOpImageMultiLayer.h"
 #include "MuT/ASTOpImagePatch.h"
 #include "MuT/ASTOpImagePixelFormat.h"
@@ -280,123 +282,13 @@ namespace mu
 
         //-------------------------------------------------------------------------------------
         case OP_TYPE::IM_PIXELFORMAT:
-        case OP_TYPE::IM_SWIZZLE:
-        {
+		case OP_TYPE::IM_SWIZZLE:
+		case OP_TYPE::IM_LAYER:
+		{
 			check(false); // moved to its own operation class
 			break;
 		}
 
-
-        //-------------------------------------------------------------------------------------
-        case OP_TYPE::IM_LAYER:
-        {
-            auto baseAt = children[op.args.ImageLayer.base].child();
-
-            // Plain masks optimization
-            auto maskAt = children[op.args.ImageLayer.mask].child();
-            if ( maskAt )
-            {
-                FVector4f colour;
-                if ( maskAt->IsImagePlainConstant( colour ))
-                {
-                    if ( colour.IsNearlyZero3(UE_SMALL_NUMBER) )
-                    {
-                        // If the mask is black, we can skip the entire operation
-                        at = children[op.args.ImageLayer.base].child();
-                    }
-                    else if ( colour.Equals( FVector4f(1,1,1,1), UE_SMALL_NUMBER ) )
-                    {
-                        // If the mask is white, we can remove it
-                        auto nop = mu::Clone<ASTOpFixed>(this);
-                        nop->op.args.ImageLayer.mask = 0;
-                        at = nop;
-                    }
-                }
-            }
-
-            // Introduce crop if mask is constant and smaller than the base
-            FImageRect sourceMaskUsage;
-            FImageDesc maskDesc;
-            bool validUsageRect = false;
-            if (!at && maskAt)
-            {
-                //MUTABLE_CPUPROFILER_SCOPE(EvaluateAreasForCrop);
-
-                validUsageRect = maskAt->GetNonBlackRect(sourceMaskUsage);
-                if (validUsageRect)
-                {
-					check(sourceMaskUsage.size[0] > 0);
-					check(sourceMaskUsage.size[1] > 0);
-					
-					FGetImageDescContext context;
-                    maskDesc = maskAt->GetImageDesc( false, &context );
-                }
-            }
-
-            if (!at && maskAt && validUsageRect)
-            {
-                // Adjust for compressed blocks (4), and some extra mips (2 more mips, which is 4)
-                constexpr int blockSize = 4 * 4;
-
-                FImageRect maskUsage;
-                maskUsage.min[0] = (sourceMaskUsage.min[0]/blockSize)*blockSize;
-                maskUsage.min[1] = (sourceMaskUsage.min[1]/blockSize)*blockSize;
-                vec2<uint16> minOffset = sourceMaskUsage.min - maskUsage.min;
-                maskUsage.size[0] = ((sourceMaskUsage.size[0]+minOffset[0]+blockSize-1)/blockSize)*blockSize;
-                maskUsage.size[1] = ((sourceMaskUsage.size[1]+minOffset[1]+blockSize-1)/blockSize)*blockSize;
-
-                // Is it worth?
-                float ratio = float(maskUsage.size[0]*maskUsage.size[1])
-                        / float(maskDesc.m_size[0]*maskDesc.m_size[1]);
-                float acceptableCropRatio = options.m_acceptableCropRatio;
-                if (ratio<acceptableCropRatio)
-                {
-					check(maskUsage.size[0] > 0);
-					check(maskUsage.size[1] > 0);
-
-                    Ptr<ASTOpFixed> cropMask = new ASTOpFixed();
-                    cropMask->op.type = OP_TYPE::IM_CROP;
-                    cropMask->SetChild( cropMask->op.args.ImageCrop.source, children[op.args.ImageLayer.mask].child() );
-                    cropMask->op.args.ImageCrop.minX = maskUsage.min[0];
-                    cropMask->op.args.ImageCrop.minY = maskUsage.min[1];
-                    cropMask->op.args.ImageCrop.sizeX = maskUsage.size[0];
-                    cropMask->op.args.ImageCrop.sizeY = maskUsage.size[1];
-
-                    Ptr<ASTOpFixed> cropBlended = new ASTOpFixed();
-                    cropBlended->op.type = OP_TYPE::IM_CROP;
-                    cropBlended->SetChild( cropBlended->op.args.ImageCrop.source, children[op.args.ImageLayer.blended].child() );
-                    cropBlended->op.args.ImageCrop.minX = maskUsage.min[0];
-                    cropBlended->op.args.ImageCrop.minY = maskUsage.min[1];
-                    cropBlended->op.args.ImageCrop.sizeX = maskUsage.size[0];
-                    cropBlended->op.args.ImageCrop.sizeY = maskUsage.size[1];
-
-                    Ptr<ASTOpFixed> cropBase = new ASTOpFixed();
-                    cropBase->op.type = OP_TYPE::IM_CROP;
-                    cropBase->SetChild( cropBase->op.args.ImageCrop.source, children[op.args.ImageLayer.base].child() );
-                    cropBase->op.args.ImageCrop.minX = maskUsage.min[0];
-                    cropBase->op.args.ImageCrop.minY = maskUsage.min[1];
-                    cropBase->op.args.ImageCrop.sizeX = maskUsage.size[0];
-                    cropBase->op.args.ImageCrop.sizeY = maskUsage.size[1];
-
-                    Ptr<ASTOpFixed> newLayer = new ASTOpFixed();
-                    newLayer->op.type = op.type;
-					newLayer->op.args.ImageLayer.blendType = op.args.ImageLayer.blendType;
-					newLayer->op.args.ImageLayer.flags = op.args.ImageLayer.flags;
-					newLayer->SetChild( newLayer->op.args.ImageLayer.base, cropBase );
-                    newLayer->SetChild( newLayer->op.args.ImageLayer.blended, cropBlended );
-                    newLayer->SetChild( newLayer->op.args.ImageLayer.mask, cropMask );
-
-                    Ptr<ASTOpImagePatch> patch = new ASTOpImagePatch();
-                    patch->base = baseAt;
-                    patch->patch = newLayer;
-                    patch->location = maskUsage.min;
-                    at = patch;
-                }
-            }
-
-            break;
-        }
-
         default:
             break;
         }
@@ -615,16 +507,16 @@ namespace mu
         case OP_TYPE::IM_LAYER:
         {
             // We move the op down the arguments
-            auto nop = mu::Clone<ASTOpFixed>(at);
+            auto nop = mu::Clone<ASTOpImageLayer>(at);
 
-            auto aOp = nop->children[nop->op.args.ImageLayer.base].child();
-            nop->SetChild(nop->op.args.ImageLayer.base, Visit(aOp, currentCropOp));
+            auto aOp = nop->base.child();
+            nop->base = Visit(aOp, currentCropOp);
 
-            auto bOp = nop->children[nop->op.args.ImageLayer.blended].child();
-            nop->SetChild(nop->op.args.ImageLayer.blended, Visit(bOp, currentCropOp));
+            auto bOp = nop->blend.child();
+            nop->blend = Visit(bOp, currentCropOp);
 
-            auto mOp = nop->children[nop->op.args.ImageLayer.mask].child();
-            nop->SetChild(nop->op.args.ImageLayer.mask, Visit(mOp, currentCropOp));
+            auto mOp = nop->mask.child();
+            nop->mask = Visit(mOp, currentCropOp);
 
             newAt = nop;
             break;
@@ -633,13 +525,13 @@ namespace mu
         case OP_TYPE::IM_LAYERCOLOUR:
         {
             // We move the op down the arguments
-            auto nop = mu::Clone<ASTOpFixed>(at);
+            auto nop = mu::Clone<ASTOpImageLayerColor>(at);
 
-            auto aOp = nop->children[nop->op.args.ImageLayerColour.base].child();
-            nop->SetChild(nop->op.args.ImageLayerColour.base, Visit(aOp, currentCropOp));
+            auto aOp = nop->base.child();
+            nop->base = Visit(aOp, currentCropOp);
 
-            auto mOp = nop->children[nop->op.args.ImageLayerColour.mask].child();
-            nop->SetChild(nop->op.args.ImageLayerColour.mask, Visit(mOp, currentCropOp));
+            auto mOp = nop->mask.child();
+            nop->mask = Visit(mOp, currentCropOp);
 
             newAt = nop;
             break;
@@ -960,248 +852,6 @@ namespace mu
             break;
         }
 
-        //-------------------------------------------------------------------------------------
-        case OP_TYPE::IM_LAYER:
-        {
-            // Layer effects may be worth sinking down switches and conditionals, to be able
-            // to apply extra optimisations
-            auto baseAt = children[op.args.ImageLayer.base].child();
-            auto blendAt = children[op.args.ImageLayer.blended].child();
-            auto maskAt = children[op.args.ImageLayer.mask].child();
-
-            // Promote conditions from the base
-            OP_TYPE baseType = baseAt->GetOpType();
-            switch ( baseType )
-            {
-            // Seems to cause operation explosion in optimizer in bandit model.
-            // moved to generic sink in the default.
-//            case OP_TYPE::IM_CONDITIONAL:
-//            {
-//                m_modified = true;
-
-//                OP op = program.m_code[baseAt];
-
-//                OP aOp = program.m_code[at];
-//                aOp.args.ImageLayer.base = program.m_code[baseAt].args.Conditional.yes;
-//                op.args.Conditional.yes = program.AddOp( aOp );
-
-//                OP bOp = program.m_code[at];
-//                bOp.args.ImageLayer.base = program.m_code[baseAt].args.Conditional.no;
-//                op.args.Conditional.no = program.AddOp( bOp );
-
-//                at = program.AddOp( op );
-//                break;
-//            }
-
-            case OP_TYPE::IM_SWITCH:
-            {
-                // Disabled:
-                // It seems to cause data explosion in optimizer in some models. Because
-                // all switch branches become unique constants
-
-//                // See if the blended has an identical switch, to optimise it too
-//                auto baseSwitch = dynamic_cast<const ASTOpSwitch*>( baseAt.get() );
-//                auto blendedSwitch = dynamic_cast<const ASTOpSwitch*>( blendAt.get() );
-//                auto maskSwitch = dynamic_cast<const ASTOpSwitch*>( maskAt.get() );
-
-//                bool blendedToo = baseSwitch->IsCompatibleWith( blendedSwitch );
-//                bool maskToo = baseSwitch->IsCompatibleWith( maskSwitch );
-
-//                // Move the layer operation down all the paths
-//                auto nop = mu::Clone<ASTOpSwitch>(baseSwitch);
-
-//                if (nop->def)
-//                {
-//                    auto defOp = mu::Clone<ASTOpFixed>(this);
-//                    defOp->SetChild( defOp->op.args.ImageLayer.base, nop->def );
-//                    if (blendedToo)
-//                    {
-//                        defOp->SetChild( defOp->op.args.ImageLayer.blended, blendedSwitch->def );
-//                    }
-//                    if (maskToo)
-//                    {
-//                        defOp->SetChild( defOp->op.args.ImageLayer.mask, maskSwitch->def );
-//                    }
-//                    nop->def = defOp;
-//                }
-
-//                for ( size_t v=0; v<nop->cases.Num(); ++v )
-//                {
-//                    if ( nop->cases[v].branch )
-//                    {
-//                        auto bOp = mu::Clone<ASTOpFixed>(this);
-//                        bOp->SetChild( bOp->op.args.ImageLayer.base, nop->cases[v].branch );
-
-//                        if (blendedToo)
-//                        {
-//                            bOp->SetChild( bOp->op.args.ImageLayer.blended, blendedSwitch->FindBranch( nop->cases[v].condition ) );
-//                        }
-
-//                        if (maskToo)
-//                        {
-//                            bOp->SetChild( bOp->op.args.ImageLayer.mask, maskSwitch->FindBranch( nop->cases[v].condition ) );
-//                        }
-
-//                        nop->cases[v].branch = bOp;
-//                    }
-//                }
-
-//                at = nop;
-                break;
-            }
-
-            default:
-            {
-                // Try generic base sink
-//                OP templateOp = program.m_code[at];
-
-//                CloneNeutralPartialTreeVisitor neutral;
-//                OP_TYPE supportedOps[] = {
-//                    OP_TYPE::IM_CONDITIONAL,
-//                    OP_TYPE::NONE
-//                };
-//                auto newRoot = neutral.Apply
-//                         (
-//                            program.m_code[at].args.ImageLayer.base, program,
-//                            DT_IMAGE, supportedOps,
-//                            &templateOp, &templateOp.args.ImageLayer.base);
-
-//                // If there is a change
-//                if (newRoot!=program.m_code[at].args.ImageLayer.base)
-//                {
-//                    m_modified = true;
-//                    at = newRoot;
-//                }
-
-                break;
-            }
-
-            }
-
-            break;
-        }
-
-
-
-		//-------------------------------------------------------------------------------------
-        case OP_TYPE::IM_LAYERCOLOUR:
-		{
-			// Layer effects may be worth sinking down switches and conditionals, to be able
-			// to apply extra optimisations
-			auto baseAt = children[op.args.ImageLayerColour.base].child();
-			auto maskAt = children[op.args.ImageLayerColour.mask].child();
-
-			// Promote conditions from the base
-			OP_TYPE baseType = baseAt->GetOpType();
-			switch (baseType)
-			{
-				// Seems to cause operation explosion in optimizer in bandit model.
-				// moved to generic sink in the default.
-	//            case OP_TYPE::IM_CONDITIONAL:
-	//            {
-	//                m_modified = true;
-
-	//                OP op = program.m_code[baseAt];
-
-	//                OP aOp = program.m_code[at];
-	//                aOp.args.ImageLayerColour.base = program.m_code[baseAt].args.Conditional.yes;
-	//                op.args.Conditional.yes = program.AddOp( aOp );
-
-	//                OP bOp = program.m_code[at];
-	//                bOp.args.ImageLayerColour.base = program.m_code[baseAt].args.Conditional.no;
-	//                op.args.Conditional.no = program.AddOp( bOp );
-
-	//                at = program.AddOp( op );
-	//                break;
-	//            }
-
-			case OP_TYPE::IM_SWITCH:
-			{
-				// Warning:
-				// It seems to cause data explosion in optimizer in some models. Because
-				// all switch branches become unique constants
-
-                // See if the blended has an identical switch, to optimise it too
-				const ASTOpSwitch* baseSwitch = dynamic_cast<const ASTOpSwitch*>( baseAt.get() );
-
-                // Mask not supported yet
-				if (maskAt)
-				{
-					break;
-				}
-
-                // Move the layer operation down base paths
-                Ptr<ASTOpSwitch> nop = mu::Clone<ASTOpSwitch>(baseSwitch);
-
-                if (nop->def)
-                {
-					Ptr<ASTOpFixed> defOp = mu::Clone<ASTOpFixed>(this);
-                    defOp->SetChild( defOp->op.args.ImageLayerColour.base, nop->def );
-                    nop->def = defOp;
-                }
-
-                for ( int32 v=0; v<nop->cases.Num(); ++v )
-                {
-                    if ( nop->cases[v].branch )
-                    {
-                        Ptr<ASTOpFixed> bOp = mu::Clone<ASTOpFixed>(this);
-                        bOp->SetChild( bOp->op.args.ImageLayerColour.base, nop->cases[v].branch );
-                        nop->cases[v].branch = bOp;
-                    }
-                }
-
-                at = nop;
-				break;
-			}
-
-			case OP_TYPE::IM_DISPLACE:
-			{
-				// Mask not supported yet. If there is a mask it wouldn't be correct to sink
-				// unless the mask was a similar displace.
-				if (maskAt)
-				{
-					break;
-				}
-
-				Ptr<ASTOpFixed> NewDisplace = mu::Clone<ASTOpFixed>(baseAt);
-
-				Ptr<ASTOp> sourceOp = NewDisplace->children[NewDisplace->op.args.ImageDisplace.source].child();
-				Ptr<ASTOpFixed> NewSource = mu::Clone<ASTOpFixed>(this);
-				NewSource->SetChild(NewSource->op.args.ImageLayerColour.base, sourceOp);
-				NewDisplace->SetChild(NewDisplace->op.args.ImageDisplace.source, NewSource);
-
-				at = NewDisplace;
-				break;
-			}
-
-			case OP_TYPE::IM_RASTERMESH:
-			{
-				// Mask not supported yet. If there is a mask it wouldn't be correct to sink.				
-				if (maskAt)
-				{
-					break;
-				}
-
-				Ptr<ASTOpFixed> NewRaster = mu::Clone<ASTOpFixed>(baseAt);
-
-				Ptr<ASTOp> sourceOp = NewRaster->children[NewRaster->op.args.ImageRasterMesh.image].child();
-				Ptr<ASTOpFixed> NewSource = mu::Clone<ASTOpFixed>(this);
-				NewSource->SetChild(NewSource->op.args.ImageLayerColour.base, sourceOp);
-				NewRaster->SetChild(NewRaster->op.args.ImageRasterMesh.image, NewSource);
-
-				at = NewRaster;
-				break;
-			}
-
-
-			default:
-				break;
-
-			}
-
-			break;
-		}
-
 
         //-----------------------------------------------------------------------------------------
         //-----------------------------------------------------------------------------------------
@@ -1585,16 +1235,16 @@ namespace mu
 
 			case OP_TYPE::IM_LAYER:
 			{
-				auto nop = mu::Clone<ASTOpFixed>(at);
+				auto nop = mu::Clone<ASTOpImageLayer>(at);
 
-				auto aOp = nop->children[nop->op.args.ImageLayer.base].child();
-				nop->SetChild(nop->op.args.ImageLayer.base, Visit(aOp, currentSinkingOp));
+				auto aOp = nop->base.child();
+				nop->base = Visit(aOp, currentSinkingOp);
 
-				auto bOp = nop->children[nop->op.args.ImageLayer.blended].child();
-				nop->SetChild(nop->op.args.ImageLayer.blended, Visit(bOp, currentSinkingOp));
+				auto bOp = nop->blend.child();
+				nop->blend = Visit(bOp, currentSinkingOp);
 
-				auto mOp = nop->children[nop->op.args.ImageLayer.mask].child();
-				nop->SetChild(nop->op.args.ImageLayer.mask, Visit(mOp, currentSinkingOp));
+				auto mOp = nop->mask.child();
+				nop->mask = Visit(mOp, currentSinkingOp);
 
 				newAt = nop;
 				break;
@@ -1602,13 +1252,13 @@ namespace mu
 
 			case OP_TYPE::IM_LAYERCOLOUR:
 			{
-				auto nop = mu::Clone<ASTOpFixed>(at);
+				auto nop = mu::Clone<ASTOpImageLayerColor>(at);
 
-				auto aOp = nop->children[nop->op.args.ImageLayerColour.base].child();
-				nop->SetChild(nop->op.args.ImageLayerColour.base, Visit(aOp, currentSinkingOp));
+				auto aOp = nop->base.child();
+				nop->base = Visit(aOp, currentSinkingOp);
 
-				auto mOp = nop->children[nop->op.args.ImageLayerColour.mask].child();
-				nop->SetChild(nop->op.args.ImageLayer.mask, Visit(mOp, currentSinkingOp));
+				auto mOp = nop->mask.child();
+				nop->mask = Visit(mOp, currentSinkingOp);
 
 				newAt = nop;
 				break;
@@ -1874,24 +1524,22 @@ namespace mu
 
             case OP_TYPE::IM_LAYER:
             {
-                Ptr<ASTOpFixed> newOp = mu::Clone<ASTOpFixed>(sourceAt);
+                Ptr<ASTOpImageLayer> newOp = mu::Clone<ASTOpImageLayer>(sourceAt);
 
                 Ptr<ASTOpFixed> baseOp = mu::Clone<ASTOpFixed>(this);
-                baseOp->SetChild( baseOp->op.args.ImageResize.source,
-                        newOp->children[newOp->op.args.ImageLayer.base] );
-                newOp->SetChild( newOp->op.args.ImageLayer.base, baseOp );
+                baseOp->SetChild( baseOp->op.args.ImageResize.source, newOp->base );
+                newOp->base = baseOp;
 
                 Ptr<ASTOpFixed> blendOp = mu::Clone<ASTOpFixed>(this);
-                blendOp->SetChild( blendOp->op.args.ImageResize.source,
-                        newOp->children[newOp->op.args.ImageLayer.blended] );
-                newOp->SetChild( newOp->op.args.ImageLayer.blended, blendOp );
+                blendOp->SetChild( blendOp->op.args.ImageResize.source, newOp->blend );
+                newOp->blend = blendOp;
 
-                auto maskAt = newOp->children[newOp->op.args.ImageLayer.mask].child();
+                auto maskAt = newOp->mask.child();
                 if (maskAt)
                 {
                     Ptr<ASTOpFixed> maskOp = mu::Clone<ASTOpFixed>(this);
                     maskOp->SetChild( maskOp->op.args.ImageResize.source, maskAt );
-                    newOp->SetChild( newOp->op.args.ImageLayer.mask, maskOp );
+                    newOp->mask = maskOp;
                 }
 
                 at = newOp;
@@ -1900,19 +1548,18 @@ namespace mu
 
             case OP_TYPE::IM_LAYERCOLOUR:
             {
-                Ptr<ASTOpFixed> newOp = mu::Clone<ASTOpFixed>(sourceAt);
+                Ptr<ASTOpImageLayerColor> newOp = mu::Clone<ASTOpImageLayerColor>(sourceAt);
 
                 Ptr<ASTOpFixed> baseOp = mu::Clone<ASTOpFixed>(this);
-                baseOp->SetChild( baseOp->op.args.ImageResize.source,
-                        newOp->children[newOp->op.args.ImageLayerColour.base] );
-                newOp->SetChild( newOp->op.args.ImageLayerColour.base, baseOp );
+                baseOp->SetChild( baseOp->op.args.ImageResize.source, newOp->base );
+                newOp->base = baseOp;
 
-                auto maskAt = newOp->children[newOp->op.args.ImageLayerColour.mask].child();
+                auto maskAt = newOp->mask.child();
                 if (maskAt)
                 {
                     Ptr<ASTOpFixed> maskOp = mu::Clone<ASTOpFixed>(this);
                     maskOp->SetChild( maskOp->op.args.ImageResize.source, maskAt );
-                    newOp->SetChild( newOp->op.args.ImageLayerColour.mask, maskOp );
+                    newOp->mask = maskOp;
                 }
 
                 at = newOp;
diff --git a/Engine/Plugins/Experimental/Mutable/Source/MutableTools/Private/MuT/Compiler.cpp b/Engine/Plugins/Experimental/Mutable/Source/MutableTools/Private/MuT/Compiler.cpp
--- a/Engine/Plugins/Experimental/Mutable/Source/MutableTools/Private/MuT/Compiler.cpp
+++ b/Engine/Plugins/Experimental/Mutable/Source/MutableTools/Private/MuT/Compiler.cpp
@@ -32,7 +32,6 @@
 #include "MuT/Node.h"
 #include "MuT/NodePrivate.h"
 #include "MuT/Table.h"
-#include "MuT/TaskManager.h"
 #include "Trace/Detail/Channel.h"
 
 #include <algorithm>
@@ -106,20 +105,6 @@ namespace mu
     }
 
 
-	//---------------------------------------------------------------------------------------------
-	void CompilerOptions::SetEnablePartialOptimisation(bool bEnabled)
-	{
-		m_pD->m_enablePartialOptimise = bEnabled;
-	}
-
-
-	//---------------------------------------------------------------------------------------------
-	void CompilerOptions::SetEnableConcurrency(bool bEnabled)
-	{
-		m_pD->m_enableConcurrency = bEnabled;
-	}
-
-
     //---------------------------------------------------------------------------------------------
     void CompilerOptions::SetUseDiskCache( bool enabled )
     {
@@ -198,14 +183,12 @@ namespace mu
     {
         MUTABLE_CPUPROFILER_SCOPE(Compile);
 
-        TaskManager* pTaskManager = new TaskManager(m_pD->m_options->GetPrivate()->m_enableConcurrency);
-
         vector< STATE_COMPILATION_DATA > states;
         Ptr<ErrorLog> genErrorLog;
         {
             CodeGenerator gen( m_pD->m_options->GetPrivate() );
 
-            gen.GenerateRoot( pNode, pTaskManager );
+            gen.GenerateRoot( pNode );
 
             check( !gen.m_states.IsEmpty() );
 
@@ -234,7 +217,7 @@ namespace mu
         // Optimize the generated code
         {
             CodeOptimiser optimiser( m_pD->m_options, states );
-            optimiser.OptimiseAST( pTaskManager );
+            optimiser.OptimiseAST( );
         }
 
 
@@ -365,8 +348,6 @@ namespace mu
 
 		UE_LOG(LogMutableCore, Verbose, TEXT("(int) %s : %ld"), TEXT("program size"), int64(program.m_opAddress.Num()));
 
-        delete pTaskManager;
-
         return pResult;
     }
 
diff --git a/Engine/Plugins/Experimental/Mutable/Source/MutableTools/Private/MuT/CompilerPrivate.h b/Engine/Plugins/Experimental/Mutable/Source/MutableTools/Private/MuT/CompilerPrivate.h
--- a/Engine/Plugins/Experimental/Mutable/Source/MutableTools/Private/MuT/CompilerPrivate.h
+++ b/Engine/Plugins/Experimental/Mutable/Source/MutableTools/Private/MuT/CompilerPrivate.h
@@ -69,8 +69,6 @@ namespace mu
         //! Detailed optimization options
         FModelOptimizationOptions m_optimisationOptions;
 
-        bool m_enablePartialOptimise = false;
-		bool m_enableConcurrency = false;
         bool m_ignoreStates = false;
         CompilerOptions::TextureLayoutStrategy m_textureLayoutStrategy = CompilerOptions::TextureLayoutStrategy::Pack;
 
diff --git a/Engine/Plugins/Experimental/Mutable/Source/MutableTools/Private/MuT/DataPacker.cpp b/Engine/Plugins/Experimental/Mutable/Source/MutableTools/Private/MuT/DataPacker.cpp
--- a/Engine/Plugins/Experimental/Mutable/Source/MutableTools/Private/MuT/DataPacker.cpp
+++ b/Engine/Plugins/Experimental/Mutable/Source/MutableTools/Private/MuT/DataPacker.cpp
@@ -18,6 +18,8 @@
 #include "MuR/RefCounted.h"
 #include "MuT/ASTOpConstantResource.h"
 #include "MuT/ASTOpImageCompose.h"
+#include "MuT/ASTOpImageLayer.h"
+#include "MuT/ASTOpImageLayerColor.h"
 #include "MuT/ASTOpImageMultiLayer.h"
 #include "MuT/ASTOpInstanceAdd.h"
 #include "MuT/ASTOpMeshExtractLayoutBlocks.h"
@@ -380,18 +382,18 @@ namespace mu
             case OP_TYPE::IM_LAYERCOLOUR:
             {
                 recurse = false;
-				const ASTOpFixed* op = dynamic_cast<const ASTOpFixed*>(node.get());
+				const ASTOpImageLayerColor* op = dynamic_cast<const ASTOpImageLayerColor*>(node.get());
 				TArray<bool> newState;
 				newState.Init(false, size_t(EImageFormat::IF_COUNT));
-				RecurseWithState( op->children[op->op.args.ImageLayerColour.base].child(), newState );
-                RecurseWithState( op->children[op->op.args.ImageLayerColour.colour].child(), newState );
+				RecurseWithState( op->base.child(), newState );
+                RecurseWithState( op->color.child(), newState );
 
-                if ( op->children[op->op.args.ImageLayerColour.mask] )
+                if ( op->mask )
                 {
                     newState[(size_t)EImageFormat::IF_L_UBYTE] = true;
                     newState[(size_t)EImageFormat::IF_L_UBYTE_RLE] = true;
 
-                    RecurseWithState( op->children[op->op.args.ImageLayerColour.mask].child(), newState );
+                    RecurseWithState( op->mask.child(), newState );
                 }
                 break;
             }
@@ -399,18 +401,18 @@ namespace mu
             case OP_TYPE::IM_LAYER:
             {
                 recurse = false;
-				const ASTOpFixed* op = dynamic_cast<const ASTOpFixed*>(node.get());
+				const ASTOpImageLayer* op = dynamic_cast<const ASTOpImageLayer*>(node.get());
 				TArray<bool> newState;
 				newState.Init(false, size_t(EImageFormat::IF_COUNT));
-				RecurseWithState( op->children[op->op.args.ImageLayer.base].child(), newState );
-                RecurseWithState( op->children[op->op.args.ImageLayer.blended].child(), newState );
+				RecurseWithState( op->base.child(), newState );
+                RecurseWithState( op->blend.child(), newState );
 
-                if (op->children[op->op.args.ImageLayer.mask])
+                if (op->mask)
                 {
                     newState[(size_t)EImageFormat::IF_L_UBYTE] = true;
                     newState[(size_t)EImageFormat::IF_L_UBYTE_RLE] = true;
 
-                    RecurseWithState( op->children[op->op.args.ImageLayer.mask].child(), newState );
+                    RecurseWithState( op->mask.child(), newState );
                 }
                 break;
             }
diff --git a/Engine/Plugins/Experimental/Mutable/Source/MutableTools/Private/MuT/TaskManager.cpp b/Engine/Plugins/Experimental/Mutable/Source/MutableTools/Private/MuT/TaskManager.cpp
deleted file mode 100644
diff --git a/Engine/Plugins/Experimental/Mutable/Source/MutableTools/Private/MuT/TaskManager.h b/Engine/Plugins/Experimental/Mutable/Source/MutableTools/Private/MuT/TaskManager.h
deleted file mode 100644
diff --git a/Engine/Plugins/Experimental/Mutable/Source/MutableTools/Public/MuT/Compiler.h b/Engine/Plugins/Experimental/Mutable/Source/MutableTools/Public/MuT/Compiler.h
--- a/Engine/Plugins/Experimental/Mutable/Source/MutableTools/Public/MuT/Compiler.h
+++ b/Engine/Plugins/Experimental/Mutable/Source/MutableTools/Public/MuT/Compiler.h
@@ -67,12 +67,6 @@ namespace mu
         //!
         void SetConstReductionEnabled( bool bEnabled );
 
-        //!
-        void SetEnablePartialOptimisation( bool bEnabled );
-
-		//!
-		void SetEnableConcurrency(bool bEnabled);
-
         //!
         void SetOptimisationMaxIteration( int maxIterations );
 
