[Mutable] Changed CPU traces to use the Mutable macro.

#rnx

diff --git a/Engine/Plugins/Experimental/Mutable/Source/CustomizableObject/Private/MuCO/CustomizableObjectSystem.cpp b/Engine/Plugins/Experimental/Mutable/Source/CustomizableObject/Private/MuCO/CustomizableObjectSystem.cpp
--- a/Engine/Plugins/Experimental/Mutable/Source/CustomizableObject/Private/MuCO/CustomizableObjectSystem.cpp
+++ b/Engine/Plugins/Experimental/Mutable/Source/CustomizableObject/Private/MuCO/CustomizableObjectSystem.cpp
@@ -2225,7 +2225,7 @@ namespace impl
 
 void UCustomizableObjectSystem::AdvanceCurrentOperation() 
 {
-	TRACE_CPUPROFILER_EVENT_SCOPE(Mutable_AdvanceCurrentOperation);
+	MUTABLE_CPUPROFILER_SCOPE(AdvanceCurrentOperation);
 
 	// See if we can clear the last reference to the mutable-thread task
 	Private->ClearMutableTaskIfDone();
@@ -2263,7 +2263,7 @@ void UCustomizableObjectSystem::AdvanceCurrentOperation()
 	{
 	case FMutableOperation::EOperationType::Discard:
 	{
-		TRACE_CPUPROFILER_EVENT_SCOPE(Mutable_OperationDiscard);
+		MUTABLE_CPUPROFILER_SCOPE(OperationDiscard);
 
 		// \TODO: Discards could be done in any case, concurrently with update operations. Should they be
 		// in their own "queue"?
@@ -2300,7 +2300,7 @@ void UCustomizableObjectSystem::AdvanceCurrentOperation()
 
 	case FMutableOperation::EOperationType::Update:
 	{
-		TRACE_CPUPROFILER_EVENT_SCOPE(Mutable_OperationUpdate);
+		MUTABLE_CPUPROFILER_SCOPE(OperationUpdate);
 
 		// Start the first task of the update process. See namespace impl comments above.
 		impl::Task_Game_StartUpdate(Private->CurrentMutableOperation);
