[mutable] Fix non-unity build problem with missing declaration.
#rnx

diff --git a/Engine/Plugins/Experimental/Mutable/Source/CustomizableObjectEditor/Public/MuCOE/MutableMeshPreviewUtils.h b/Engine/Plugins/Experimental/Mutable/Source/CustomizableObjectEditor/Public/MuCOE/MutableMeshPreviewUtils.h
--- a/Engine/Plugins/Experimental/Mutable/Source/CustomizableObjectEditor/Public/MuCOE/MutableMeshPreviewUtils.h
+++ b/Engine/Plugins/Experimental/Mutable/Source/CustomizableObjectEditor/Public/MuCOE/MutableMeshPreviewUtils.h
@@ -3,6 +3,7 @@
 #pragma once
 
 #include "MuR/System.h"
+#include "MuR/Mesh.h"
 
 class USkeletalMesh;
 
