[mutable] Fix crash when compiling a CO using a TextureParameter node.

#rnx

diff --git a/Engine/Plugins/Experimental/Mutable/Source/CustomizableObjectEditor/Private/MuCOE/GenerateMutableSource/GenerateMutableSource.cpp b/Engine/Plugins/Experimental/Mutable/Source/CustomizableObjectEditor/Private/MuCOE/GenerateMutableSource/GenerateMutableSource.cpp
--- a/Engine/Plugins/Experimental/Mutable/Source/CustomizableObjectEditor/Private/MuCOE/GenerateMutableSource/GenerateMutableSource.cpp
+++ b/Engine/Plugins/Experimental/Mutable/Source/CustomizableObjectEditor/Private/MuCOE/GenerateMutableSource/GenerateMutableSource.cpp
@@ -819,14 +819,9 @@ mu::NodeObjectPtr GenerateMutableSource(const UEdGraphPin * Pin, FMutableGraphGe
 				LODNode->SetComponent(MeshComponentIndex, ComponentNode);
 			}
 
-			if (GenerationContext.CurrentLOD < GenerationContext.FirstLODAvailable)
-			{
-				continue;
-			} 
-
 			const bool bUseAutomaticLods = 
 					GenerationContext.CurrentAutoLODStrategy == ECustomizableObjectAutomaticLODStrategy::AutomaticFromMesh;
-			LastDefinedLOD = (LODIndex < NumLODs) && (LastDefinedLOD == -1 || !bUseAutomaticLods) ? LODIndex : LastDefinedLOD;
+			LastDefinedLOD = (LODIndex < NumLODs) && (LastDefinedLOD == INDEX_NONE || !bUseAutomaticLods) ? LODIndex : LastDefinedLOD;
 
 			// It turns out LODToGenerate is always LastDefinedLOD.
 			const int32& LODToGenerate = LastDefinedLOD;
@@ -834,6 +829,11 @@ mu::NodeObjectPtr GenerateMutableSource(const UEdGraphPin * Pin, FMutableGraphGe
 			{
 				continue;
 			}
+
+			if (GenerationContext.CurrentLOD < GenerationContext.FirstLODAvailable)
+			{
+				continue;
+			} 
 			
 			TArray<UEdGraphPin*> ConnectedLODPins = FollowInputPinArray(*TypedNodeObj->LODPin(LODToGenerate));
 
@@ -1401,7 +1401,7 @@ int32 ComputeLODBias(const FMutableGraphGenerationContext& GenerationContext, co
 		GenerationContext.CurrentAutoLODStrategy == ECustomizableObjectAutomaticLODStrategy::AutomaticFromMesh)
 	{
 		// Only if the texture actually uses a layout. Otherwise it could be a special texture we shouldn't scale.
-		if (MaterialNode->GetImageUVLayout(ImageIndex) >= 0)
+		if (!MaterialNode || MaterialNode->GetImageUVLayout(ImageIndex) >= 0)
 		{
 			// \todo: make it an object property to be tweaked
 			int MipsToSkipPerLOD = 1;
diff --git a/Engine/Plugins/Experimental/Mutable/Source/CustomizableObjectEditor/Private/MuCOE/GenerateMutableSource/GenerateMutableSourceImage.cpp b/Engine/Plugins/Experimental/Mutable/Source/CustomizableObjectEditor/Private/MuCOE/GenerateMutableSource/GenerateMutableSourceImage.cpp
--- a/Engine/Plugins/Experimental/Mutable/Source/CustomizableObjectEditor/Private/MuCOE/GenerateMutableSource/GenerateMutableSourceImage.cpp
+++ b/Engine/Plugins/Experimental/Mutable/Source/CustomizableObjectEditor/Private/MuCOE/GenerateMutableSource/GenerateMutableSourceImage.cpp
@@ -251,9 +251,17 @@ mu::NodeImagePtr GenerateMutableSourceImage(const UEdGraphPin* Pin, FMutableGrap
 		ResizeNode->SetBase(FormatNode);
 		ResizeNode->SetRelative(false);
 
-		if (TypedNodeParam->DefaultValue)
+		const UTexture2D* ReferenceTexture = TypedNodeParam->DefaultValue;
+		if (ReferenceTexture)
 		{
-			ResizeNode->SetSize(TypedNodeParam->DefaultValue->GetResource()->GetSizeX(), TypedNodeParam->DefaultValue->GetResource()->GetSizeY());
+			const int32 LODBias = ComputeLODBias(GenerationContext, ReferenceTexture, MaxTextureSize, nullptr, INDEX_NONE);
+
+			int32 Width = MaxTextureSize > 0 ? FMath::Min(MaxTextureSize, ReferenceTexture->GetImportedSize().X) : ReferenceTexture->GetImportedSize().X;
+			int32 Height = MaxTextureSize > 0 ? FMath::Min(MaxTextureSize, ReferenceTexture->GetImportedSize().Y) : ReferenceTexture->GetImportedSize().Y;
+			Width = Width >> LODBias;
+			Height = Height >> LODBias;
+
+			ResizeNode->SetSize(FMath::Max(Width,1), FMath::Max(Height, 1));
 		}
 		else
 		{
