[mutable] Add some checks to ensure we can save the compiled data

#rnx

diff --git a/Engine/Plugins/Experimental/Mutable/Source/CustomizableObjectEditor/Private/MuCOE/CustomizableObjectCompileRunnable.cpp b/Engine/Plugins/Experimental/Mutable/Source/CustomizableObjectEditor/Private/MuCOE/CustomizableObjectCompileRunnable.cpp
--- a/Engine/Plugins/Experimental/Mutable/Source/CustomizableObjectEditor/Private/MuCOE/CustomizableObjectCompileRunnable.cpp
+++ b/Engine/Plugins/Experimental/Mutable/Source/CustomizableObjectEditor/Private/MuCOE/CustomizableObjectCompileRunnable.cpp
@@ -220,36 +220,52 @@ uint32 FCustomizableObjectSaveDDRunnable::Run()
 		FileManager.MakeDirectory(*FolderPath, true);
 
 		// Delete files...
-		FileManager.Delete(*CompildeDataFullFileName, true, false, true);
-		FileManager.Delete(*StreamableDataFullFileName, true, false, true);
+		bool bFilesDeleted = true;
+		if (!FileManager.Delete(*CompildeDataFullFileName, true, false, true))
+		{
+			UE_LOG(LogMutable, Error, TEXT("Failed to delete compiled data in file [%s]."), *CompildeDataFullFileName);
+			bFilesDeleted = false;
+		}
+
+		if (!FileManager.Delete(*StreamableDataFullFileName, true, false, true))
+		{
+			UE_LOG(LogMutable, Error, TEXT("Failed to delete streamed data in file [%s]."), *StreamableDataFullFileName);
+			bFilesDeleted = false;
+		}
 
-		// Create file writers...
-		FArchive* ModelMemoryWriter = FileManager.CreateFileWriter(*CompildeDataFullFileName);
-		FArchive* StreamableMemoryWriter = FileManager.CreateFileWriter(*StreamableDataFullFileName);
+		// Store current compiled data
+		if (bFilesDeleted)
+		{
+			// Create file writers...
+			FArchive* ModelMemoryWriter = FileManager.CreateFileWriter(*CompildeDataFullFileName);
+			FArchive* StreamableMemoryWriter = FileManager.CreateFileWriter(*StreamableDataFullFileName);
+			check(ModelMemoryWriter);
+			check(StreamableMemoryWriter);
 
-		// Serailize headers to validate data
-		*ModelMemoryWriter << CustomizableObjectHeader;
-		*StreamableMemoryWriter << CustomizableObjectHeader;
+			// Serailize headers to validate data
+			*ModelMemoryWriter << CustomizableObjectHeader;
+			*StreamableMemoryWriter << CustomizableObjectHeader;
 
-		// Serialize Customizable Object's Data to disk
-		ModelMemoryWriter->Serialize(reinterpret_cast<void*>(Bytes.GetData()), Bytes.Num() * sizeof(uint8));
-		Bytes.Empty();
+			// Serialize Customizable Object's Data to disk
+			ModelMemoryWriter->Serialize(reinterpret_cast<void*>(Bytes.GetData()), Bytes.Num() * sizeof(uint8));
+			Bytes.Empty();
 
-		// Serialize mu::Model and streamable resources
-		*ModelMemoryWriter << bModelSerialized;
+			// Serialize mu::Model and streamable resources
+			*ModelMemoryWriter << bModelSerialized;
 
-		FUnrealMutableModelBulkStreamer Streamer(ModelMemoryWriter, StreamableMemoryWriter);
-		mu::Model::Serialise(Model.get(), Streamer);
+			FUnrealMutableModelBulkStreamer Streamer(ModelMemoryWriter, StreamableMemoryWriter);
+			mu::Model::Serialise(Model.get(), Streamer);
 
-		// Save to disk
-		ModelMemoryWriter->Flush();
-		StreamableMemoryWriter->Flush();
+			// Save to disk
+			ModelMemoryWriter->Flush();
+			StreamableMemoryWriter->Flush();
 
-		ModelMemoryWriter->Close();
-		StreamableMemoryWriter->Close();
+			ModelMemoryWriter->Close();
+			StreamableMemoryWriter->Close();
 
-		delete ModelMemoryWriter;
-		delete StreamableMemoryWriter;
+			delete ModelMemoryWriter;
+			delete StreamableMemoryWriter;
+		}
 	}
 
 	bThreadCompleted = true;
