#Mutable Disable assert, for now. Causing cook crash

diff --git a/Engine/Plugins/Experimental/Mutable/Source/MutableTools/Private/MuT/CodeGenerator_Image.cpp b/Engine/Plugins/Experimental/Mutable/Source/MutableTools/Private/MuT/CodeGenerator_Image.cpp
--- a/Engine/Plugins/Experimental/Mutable/Source/MutableTools/Private/MuT/CodeGenerator_Image.cpp
+++ b/Engine/Plugins/Experimental/Mutable/Source/MutableTools/Private/MuT/CodeGenerator_Image.cpp
@@ -2035,7 +2035,7 @@ namespace mu
     {
         Ptr<ASTOp> result = at;
 
-        check( size[0]>0 && size[1]>0 );
+        //check( size[0]>0 && size[1]>0 );
 
         if ( at->GetImageDesc().m_size != size )
         {
