[Mutable] Projector Parameter fix Copy/Paste and Transform Reset.

#rnx

diff --git a/Engine/Plugins/Experimental/Mutable/Source/CustomizableObjectEditor/Private/MuCOE/CustomizableObjectEditor.cpp b/Engine/Plugins/Experimental/Mutable/Source/CustomizableObjectEditor/Private/MuCOE/CustomizableObjectEditor.cpp
--- a/Engine/Plugins/Experimental/Mutable/Source/CustomizableObjectEditor/Private/MuCOE/CustomizableObjectEditor.cpp
+++ b/Engine/Plugins/Experimental/Mutable/Source/CustomizableObjectEditor/Private/MuCOE/CustomizableObjectEditor.cpp
@@ -1813,7 +1813,7 @@ void FCustomizableObjectEditor::Tick( float InDeltaTime )
 			ResetProjectorVisibilityNoUpdate();
 		}
 
-		if (SelectedGraphNodesChanged && !SelectedProjectorParameterNotNode)
+		if (SelectedGraphNodesChanged && !SelectedProjectorParameterNotNode || PreviewInstance && PreviewInstance->TempUpdateGizmoInViewport)
 		{
 			ManagingProjector = true;
 			Viewport->ResetProjectorVisibility(false);
diff --git a/Engine/Plugins/Experimental/Mutable/Source/CustomizableObjectEditor/Private/MuCOE/CustomizableObjectEditorViewportClient.cpp b/Engine/Plugins/Experimental/Mutable/Source/CustomizableObjectEditor/Private/MuCOE/CustomizableObjectEditorViewportClient.cpp
--- a/Engine/Plugins/Experimental/Mutable/Source/CustomizableObjectEditor/Private/MuCOE/CustomizableObjectEditorViewportClient.cpp
+++ b/Engine/Plugins/Experimental/Mutable/Source/CustomizableObjectEditor/Private/MuCOE/CustomizableObjectEditorViewportClient.cpp
@@ -374,6 +374,7 @@ void GizmoRTSProxy::ProjectorParameterChanged(UCustomizableObjectNodeProjectorPa
 			{
 				Value.Position = Node->DefaultValue.Position;
 				Value.Direction = Node->DefaultValue.Direction;
+				Value.Scale = Node->DefaultValue.Scale;
 				Value.Up = Node->DefaultValue.Up;
 				Node->ParameterSetModified = -1;
 				break;
