[Mutable] Fix parentheses warning.

diff --git a/Engine/Plugins/Experimental/Mutable/Source/CustomizableObjectEditor/Private/MuCOE/CustomizableObjectEditor.cpp b/Engine/Plugins/Experimental/Mutable/Source/CustomizableObjectEditor/Private/MuCOE/CustomizableObjectEditor.cpp
--- a/Engine/Plugins/Experimental/Mutable/Source/CustomizableObjectEditor/Private/MuCOE/CustomizableObjectEditor.cpp
+++ b/Engine/Plugins/Experimental/Mutable/Source/CustomizableObjectEditor/Private/MuCOE/CustomizableObjectEditor.cpp
@@ -1813,7 +1813,7 @@ void FCustomizableObjectEditor::Tick( float InDeltaTime )
 			ResetProjectorVisibilityNoUpdate();
 		}
 
-		if (SelectedGraphNodesChanged && !SelectedProjectorParameterNotNode || PreviewInstance && PreviewInstance->TempUpdateGizmoInViewport)
+		if ((SelectedGraphNodesChanged && !SelectedProjectorParameterNotNode) || (PreviewInstance && PreviewInstance->TempUpdateGizmoInViewport))
 		{
 			ManagingProjector = true;
 			Viewport->ResetProjectorVisibility(false);
