[mutable] debugger: Add option to compile forcing a large mip bias.
Remove unnecessary debugger boilerplate.
#rnx

diff --git a/Engine/Plugins/Experimental/Mutable/Source/CustomizableObject/Classes/MuCO/CustomizableObject.h b/Engine/Plugins/Experimental/Mutable/Source/CustomizableObject/Classes/MuCO/CustomizableObject.h
--- a/Engine/Plugins/Experimental/Mutable/Source/CustomizableObject/Classes/MuCO/CustomizableObject.h
+++ b/Engine/Plugins/Experimental/Mutable/Source/CustomizableObject/Classes/MuCO/CustomizableObject.h
@@ -209,6 +209,9 @@ struct FCompilationOptions
 	// Used to reduce the number of notifications when compiling objects
 	bool bSilentCompilation = true;
 
+	/** Force a very big number on the mips to skip during compilation. Useful to debug special cooks of the data. */
+	bool bForceLargeLODBias = false;
+
 };
 
 
diff --git a/Engine/Plugins/Experimental/Mutable/Source/CustomizableObjectEditor/Private/MuCOE/CustomizableObjectDebuggerActions.cpp b/Engine/Plugins/Experimental/Mutable/Source/CustomizableObjectEditor/Private/MuCOE/CustomizableObjectDebuggerActions.cpp
deleted file mode 100644
diff --git a/Engine/Plugins/Experimental/Mutable/Source/CustomizableObjectEditor/Private/MuCOE/GenerateMutableSource/GenerateMutableSource.cpp b/Engine/Plugins/Experimental/Mutable/Source/CustomizableObjectEditor/Private/MuCOE/GenerateMutableSource/GenerateMutableSource.cpp
--- a/Engine/Plugins/Experimental/Mutable/Source/CustomizableObjectEditor/Private/MuCOE/GenerateMutableSource/GenerateMutableSource.cpp
+++ b/Engine/Plugins/Experimental/Mutable/Source/CustomizableObjectEditor/Private/MuCOE/GenerateMutableSource/GenerateMutableSource.cpp
@@ -1367,6 +1367,12 @@ int32 ComputeLODBias(const FMutableGraphGenerationContext& GenerationContext, co
 {
 	int32 LODBias = 0;
 
+	if (GenerationContext.Options.bForceLargeLODBias)
+	{
+		// This seems to be the highest LODBias we get during cook.
+		return 11;
+	}
+
 	// We used to calculate the lod bias directly from the group like this:
 	//int LODBias = 0;
 	//if (LODSettings.TextureLODGroups.IsValidIndex(ReferenceTexture->LODGroup))
diff --git a/Engine/Plugins/Experimental/Mutable/Source/CustomizableObjectEditor/Private/MuCOE/SMutableObjectViewer.cpp b/Engine/Plugins/Experimental/Mutable/Source/CustomizableObjectEditor/Private/MuCOE/SMutableObjectViewer.cpp
--- a/Engine/Plugins/Experimental/Mutable/Source/CustomizableObjectEditor/Private/MuCOE/SMutableObjectViewer.cpp
+++ b/Engine/Plugins/Experimental/Mutable/Source/CustomizableObjectEditor/Private/MuCOE/SMutableObjectViewer.cpp
@@ -402,6 +402,19 @@ TSharedRef<SWidget> SMutableObjectViewer::GenerateCompileOptionsMenuContent()
 			NAME_None,
 			EUserInterfaceActionType::ToggleButton
 		);
+
+		// Debug LODBias
+		MenuBuilder.AddMenuEntry(
+			LOCTEXT("ForceLargeLODBias", "Force a large texture LODBias."),
+			LOCTEXT("ForceLargeLODBiasTooltip", "This is useful to test compilation of special cook modes."),
+			FSlateIcon(),
+			FUIAction(
+				FExecuteAction::CreateLambda([this]() { CompileOptions.bForceLargeLODBias = !CompileOptions.bForceLargeLODBias; }),
+				FCanExecuteAction(),
+				FIsActionChecked::CreateLambda([this]() { return CompileOptions.bForceLargeLODBias; })),
+			NAME_None,
+			EUserInterfaceActionType::ToggleButton
+		);
 	}
 	MenuBuilder.EndSection();
 
diff --git a/Engine/Plugins/Experimental/Mutable/Source/CustomizableObjectEditor/Public/MuCOE/CustomizableObjectDebuggerActions.h b/Engine/Plugins/Experimental/Mutable/Source/CustomizableObjectEditor/Public/MuCOE/CustomizableObjectDebuggerActions.h
deleted file mode 100644
