[mutable] Fix optimization crash with projectors on 1-channel images.
#rnx

diff --git a/Engine/Plugins/Experimental/Mutable/Source/MutableTools/Private/MuT/ASTOpImageSwizzle.cpp b/Engine/Plugins/Experimental/Mutable/Source/MutableTools/Private/MuT/ASTOpImageSwizzle.cpp
--- a/Engine/Plugins/Experimental/Mutable/Source/MutableTools/Private/MuT/ASTOpImageSwizzle.cpp
+++ b/Engine/Plugins/Experimental/Mutable/Source/MutableTools/Private/MuT/ASTOpImageSwizzle.cpp
@@ -242,6 +242,11 @@ namespace mu
 			}
 		}
 
+		if (!channelSourceAt)
+		{
+			return at;
+		}
+
 		OP_TYPE sourceType = channelSourceAt->GetOpType();
 
 		if (bAllChannelsAreTheSame && channelSourceAt)
