[mutable] Avoid crash when image switching textures don't have the same format. Try to get the best format when getting the layer image base texture pixel format.
Avoid crash if a skeleton resource is null during generation.
#rnx

diff --git a/Engine/Plugins/Experimental/Mutable/Source/CustomizableObjectEditor/Private/MuCOE/GenerateMutableSource/GenerateMutableSourceMesh.cpp b/Engine/Plugins/Experimental/Mutable/Source/CustomizableObjectEditor/Private/MuCOE/GenerateMutableSource/GenerateMutableSourceMesh.cpp
--- a/Engine/Plugins/Experimental/Mutable/Source/CustomizableObjectEditor/Private/MuCOE/GenerateMutableSource/GenerateMutableSourceMesh.cpp
+++ b/Engine/Plugins/Experimental/Mutable/Source/CustomizableObjectEditor/Private/MuCOE/GenerateMutableSource/GenerateMutableSourceMesh.cpp
@@ -482,7 +482,15 @@ mu::MeshPtr ConvertSkeletalMeshToMutable(USkeletalMesh* InSkeletalMesh, int LOD,
 	}
 	
 	USkeleton* InSkeleton = InSkeletalMesh->GetSkeleton();
-	check(InSkeleton);
+	if (!InSkeleton)
+	{
+		FString Msg = FString::Printf(
+			TEXT("No skeleton provided when converting SkeletalMesh [%s]."),
+			*InSkeletalMesh->GetName()
+		);
+		GenerationContext.Compiler->CompilerLog(FText::FromString(Msg), CurrentNode);
+		return nullptr;
+	}
 
 	TObjectPtr<const USkeletalMesh> SkeletalMesh = InSkeletalMesh;
 
diff --git a/Engine/Plugins/Experimental/Mutable/Source/MutableTools/Private/MuT/CodeGenerator_Image.cpp b/Engine/Plugins/Experimental/Mutable/Source/MutableTools/Private/MuT/CodeGenerator_Image.cpp
--- a/Engine/Plugins/Experimental/Mutable/Source/MutableTools/Private/MuT/CodeGenerator_Image.cpp
+++ b/Engine/Plugins/Experimental/Mutable/Source/MutableTools/Private/MuT/CodeGenerator_Image.cpp
@@ -320,7 +320,7 @@ namespace mu
 
 		FImageSize TargetSize((uint16)m_imageState.Last().m_imageRect.size[0],(uint16)m_imageState.Last().m_imageRect.size[1]);
 
-		EImageFormat baseFormat = base->GetImageDesc().m_format;
+		EImageFormat baseFormat = base->GetImageDesc( true ).m_format;
         //base = GenerateImageFormat( base, EImageFormat::IF_RGB_UBYTE );
         base = GenerateImageSize( base, TargetSize);
         op->base = base;
@@ -1981,22 +1981,20 @@ namespace mu
 
 
     //---------------------------------------------------------------------------------------------
-    Ptr<ASTOp> CodeGenerator::GenerateImageFormat( Ptr<ASTOp> at, EImageFormat format )
+    Ptr<ASTOp> CodeGenerator::GenerateImageFormat( Ptr<ASTOp> Op, EImageFormat InFormat)
     {
-        Ptr<ASTOp> result = at;
+        Ptr<ASTOp> Result = Op;
 
-        check( format != EImageFormat::IF_NONE );
-
-        if ( at && at->GetImageDesc().m_format != format )
+        if (InFormat!=EImageFormat::IF_NONE && Op && Op->GetImageDesc().m_format!=InFormat)
         {
             // Generate the format change code
             Ptr<ASTOpImagePixelFormat> op = new ASTOpImagePixelFormat();
-            op->Source = at;
-            op->Format = format;
-            result = op;
+            op->Source = Op;
+            op->Format = InFormat;
+			Result = op;
         }
 
-        return result;
+        return Result;
     }
 
 
