[mutable] Now TextureProject Nodes take into account the LODBias of a reference texture.

#rnx

diff --git a/Engine/Plugins/Experimental/Mutable/Source/CustomizableObjectEditor/Classes/MuCOE/Nodes/CustomizableObjectNodeTextureProject.h b/Engine/Plugins/Experimental/Mutable/Source/CustomizableObjectEditor/Classes/MuCOE/Nodes/CustomizableObjectNodeTextureProject.h
--- a/Engine/Plugins/Experimental/Mutable/Source/CustomizableObjectEditor/Classes/MuCOE/Nodes/CustomizableObjectNodeTextureProject.h
+++ b/Engine/Plugins/Experimental/Mutable/Source/CustomizableObjectEditor/Classes/MuCOE/Nodes/CustomizableObjectNodeTextureProject.h
@@ -32,15 +32,22 @@ public:
 	UPROPERTY(EditAnywhere, Category = CustomizableObject)
 	int32 Layout = 0;
 
+	/** Set the width of the Texture. If the reference texture is not null, it overrdies the Size X property */
 	UPROPERTY(EditAnywhere, Category = CustomizableObject)
-	int32 TextureSizeX = 0; 
+	int32 TextureSizeX = 0;
 
+	/** Set the height of the Texture. If the reference texture is not null, it overrdies the Size Y property */
 	UPROPERTY(EditAnywhere, Category = CustomizableObject)
 	int32 TextureSizeY = 0;
 
 	UPROPERTY(EditAnywhere, Category = CustomizableObject)
 	int32 Textures = 0;
 
+	/** Reference Texture used to decide the texture properties of the mutable-generated textures
+	* connected to this material (e.g. LODBias, Size X,...). If null, mutable default texture properties will be applied. */
+	UPROPERTY(EditAnywhere, Category = CustomizableObject)
+	TObjectPtr<UTexture2D> ReferenceTexture = nullptr;
+
 	// UObject interface.
 	void PostEditChangeProperty(FPropertyChangedEvent& PropertyChangedEvent) override;
 	virtual void Serialize(FArchive& Ar) override;
diff --git a/Engine/Plugins/Experimental/Mutable/Source/CustomizableObjectEditor/Private/MuCOE/GenerateMutableSource/GenerateMutableSource.cpp b/Engine/Plugins/Experimental/Mutable/Source/CustomizableObjectEditor/Private/MuCOE/GenerateMutableSource/GenerateMutableSource.cpp
--- a/Engine/Plugins/Experimental/Mutable/Source/CustomizableObjectEditor/Private/MuCOE/GenerateMutableSource/GenerateMutableSource.cpp
+++ b/Engine/Plugins/Experimental/Mutable/Source/CustomizableObjectEditor/Private/MuCOE/GenerateMutableSource/GenerateMutableSource.cpp
@@ -1407,7 +1407,7 @@ int32 ComputeLODBias(const FMutableGraphGenerationContext& GenerationContext, co
 		GenerationContext.CurrentAutoLODStrategy == ECustomizableObjectAutomaticLODStrategy::AutomaticFromMesh)
 	{
 		// Only if the texture actually uses a layout. Otherwise it could be a special texture we shouldn't scale.
-		if (!MaterialNode || MaterialNode->GetImageUVLayout(ImageIndex) >= 0)
+		if (MaterialNode && MaterialNode->GetImageUVLayout(ImageIndex) >= 0)
 		{
 			// \todo: make it an object property to be tweaked
 			int MipsToSkipPerLOD = 1;
diff --git a/Engine/Plugins/Experimental/Mutable/Source/CustomizableObjectEditor/Private/MuCOE/GenerateMutableSource/GenerateMutableSourceImage.cpp b/Engine/Plugins/Experimental/Mutable/Source/CustomizableObjectEditor/Private/MuCOE/GenerateMutableSource/GenerateMutableSourceImage.cpp
--- a/Engine/Plugins/Experimental/Mutable/Source/CustomizableObjectEditor/Private/MuCOE/GenerateMutableSource/GenerateMutableSourceImage.cpp
+++ b/Engine/Plugins/Experimental/Mutable/Source/CustomizableObjectEditor/Private/MuCOE/GenerateMutableSource/GenerateMutableSourceImage.cpp
@@ -243,10 +243,6 @@ mu::NodeImagePtr GenerateMutableSourceImage(const UEdGraphPin* Pin, FMutableGrap
 
 		
 
-		//TODO(Uncoment this)
-		// //Result = FormatNode;
-		//if (TypedNodeParam->bForceFixedSize)
-		//{
 		mu::NodeImageResizePtr ResizeNode = new mu::NodeImageResize();
 		ResizeNode->SetBase(FormatNode);
 		ResizeNode->SetRelative(false);
@@ -270,7 +266,6 @@ mu::NodeImagePtr GenerateMutableSourceImage(const UEdGraphPin* Pin, FMutableGrap
 		}
 
 		Result = ResizeNode;
-		//}
 	}
 
 	else if (const UCustomizableObjectNodeMesh* TypedNodeMesh = Cast<UCustomizableObjectNodeMesh>(Node))
@@ -617,7 +612,26 @@ mu::NodeImagePtr GenerateMutableSourceImage(const UEdGraphPin* Pin, FMutableGrap
 		}
 
 		ImageNode->SetLayout(TypedNodeProject->Layout);
-		ImageNode->SetImageSize( FUintVector2(TypedNodeProject->TextureSizeX, TypedNodeProject->TextureSizeY) );
+		FUintVector2 TextureSize(TypedNodeProject->TextureSizeX, TypedNodeProject->TextureSizeY);
+
+		// Calculating Texture size using Reference texture parameters
+		if (TypedNodeProject->ReferenceTexture)
+		{
+			int32 LODBias = ComputeLODBias(GenerationContext, TypedNodeProject->ReferenceTexture, TypedNodeProject->ReferenceTexture->MaxTextureSize, nullptr, INDEX_NONE);
+
+			if (TextureSize.X > 0 && TextureSize.Y > 0)
+			{
+				TextureSize.X = TextureSize.X >> LODBias;
+				TextureSize.Y = TextureSize.Y >> LODBias;
+			}
+			else
+			{
+				TextureSize.X = TypedNodeProject->ReferenceTexture->GetImportedSize().X >> LODBias;
+				TextureSize.Y = TypedNodeProject->ReferenceTexture->GetImportedSize().Y >> LODBias;
+			}
+		}
+
+		ImageNode->SetImageSize(TextureSize);
 
 		if (const UEdGraphPin* ConnectedPin = FollowInputPin(*TypedNodeProject->AngleFadeStartPin()))
 		{
