[mutable] Prevent constant-optimize of subtrees with "plain image" operations to reduce constants on disk.
#rnx

diff --git a/Engine/Plugins/Experimental/Mutable/Source/MutableTools/Private/MuT/CodeOptimiser.cpp b/Engine/Plugins/Experimental/Mutable/Source/MutableTools/Private/MuT/CodeOptimiser.cpp
--- a/Engine/Plugins/Experimental/Mutable/Source/MutableTools/Private/MuT/CodeOptimiser.cpp
+++ b/Engine/Plugins/Experimental/Mutable/Source/MutableTools/Private/MuT/CodeOptimiser.cpp
@@ -638,6 +638,7 @@ namespace mu
 				case OP_TYPE::ME_CLIPMORPHPLANE:
 				case OP_TYPE::ME_APPLYPOSE:
 				case OP_TYPE::ME_REMOVEMASK:
+				case OP_TYPE::IM_PLAINCOLOUR:
 					hasSpecialOpInSubtree = true;
 					break;
 
