[mutable] Optimize code generation for switching operations when they fall below a silmilar operation with the same variable.
#rnx

diff --git a/Engine/Plugins/Experimental/Mutable/Source/MutableTools/Private/MuT/AST.cpp b/Engine/Plugins/Experimental/Mutable/Source/MutableTools/Private/MuT/AST.cpp
--- a/Engine/Plugins/Experimental/Mutable/Source/MutableTools/Private/MuT/AST.cpp
+++ b/Engine/Plugins/Experimental/Mutable/Source/MutableTools/Private/MuT/AST.cpp
@@ -89,15 +89,15 @@ ASTChild& ASTChild::operator=( ASTChild&& rhs )
 //-------------------------------------------------------------------------------------------------
 //-------------------------------------------------------------------------------------------------
 //-------------------------------------------------------------------------------------------------
-void ASTOp::ForEachParent(const TFunctionRef<void(ASTOp*)> f )
+void ASTOp::ForEachParent(const TFunctionRef<void(ASTOp*)> f) const
 {
-    for( auto& p: m_parents )
-    {
-        if (p)
-        {
-            f(p);
-        }
-    }
+	for (auto& p : m_parents)
+	{
+		if (p)
+		{
+			f(p);
+		}
+	}
 }
 
 
@@ -859,10 +859,10 @@ void ASTOp::Traverse_BottomUp_Unique
 //-------------------------------------------------------------------------------------------------
 //-------------------------------------------------------------------------------------------------
 //-------------------------------------------------------------------------------------------------
-int ASTOp::GetParentCount()
+int ASTOp::GetParentCount() const
 {
     int result=0;
-    ForEachParent( [&](ASTOp* p)
+    ForEachParent( [&](const ASTOp* p)
     {
         if (p!=nullptr) ++result;
     });
diff --git a/Engine/Plugins/Experimental/Mutable/Source/MutableTools/Private/MuT/AST.h b/Engine/Plugins/Experimental/Mutable/Source/MutableTools/Private/MuT/AST.h
--- a/Engine/Plugins/Experimental/Mutable/Source/MutableTools/Private/MuT/AST.h
+++ b/Engine/Plugins/Experimental/Mutable/Source/MutableTools/Private/MuT/AST.h
@@ -4,7 +4,6 @@
 
 #include "MuT/Platform.h"
 #include "MuR/Image.h"
-#include "MuR/MemoryPrivate.h"
 #include "MuR/ModelPrivate.h"
 #include "MuR/MutableMath.h"
 #include "MuR/MutableMemory.h"
@@ -416,7 +415,7 @@ namespace mu
         virtual void ForEachChild( const TFunctionRef<void(ASTChild&)> ) = 0;
 
         //! Run something for each parent operation+.
-        void ForEachParent( const TFunctionRef<void(ASTOp*)> );
+		void ForEachParent(const TFunctionRef<void(ASTOp*)>) const;
 
         //! Run something for each child operation, with a chance to modify it.
         virtual bool operator==( const ASTOp& other ) const;
@@ -666,7 +665,7 @@ namespace mu
 
     public:
 
-        int GetParentCount();
+        int GetParentCount() const;
 
         //! Make all parents of this node point at the other node instead.
         static void Replace( const Ptr<ASTOp>& node, const Ptr<ASTOp>& other );
diff --git a/Engine/Plugins/Experimental/Mutable/Source/MutableTools/Private/MuT/ASTOpSwitch.cpp b/Engine/Plugins/Experimental/Mutable/Source/MutableTools/Private/MuT/ASTOpSwitch.cpp
--- a/Engine/Plugins/Experimental/Mutable/Source/MutableTools/Private/MuT/ASTOpSwitch.cpp
+++ b/Engine/Plugins/Experimental/Mutable/Source/MutableTools/Private/MuT/ASTOpSwitch.cpp
@@ -141,7 +141,7 @@ namespace mu
 
 
 	//-------------------------------------------------------------------------------------------------
-	mu::Ptr<ASTOp> ASTOpSwitch::FindBranch(int32_t condition) const
+	mu::Ptr<ASTOp> ASTOpSwitch::FindBranch(int32 condition) const
 	{
 		for (const auto& c : cases)
 		{
@@ -479,6 +479,116 @@ namespace mu
 	        }
 		}
 
+		// Ad-hoc logic optimization: check if all code paths leading to this operation have a switch with the same variable
+		// and the option on those switches for the path that connects to this one is always the same. In that case, we can 
+		// remove this switch and replace it by the value it has for that option. 
+		// This is something the generic logic optimizer should do whan re-enabled.
+		{
+			// List of parent operations that we have visited, and the child we have visited them from.
+			TSet<TTuple<const ASTOp*, const ASTOp*>> Visited;
+			Visited.Reserve(64);
+
+			// First is parent, second is what child we are reaching the parent from. This is necessary to find out what 
+			// switch branch we reach the parent from, if it is a switch.
+			TArray< TTuple<const ASTOp*, const ASTOp*>, TInlineAllocator<16>> Pending;
+			ForEachParent([this,&Pending](ASTOp* Parent)
+				{
+					Pending.Add({ Parent,this});
+				});
+
+			bool bAllPathsHaveMatchingSwitch = true;
+
+			// Switch option value of all parent compatible switches (if any)
+			int32 MatchingSwitchOption = -1;
+
+			while (!Pending.IsEmpty() && bAllPathsHaveMatchingSwitch)
+			{
+				TTuple<const ASTOp*, const ASTOp*> ParentPair = Pending.Pop();
+				bool bAlreadyVisited = false;
+				Visited.Add(ParentPair, &bAlreadyVisited);
+
+				if (!bAlreadyVisited)
+				{
+					const ASTOp* Parent = ParentPair.Get<0>();
+					const ASTOp* ParentChild = ParentPair.Get<1>();
+
+					bool bIsMatchingSwitch = false;
+
+					// TODO: Probably it could be a any switch, it doesn't need to be of the same type.
+					if (Parent->GetOpType() == GetOpType())
+					{
+						const ASTOpSwitch* ParentSwitch = dynamic_cast<const ASTOpSwitch*>(Parent);
+						check(ParentSwitch);
+
+						// To be compatible the switch must be on the same variable
+						if (ParentSwitch->variable==variable)
+						{
+							bIsMatchingSwitch = true;
+							
+							// Find what switch option we are reaching it from
+							bool bIsSingleOption = true;
+							int OptionIndex = -1;
+							for (int32 CaseIndex = 0; CaseIndex < ParentSwitch->cases.Num(); ++CaseIndex)
+							{
+								if (ParentSwitch->cases[CaseIndex].branch.child().get() == ParentChild)
+								{
+									if (OptionIndex != -1)
+									{
+										// This means the same child is connected to more than one switch options
+										// so we cannot optimize.
+										// \TODO: We could if we track a "set of options" for all switches instead of just one.
+										bIsSingleOption = false;
+										break;
+									}
+									else
+									{
+										OptionIndex = CaseIndex;
+									}
+								}
+							}
+
+							// If we did reach it from one single option
+							if (bIsSingleOption && OptionIndex!=-1)
+							{
+								if (MatchingSwitchOption<0)
+								{
+									MatchingSwitchOption = ParentSwitch->cases[OptionIndex].condition;
+								}
+								else if (MatchingSwitchOption!= ParentSwitch->cases[OptionIndex].condition)
+								{
+									bAllPathsHaveMatchingSwitch = false;
+								}
+							}
+						}
+					}
+					
+					if (!bIsMatchingSwitch)
+					{
+						// If it has no parents, then the optimization cannot be applied
+						bool bHasParent = false;
+						Parent->ForEachParent([&bHasParent,this,&Pending,Parent](ASTOp* ParentParent)
+							{
+								Pending.Add({ ParentParent,Parent });
+								bHasParent = true;
+							});
+
+						if (!bHasParent)
+						{
+							// We reached a root without a matching switch along the path.
+							bAllPathsHaveMatchingSwitch = false;
+						}
+					}
+				}
+			}
+
+			if (bAllPathsHaveMatchingSwitch && MatchingSwitchOption>=0)
+			{
+				// We can remove this switch, all paths leading to it have the same condition for this switches variable.
+				return FindBranch(MatchingSwitchOption);
+			}
+
+		}
+
 		return nullptr;
 	}
 
diff --git a/Engine/Plugins/Experimental/Mutable/Source/MutableTools/Private/MuT/ASTOpSwitch.h b/Engine/Plugins/Experimental/Mutable/Source/MutableTools/Private/MuT/ASTOpSwitch.h
--- a/Engine/Plugins/Experimental/Mutable/Source/MutableTools/Private/MuT/ASTOpSwitch.h
+++ b/Engine/Plugins/Experimental/Mutable/Source/MutableTools/Private/MuT/ASTOpSwitch.h
@@ -4,7 +4,6 @@
 
 #include "HAL/PlatformMath.h"
 #include "MuR/Image.h"
-#include "MuR/MemoryPrivate.h"
 #include "MuR/Operations.h"
 #include "MuR/Ptr.h"
 #include "MuT/AST.h"
@@ -34,13 +33,13 @@ template <class SCALAR> class vec4;
 
 		struct FCase
 		{
-			FCase(int32_t cond, Ptr<ASTOp> parent, Ptr<ASTOp> b)
+			FCase(int32 cond, Ptr<ASTOp> parent, Ptr<ASTOp> b)
 				: condition(cond)
 				, branch(parent.get(), b)
 			{
 			}
 
-			int32_t condition;
+			int32 condition;
 			ASTChild branch;
 
 			//!
@@ -85,7 +84,7 @@ template <class SCALAR> class vec4;
 		bool IsCompatibleWith(const ASTOpSwitch* other) const;
 
 		//!
-		Ptr<ASTOp> FindBranch(int32_t condition) const;
+		Ptr<ASTOp> FindBranch(int32 condition) const;
 
 	};
 
