[mutable] Remove unnecessary header and all its references.
#rnx

diff --git a/Engine/Plugins/Experimental/Mutable/Source/CustomizableObjectEditor/Private/MuCOE/SMutableCodeViewer.h b/Engine/Plugins/Experimental/Mutable/Source/CustomizableObjectEditor/Private/MuCOE/SMutableCodeViewer.h
--- a/Engine/Plugins/Experimental/Mutable/Source/CustomizableObjectEditor/Private/MuCOE/SMutableCodeViewer.h
+++ b/Engine/Plugins/Experimental/Mutable/Source/CustomizableObjectEditor/Private/MuCOE/SMutableCodeViewer.h
@@ -15,7 +15,6 @@
 #include "Layout/Visibility.h"
 #include "Misc/Optional.h"
 #include "MuR/Image.h"
-#include "MuR/MemoryPrivate.h"
 #include "MuR/Mesh.h"
 #include "MuR/Model.h"
 #include "MuR/ModelPrivate.h"
diff --git a/Engine/Plugins/Experimental/Mutable/Source/CustomizableObjectEditor/Private/MuCOE/SMutableConstantsWidget.h b/Engine/Plugins/Experimental/Mutable/Source/CustomizableObjectEditor/Private/MuCOE/SMutableConstantsWidget.h
--- a/Engine/Plugins/Experimental/Mutable/Source/CustomizableObjectEditor/Private/MuCOE/SMutableConstantsWidget.h
+++ b/Engine/Plugins/Experimental/Mutable/Source/CustomizableObjectEditor/Private/MuCOE/SMutableConstantsWidget.h
@@ -12,7 +12,6 @@
 #include "Internationalization/Text.h"
 #include "Misc/Optional.h"
 #include "MuR/Image.h"
-#include "MuR/MemoryPrivate.h"
 #include "MuR/Mesh.h"
 #include "MuR/MutableMath.h"
 #include "Templates/SharedPointer.h"
diff --git a/Engine/Plugins/Experimental/Mutable/Source/MutableTools/Private/MuT/NodeColourArithmeticOperationPrivate.h b/Engine/Plugins/Experimental/Mutable/Source/MutableTools/Private/MuT/NodeColourArithmeticOperationPrivate.h
--- a/Engine/Plugins/Experimental/Mutable/Source/MutableTools/Private/MuT/NodeColourArithmeticOperationPrivate.h
+++ b/Engine/Plugins/Experimental/Mutable/Source/MutableTools/Private/MuT/NodeColourArithmeticOperationPrivate.h
@@ -5,8 +5,6 @@
 #include "MuT/NodeColourPrivate.h"
 #include "MuT/NodeColourArithmeticOperation.h"
 
-#include "MuR/MemoryPrivate.h"
-
 
 namespace mu
 {
diff --git a/Engine/Plugins/Experimental/Mutable/Source/MutableTools/Private/MuT/NodeColourFromScalarsPrivate.h b/Engine/Plugins/Experimental/Mutable/Source/MutableTools/Private/MuT/NodeColourFromScalarsPrivate.h
--- a/Engine/Plugins/Experimental/Mutable/Source/MutableTools/Private/MuT/NodeColourFromScalarsPrivate.h
+++ b/Engine/Plugins/Experimental/Mutable/Source/MutableTools/Private/MuT/NodeColourFromScalarsPrivate.h
@@ -8,8 +8,6 @@
 #include "MuT/NodeImage.h"
 #include "MuT/AST.h"
 
-#include "MuR/MemoryPrivate.h"
-
 
 namespace mu
 {
diff --git a/Engine/Plugins/Experimental/Mutable/Source/MutableTools/Private/MuT/NodeColourParameter.cpp b/Engine/Plugins/Experimental/Mutable/Source/MutableTools/Private/MuT/NodeColourParameter.cpp
--- a/Engine/Plugins/Experimental/Mutable/Source/MutableTools/Private/MuT/NodeColourParameter.cpp
+++ b/Engine/Plugins/Experimental/Mutable/Source/MutableTools/Private/MuT/NodeColourParameter.cpp
@@ -4,15 +4,11 @@
 #include "MuT/NodeColourParameter.h"
 
 #include "Misc/AssertionMacros.h"
-#include "MuR/MemoryPrivate.h"
 #include "MuR/MutableMath.h"
 #include "MuT/NodeColourParameterPrivate.h"
 #include "MuT/NodePrivate.h"
 #include "MuT/NodeRange.h"
 
-#include <memory>
-#include <utility>
-
 
 namespace mu
 {
diff --git a/Engine/Plugins/Experimental/Mutable/Source/MutableTools/Private/MuT/NodeColourSwitch.cpp b/Engine/Plugins/Experimental/Mutable/Source/MutableTools/Private/MuT/NodeColourSwitch.cpp
--- a/Engine/Plugins/Experimental/Mutable/Source/MutableTools/Private/MuT/NodeColourSwitch.cpp
+++ b/Engine/Plugins/Experimental/Mutable/Source/MutableTools/Private/MuT/NodeColourSwitch.cpp
@@ -4,14 +4,10 @@
 #include "MuT/NodeColourSwitch.h"
 
 #include "Misc/AssertionMacros.h"
-#include "MuR/MemoryPrivate.h"
 #include "MuT/NodeColourSwitchPrivate.h"
 #include "MuT/NodePrivate.h"
 #include "MuT/NodeScalar.h"
 
-#include <memory>
-#include <utility>
-
 namespace mu
 {
 
diff --git a/Engine/Plugins/Experimental/Mutable/Source/MutableTools/Private/MuT/NodeComponentEditPrivate.h b/Engine/Plugins/Experimental/Mutable/Source/MutableTools/Private/MuT/NodeComponentEditPrivate.h
--- a/Engine/Plugins/Experimental/Mutable/Source/MutableTools/Private/MuT/NodeComponentEditPrivate.h
+++ b/Engine/Plugins/Experimental/Mutable/Source/MutableTools/Private/MuT/NodeComponentEditPrivate.h
@@ -10,8 +10,6 @@
 #include "MuT/NodeMesh.h"
 #include "MuT/NodePatchMesh.h"
 #include "MuT/NodeImage.h"
-
-#include "MuR/MemoryPrivate.h"
 #include "MuT/NodeSurface.h"
 
 
diff --git a/Engine/Plugins/Experimental/Mutable/Source/MutableTools/Private/MuT/NodeComponentNewPrivate.h b/Engine/Plugins/Experimental/Mutable/Source/MutableTools/Private/MuT/NodeComponentNewPrivate.h
--- a/Engine/Plugins/Experimental/Mutable/Source/MutableTools/Private/MuT/NodeComponentNewPrivate.h
+++ b/Engine/Plugins/Experimental/Mutable/Source/MutableTools/Private/MuT/NodeComponentNewPrivate.h
@@ -11,8 +11,6 @@
 #include "MuT/NodeMesh.h"
 #include "MuT/NodePatchMesh.h"
 #include "MuT/NodeImage.h"
-
-#include "MuR/MemoryPrivate.h"
 #include "MuT/NodeSurface.h"
 
 
diff --git a/Engine/Plugins/Experimental/Mutable/Source/MutableTools/Private/MuT/NodeMeshClipMorphPlane.cpp b/Engine/Plugins/Experimental/Mutable/Source/MutableTools/Private/MuT/NodeMeshClipMorphPlane.cpp
--- a/Engine/Plugins/Experimental/Mutable/Source/MutableTools/Private/MuT/NodeMeshClipMorphPlane.cpp
+++ b/Engine/Plugins/Experimental/Mutable/Source/MutableTools/Private/MuT/NodeMeshClipMorphPlane.cpp
@@ -4,17 +4,12 @@
 #include "MuT/NodeMeshClipMorphPlane.h"
 
 #include "Misc/AssertionMacros.h"
-#include "MuR/MemoryPrivate.h"
 #include "MuR/MutableMath.h"
 #include "MuT/NodeLayout.h"
 #include "MuT/NodeMeshClipMorphPlanePrivate.h"
 #include "MuT/NodeMeshPrivate.h"
 #include "MuT/NodePrivate.h"
 
-#include <memory>
-#include <utility>
-
-
 
 namespace mu
 {
diff --git a/Engine/Plugins/Experimental/Mutable/Source/MutableTools/Private/MuT/NodeMeshReshape.cpp b/Engine/Plugins/Experimental/Mutable/Source/MutableTools/Private/MuT/NodeMeshReshape.cpp
--- a/Engine/Plugins/Experimental/Mutable/Source/MutableTools/Private/MuT/NodeMeshReshape.cpp
+++ b/Engine/Plugins/Experimental/Mutable/Source/MutableTools/Private/MuT/NodeMeshReshape.cpp
@@ -5,16 +5,11 @@
 #include "Containers/Array.h"
 #include "HAL/PlatformCrt.h"
 #include "Misc/AssertionMacros.h"
-#include "MuR/MemoryPrivate.h"
-#include "MuR/RefCounted.h"
 #include "MuT/NodeLayout.h"
 #include "MuT/NodeMeshPrivate.h"
 #include "MuT/NodeMeshReshapePrivate.h"
 #include "MuT/NodePrivate.h"
 
-#include <memory>
-
-
 
 namespace mu
 {
diff --git a/Engine/Plugins/Experimental/Mutable/Source/MutableTools/Private/MuT/NodeMeshVariation.cpp b/Engine/Plugins/Experimental/Mutable/Source/MutableTools/Private/MuT/NodeMeshVariation.cpp
--- a/Engine/Plugins/Experimental/Mutable/Source/MutableTools/Private/MuT/NodeMeshVariation.cpp
+++ b/Engine/Plugins/Experimental/Mutable/Source/MutableTools/Private/MuT/NodeMeshVariation.cpp
@@ -1,16 +1,11 @@
 // Copyright Epic Games, Inc. All Rights Reserved.
 
-
 #include "MuT/NodeMeshVariation.h"
 
 #include "Misc/AssertionMacros.h"
-#include "MuR/MemoryPrivate.h"
 #include "MuT/NodeMeshVariationPrivate.h"
 #include "MuT/NodePrivate.h"
 
-#include <memory>
-#include <utility>
-
 
 namespace mu
 {
diff --git a/Engine/Plugins/Experimental/Mutable/Source/MutableTools/Private/MuT/NodeModifier.cpp b/Engine/Plugins/Experimental/Mutable/Source/MutableTools/Private/MuT/NodeModifier.cpp
--- a/Engine/Plugins/Experimental/Mutable/Source/MutableTools/Private/MuT/NodeModifier.cpp
+++ b/Engine/Plugins/Experimental/Mutable/Source/MutableTools/Private/MuT/NodeModifier.cpp
@@ -4,7 +4,6 @@
 #include "MuT/NodeModifier.h"
 
 #include "Misc/AssertionMacros.h"
-#include "MuR/MemoryPrivate.h"
 #include "MuR/Serialisation.h"
 #include "MuR/SerialisationPrivate.h"
 #include "MuT/NodeModifierMeshClipDeform.h"
@@ -13,10 +12,6 @@
 #include "MuT/NodeModifierPrivate.h"
 #include "MuT/NodePrivate.h"
 
-#include <stdint.h>
-#include <memory>
-#include <utility>
-
 
 namespace mu
 {
@@ -48,13 +43,13 @@ namespace mu
 	//---------------------------------------------------------------------------------------------
 	void NodeModifier::Serialise( const NodeModifier* p, OutputArchive& arch )
 	{
-        uint32_t ver = 0;
+        uint32 ver = 0;
 		arch << ver;
 
 #define SERIALISE_CHILDREN( C, ID ) \
         ( const C* pTyped##ID = dynamic_cast<const C*>(p) )			\
         {                                                           \
-            arch << (uint32_t)ID;									\
+            arch << (uint32)ID;									\
             C::Serialise( pTyped##ID, arch );						\
 		}
 
@@ -70,11 +65,11 @@ namespace mu
 	//---------------------------------------------------------------------------------------------
 	NodeModifierPtr NodeModifier::StaticUnserialise( InputArchive& arch )
 	{
-        uint32_t ver;
+        uint32 ver;
 		arch >> ver;
 		check( ver == 0 );
 
-        uint32_t id;
+        uint32 id;
 		arch >> id;
 
 		switch (id)
diff --git a/Engine/Plugins/Experimental/Mutable/Source/MutableTools/Private/MuT/NodePrivate.h b/Engine/Plugins/Experimental/Mutable/Source/MutableTools/Private/MuT/NodePrivate.h
--- a/Engine/Plugins/Experimental/Mutable/Source/MutableTools/Private/MuT/NodePrivate.h
+++ b/Engine/Plugins/Experimental/Mutable/Source/MutableTools/Private/MuT/NodePrivate.h
@@ -7,9 +7,7 @@
 #include "MuT/Visitor.h"
 #include "MuT/AST.h"
 #include "MuR/Operations.h"
-#include "MuR/MemoryPrivate.h"
 
-#include <memory>
 
 namespace mu
 {
diff --git a/Engine/Plugins/Experimental/Mutable/Source/MutableTools/Private/MuT/NodeSurfaceNew.cpp b/Engine/Plugins/Experimental/Mutable/Source/MutableTools/Private/MuT/NodeSurfaceNew.cpp
--- a/Engine/Plugins/Experimental/Mutable/Source/MutableTools/Private/MuT/NodeSurfaceNew.cpp
+++ b/Engine/Plugins/Experimental/Mutable/Source/MutableTools/Private/MuT/NodeSurfaceNew.cpp
@@ -5,7 +5,6 @@
 
 #include "HAL/PlatformMath.h"
 #include "Misc/AssertionMacros.h"
-#include "MuR/MemoryPrivate.h"
 #include "MuT/NodeColour.h"
 #include "MuT/NodeImage.h"
 #include "MuT/NodeMesh.h"
@@ -15,8 +14,6 @@
 #include "MuT/NodeString.h"
 #include "MuT/NodeSurfaceNewPrivate.h"
 
-#include <memory>
-#include <utility>
 
 namespace mu
 {
