[mutable] Smart pointer cleanup. Move to proper atomics. Remove unnecessary weak pointer support.
Fix mistreatment of image-interpolate operations when number of mips differ.
#rnx

diff --git a/Engine/Plugins/Experimental/Mutable/Source/CustomizableObject/Classes/MuCO/CustomizableObject.h b/Engine/Plugins/Experimental/Mutable/Source/CustomizableObject/Classes/MuCO/CustomizableObject.h
--- a/Engine/Plugins/Experimental/Mutable/Source/CustomizableObject/Classes/MuCO/CustomizableObject.h
+++ b/Engine/Plugins/Experimental/Mutable/Source/CustomizableObject/Classes/MuCO/CustomizableObject.h
@@ -1417,10 +1417,10 @@ public:
 	// See UCustomizableObjectSystem::LockObject()
 	bool IsLocked() const;
 
-	mu::Model* GetModel() const;
+	TSharedPtr<mu::Model, ESPMode::ThreadSafe> GetModel() const;
 
 #if WITH_EDITOR
-	void SetModel(mu::Model* Model);
+	void SetModel(TSharedPtr<mu::Model, ESPMode::ThreadSafe> Model);
 #endif
 
 	int32 GetNumLODs() const;
diff --git a/Engine/Plugins/Experimental/Mutable/Source/CustomizableObject/Private/MuCO/CustomizableObject.cpp b/Engine/Plugins/Experimental/Mutable/Source/CustomizableObject/Private/MuCO/CustomizableObject.cpp
--- a/Engine/Plugins/Experimental/Mutable/Source/CustomizableObject/Private/MuCO/CustomizableObject.cpp
+++ b/Engine/Plugins/Experimental/Mutable/Source/CustomizableObject/Private/MuCO/CustomizableObject.cpp
@@ -358,7 +358,7 @@ void UCustomizableObject::ClearCompiledData()
 
 void UCustomizableObject::UpdateCompiledDataFromModel()
 {
-	mu::ModelPtr Model = PrivateData->GetModel();
+	TSharedPtr<mu::Model, ESPMode::ThreadSafe> Model = PrivateData->GetModel();
 
 	// Generate a map that using the resource id tells the offset and size of the resource inside the bulk data
 	if(Model)
@@ -555,9 +555,9 @@ void UCustomizableObject::LoadCompiledData(FArchive& MemoryReader, bool bSkipEdi
 		{
 			UnrealMutableInputStream stream(MemoryReader);
 			mu::InputArchive arch(&stream);
-			mu::ModelPtr Model = mu::Model::StaticUnserialise( arch );
+			TSharedPtr<mu::Model, ESPMode::ThreadSafe> Model = mu::Model::StaticUnserialise( arch );
 
-			PrivateData->SetModel(Model.get());
+			PrivateData->SetModel(Model);
 		}
 	}
 
@@ -869,7 +869,7 @@ void UCustomizableObject::SaveEmbeddedData(FArchive& Ar)
 
 			UnrealMutableOutputStream stream(Ar);
 			mu::OutputArchive arch(&stream);
-			mu::Model::Serialise(PrivateData->GetModel(), arch);
+			mu::Model::Serialise(PrivateData->GetModel().Get(), arch);
 		}
 
 		UE_LOG(LogMutable, Log, TEXT("Saved embedded data for Customizable Object [%s] now at position %d."), *GetName(), int(Ar.Tell()));
@@ -911,8 +911,8 @@ void UCustomizableObject::LoadEmbeddedData(FArchive& Ar)
 		// Load model
 		UnrealMutableInputStream stream(Ar);
 		mu::InputArchive arch(&stream);
-		mu::ModelPtr Model = mu::Model::StaticUnserialise( arch );
-		PrivateData->SetModel( Model.get() );
+		TSharedPtr<mu::Model, ESPMode::ThreadSafe> Model = mu::Model::StaticUnserialise( arch );
+		PrivateData->SetModel( Model );
 
 		// Create parameter properties
 		UpdateParameterPropertiesFromModel();
@@ -1089,13 +1089,13 @@ UCustomizableObjectInstance* UCustomizableObject::CreateInstance()
 }
 
 
-mu::Model* UCustomizableObject::GetModel() const
+TSharedPtr<mu::Model, ESPMode::ThreadSafe> UCustomizableObject::GetModel() const
 {
 	return PrivateData->GetModel();
 }
 
 #if WITH_EDITOR
-void UCustomizableObject::SetModel(mu::Model* Model)
+void UCustomizableObject::SetModel(TSharedPtr<mu::Model, ESPMode::ThreadSafe> Model)
 {
 	PrivateData->SetModel(Model);
 	
@@ -1179,7 +1179,7 @@ void UCustomizableObject::UpdateParameterPropertiesFromModel()
 {
 	if (PrivateData->GetModel())
 	{
-		mu::ParametersPtr MutableParameters = PrivateData->GetModel()->NewParameters();
+		mu::ParametersPtr MutableParameters = mu::Model::NewParameters(PrivateData->GetModel());
 		int paramCount = MutableParameters->GetCount();
 
 		ParameterProperties.Reset(paramCount);
@@ -1583,21 +1583,21 @@ FGuid UCustomizableObject::GetCompilationGuid() const
 //-------------------------------------------------------------------------------------------------
 //-------------------------------------------------------------------------------------------------
 
-void FCustomizableObjectPrivateData::SetModel(mu::Model* Model)
+void FCustomizableObjectPrivateData::SetModel(const TSharedPtr<mu::Model, ESPMode::ThreadSafe>& Model)
 {
 	MutableModel = Model;
 }
 
 
-mu::Model* FCustomizableObjectPrivateData::GetModel()
+const TSharedPtr<mu::Model, ESPMode::ThreadSafe>& FCustomizableObjectPrivateData::GetModel()
 {
-	return MutableModel.get();
+	return MutableModel;
 }
 
 
-const mu::Model* FCustomizableObjectPrivateData::GetModel() const
+TSharedPtr<const mu::Model, ESPMode::ThreadSafe> FCustomizableObjectPrivateData::GetModel() const
 {
-	return MutableModel.get();
+	return MutableModel;
 }
 
 //-------------------------------------------------------------------------------------------------
@@ -1689,7 +1689,7 @@ void UCustomizableObjectBulk::PrepareBulkData(UCustomizableObject* InOuter, cons
 	BulkDataFileNames.Empty();
 	
 	// Split the Streamable data into several separate files and fix up FileIndex and Offset of each StreamableBlock
-	if(mu::ModelPtr Model = CustomizableObject->GetModel())
+	if(TSharedPtr<const mu::Model, ESPMode::ThreadSafe> Model = CustomizableObject->GetModel())
 	{
 		const uint64 MaxChunkSize = UCustomizableObjectSystem::GetInstance()->GetMaxChunkSizeForPlatform(TargetPlatform);
 
diff --git a/Engine/Plugins/Experimental/Mutable/Source/CustomizableObject/Private/MuCO/CustomizableObjectInstanceDescriptor.cpp b/Engine/Plugins/Experimental/Mutable/Source/CustomizableObject/Private/MuCO/CustomizableObjectInstanceDescriptor.cpp
--- a/Engine/Plugins/Experimental/Mutable/Source/CustomizableObject/Private/MuCO/CustomizableObjectInstanceDescriptor.cpp
+++ b/Engine/Plugins/Experimental/Mutable/Source/CustomizableObject/Private/MuCO/CustomizableObjectInstanceDescriptor.cpp
@@ -480,7 +480,7 @@ mu::ParametersPtr FCustomizableObjectInstanceDescriptor::GetParameters() const
 		return nullptr;
 	}
 
-	mu::ParametersPtr MutableParameters = CustomizableObject->GetPrivate()->GetModel()->NewParameters();
+	mu::ParametersPtr MutableParameters = mu::Model::NewParameters( CustomizableObject->GetPrivate()->GetModel() );
 
 	const int32 ParamCount = MutableParameters->GetCount();
 	for (int32 ParamIndex = 0; ParamIndex < ParamCount; ++ParamIndex)
@@ -707,7 +707,7 @@ void FCustomizableObjectInstanceDescriptor::ReloadParameters()
 		return;
 	}
 
-	mu::ParametersPtr MutableParameters = CustomizableObject->GetPrivate()->GetModel()->NewParameters();
+	mu::ParametersPtr MutableParameters = mu::Model::NewParameters(CustomizableObject->GetPrivate()->GetModel());
 
 	int32 ParamCount = MutableParameters->GetCount();
 	for (int32 ParamIndex = 0; ParamIndex < ParamCount; ++ParamIndex)
@@ -1829,7 +1829,7 @@ bool FCustomizableObjectInstanceDescriptor::IsParamMultidimensional(const int32
 {
 	check(CustomizableObject);
 
-	const mu::ParametersPtr MutableParameters = CustomizableObject->GetPrivate()->GetModel()->NewParameters();
+	const mu::ParametersPtr MutableParameters = mu::Model::NewParameters(CustomizableObject->GetPrivate()->GetModel());
 	check(ParamIndex < MutableParameters->GetCount());
 	const mu::RangeIndexPtr RangeIdxPtr = MutableParameters->NewRangeIndex(ParamIndex);
 
@@ -2029,7 +2029,7 @@ FCustomizableObjectProjector FCustomizableObjectInstanceDescriptor::GetProjector
 {
 	check(CustomizableObject);
 
-	const mu::ParametersPtr MutableParameters = CustomizableObject->GetPrivate()->GetModel()->NewParameters();
+	const mu::ParametersPtr MutableParameters = mu::Model::NewParameters(CustomizableObject->GetPrivate()->GetModel());
 	check(ParamIndex < MutableParameters->GetCount());
 
 	FCustomizableObjectProjector Projector;
diff --git a/Engine/Plugins/Experimental/Mutable/Source/CustomizableObject/Private/MuCO/CustomizableObjectMipDataProvider.cpp b/Engine/Plugins/Experimental/Mutable/Source/CustomizableObject/Private/MuCO/CustomizableObjectMipDataProvider.cpp
--- a/Engine/Plugins/Experimental/Mutable/Source/CustomizableObject/Private/MuCO/CustomizableObjectMipDataProvider.cpp
+++ b/Engine/Plugins/Experimental/Mutable/Source/CustomizableObject/Private/MuCO/CustomizableObjectMipDataProvider.cpp
@@ -80,14 +80,14 @@ namespace impl
 		// This runs in a worker thread.
 		check(OperationData.IsValid());
 		check(OperationData->UpdateContext->System.get());
-		check(OperationData->UpdateContext->Model.get());
+		check(OperationData->UpdateContext->Model);
 		check(OperationData->UpdateContext->Parameters.get());
 
 		mu::SystemPtr System = OperationData->UpdateContext->System;
-		mu::ModelPtr Model = OperationData->UpdateContext->Model;
+		const TSharedPtr<mu::Model,ESPMode::ThreadSafe> Model = OperationData->UpdateContext->Model;
 
 		// For now, we are forcing the recreation of mutable-side instances with every update.
-		mu::Instance::ID InstanceID = System->NewInstance(Model.get());
+		mu::Instance::ID InstanceID = System->NewInstance(Model);
 		UE_LOG(LogMutable, Log, TEXT("Creating instance with id [%d] "), InstanceID)
 
 		const mu::Instance* Instance = nullptr;
diff --git a/Engine/Plugins/Experimental/Mutable/Source/CustomizableObject/Private/MuCO/CustomizableObjectPrivate.h b/Engine/Plugins/Experimental/Mutable/Source/CustomizableObject/Private/MuCO/CustomizableObjectPrivate.h
--- a/Engine/Plugins/Experimental/Mutable/Source/CustomizableObject/Private/MuCO/CustomizableObjectPrivate.h
+++ b/Engine/Plugins/Experimental/Mutable/Source/CustomizableObject/Private/MuCO/CustomizableObjectPrivate.h
@@ -11,13 +11,13 @@ class FCustomizableObjectPrivateData
 {
 private:
 
-	mu::ModelPtr MutableModel;
+	TSharedPtr<mu::Model, ESPMode::ThreadSafe> MutableModel;
 
 public:
 
-	void SetModel(mu::Model* Model);
-	mu::Model* GetModel();
-	const mu::Model* GetModel() const;
+	void SetModel(const TSharedPtr<mu::Model, ESPMode::ThreadSafe>& Model);
+	const TSharedPtr<mu::Model, ESPMode::ThreadSafe>& GetModel();
+	TSharedPtr<const mu::Model, ESPMode::ThreadSafe> GetModel() const;
 
 	// See UCustomizableObjectSystem::LockObject. Must only be modified from the game thread
 	bool bLocked = false;
diff --git a/Engine/Plugins/Experimental/Mutable/Source/CustomizableObject/Private/MuCO/CustomizableObjectSystem.cpp b/Engine/Plugins/Experimental/Mutable/Source/CustomizableObject/Private/MuCO/CustomizableObjectSystem.cpp
--- a/Engine/Plugins/Experimental/Mutable/Source/CustomizableObject/Private/MuCO/CustomizableObjectSystem.cpp
+++ b/Engine/Plugins/Experimental/Mutable/Source/CustomizableObject/Private/MuCO/CustomizableObjectSystem.cpp
@@ -1098,7 +1098,7 @@ void FCustomizableObjectSystemPrivate::UpdateStreamingLimit()
 namespace impl
 {
 
-	void Subtask_Mutable_UpdateParameterDecorations(const TSharedPtr<FMutableOperationData>& OperationData, const mu::ModelPtr& MutableModel, const mu::Parameters* MutableParameters)
+	void Subtask_Mutable_UpdateParameterDecorations(const TSharedPtr<FMutableOperationData>& OperationData, const TSharedPtr<mu::Model, ESPMode::ThreadSafe>& MutableModel, const mu::Parameters* MutableParameters)
 	{
 		MUTABLE_CPUPROFILER_SCOPE(UCustomizableInstancePrivateData::UpdateParameterDecorations)
 
@@ -1124,7 +1124,7 @@ namespace impl
 					for (int32 i = 0; i < MutableParameters->GetAdditionalImageCount(p); ++i)
 					{
 						mu::ImagePtrConst Image = MutableSystem->BuildParameterAdditionalImage(
-							MutableModel.get(),
+							MutableModel,
 							MutableParameters,
 							p, i);
 
@@ -1141,7 +1141,7 @@ namespace impl
 
 			TArray<bool> relevant;
 			relevant.SetNumZeroed(NumParameter);
-			MutableSystem->GetParameterRelevancy(MutableModel.get(), MutableParameters, relevant.GetData());
+			MutableSystem->GetParameterRelevancy(MutableModel, MutableParameters, relevant.GetData());
 
 			for (int32 p = 0; p < NumParameter; ++p)
 			{
@@ -1155,7 +1155,7 @@ namespace impl
 	}
 
 
-	void Subtask_Mutable_BeginUpdate_GetMesh(const TSharedPtr<FMutableOperationData>& OperationData, mu::ModelPtr Model, const mu::Parameters* MutableParameters, int32 State)
+	void Subtask_Mutable_BeginUpdate_GetMesh(const TSharedPtr<FMutableOperationData>& OperationData, TSharedPtr<mu::Model, ESPMode::ThreadSafe> Model, const mu::Parameters* MutableParameters, int32 State)
 	{
 		MUTABLE_CPUPROFILER_SCOPE(Subtask_Mutable_BeginUpdate);
 
@@ -1167,7 +1167,7 @@ namespace impl
 		mu::System* System = UCustomizableObjectSystem::GetInstance()->GetPrivate()->MutableSystem.get();
 
 		// For now, we are forcing the recreation of mutable-side instances with every update.
-		OperationData->InstanceID = System->NewInstance(Model.get());
+		OperationData->InstanceID = System->NewInstance(Model);
 		UE_LOG(LogMutable, Log, TEXT("Creating instance with id [%d] "), OperationData->InstanceID)
 
 		const mu::Instance* Instance = nullptr;
@@ -1310,7 +1310,7 @@ namespace impl
 	}
 
 	
-	void Subtask_Mutable_GetImages(const TSharedPtr<FMutableOperationData>& OperationData, mu::ModelPtr Model, const mu::Parameters* MutableParameters, int32 State)
+	void Subtask_Mutable_GetImages(const TSharedPtr<FMutableOperationData>& OperationData, TSharedPtr<mu::Model, ESPMode::ThreadSafe> Model, const mu::Parameters* MutableParameters, int32 State)
 	{
 		MUTABLE_CPUPROFILER_SCOPE(Subtask_Mutable_GetImages);
 
@@ -1694,7 +1694,7 @@ namespace impl
 	}
 
 
-	void Task_Mutable_Update_GetMesh(TSharedPtr<FMutableOperationData> OperationData, mu::ModelPtr Model, mu::ParametersPtrConst Parameters, bool bBuildParameterDecorations, int32 State)
+	void Task_Mutable_Update_GetMesh(TSharedPtr<FMutableOperationData> OperationData, TSharedPtr<mu::Model, ESPMode::ThreadSafe> Model, mu::ParametersPtrConst Parameters, bool bBuildParameterDecorations, int32 State)
 	{
 		// This runs in a worker thread.
 		MUTABLE_CPUPROFILER_SCOPE(Task_Mutable_GetMesh)
@@ -1729,7 +1729,7 @@ namespace impl
 	}
 
 
-	void Task_Mutable_Update_GetImages(TSharedPtr<FMutableOperationData> OperationData, mu::ModelPtr Model, mu::ParametersPtrConst Parameters, int32 State)
+	void Task_Mutable_Update_GetImages(TSharedPtr<FMutableOperationData> OperationData, TSharedPtr<mu::Model, ESPMode::ThreadSafe> Model, mu::ParametersPtrConst Parameters, int32 State)
 	{
 		// This runs in a worker thread.
 		MUTABLE_CPUPROFILER_SCOPE(Task_Mutable_GetImages)
@@ -2051,7 +2051,7 @@ namespace impl
 		FGraphEventRef Mutable_GetImagesTask;
 		{
 			// Task inputs
-			mu::ModelPtr Model = CustomizableObject->GetPrivate()->GetModel();
+			TSharedPtr<mu::Model, ESPMode::ThreadSafe> Model = CustomizableObject->GetPrivate()->GetModel();
 			int32 State = CustomizableObjectInstancePtr->GetState();
 
 			Mutable_GetImagesTask = System->GetPrivate()->AddMutableThreadTask(
@@ -2188,7 +2188,7 @@ namespace impl
 			TWeakObjectPtr<UCustomizableObjectInstance> CustomizableObjectInstancePtr = CandidateInstance;
 			TSharedPtr<FMutableOperation> CurrentMutableOperation = System->GetPrivate()->CurrentMutableOperation;
 			bool bBuildParameterDecorations = CurrentMutableOperation->IsBuildParameterDecorations();
-			mu::ModelPtr Model = CustomizableObject->GetPrivate()->GetModel();
+			TSharedPtr<mu::Model, ESPMode::ThreadSafe> Model = CustomizableObject->GetPrivate()->GetModel();
 			int32 State = CustomizableObjectInstancePtr->GetState();
 
 			Mutable_GetMeshTask = System->GetPrivate()->AddMutableThreadTask(
diff --git a/Engine/Plugins/Experimental/Mutable/Source/CustomizableObject/Private/MuCO/UnrealMutableModelDiskStreamer.cpp b/Engine/Plugins/Experimental/Mutable/Source/CustomizableObject/Private/MuCO/UnrealMutableModelDiskStreamer.cpp
--- a/Engine/Plugins/Experimental/Mutable/Source/CustomizableObject/Private/MuCO/UnrealMutableModelDiskStreamer.cpp
+++ b/Engine/Plugins/Experimental/Mutable/Source/CustomizableObject/Private/MuCO/UnrealMutableModelDiskStreamer.cpp
@@ -144,14 +144,14 @@ bool FUnrealMutableModelBulkStreamer::PrepareStreamingForObject(UCustomizableObj
 	// Is the object already prepared for streaming?
 	bool bAlreadyStreaming = Objects.FindByPredicate(
 		[CustomizableObject](const FObjectData& d) 
-		{ return d.Model.Pin() == CustomizableObject->GetModel(); })
+		{ return d.Model.Pin().Get() == CustomizableObject->GetModel().Get(); })
 		!= 
 		nullptr;
 
 	if (!bAlreadyStreaming)
 	{
 		FObjectData NewData;
-		NewData.Model = mu::WeakPtr<mu::Model>(CustomizableObject->GetModel());
+		NewData.Model = TWeakPtr<const mu::Model>(CustomizableObject->GetModel());
 
 #if WITH_EDITOR
 		FString FolderPath = CustomizableObject->GetCompiledDataFolderPath(true);
@@ -299,7 +299,7 @@ mu::ModelStreamer::OPERATION_ID FUnrealMutableModelBulkStreamer::BeginReadBlock(
 
 	// Find the object we are streaming for
 	FObjectData* ObjectData = Objects.FindByPredicate(
-		[Model](const FObjectData& d) {return d.Model.Pin()==Model; });
+		[Model](const FObjectData& d) {return d.Model.Pin().Get() == Model; });
 
 	if (!ObjectData)
 	{
diff --git a/Engine/Plugins/Experimental/Mutable/Source/CustomizableObject/Public/MuCO/CustomizableObjectMipDataProvider.h b/Engine/Plugins/Experimental/Mutable/Source/CustomizableObject/Public/MuCO/CustomizableObjectMipDataProvider.h
--- a/Engine/Plugins/Experimental/Mutable/Source/CustomizableObject/Public/MuCO/CustomizableObjectMipDataProvider.h
+++ b/Engine/Plugins/Experimental/Mutable/Source/CustomizableObject/Public/MuCO/CustomizableObjectMipDataProvider.h
@@ -31,7 +31,7 @@ class UTexture;
 struct FMutableUpdateContext
 {
 	mu::Ptr<mu::System> System;
-	mu::Ptr<mu::Model> Model;
+	TSharedPtr<mu::Model, ESPMode::ThreadSafe> Model;
 	mu::Ptr<const mu::Parameters> Parameters;
 	int32 State = -1;
 
diff --git a/Engine/Plugins/Experimental/Mutable/Source/CustomizableObject/Public/MuCO/UnrealMutableModelDiskStreamer.h b/Engine/Plugins/Experimental/Mutable/Source/CustomizableObject/Public/MuCO/UnrealMutableModelDiskStreamer.h
--- a/Engine/Plugins/Experimental/Mutable/Source/CustomizableObject/Public/MuCO/UnrealMutableModelDiskStreamer.h
+++ b/Engine/Plugins/Experimental/Mutable/Source/CustomizableObject/Public/MuCO/UnrealMutableModelDiskStreamer.h
@@ -104,7 +104,7 @@ protected:
 	/** Streaming data for one object. */
 	struct FObjectData
 	{
-		mu::WeakPtr<mu::Model> Model;
+		TWeakPtr<const mu::Model> Model;
 		TArray<IAsyncReadFileHandle*> ReadFileHandles;
 		TMap<OPERATION_ID, IAsyncReadRequest*> CurrentReadRequests;
 		TMap<uint64, FMutableStreamableBlock> StreamableBlocks;
diff --git a/Engine/Plugins/Experimental/Mutable/Source/CustomizableObjectEditor/Private/MuCOE/CustomizableObjectCompileRunnable.cpp b/Engine/Plugins/Experimental/Mutable/Source/CustomizableObjectEditor/Private/MuCOE/CustomizableObjectCompileRunnable.cpp
--- a/Engine/Plugins/Experimental/Mutable/Source/CustomizableObjectEditor/Private/MuCOE/CustomizableObjectCompileRunnable.cpp
+++ b/Engine/Plugins/Experimental/Mutable/Source/CustomizableObjectEditor/Private/MuCOE/CustomizableObjectCompileRunnable.cpp
@@ -198,7 +198,7 @@ FCustomizableObjectSaveDDRunnable::FCustomizableObjectSaveDDRunnable(UCustomizab
 
 uint32 FCustomizableObjectSaveDDRunnable::Run()
 {
-	bool bModelSerialized = Model.get() != nullptr;
+	bool bModelSerialized = Model.Get() != nullptr;
 
 	if (Options.bIsCooking && !Options.bSaveCookedDataToDisk)
 	{
@@ -210,7 +210,7 @@ uint32 FCustomizableObjectSaveDDRunnable::Run()
 		if (bModelSerialized)
 		{
 			FUnrealMutableModelBulkStreamer Streamer(&ModelMemoryWriter, &StreamableMemoryWriter);
-			mu::Model::Serialise(Model.get(), Streamer);
+			mu::Model::Serialise(Model.Get(), Streamer);
 		}
 	}
 	else if(bModelSerialized) // Save CO data + mu::Model and streamable resources to disk
@@ -254,7 +254,7 @@ uint32 FCustomizableObjectSaveDDRunnable::Run()
 			*ModelMemoryWriter << bModelSerialized;
 
 			FUnrealMutableModelBulkStreamer Streamer(ModelMemoryWriter, StreamableMemoryWriter);
-			mu::Model::Serialise(Model.get(), Streamer);
+			mu::Model::Serialise(Model.Get(), Streamer);
 
 			// Save to disk
 			ModelMemoryWriter->Flush();
diff --git a/Engine/Plugins/Experimental/Mutable/Source/CustomizableObjectEditor/Private/MuCOE/CustomizableObjectCompileRunnable.h b/Engine/Plugins/Experimental/Mutable/Source/CustomizableObjectEditor/Private/MuCOE/CustomizableObjectCompileRunnable.h
--- a/Engine/Plugins/Experimental/Mutable/Source/CustomizableObjectEditor/Private/MuCOE/CustomizableObjectCompileRunnable.h
+++ b/Engine/Plugins/Experimental/Mutable/Source/CustomizableObjectEditor/Private/MuCOE/CustomizableObjectCompileRunnable.h
@@ -62,7 +62,7 @@ public:
 
 public:
 
-	mu::ModelPtr Model;
+	TSharedPtr<mu::Model, ESPMode::ThreadSafe> Model;
 
 	// Texture packing strategy
 	bool bDisableTextureLayout;
@@ -112,7 +112,7 @@ private:
 	FString CompildeDataFullFileName;
 	FString StreamableDataFullFileName;
 
-	mu::ModelPtr Model;
+	TSharedPtr<mu::Model, ESPMode::ThreadSafe> Model;
 
 	// Bytes where the model is stored
 	TArray64<uint8> Bytes;
diff --git a/Engine/Plugins/Experimental/Mutable/Source/CustomizableObjectEditor/Private/MuCOE/CustomizableObjectCompiler.cpp b/Engine/Plugins/Experimental/Mutable/Source/CustomizableObjectEditor/Private/MuCOE/CustomizableObjectCompiler.cpp
--- a/Engine/Plugins/Experimental/Mutable/Source/CustomizableObjectEditor/Private/MuCOE/CustomizableObjectCompiler.cpp
+++ b/Engine/Plugins/Experimental/Mutable/Source/CustomizableObjectEditor/Private/MuCOE/CustomizableObjectCompiler.cpp
@@ -1401,9 +1401,9 @@ void FCustomizableObjectCompiler::FinishCompilation()
 	check(CompileTask.IsValid());
 
 	UpdateCompilerLogData();
-	mu::ModelPtr Model = CompileTask->Model;
+	TSharedPtr<mu::Model, ESPMode::ThreadSafe> Model = CompileTask->Model;
 
-	CurrentObject->SetModel(Model.get());
+	CurrentObject->SetModel(Model);
 
 	// Reset all instances, as the parameters may need to be rebuilt.
 	for (TObjectIterator<UCustomizableObjectInstance> It; It; ++It)
diff --git a/Engine/Plugins/Experimental/Mutable/Source/CustomizableObjectEditor/Private/MuCOE/CustomizableObjectDetails.cpp b/Engine/Plugins/Experimental/Mutable/Source/CustomizableObjectEditor/Private/MuCOE/CustomizableObjectDetails.cpp
--- a/Engine/Plugins/Experimental/Mutable/Source/CustomizableObjectEditor/Private/MuCOE/CustomizableObjectDetails.cpp
+++ b/Engine/Plugins/Experimental/Mutable/Source/CustomizableObjectEditor/Private/MuCOE/CustomizableObjectDetails.cpp
@@ -56,7 +56,7 @@ public:
 		ParameterIndex = -1;
 	}
 
-	mu::ModelPtrConst Model;
+	TSharedPtr<mu::Model, ESPMode::ThreadSafe> Model;
 	int StateIndex;
 	int ParameterIndex;
 
@@ -70,7 +70,7 @@ public:
 		}
 		else
 		{
-			mu::ParametersPtr TempParams = Model->NewParameters();
+			mu::ParametersPtr TempParams = mu::Model::NewParameters(Model);
 			res = TempParams->GetName( ParameterIndex );
 		}
 
diff --git a/Engine/Plugins/Experimental/Mutable/Source/CustomizableObjectEditor/Private/MuCOE/SMutableCodeViewer.cpp b/Engine/Plugins/Experimental/Mutable/Source/CustomizableObjectEditor/Private/MuCOE/SMutableCodeViewer.cpp
--- a/Engine/Plugins/Experimental/Mutable/Source/CustomizableObjectEditor/Private/MuCOE/SMutableCodeViewer.cpp
+++ b/Engine/Plugins/Experimental/Mutable/Source/CustomizableObjectEditor/Private/MuCOE/SMutableCodeViewer.cpp
@@ -431,10 +431,10 @@ FString SMutableCodeViewer::GetReferencerName() const
 }
 
 
-void SMutableCodeViewer::SetCurrentModel(const mu::ModelPtr& InMutableModel)
+void SMutableCodeViewer::SetCurrentModel(const TSharedPtr<mu::Model, ESPMode::ThreadSafe>& InMutableModel)
 {
 	MutableModel = InMutableModel;
-	PreviewParameters = MutableModel->NewParameters();
+	PreviewParameters = mu::Model::NewParameters(MutableModel);
 
 	RootNodes.Empty();
 	RootNodeAddresses.Empty();
@@ -471,7 +471,7 @@ void SMutableCodeViewer::SetCurrentModel(const mu::ModelPtr& InMutableModel)
 }
 
 
-void SMutableCodeViewer::Construct(const FArguments& InArgs, const mu::ModelPtr& InMutableModel/*, const TSharedPtr<SDockTab>& ConstructUnderMajorTab*/)
+void SMutableCodeViewer::Construct(const FArguments& InArgs, const TSharedPtr<mu::Model, ESPMode::ThreadSafe>& InMutableModel/*, const TSharedPtr<SDockTab>& ConstructUnderMajorTab*/)
 {
 	// Min width allowed for the column. Needed to avoid having issues with the constants space being to small
 	// and then getting too tall on the y axis crashing the UI drawer.
@@ -519,7 +519,7 @@ void SMutableCodeViewer::Construct(const FArguments& InArgs, const mu::ModelPtr&
 						mu::OutputFileStream Stream(TCHAR_TO_ANSI(*SaveFileName));
 						Stream.Write(MUTABLE_COMPILED_MODEL_FILETAG, 4);
 						mu::OutputArchive Archive(&Stream);
-						mu::Model::Serialise(InMutableModel.get(), Archive);
+						mu::Model::Serialise(InMutableModel.Get(), Archive);
 						const uint32_t CodeVersion = MUTABLE_COMPILED_MODEL_CODE_VERSION;
 						Stream.Write((const uint8*)&CodeVersion, sizeof(uint32_t));
 						Stream.Flush();
@@ -2559,20 +2559,16 @@ FReply SMutableCodeViewer::OnDrop(const FGeometry& MyGeometry, const FDragDropEv
 					if (!FMemory::Memcmp(MutableSourceTag, MUTABLE_COMPILED_MODEL_FILETAG, 4))
 					{
 						mu::InputArchive arch(&stream);
-						mu::ModelPtr Model = mu::Model::StaticUnserialise(arch);
+						TSharedPtr<mu::Model, ESPMode::ThreadSafe> Model = mu::Model::StaticUnserialise(arch);
 						SetCurrentModel(Model);
 
 						TreeView->RequestTreeRefresh();
 
 						return FReply::Handled();
 					}
-
-					return FReply::Unhandled();
 				}
 			}
 		}
-
-		return FReply::Unhandled();
 	}
 
 	return FReply::Unhandled();
diff --git a/Engine/Plugins/Experimental/Mutable/Source/CustomizableObjectEditor/Private/MuCOE/SMutableCodeViewer.h b/Engine/Plugins/Experimental/Mutable/Source/CustomizableObjectEditor/Private/MuCOE/SMutableCodeViewer.h
--- a/Engine/Plugins/Experimental/Mutable/Source/CustomizableObjectEditor/Private/MuCOE/SMutableCodeViewer.h
+++ b/Engine/Plugins/Experimental/Mutable/Source/CustomizableObjectEditor/Private/MuCOE/SMutableCodeViewer.h
@@ -74,7 +74,7 @@ public:
 		
 	SLATE_END_ARGS()
 
-	void Construct(const FArguments& InArgs, const mu::ModelPtr& InMutableModel /*, const TSharedPtr<SDockTab>& ConstructUnderMajorTab*/);
+	void Construct(const FArguments& InArgs, const TSharedPtr<mu::Model, ESPMode::ThreadSafe>& InMutableModel /*, const TSharedPtr<SDockTab>& ConstructUnderMajorTab*/);
 
 	// SWidget interface
 	virtual void Tick(const FGeometry& AllottedGeometry, const double InCurrentTime, const float InDeltaTime) override;
@@ -88,7 +88,7 @@ public:
 private:
 
 	/** The Mutable Model that we are showing. */
-	mu::ModelPtr MutableModel;
+	TSharedPtr<mu::Model, ESPMode::ThreadSafe> MutableModel;
 	
 	/** Selected model operation for preview. */
 	mu::OP::ADDRESS SelectedOperationAddress = 0;
@@ -193,7 +193,7 @@ private:
 	TMap< mu::OP::ADDRESS, TSharedPtr<FMutableCodeTreeElement>> ExpandedElements;
 
 	/** Prepare the widget for the given model. */
-	void SetCurrentModel(const mu::ModelPtr&);
+	void SetCurrentModel(const TSharedPtr<mu::Model, ESPMode::ThreadSafe>&);
 
 	/** Before any UI operation generate all the elements that may be navigable over the tree. No children of duplicated 
 	 * addresses will be generated.
@@ -736,7 +736,7 @@ public:
 class FMutableCodeTreeElement : public TSharedFromThis<FMutableCodeTreeElement>
 {
 public:
-	FMutableCodeTreeElement(int32 InIndexOnTree ,const mu::ModelPtr& InModel, mu::OP::ADDRESS InOperation, const FString& InCaption,const FSlateColor InLabelColor, const TSharedPtr<FMutableCodeTreeElement>* InDuplicatedOf = nullptr)
+	FMutableCodeTreeElement(int32 InIndexOnTree ,const TSharedPtr<mu::Model, ESPMode::ThreadSafe>& InModel, mu::OP::ADDRESS InOperation, const FString& InCaption,const FSlateColor InLabelColor, const TSharedPtr<FMutableCodeTreeElement>* InDuplicatedOf = nullptr)
 	{
 		MutableModel = InModel;
 		MutableOperation = InOperation;
@@ -791,7 +791,7 @@ public:
 public:
 
 	/** */
-	mu::ModelPtr MutableModel;
+	TSharedPtr<mu::Model, ESPMode::ThreadSafe> MutableModel;
 
 	/** Mutable Graph Node represented in this tree row. */
 	mu::OP::ADDRESS MutableOperation;
diff --git a/Engine/Plugins/Experimental/Mutable/Source/MutableRuntime/Private/MuR/CodeRunner.cpp b/Engine/Plugins/Experimental/Mutable/Source/MutableRuntime/Private/MuR/CodeRunner.cpp
--- a/Engine/Plugins/Experimental/Mutable/Source/MutableRuntime/Private/MuR/CodeRunner.cpp
+++ b/Engine/Plugins/Experimental/Mutable/Source/MutableRuntime/Private/MuR/CodeRunner.cpp
@@ -79,7 +79,7 @@ namespace mu
     //---------------------------------------------------------------------------------------------
     CodeRunner::CodeRunner(const SettingsPtrConst& pSettings,
 		class System::Private* s,
-		const Model* pModel,
+		const TSharedPtr<const Model>& pModel,
 		const Parameters* pParams,
 		OP::ADDRESS at,
 		uint32 lodMask, uint8 executionOptions, FScheduledOp::EType Type )
@@ -3233,10 +3233,9 @@ namespace mu
 
                     if (pMin && pMax)
                     {
-                        ImagePtr pNew = new Image( pMin->GetSizeX(),
-                                             pMin->GetSizeY(),
-                                             pMin->GetLODCount(),
-                                             pMin->GetFormat() );
+						int32 LevelCount = FMath::Max(pMin->GetLODCount(), pMax->GetLODCount());
+						
+						ImagePtr pNew = new Image( pMin->GetSizeX(), pMin->GetSizeY(), LevelCount, pMin->GetFormat() );
 
 						// Be defensive: ensure image sizes match.
 						if (pMin->GetSize() != pMax->GetSize())
@@ -3245,18 +3244,33 @@ namespace mu
 							pMax = ImageResizeLinear(0, pMax.get(), pMin->GetSize());
 						}
 
-						if (pMax->GetLODCount() != pMin->GetLODCount())
+						if (pMin->GetLODCount() != LevelCount)
+						{
+							MUTABLE_CPUPROFILER_SCOPE(Mipmap_ForInterpolate);
+
+							ImagePtr pDest = new Image(pMin->GetSizeX(), pMin->GetSizeY(), LevelCount, pMin->GetFormat());
+
+							SCRATCH_IMAGE_MIPMAP scratch;
+							FMipmapGenerationSettings settings{};
+
+							ImageMipmap_PrepareScratch(pDest.get(), pMin.get(), LevelCount, &scratch);
+							ImageMipmap(m_pSettings->GetPrivate()->m_imageCompressionQuality, pDest.get(), pMin.get(), LevelCount,
+								&scratch, settings);
+
+							pMin = pDest;
+						}
+
+						if (pMax->GetLODCount() != LevelCount)
 						{
 							MUTABLE_CPUPROFILER_SCOPE(Mipmap_ForInterpolate);
 
-							int32 levelCount = pMin->GetLODCount();
-							ImagePtr pDest = new Image(pMax->GetSizeX(), pMax->GetSizeY(), levelCount, pMax->GetFormat());
+							ImagePtr pDest = new Image(pMax->GetSizeX(), pMax->GetSizeY(), LevelCount, pMax->GetFormat());
 
 							SCRATCH_IMAGE_MIPMAP scratch;
 							FMipmapGenerationSettings settings{};
 
-							ImageMipmap_PrepareScratch(pDest.get(), pMax.get(), levelCount, &scratch);
-							ImageMipmap(m_pSettings->GetPrivate()->m_imageCompressionQuality, pDest.get(), pMax.get(), levelCount,
+							ImageMipmap_PrepareScratch(pDest.get(), pMax.get(), LevelCount, &scratch);
+							ImageMipmap(m_pSettings->GetPrivate()->m_imageCompressionQuality, pDest.get(), pMax.get(), LevelCount,
 								&scratch, settings);
 
 							pMax = pDest;
diff --git a/Engine/Plugins/Experimental/Mutable/Source/MutableRuntime/Private/MuR/CodeRunner.h b/Engine/Plugins/Experimental/Mutable/Source/MutableRuntime/Private/MuR/CodeRunner.h
--- a/Engine/Plugins/Experimental/Mutable/Source/MutableRuntime/Private/MuR/CodeRunner.h
+++ b/Engine/Plugins/Experimental/Mutable/Source/MutableRuntime/Private/MuR/CodeRunner.h
@@ -42,7 +42,7 @@ namespace  mu
     {
     public:
 		CodeRunner(const SettingsPtrConst&, class System::Private*, 
-			const Model* pModel, const Parameters* pParams,
+			const TSharedPtr<const Model>&, const Parameters* pParams,
 			OP::ADDRESS at, uint32 lodMask, uint8 executionOptions, FScheduledOp::EType );
 
     protected:
@@ -146,7 +146,7 @@ namespace  mu
 #endif
 
 			virtual ~FIssuedTask() {}
-			virtual bool Prepare(CodeRunner*, const Model*, bool& bOutFailed) { bOutFailed = false; return true; }
+			virtual bool Prepare(CodeRunner*, const TSharedPtr<const Model>&, bool& bOutFailed) { bOutFailed = false; return true; }
 			virtual void DoWork() {}
 			virtual void Complete(CodeRunner*) = 0;
 			virtual bool IsComplete(CodeRunner*)
@@ -170,7 +170,7 @@ namespace  mu
 			}
 
 			// FIssuedTask interface
-			bool Prepare(CodeRunner*, const Model*, bool& bOutFailed) override;
+			bool Prepare(CodeRunner*, const TSharedPtr<const Model>&, bool& bOutFailed) override;
 			void Complete(CodeRunner*) override;
 			bool IsComplete(CodeRunner*) override;
 
@@ -189,7 +189,7 @@ namespace  mu
 			}
 
 			// FIssuedTask interface
-			bool Prepare(CodeRunner*, const Model*, bool& bOutFailed) override;
+			bool Prepare(CodeRunner*, const TSharedPtr<const Model>&, bool& bOutFailed) override;
 			void Complete(CodeRunner*) override;
 			bool IsComplete(CodeRunner*) override;
 
@@ -302,7 +302,7 @@ namespace  mu
 		bool bUnrecoverableError = false;
 
 		System::Private* m_pSystem = nullptr;
-		const Model* m_pModel = nullptr;
+		TSharedPtr<const Model> m_pModel = nullptr;
 		const Parameters* m_pParams = nullptr;
 		uint32 m_lodMask = 0;
 
diff --git a/Engine/Plugins/Experimental/Mutable/Source/MutableRuntime/Private/MuR/CodeRunnerTasked.cpp b/Engine/Plugins/Experimental/Mutable/Source/MutableRuntime/Private/MuR/CodeRunnerTasked.cpp
--- a/Engine/Plugins/Experimental/Mutable/Source/MutableRuntime/Private/MuR/CodeRunnerTasked.cpp
+++ b/Engine/Plugins/Experimental/Mutable/Source/MutableRuntime/Private/MuR/CodeRunnerTasked.cpp
@@ -156,7 +156,7 @@ namespace mu
 				// Special processing in case it is an ImageDesc operation
 				if ( item.Type==FScheduledOp::EType::ImageDesc )
 				{
-					RunCodeImageDesc(item, m_pParams, m_pModel, m_lodMask);
+					RunCodeImageDesc(item, m_pParams, m_pModel.Get(), m_lodMask);
 					continue;
 				}
 
@@ -210,7 +210,7 @@ namespace mu
 				else
 				{
 					// Run immediately
-					RunCode(item, m_pParams, m_pModel, m_lodMask);
+					RunCode(item, m_pParams, m_pModel.Get(), m_lodMask);
 
 					if (ScheduledStagePerOp[item] == item.Stage + 1)
 					{
@@ -435,7 +435,6 @@ namespace mu
 
 		bool bApplyColorBlendToAlpha = (Args.flags & OP::ImageLayerArgs::F_APPLY_TO_ALPHA) != 0;
 
-		// Investigate why this causes problems sometimes 
 		//ImagePtr pNew = mu::CloneOrTakeOver<Image>(m_base.get());
 		ImagePtr pNew = m_base->Clone();
 		check(pNew->GetDataSize() == m_base->GetDataSize());
@@ -1192,7 +1191,7 @@ namespace mu
 	//---------------------------------------------------------------------------------------------
 	//---------------------------------------------------------------------------------------------
 	//---------------------------------------------------------------------------------------------
-	bool CodeRunner::FLoadMeshRomTask::Prepare(CodeRunner* runner, const Model* InModel, bool& bOutFailed )
+	bool CodeRunner::FLoadMeshRomTask::Prepare(CodeRunner* runner, const TSharedPtr<const Model>& InModel, bool& bOutFailed )
 	{
 		MUTABLE_CPUPROFILER_SCOPE(FLoadMeshRomTask_Prepare);
 
@@ -1214,7 +1213,7 @@ namespace mu
 				runner->m_pSystem->m_modelCache.UpdateForLoad(RomIndex, InModel,
 					[&](const Model* pModel, int romInd)
 					{
-						return (pModel == InModel)
+						return (pModel == InModel.Get())
 							&&
 							runner->m_romPendingOps[romInd] > 0;
 					});
@@ -1249,7 +1248,7 @@ namespace mu
 			}
 
 			uint32 RomId = program.m_roms[RomIndex].Id;
-			op->m_streamID = runner->m_pSystem->m_pStreamInterface->BeginReadBlock(InModel, RomId, op->m_streamBuffer.GetData(), RomSize);
+			op->m_streamID = runner->m_pSystem->m_pStreamInterface->BeginReadBlock(InModel.Get(), RomId, op->m_streamBuffer.GetData(), RomSize);
 			if (op->m_streamID < 0)
 			{
 				bOutFailed = true;
@@ -1273,7 +1272,7 @@ namespace mu
 		// This runs in the runner thread
 
 		// Process the constant op normally, now that the rom is loaded.
-		runner->RunCode(Op, runner->m_pParams, runner->m_pModel, runner->m_lodMask);
+		runner->RunCode(Op, runner->m_pParams, runner->m_pModel.Get(), runner->m_lodMask);
 
 		runner->m_pSystem->m_modelCache.MarkRomUsed(RomIndex, runner->m_pModel);
 		--runner->m_romPendingOps[RomIndex];
@@ -1291,7 +1290,7 @@ namespace mu
 	//---------------------------------------------------------------------------------------------
 	//---------------------------------------------------------------------------------------------
 	//---------------------------------------------------------------------------------------------
-	bool CodeRunner::FLoadImageRomsTask::Prepare(CodeRunner* runner, const Model* InModel, bool& bOutFailed )
+	bool CodeRunner::FLoadImageRomsTask::Prepare(CodeRunner* runner, const TSharedPtr<const Model>& InModel, bool& bOutFailed )
 	{
 		MUTABLE_CPUPROFILER_SCOPE(FLoadImageRomsTask_Prepare);
 
@@ -1334,7 +1333,7 @@ namespace mu
 				runner->m_pSystem->m_modelCache.UpdateForLoad(RomIndex, InModel,
 					[&](const Model* pModel, int romInd)
 					{
-						return (pModel == InModel)
+						return (pModel == InModel.Get())
 							&&
 							runner->m_romPendingOps[romInd] > 0;
 					});
@@ -1369,7 +1368,7 @@ namespace mu
 			}
 
 			uint32 RomId = program.m_roms[RomIndex].Id;
-			op->m_streamID = runner->m_pSystem->m_pStreamInterface->BeginReadBlock(InModel, RomId, op->m_streamBuffer.GetData(), RomSize);
+			op->m_streamID = runner->m_pSystem->m_pStreamInterface->BeginReadBlock(InModel.Get(), RomId, op->m_streamBuffer.GetData(), RomSize);
 			if (op->m_streamID < 0)
 			{
 				bOutFailed = true;
@@ -1392,7 +1391,7 @@ namespace mu
 		// This runs in the runner thread
 
 		// Process the constant op normally, now that the rom is loaded.
-		runner->RunCode(Op, runner->m_pParams, runner->m_pModel, runner->m_lodMask);
+		runner->RunCode(Op, runner->m_pParams, runner->m_pModel.Get(), runner->m_lodMask);
 
 		FProgram& program = runner->m_pModel->GetPrivate()->m_program;
 		for (int32 i = 0; i < LODIndexCount; ++i)
diff --git a/Engine/Plugins/Experimental/Mutable/Source/MutableRuntime/Private/MuR/CodeVisitor.h b/Engine/Plugins/Experimental/Mutable/Source/MutableRuntime/Private/MuR/CodeVisitor.h
--- a/Engine/Plugins/Experimental/Mutable/Source/MutableRuntime/Private/MuR/CodeVisitor.h
+++ b/Engine/Plugins/Experimental/Mutable/Source/MutableRuntime/Private/MuR/CodeVisitor.h
@@ -461,7 +461,7 @@ class Parameters;
         DiscreteCoveredCodeVisitorBase
             (
                 System::Private* pSystem,
-                const ModelPtrConst& pModel,
+                const TSharedPtr<const Model>& pModel,
                 const ParametersPtrConst& pParams,
                 unsigned lodMask,
                 bool skipResources=false
@@ -654,7 +654,7 @@ class Parameters;
 
     protected:
         System::Private* m_pSystem = nullptr;
-        ModelPtrConst m_pModel;
+		TSharedPtr<const Model> m_pModel;
         const Parameters* m_pParams = nullptr;
         unsigned m_lodMask = 0;
     };
@@ -687,7 +687,7 @@ class Parameters;
         UniqueDiscreteCoveredCodeVisitor
             (
                 System::Private* pSystem,
-                const ModelPtrConst& pModel,
+				const TSharedPtr<const Model>& pModel,
                 const ParametersPtrConst& pParams,
                 unsigned lodMask
             )
diff --git a/Engine/Plugins/Experimental/Mutable/Source/MutableRuntime/Private/MuR/Model.cpp b/Engine/Plugins/Experimental/Mutable/Source/MutableRuntime/Private/MuR/Model.cpp
--- a/Engine/Plugins/Experimental/Mutable/Source/MutableRuntime/Private/MuR/Model.cpp
+++ b/Engine/Plugins/Experimental/Mutable/Source/MutableRuntime/Private/MuR/Model.cpp
@@ -292,11 +292,11 @@ namespace mu
 
 
     //---------------------------------------------------------------------------------------------
-    ModelPtr Model::StaticUnserialise( InputArchive& arch )
+	TSharedPtr<Model> Model::StaticUnserialise( InputArchive& arch )
     {
 		LLM_SCOPE_BYNAME(TEXT("MutableRuntime"));
 
-        ModelPtr pResult = new Model();
+		TSharedPtr<Model> pResult = MakeShared<Model>();
         arch >> *pResult->m_pD;
         return pResult;
     }
@@ -339,18 +339,19 @@ namespace mu
 
 
     //---------------------------------------------------------------------------------------------
-    ParametersPtr Model::NewParameters( const Parameters* pOld ) const
+    ParametersPtr Model::NewParameters(TSharedPtr<const Model> Model, const Parameters* pOld )
     {
 		LLM_SCOPE_BYNAME(TEXT("MutableRuntime"));
 
         ParametersPtr pRes = new Parameters();
 
-        pRes->GetPrivate()->m_pModel = this;
+        pRes->GetPrivate()->m_pModel = Model;
 
-        pRes->GetPrivate()->m_values.SetNum( m_pD->m_program.m_parameters.Num() );
-        for ( int32 p=0; p<m_pD->m_program.m_parameters.Num(); ++p )
+		const FProgram& Program = Model->GetPrivate()->m_program;
+        pRes->GetPrivate()->m_values.SetNum(Program.m_parameters.Num());
+        for ( int32 p=0; p< Program.m_parameters.Num(); ++p )
         {
-            pRes->GetPrivate()->m_values[p] = m_pD->m_program.m_parameters[p].m_defaultValue;
+            pRes->GetPrivate()->m_values[p] = Program.m_parameters[p].m_defaultValue;
         }
 
         // Copy old values
@@ -721,7 +722,7 @@ namespace mu
     //---------------------------------------------------------------------------------------------
     ModelParametersGenerator::ModelParametersGenerator
         (
-            const Model* pModel,
+			TSharedPtr<const Model> pModel,
             System* pSystem
         )
     {
@@ -791,7 +792,7 @@ namespace mu
     {
 		LLM_SCOPE_BYNAME(TEXT("MutableRuntime"));
 
-		Ptr<Parameters> res = m_pD->m_pModel->NewParameters();
+		Ptr<Parameters> res = Model::NewParameters(m_pD->m_pModel);
 
 		const FProgram& Program = m_pD->m_pModel->GetPrivate()->m_program;
         int paramCount = Program.m_parameters.Num();
@@ -938,7 +939,7 @@ namespace mu
     {
 		LLM_SCOPE_BYNAME(TEXT("MutableRuntime"));
 
-        ParametersPtr res = m_pD->m_pModel->NewParameters();
+        ParametersPtr res = Model::NewParameters(m_pD->m_pModel);
 
 		const FProgram& Program = m_pD->m_pModel->GetPrivate()->m_program;
         int32 ParamCount = Program.m_parameters.Num();
diff --git a/Engine/Plugins/Experimental/Mutable/Source/MutableRuntime/Private/MuR/ModelPrivate.h b/Engine/Plugins/Experimental/Mutable/Source/MutableRuntime/Private/MuR/ModelPrivate.h
--- a/Engine/Plugins/Experimental/Mutable/Source/MutableRuntime/Private/MuR/ModelPrivate.h
+++ b/Engine/Plugins/Experimental/Mutable/Source/MutableRuntime/Private/MuR/ModelPrivate.h
@@ -730,7 +730,7 @@ namespace mu
     public:
 
         //!
-        Ptr<const Model> m_pModel;
+        TSharedPtr<const Model> m_pModel;
         Ptr<System> m_pSystem;
 
         //! Number of possible instances
diff --git a/Engine/Plugins/Experimental/Mutable/Source/MutableRuntime/Private/MuR/Parameters.cpp b/Engine/Plugins/Experimental/Mutable/Source/MutableRuntime/Private/MuR/Parameters.cpp
--- a/Engine/Plugins/Experimental/Mutable/Source/MutableRuntime/Private/MuR/Parameters.cpp
+++ b/Engine/Plugins/Experimental/Mutable/Source/MutableRuntime/Private/MuR/Parameters.cpp
@@ -52,13 +52,6 @@ namespace mu
 	}
 
 
-	//---------------------------------------------------------------------------------------------
-	void Parameters::SetUnserialisedModel( Model* pModel )
-	{
-		m_pD->m_pModel = pModel;
-	}
-
-
     //---------------------------------------------------------------------------------------------
     void Parameters::SerialisePortable( const Parameters* params, OutputArchive& arch )
     {
@@ -80,11 +73,11 @@ namespace mu
 
 
     //---------------------------------------------------------------------------------------------
-    ParametersPtr Parameters::UnserialisePortable( InputArchive& arch, const Model* pModel )
+    ParametersPtr Parameters::UnserialisePortable( InputArchive& arch, TSharedPtr<const Model> pModel )
     {
 		LLM_SCOPE_BYNAME(TEXT("MutableRuntime"));
 
-        ParametersPtr pResult = pModel->NewParameters();
+        ParametersPtr pResult = Model::NewParameters(pModel);
         size_t modelParameters = pModel->GetPrivate()->m_program.m_parameters.Num();
 
         int32_t ver;
@@ -1174,11 +1167,11 @@ namespace mu
     {
         check( index >= 0 && index < GetRangeCount() );
         check( index < int(m_pD->m_pParameters->GetPrivate()->m_values.Num()) );
-        ModelPtrConst m_pModel = m_pD->m_pParameters->GetPrivate()->m_pModel;
-        check( m_pModel );
-        check( index < int(m_pModel->GetPrivate()->m_program.m_ranges.Num()) );
+        TSharedPtr<const Model> pModel = m_pD->m_pParameters->GetPrivate()->m_pModel;
+        check(pModel);
+        check( index < int(pModel->GetPrivate()->m_program.m_ranges.Num()) );
 
-        return m_pModel->GetPrivate()->m_program.m_ranges[index].m_name.c_str();
+        return pModel->GetPrivate()->m_program.m_ranges[index].m_name.c_str();
     }
 
 
@@ -1187,11 +1180,11 @@ namespace mu
     {
         check( index >= 0 && index < GetRangeCount() );
         check( index < int(m_pD->m_pParameters->GetPrivate()->m_values.Num()) );
-        ModelPtrConst m_pModel = m_pD->m_pParameters->GetPrivate()->m_pModel;
-        check( m_pModel );
-        check( index < int(m_pModel->GetPrivate()->m_program.m_ranges.Num()) );
+		TSharedPtr<const Model> pModel = m_pD->m_pParameters->GetPrivate()->m_pModel;
+		check(pModel);
+        check( index < int(pModel->GetPrivate()->m_program.m_ranges.Num()) );
 
-        return m_pModel->GetPrivate()->m_program.m_ranges[index].m_uid.c_str();
+        return pModel->GetPrivate()->m_program.m_ranges[index].m_uid.c_str();
     }
 
 
diff --git a/Engine/Plugins/Experimental/Mutable/Source/MutableRuntime/Private/MuR/ParametersPrivate.h b/Engine/Plugins/Experimental/Mutable/Source/MutableRuntime/Private/MuR/ParametersPrivate.h
--- a/Engine/Plugins/Experimental/Mutable/Source/MutableRuntime/Private/MuR/ParametersPrivate.h
+++ b/Engine/Plugins/Experimental/Mutable/Source/MutableRuntime/Private/MuR/ParametersPrivate.h
@@ -463,7 +463,7 @@ namespace mu
         //! Warning: update Parameters::Clone method if this members change.
 
         //! Run-time data
-        ModelPtrConst m_pModel;
+        TSharedPtr<const Model> m_pModel;
 
         //! Values for the parameters if they are not multidimensional.
 		TArray<PARAMETER_VALUE> m_values;
diff --git a/Engine/Plugins/Experimental/Mutable/Source/MutableRuntime/Private/MuR/System.cpp b/Engine/Plugins/Experimental/Mutable/Source/MutableRuntime/Private/MuR/System.cpp
--- a/Engine/Plugins/Experimental/Mutable/Source/MutableRuntime/Private/MuR/System.cpp
+++ b/Engine/Plugins/Experimental/Mutable/Source/MutableRuntime/Private/MuR/System.cpp
@@ -194,7 +194,7 @@ namespace mu
 
 
     //---------------------------------------------------------------------------------------------
-    Instance::ID System::NewInstance( const ModelPtrConst& pModel )
+    Instance::ID System::NewInstance( const TSharedPtr<const Model>& pModel )
     {
 		LLM_SCOPE_BYNAME(TEXT("MutableRuntime"));
 		MUTABLE_CPUPROFILER_SCOPE(NewInstance);
@@ -268,10 +268,10 @@ namespace mu
 
 		OP::ADDRESS rootAt = pLiveInstance->m_pModel->GetPrivate()->m_program.m_states[stateIndex].m_root;
 
-		m_pD->PrepareCache(pLiveInstance->m_pModel, stateIndex);
+		m_pD->PrepareCache(pLiveInstance->m_pModel.Get(), stateIndex);
 		pLiveInstance->m_pOldParameters = pParams->Clone();
 
-		m_pD->RunCode(pLiveInstance->m_pModel.get(), pParams.get(), rootAt, lodMask);
+		m_pD->RunCode(pLiveInstance->m_pModel, pParams.get(), rootAt, lodMask);
 
 		InstancePtrConst pResult = pLiveInstance->m_memory->GetInstance(FCacheAddress(rootAt, 0, 0));
 
@@ -356,7 +356,7 @@ namespace mu
 		{
 			if (res.m_id == imageId)
 			{
-				const mu::Model* Model = pLiveInstance->m_pModel.get();
+				const mu::Model* Model = pLiveInstance->m_pModel.Get();
 				const mu::FProgram& program = Model->GetPrivate()->m_program;
 
 				int32 VarValue = CVarClearImageDescCache.GetValueOnAnyThread();
@@ -372,7 +372,7 @@ namespace mu
 				if (GetOpDataType(opType) == DT_IMAGE)
 				{
 					int8 executionOptions = 0;
-					CodeRunner Runner(m_pD->m_pSettings, m_pD, Model, pLiveInstance->m_pOldParameters.get(), at, System::AllLODs, executionOptions, FScheduledOp::EType::ImageDesc);
+					CodeRunner Runner(m_pD->m_pSettings, m_pD, pLiveInstance->m_pModel, pLiveInstance->m_pOldParameters.get(), at, System::AllLODs, executionOptions, FScheduledOp::EType::ImageDesc);
 					Runner.Run();
 					Runner.GetImageDescResult(OutDesc);
 				}
@@ -461,7 +461,7 @@ namespace mu
 
 
     //---------------------------------------------------------------------------------------------
-    ImagePtrConst System::BuildParameterAdditionalImage( const Ptr<const Model>& pModel,
+    ImagePtrConst System::BuildParameterAdditionalImage(const TSharedPtr<const Model>& pModel,
                                                          const Ptr<const Parameters>& pParams,
                                                          int32 parameter,
                                                          int32 ImageIndex )
@@ -499,7 +499,7 @@ namespace mu
         RelevantParameterVisitor
             (
                 System::Private* pSystem,
-                const Ptr<const Model>& pModel,
+				const TSharedPtr<const Model>& pModel,
                 const Ptr<const Parameters>& pParams,
                 bool* pFlags
             )
@@ -552,7 +552,7 @@ namespace mu
 
 
     //---------------------------------------------------------------------------------------------
-    void System::GetParameterRelevancy( const ModelPtrConst& pModel,
+    void System::GetParameterRelevancy(const TSharedPtr<const Model>& pModel,
                                         const ParametersPtrConst& pParameters,
                                         bool* pFlags )
     {
@@ -619,14 +619,14 @@ namespace mu
 
 
 	//---------------------------------------------------------------------------------------------
-	void System::Private::BeginBuild(const Ptr<const Model>& pModel)
+	void System::Private::BeginBuild(const TSharedPtr<const Model>& pModel)
 	{
 		// We don't have a FLiveInstance, let's create the memory
 		// \TODO: There is no clear moment to remove this... EndBuild?
 		m_memory = MakeShared<FProgramCache>();
 		m_memory->Init(pModel->GetPrivate()->m_program.m_opAddress.Num());
 
-		PrepareCache(pModel, -1);
+		PrepareCache(pModel.Get(), -1);
 	}
 
 
@@ -638,10 +638,10 @@ namespace mu
 
 
 	//---------------------------------------------------------------------------------------------
-	void System::Private::RunCode(const Ptr<const Model>& InModel,
+	void System::Private::RunCode(const TSharedPtr<const Model>& InModel,
 		const Parameters* InParameters, OP::ADDRESS InCodeRoot, uint32 InLODs, uint8 executionOptions)
 	{
-		CodeRunner Runner(m_pSettings, this, InModel.get(), InParameters, InCodeRoot, InLODs, 
+		CodeRunner Runner(m_pSettings, this, InModel, InParameters, InCodeRoot, InLODs, 
 			executionOptions, FScheduledOp::EType::Full);
 		Runner.Run();
 		bUnrecoverableError = Runner.bUnrecoverableError;
@@ -649,11 +649,11 @@ namespace mu
 
 
 	//---------------------------------------------------------------------------------------------
-	bool System::Private::BuildBool(const Ptr<const Model>& pModel,
+	bool System::Private::BuildBool(const TSharedPtr<const Model>& pModel,
 		const Parameters* pParams,
 		OP::ADDRESS at)
 	{
-		RunCode(pModel.get(), pParams, at);
+		RunCode(pModel, pParams, at);
 		if (bUnrecoverableError)
 		{
 			return false;
@@ -663,11 +663,11 @@ namespace mu
 
 
 	//---------------------------------------------------------------------------------------------
-	float System::Private::BuildScalar(const Ptr<const Model>& pModel,
+	float System::Private::BuildScalar(const TSharedPtr<const Model>& pModel,
 		const Parameters* pParams,
 		OP::ADDRESS at)
 	{
-		RunCode(pModel.get(), pParams, at);
+		RunCode(pModel, pParams, at);
 		if (bUnrecoverableError)
 		{
 			return 0.0f;
@@ -677,11 +677,11 @@ namespace mu
 
 
 	//---------------------------------------------------------------------------------------------
-	int System::Private::BuildInt(const Ptr<const Model>& pModel,
+	int System::Private::BuildInt(const TSharedPtr<const Model>& pModel,
 		const Parameters* pParams,
 		OP::ADDRESS at)
 	{
-		RunCode(pModel.get(), pParams, at);
+		RunCode(pModel, pParams, at);
 		if (bUnrecoverableError)
 		{
 			return 0;
@@ -692,7 +692,7 @@ namespace mu
 
 
 	//---------------------------------------------------------------------------------------------
-	void System::Private::BuildColour(const Ptr<const Model>& pModel,
+	void System::Private::BuildColour(const TSharedPtr<const Model>& pModel,
 		const Parameters* pParams,
 		OP::ADDRESS at,
 		float* pR,
@@ -704,7 +704,7 @@ namespace mu
 		mu::OP_TYPE opType = pModel->GetPrivate()->m_program.GetOpType(at);
 		if (GetOpDataType(opType) == DT_COLOUR)
 		{
-			RunCode(pModel.get(), pParams, at);
+			RunCode(pModel, pParams, at);
 			if (bUnrecoverableError)
 			{
 				if (pR) *pR = 0.0f;
@@ -722,9 +722,9 @@ namespace mu
 
 	
 	//---------------------------------------------------------------------------------------------
-	Ptr<const Projector> System::Private::BuildProjector(const Ptr<const Model>& pModel, const Parameters* pParams, OP::ADDRESS at)
+	Ptr<const Projector> System::Private::BuildProjector(const TSharedPtr<const Model>& pModel, const Parameters* pParams, OP::ADDRESS at)
 	{
-    	RunCode(pModel.get(), pParams, at);
+    	RunCode(pModel, pParams, at);
 		if (bUnrecoverableError)
 		{
 			return nullptr;
@@ -734,7 +734,7 @@ namespace mu
 
 	
 	//---------------------------------------------------------------------------------------------
-	Ptr<const Image> System::Private::BuildImage(const Ptr<const Model>& pModel,
+	Ptr<const Image> System::Private::BuildImage(const TSharedPtr<const Model>& pModel,
 		const Parameters* pParams,
 		OP::ADDRESS at, int MipsToSkip)
 	{
@@ -742,7 +742,7 @@ namespace mu
 		if (GetOpDataType(opType) == DT_IMAGE)
 		{
 			m_memory->ClearCacheLayer0();
-			RunCode(pModel.get(), pParams, at, System::AllLODs, uint8(MipsToSkip));
+			RunCode(pModel, pParams, at, System::AllLODs, uint8(MipsToSkip));
 			if (bUnrecoverableError)
 			{
 				return nullptr;
@@ -756,7 +756,7 @@ namespace mu
 
 
 	//---------------------------------------------------------------------------------------------
-	MeshPtrConst System::Private::BuildMesh(const Ptr<const Model>& pModel,
+	MeshPtrConst System::Private::BuildMesh(const TSharedPtr<const Model>& pModel,
 		const Parameters* pParams,
 		OP::ADDRESS at)
 	{
@@ -764,7 +764,7 @@ namespace mu
 		if (GetOpDataType(opType) == DT_MESH)
 		{
 			m_memory->ClearCacheLayer0();
-			RunCode(pModel.get(), pParams, at);
+			RunCode(pModel, pParams, at);
 			if (bUnrecoverableError)
 			{
 				return nullptr;
@@ -778,7 +778,7 @@ namespace mu
 
 
 	//---------------------------------------------------------------------------------------------
-	LayoutPtrConst System::Private::BuildLayout(const Ptr<const Model>& pModel,
+	LayoutPtrConst System::Private::BuildLayout(const TSharedPtr<const Model>& pModel,
 		const Parameters* pParams,
 		OP::ADDRESS at)
 	{
@@ -789,7 +789,7 @@ namespace mu
 			mu::OP_TYPE opType = pModel->GetPrivate()->m_program.GetOpType(at);
 			if (GetOpDataType(opType) == DT_LAYOUT)
 			{
-				RunCode(pModel.get(), pParams, at);
+				RunCode(pModel, pParams, at);
 				if (bUnrecoverableError)
 				{
 					return nullptr;
@@ -803,7 +803,7 @@ namespace mu
 
 
 	//---------------------------------------------------------------------------------------------
-	Ptr<const String> System::Private::BuildString(const Ptr<const Model>& pModel,
+	Ptr<const String> System::Private::BuildString(const TSharedPtr<const Model>& pModel,
 		const Parameters* pParams,
 		OP::ADDRESS at)
 	{
@@ -814,7 +814,7 @@ namespace mu
 			mu::OP_TYPE opType = pModel->GetPrivate()->m_program.GetOpType(at);
 			if (GetOpDataType(opType) == DT_STRING)
 			{
-				RunCode(pModel.get(), pParams, at);
+				RunCode(pModel, pParams, at);
 				if (bUnrecoverableError)
 				{
 					return nullptr;
@@ -828,7 +828,7 @@ namespace mu
 
 
 	//---------------------------------------------------------------------------------------------
-	void System::Private::PrepareCache(const Ptr<const Model>& pModel, int state)
+	void System::Private::PrepareCache( const Model* pModel, int state)
 	{
 		MUTABLE_CPUPROFILER_SCOPE(PrepareCache);
 
@@ -852,16 +852,16 @@ namespace mu
     //---------------------------------------------------------------------------------------------
     //---------------------------------------------------------------------------------------------
     //---------------------------------------------------------------------------------------------
-	FModelCache::FModelCacheEntry& FModelCache::GetModelCache( const Model* m )
+	FModelCache::FModelCacheEntry& FModelCache::GetModelCache(const TSharedPtr<const Model>& InModel )
     {
-        check(m);
+        check(InModel);
 
         for(FModelCacheEntry& c:m_cachePerModel)
         {
-            mu::Ptr<mu::Model> pCandidate = c.m_pModel.Pin();
+			TSharedPtr<const Model> pCandidate = c.m_pModel.Pin();
             if (pCandidate)
             {
-                if (pCandidate==m)
+                if (pCandidate==InModel)
                 {
                     return c;
                 }
@@ -875,7 +875,7 @@ namespace mu
 
         // Not found. Add new
 		FModelCacheEntry n;
-        n.m_pModel = WeakPtr<Model>(m);
+        n.m_pModel = TWeakPtr<const Model>(InModel);
         m_cachePerModel.Add(n);
         return m_cachePerModel.Last();
     }
@@ -889,7 +889,7 @@ namespace mu
 		uint64 totalMemory = 0;
         for (FModelCacheEntry& m : m_cachePerModel)
         {
-			mu::Ptr<mu::Model> pCacheModel = m.m_pModel.Pin();
+			TSharedPtr<const Model> pCacheModel = m.m_pModel.Pin();
             if (pCacheModel)
             {
 				mu::FProgram& program = pCacheModel->GetPrivate()->m_program;
@@ -910,12 +910,12 @@ namespace mu
             bool finished = totalMemory < m_romBudget;
             while (!finished)
             {
-                ModelPtr lowestPriorityModel;
+				TSharedPtr<const Model> lowestPriorityModel;
                 int32 lowestPriorityRom = -1;
                 float lowestPriority = 0.0f;
                 for (FModelCacheEntry& modelCache : m_cachePerModel)
                 {
-					mu::Ptr<mu::Model> pCacheModel = modelCache.m_pModel.Pin();
+					TSharedPtr<const Model> pCacheModel = modelCache.m_pModel.Pin();
                     if (pCacheModel)
                     {
 						mu::FProgram& program = pCacheModel->GetPrivate()->m_program;
@@ -928,7 +928,7 @@ namespace mu
 
                             if (bIsLoaded
                                 &&
-                                (!isRomLockedFunc(pCacheModel.get(),RomIndex)) )
+                                (!isRomLockedFunc(pCacheModel.Get(),RomIndex)) )
                             {
                                 constexpr float factorWeight = 100.0f;
                                 constexpr float factorTime = -1.0f;
@@ -970,7 +970,7 @@ namespace mu
 
 
     //---------------------------------------------------------------------------------------------
-    void FModelCache::MarkRomUsed( int romIndex, const Model* pModel )
+    void FModelCache::MarkRomUsed( int romIndex, const TSharedPtr<const Model>& pModel )
     {
         check(pModel);
 
@@ -1001,7 +1001,7 @@ namespace mu
 
 
     //---------------------------------------------------------------------------------------------
-    void FModelCache::UpdateForLoad( int romIndex, const Model* pModel,
+    void FModelCache::UpdateForLoad( int romIndex, const TSharedPtr<const Model>& pModel,
                                      TFunctionRef<bool(const Model*,int)> isRomLockedFunc )
     {
 		MarkRomUsed( romIndex, pModel);
diff --git a/Engine/Plugins/Experimental/Mutable/Source/MutableRuntime/Private/MuR/SystemPrivate.h b/Engine/Plugins/Experimental/Mutable/Source/MutableRuntime/Private/MuR/SystemPrivate.h
--- a/Engine/Plugins/Experimental/Mutable/Source/MutableRuntime/Private/MuR/SystemPrivate.h
+++ b/Engine/Plugins/Experimental/Mutable/Source/MutableRuntime/Private/MuR/SystemPrivate.h
@@ -870,7 +870,7 @@ namespace mu
         struct FModelCacheEntry
         {
             //!
-            WeakPtr<Model> m_pModel;
+            TWeakPtr<const Model> m_pModel;
 
             // Rom streaming management
             TArray< TPair<uint64,uint64> > m_romWeight;
@@ -888,11 +888,11 @@ namespace mu
 			TFunctionRef<bool(const Model*, int)> IsRomLockedFunc = [](const Model*, int) {return false;});
 
         //!
-        void UpdateForLoad( int romIndex, const Model* pModel, TFunctionRef<bool(const Model*,int)> isRomLockedFunc );
-        void MarkRomUsed( int romIndex, const Model* pModel );
+        void UpdateForLoad( int romIndex, const TSharedPtr<const Model>& pModel, TFunctionRef<bool(const Model*,int)> isRomLockedFunc );
+        void MarkRomUsed( int romIndex, const TSharedPtr<const Model>& pModel );
 
         //! Private helper
-		FModelCacheEntry& GetModelCache( const Model* m );
+		FModelCacheEntry& GetModelCache(const TSharedPtr<const Model>&);
     };
 
 
@@ -911,18 +911,18 @@ namespace mu
         //-----------------------------------------------------------------------------------------
 
 		/** This method can be used to internally prepare for code execution. */
-		MUTABLERUNTIME_API void BeginBuild(const Ptr<const Model>&);
+		MUTABLERUNTIME_API void BeginBuild(const TSharedPtr<const Model>&);
 		MUTABLERUNTIME_API void EndBuild();
 
-		MUTABLERUNTIME_API bool BuildBool(const Ptr<const Model>&, const Parameters* pParams, OP::ADDRESS at) ;
-		MUTABLERUNTIME_API int BuildInt(const Ptr<const Model>&, const Parameters* pParams, OP::ADDRESS at) ;
-		MUTABLERUNTIME_API float BuildScalar(const Ptr<const Model>&, const Parameters* pParams, OP::ADDRESS at) ;
-		MUTABLERUNTIME_API void BuildColour(const Ptr<const Model>&, const Parameters* pParams, OP::ADDRESS at, float* pR, float* pG, float* pB) ;
-		MUTABLERUNTIME_API Ptr<const String> BuildString(const Ptr<const Model>&, const Parameters* pParams, OP::ADDRESS at) ;
-		MUTABLERUNTIME_API Ptr<const Image> BuildImage(const Ptr<const Model>&, const Parameters* pParams, OP::ADDRESS at, int32 MipsToSkip) ;
-		MUTABLERUNTIME_API Ptr<const Mesh> BuildMesh(const Ptr<const Model>&, const Parameters* pParams, OP::ADDRESS at) ;
-		MUTABLERUNTIME_API Ptr<const Layout> BuildLayout(const Ptr<const Model>&, const Parameters* pParams, OP::ADDRESS at) ;
-    	MUTABLERUNTIME_API Ptr<const Projector> BuildProjector(const Ptr<const Model>&, const Parameters* pParams, OP::ADDRESS at) ;
+		MUTABLERUNTIME_API bool BuildBool(const TSharedPtr<const Model>&, const Parameters* pParams, OP::ADDRESS at) ;
+		MUTABLERUNTIME_API int BuildInt(const TSharedPtr<const Model>&, const Parameters* pParams, OP::ADDRESS at) ;
+		MUTABLERUNTIME_API float BuildScalar(const TSharedPtr<const Model>&, const Parameters* pParams, OP::ADDRESS at) ;
+		MUTABLERUNTIME_API void BuildColour(const TSharedPtr<const Model>&, const Parameters* pParams, OP::ADDRESS at, float* pR, float* pG, float* pB) ;
+		MUTABLERUNTIME_API Ptr<const String> BuildString(const TSharedPtr<const Model>&, const Parameters* pParams, OP::ADDRESS at) ;
+		MUTABLERUNTIME_API Ptr<const Image> BuildImage(const TSharedPtr<const Model>&, const Parameters* pParams, OP::ADDRESS at, int32 MipsToSkip) ;
+		MUTABLERUNTIME_API Ptr<const Mesh> BuildMesh(const TSharedPtr<const Model>&, const Parameters* pParams, OP::ADDRESS at) ;
+		MUTABLERUNTIME_API Ptr<const Layout> BuildLayout(const TSharedPtr<const Model>&, const Parameters* pParams, OP::ADDRESS at) ;
+    	MUTABLERUNTIME_API Ptr<const Projector> BuildProjector(const TSharedPtr<const Model>&, const Parameters* pParams, OP::ADDRESS at) ;
 		void ClearCache() ;
 		void SetStreamingCache(uint64 bytes) ;
 
@@ -950,7 +950,7 @@ namespace mu
         {
             Instance::ID m_instanceID;
             InstancePtrConst m_pInstance;
-            ModelPtrConst m_pModel;
+            TSharedPtr<const Model> m_pModel;
 
             int m_state;
             ParametersPtr m_pOldParameters;
@@ -1007,12 +1007,12 @@ namespace mu
 
 	public:
 
-		void RunCode(const Ptr<const Model>& pModel,
+		void RunCode(const TSharedPtr<const Model>& pModel,
 			const Parameters* pParams,
 			OP::ADDRESS at, uint32 LODs = System::AllLODs, uint8 executionOptions = 0);
 
 		//!
-		void PrepareCache(const Ptr<const Model>& pModel, int state);
+		void PrepareCache(const Model*, int32 State);
 
 		//! Cache related members
 		//! @{
diff --git a/Engine/Plugins/Experimental/Mutable/Source/MutableRuntime/Public/MuR/Model.h b/Engine/Plugins/Experimental/Mutable/Source/MutableRuntime/Public/MuR/Model.h
--- a/Engine/Plugins/Experimental/Mutable/Source/MutableRuntime/Public/MuR/Model.h
+++ b/Engine/Plugins/Experimental/Mutable/Source/MutableRuntime/Public/MuR/Model.h
@@ -6,6 +6,7 @@
 #include "MuR/Ptr.h"
 #include "MuR/RefCounted.h"
 #include "Templates/Function.h"
+#include "Templates/SharedPointer.h"
 
 
 //! This version number changes whenever there is a compatibility-breaking change in the Model
@@ -25,11 +26,6 @@ namespace mu
 	typedef Ptr<Parameters> ParametersPtr;
 	typedef Ptr<const Parameters> ParametersPtrConst;
 
-    class Model;
-
-    typedef Ptr<Model> ModelPtr;
-    typedef Ptr<const Model> ModelPtrConst;
-
     class ModelParametersGenerator;
 
     typedef Ptr<ModelParametersGenerator> ModelParametersGeneratorPtr;
@@ -46,7 +42,7 @@ namespace mu
     //! When values are given to the parameters, specific Instances can be built, which hold the
     //! built application-usable data.
 	//! \ingroup runtime
-    class MUTABLERUNTIME_API Model : public RefCountedWeak
+    class MUTABLERUNTIME_API Model
 	{
 	public:
 
@@ -55,8 +51,11 @@ namespace mu
 		//-----------------------------------------------------------------------------------------
 		Model();
 
+		//! Don't call directly. Manage with a TSharedPtr.
+		~Model();
+
 		static void Serialise( const Model* p, OutputArchive& arch );
-		static ModelPtr StaticUnserialise( InputArchive& arch );
+		static TSharedPtr<Model> StaticUnserialise( InputArchive& arch );
 
 		//! Special serialise operation that serialises the data in separate "files". An object
         //! with the ModelStreamer interface is responsible of storing this data and providing
@@ -80,7 +79,7 @@ namespace mu
 		//! Create a set of new parameters of the model with the default values.
 		//! If old parameters are provided, they will be reused when possible instead of the
 		//! default values.
-        ParametersPtr NewParameters( const Parameters* pOld = nullptr ) const;
+        static ParametersPtr NewParameters( TSharedPtr<const Model> Model, const Parameters* pOld = nullptr );
 
 		//! Get the number of states in the model.
 		int GetStateCount() const;
@@ -113,11 +112,6 @@ namespace mu
 
 		Private* GetPrivate() const;
 
-	protected:
-
-		//! Forbidden. Manage with the Ptr<> template.
-        ~Model() override;
-
 	private:
 
 		Private* m_pD;
@@ -138,7 +132,7 @@ namespace mu
     public:
 
         //!
-        ModelParametersGenerator( const Model* pModel, System* pSystem );
+        ModelParametersGenerator( TSharedPtr<const Model>, System* );
 
         //! Return the number of different possible instances that can be built from the model.
         int64 GetInstanceCount();
diff --git a/Engine/Plugins/Experimental/Mutable/Source/MutableRuntime/Public/MuR/Parameters.h b/Engine/Plugins/Experimental/Mutable/Source/MutableRuntime/Public/MuR/Parameters.h
--- a/Engine/Plugins/Experimental/Mutable/Source/MutableRuntime/Public/MuR/Parameters.h
+++ b/Engine/Plugins/Experimental/Mutable/Source/MutableRuntime/Public/MuR/Parameters.h
@@ -180,15 +180,14 @@ namespace mu
         //! Serialisation
         static void Serialise( const Parameters* p, OutputArchive& arch );
         static ParametersPtr StaticUnserialise( InputArchive& arch );
-        void SetUnserialisedModel( Model* pModel );
 
         //! Portable serialisation, which tries to keep values when:
         //! - new parametes have been added
         //! - parameters have been removed
         //! - parameter types have changed
         //! - parameter names have changed, and the application is using UIDs
-        static void SerialisePortable( const Parameters* p, OutputArchive& arch );
-        static ParametersPtr UnserialisePortable( InputArchive& arch, const Model* pModel );
+        static void SerialisePortable( const Parameters*, OutputArchive& );
+        static ParametersPtr UnserialisePortable( InputArchive&, TSharedPtr<const Model> );
 
 		//! Deep clone this object.
 		ParametersPtr Clone() const;
diff --git a/Engine/Plugins/Experimental/Mutable/Source/MutableRuntime/Public/MuR/RefCounted.h b/Engine/Plugins/Experimental/Mutable/Source/MutableRuntime/Public/MuR/RefCounted.h
--- a/Engine/Plugins/Experimental/Mutable/Source/MutableRuntime/Public/MuR/RefCounted.h
+++ b/Engine/Plugins/Experimental/Mutable/Source/MutableRuntime/Public/MuR/RefCounted.h
@@ -14,219 +14,70 @@ namespace mu
 {
 
 
-    //! \brief %Base class for all reference counted objects.
-    //!
-    //! Any subclass of this class can be managed using smart pointers through the Ptr<T> template.
-    //! \warning This base allow multi-threaded manipulation of smart pointers, since the count
-    //! increments and decrements are atomic.
-    //! \ingroup runtime
-    class MUTABLERUNTIME_API RefCounted : public Base
-    {
-    public:
-
-        inline void IncRef() const
-        {
-			FPlatformAtomics::InterlockedIncrement(&m_refCount);
-        }
-
-        inline void DecRef() const
-        {
+	//! \brief %Base class for all reference counted objects.
+	//!
+	//! Any subclass of this class can be managed using smart pointers through the Ptr<T> template.
+	//! \warning This base allow multi-threaded manipulation of smart pointers, since the count
+	//! increments and decrements are atomic.
+	//! \ingroup runtime
+	class MUTABLERUNTIME_API RefCounted : public Base
+	{
+	public:
+
+		inline void IncRef() const
+		{
+			m_refCount++;
+		}
+
+		inline void DecRef() const
+		{
 			LLM_SCOPE_BYNAME(TEXT("MutableRuntime"));
 
-            // Warning: This has to be in the same line: we cannot check m_refCount after decrement
-            // since another thread could have decremented it
-            if (FPlatformAtomics::InterlockedDecrement(&m_refCount)==0)
-            {
-                delete this;
-            }
-        }
+			if (m_refCount.fetch_sub(1) == 1)
+			{
+				delete this;
+			}
+		}
 
-        inline int GetRefCount() const
-        {
-            return int( FPlatformAtomics::InterlockedExchange(&m_refCount, m_refCount) );
-        }
+		inline int32 GetRefCount() const
+		{
+			return m_refCount;
+		}
 
-        RefCounted( const RefCounted& ) = delete;
-        RefCounted( const RefCounted&& ) = delete;
-        RefCounted& operator=( const RefCounted& ) = delete;
-        RefCounted& operator=( const RefCounted&& ) = delete;
+		RefCounted(const RefCounted&) = delete;
+		RefCounted(const RefCounted&&) = delete;
+		RefCounted& operator=(const RefCounted&) = delete;
+		RefCounted& operator=(const RefCounted&&) = delete;
 
-    protected:
+	protected:
 
-        RefCounted()
-        {
-            m_refCount = 0;
-        }
+		RefCounted()
+		{
+			m_refCount = 0;
+		}
 
-        inline virtual ~RefCounted() = default;
+		inline virtual ~RefCounted() = default;
 
-    private:
+	private:
 
-        mutable volatile int32 m_refCount;
+		mutable std::atomic<int32> m_refCount;
 
-    };
+	};
 
 
-    inline void mutable_ptr_add_ref( const RefCounted* p )
-    {
-        if (p) p->IncRef();
-    }
+	inline void mutable_ptr_add_ref(const RefCounted* p)
+	{
+		if (p) p->IncRef();
+	}
 
 
-    inline void mutable_ptr_release( const RefCounted* p )
-    {
-        if (p) p->DecRef();
-    }
-
-
-    //! RefCounted with weak pointer support
-    class MUTABLERUNTIME_API RefCountedWeak : public Base
-    {
-    public:
-
-        inline void IncRef() const
-        {
-			FPlatformAtomics::InterlockedIncrement(&m_refCount);
-        }
-
-        inline void DecRef() const
-        {
-			LLM_SCOPE_BYNAME(TEXT("MutableRuntime"));
-
-            // Lock the control block to prevent creation of pointers while we work
-            Ptr<ControlBlock> pControlBlock = m_pControlBlock;
-            while (pControlBlock->m_lock.test_and_set(std::memory_order_acquire));
-
-            // Warning: This has to be in the same line: we cannot check m_refCount after decrement
-            // since another thread could have decremented it
-            if (FPlatformAtomics::InterlockedDecrement(&m_refCount)==0)
-            {
-                delete this;
-            }
-
-            pControlBlock->m_lock.clear();
-        }
-
-        inline int GetRefCount() const
-        {
-            return int(FPlatformAtomics::InterlockedExchange(&m_refCount, m_refCount));
-        }
-
-        RefCountedWeak( const RefCountedWeak& ) = delete;
-        RefCountedWeak( const RefCountedWeak&& ) = delete;
-        RefCountedWeak& operator=( const RefCountedWeak& ) = delete;
-        RefCountedWeak& operator=( const RefCountedWeak&& ) = delete;
-
-    protected:
-
-        RefCountedWeak()
-        {
-			LLM_SCOPE_BYNAME(TEXT("MutableRuntime"));
-
-            m_refCount = 0;
-            m_pControlBlock = new ControlBlock;
-            m_pControlBlock->m_pObject = this;
-        }
-
-        inline virtual ~RefCountedWeak()
-        {
-            // Control block is locked as per DecRef
-            m_pControlBlock->m_pObject = nullptr;
-        }
-
-
-    private:
-
-        // Normal intrusive smart pointer support
-        mutable volatile int32 m_refCount;
-
-        // Weak pointer support
-        class ControlBlock : public RefCounted
-        {
-        private:
-            RefCountedWeak* m_pObject = nullptr;
-            std::atomic_flag m_lock = ATOMIC_FLAG_INIT;
-
-            friend class RefCountedWeak;
-            template<class T> friend class WeakPtr;
-        };
-
-        Ptr<ControlBlock> m_pControlBlock;
-
-        template<class T> friend class WeakPtr;
-
-    };
-
-
-    inline void mutable_ptr_add_ref( const RefCountedWeak* p )
-    {
-        if (p) p->IncRef();
-    }
-
-
-    inline void mutable_ptr_release( const RefCountedWeak* p )
-    {
-        if (p) p->DecRef();
-    }
-
-
-    //!
-    template<class T>
-    class WeakPtr
-    {
-    private:
-        Ptr<RefCountedWeak::ControlBlock> m_pControlBlock;
-
-    public:
-
-        WeakPtr() = default;
-
-        WeakPtr( const Ptr<const T>& ptr )
-        {
-            if (ptr)
-            {
-                m_pControlBlock = ptr->m_pControlBlock;
-            }
-        }
-
-        inline Ptr<T> Pin()
-        {
-            Ptr<T> result;
-
-            if (m_pControlBlock)
-            {
-                while (m_pControlBlock->m_lock.test_and_set(std::memory_order_acquire));
-                result = (T*)m_pControlBlock->m_pObject;
-                m_pControlBlock->m_lock.clear();
-            }
-
-            return result;
-        }
-
-		inline Ptr<const T> Pin() const
-        {
-            Ptr<const T> result;
-
-            if (m_pControlBlock)
-            {
-                while (m_pControlBlock->m_lock.test_and_set(std::memory_order_acquire));
-                result = (const T*)m_pControlBlock->m_pObject;
-                m_pControlBlock->m_lock.clear();
-            }
-
-            return result;
-        }
-
-		inline void Reset()
-        {
-            m_pControlBlock = nullptr;
-        }
-
-		inline bool IsNull() const
-        {
-            return Pin().get()==nullptr;
-        }
-    };
+	inline void mutable_ptr_release(const RefCounted* p)
+	{
+		if (p)
+		{
+			p->DecRef();
+		}
+	}
 
 }
 
diff --git a/Engine/Plugins/Experimental/Mutable/Source/MutableRuntime/Public/MuR/System.h b/Engine/Plugins/Experimental/Mutable/Source/MutableRuntime/Public/MuR/System.h
--- a/Engine/Plugins/Experimental/Mutable/Source/MutableRuntime/Public/MuR/System.h
+++ b/Engine/Plugins/Experimental/Mutable/Source/MutableRuntime/Public/MuR/System.h
@@ -115,7 +115,7 @@ namespace mu
         //! no longer needed.
         //! \param pModel Model to build an instance of
         //! \return An identifier that is always bigger than 0.
-        Instance::ID NewInstance( const Ptr<const Model>& Model );
+        Instance::ID NewInstance(const TSharedPtr<const Model>& Model );
 
         //! \brief Update an instance with a new parameter set and/or state.
         //!
@@ -160,7 +160,7 @@ namespace mu
 
 		//! Build one of the images defined in a model parameter as additional description. These
 		//! images can be used for colour bars, icons, etc...
-		Ptr<const Image> BuildParameterAdditionalImage( const Ptr<const Model>& Model,
+		Ptr<const Image> BuildParameterAdditionalImage(const TSharedPtr<const Model>& Model,
 														const Ptr<const Parameters>& Params,
 														int32 Parameter,
 														int32 ImageIndex );
@@ -173,7 +173,7 @@ namespace mu
         //! \param pParameters Parameter set that we want to find the relevancy of.
         //! \param pFlags is a pointer to a preallocated array of booleans that contains at least
 		//! pParameters->GetCount() elements.
-        void GetParameterRelevancy( const Ptr<const Model>& Model,
+        void GetParameterRelevancy(const TSharedPtr<const Model>& Model,
 									const Ptr<const Parameters>& Parameters,
 									bool* Flags );
 
diff --git a/Engine/Plugins/Experimental/Mutable/Source/MutableTools/Private/MuT/CodeOptimiser.cpp b/Engine/Plugins/Experimental/Mutable/Source/MutableTools/Private/MuT/CodeOptimiser.cpp
--- a/Engine/Plugins/Experimental/Mutable/Source/MutableTools/Private/MuT/CodeOptimiser.cpp
+++ b/Engine/Plugins/Experimental/Mutable/Source/MutableTools/Private/MuT/CodeOptimiser.cpp
@@ -461,7 +461,7 @@ namespace mu
 			FLinkerOptions LinkerOptions;
 			LinkerOptions.MinTextureResidentMipCount = 255;
 
-			ModelPtrConst model = new Model;
+			TSharedPtr<const Model> model = MakeShared<Model>();
 			ASTOp::FullLink( cloned, model->GetPrivate()->m_program, &LinkerOptions);
 			OP::ADDRESS at = cloned->linkedAddress;
 
@@ -469,7 +469,7 @@ namespace mu
 			state.m_root = at;
 			model->GetPrivate()->m_program.m_states.Add(state);
 
-			ParametersPtr localParams = model->NewParameters();
+			ParametersPtr localParams = Model::NewParameters(model);
 			pSystem->GetPrivate()->BeginBuild( model );
 
 			// Calculate the value and replace this op by a constant
diff --git a/Engine/Plugins/Experimental/Mutable/Source/MutableTools/Private/MuT/Compiler.cpp b/Engine/Plugins/Experimental/Mutable/Source/MutableTools/Private/MuT/Compiler.cpp
--- a/Engine/Plugins/Experimental/Mutable/Source/MutableTools/Private/MuT/Compiler.cpp
+++ b/Engine/Plugins/Experimental/Mutable/Source/MutableTools/Private/MuT/Compiler.cpp
@@ -178,7 +178,7 @@ namespace mu
 
 
     //---------------------------------------------------------------------------------------------
-    ModelPtr Compiler::Compile( const Ptr<Node>& pNode )
+	TSharedPtr<Model> Compiler::Compile( const Ptr<Node>& pNode )
     {
         MUTABLE_CPUPROFILER_SCOPE(Compile);
 
@@ -221,7 +221,7 @@ namespace mu
 
 
         // Link the program and generate state data.
-        ModelPtr pResult = new Model();
+		TSharedPtr<Model> pResult = MakeShared<Model>();
         auto& program = pResult->GetPrivate()->m_program;
         for( auto& s: states )
         {
@@ -343,7 +343,7 @@ namespace mu
 
 		// Pack data
 		int32 MinimumBytesPerRom = 1024; // \TODO: compilation parameter
-		m_pD->GenerateRoms(pResult.get(),MinimumBytesPerRom);
+		m_pD->GenerateRoms(pResult.Get(),MinimumBytesPerRom);
 
 		UE_LOG(LogMutableCore, Verbose, TEXT("(int) %s : %ld"), TEXT("program size"), int64(program.m_opAddress.Num()));
 
diff --git a/Engine/Plugins/Experimental/Mutable/Source/MutableTools/Public/MuT/Compiler.h b/Engine/Plugins/Experimental/Mutable/Source/MutableTools/Public/MuT/Compiler.h
--- a/Engine/Plugins/Experimental/Mutable/Source/MutableTools/Public/MuT/Compiler.h
+++ b/Engine/Plugins/Experimental/Mutable/Source/MutableTools/Public/MuT/Compiler.h
@@ -150,7 +150,7 @@ namespace mu
         Compiler( CompilerOptionsPtr options=nullptr );
 
         //! Compile the expression into a run-time model.
-        ModelPtr Compile( const Ptr<Node>& pNode );
+		TSharedPtr<Model> Compile( const Ptr<Node>& pNode );
 
         //! Return the log of messages of all the compile operations executed so far.
         ErrorLogPtrConst GetLog() const;
