[mutable] Another non-unity build fix.
#rnx

diff --git a/Engine/Plugins/Experimental/Mutable/Source/MutableTools/Public/MuT/Compiler.h b/Engine/Plugins/Experimental/Mutable/Source/MutableTools/Public/MuT/Compiler.h
--- a/Engine/Plugins/Experimental/Mutable/Source/MutableTools/Public/MuT/Compiler.h
+++ b/Engine/Plugins/Experimental/Mutable/Source/MutableTools/Public/MuT/Compiler.h
@@ -5,6 +5,7 @@
 #include "HAL/PlatformMath.h"
 #include "MuR/Ptr.h"
 #include "MuR/RefCounted.h"
+#include "Templates/SharedPointer.h"
 
 
 namespace mu
