[mutable] Fix COI descriptor desynchronization causing multiple updates

#rnx

diff --git a/Engine/Plugins/Experimental/Mutable/Source/CustomizableObject/Private/MuCO/CustomizableObjectInstance.cpp b/Engine/Plugins/Experimental/Mutable/Source/CustomizableObject/Private/MuCO/CustomizableObjectInstance.cpp
--- a/Engine/Plugins/Experimental/Mutable/Source/CustomizableObject/Private/MuCO/CustomizableObjectInstance.cpp
+++ b/Engine/Plugins/Experimental/Mutable/Source/CustomizableObject/Private/MuCO/CustomizableObjectInstance.cpp
@@ -228,7 +228,7 @@ void UCustomizableInstancePrivateData::SetMinMaxLODToLoad(UCustomizableObjectIns
 	const bool bIsDowngradeLODUpdate = Public->GetCurrentMinLOD() >= 0 && NewMinLOD > Public->GetCurrentMinLOD();
 	SetCOInstanceFlags(bIsDowngradeLODUpdate ? PendingLODsDowngrade : ECONone);
 
-	if (!bIsDowngradeLODUpdate && bLimitLODUpgrades)
+	if (!bIsDowngradeLODUpdate && bLimitLODUpgrades && HasCOInstanceFlags(Generated))
 	{
 		if (FPlatformTime::Seconds() - LastUpdateTime < 3.f) // This is a sort of hysteresis/cooldown to prevent too many instances being discarded/updated per second
 		{
@@ -259,6 +259,8 @@ void UCustomizableInstancePrivateData::SetMinMaxLODToLoad(UCustomizableObjectIns
 
 	if (bLODChanged)
 	{
+		Public->UpdateDescriptorRuntimeHash.UpdateMinMaxLOD(NewMinLOD, NewMaxLOD);
+
 		if (const FMutableQueueElem* QueueElem = UCustomizableObjectSystem::GetInstance()->GetPrivate()->MutableOperationQueue.Get(Public))
 		{
 			QueueElem->Operation->InstanceDescriptorRuntimeHash.UpdateMinMaxLOD(NewMinLOD, NewMaxLOD);
diff --git a/Engine/Plugins/Experimental/Mutable/Source/CustomizableObject/Private/MuCO/CustomizableObjectInstanceDescriptor.cpp b/Engine/Plugins/Experimental/Mutable/Source/CustomizableObject/Private/MuCO/CustomizableObjectInstanceDescriptor.cpp
--- a/Engine/Plugins/Experimental/Mutable/Source/CustomizableObject/Private/MuCO/CustomizableObjectInstanceDescriptor.cpp
+++ b/Engine/Plugins/Experimental/Mutable/Source/CustomizableObject/Private/MuCO/CustomizableObjectInstanceDescriptor.cpp
@@ -1022,6 +1022,30 @@ void FCustomizableObjectInstanceDescriptor::ReloadParameters()
 }
 
 
+int32 FCustomizableObjectInstanceDescriptor::GetMinLod() const
+{
+	return MinLOD;
+}
+
+
+void FCustomizableObjectInstanceDescriptor::SetMinLod(int32 InMinLOD)
+{
+	MinLOD = InMinLOD;
+}
+
+
+int32 FCustomizableObjectInstanceDescriptor::GetMaxLod() const
+{
+	return MaxLOD;
+}
+
+
+void FCustomizableObjectInstanceDescriptor::SetMaxLod(int32 InMaxLOD)
+{
+	MaxLOD = InMaxLOD;
+}
+
+
 TArray<FCustomizableObjectBoolParameterValue>& FCustomizableObjectInstanceDescriptor::GetBoolParameters()
 {
 	return BoolParameters;
@@ -2310,8 +2334,6 @@ FDescriptorHash::FDescriptorHash(const FCustomizableObjectInstanceDescriptor& De
 	
 	Hash = HashCombine(Hash, GetTypeHash(Descriptor.State));
 	Hash = HashCombine(Hash, GetTypeHash(Descriptor.bBuildParameterDecorations));
-	Hash = HashCombine(Hash, GetTypeHash(Descriptor.MinLOD));
-	Hash = HashCombine(Hash, GetTypeHash(Descriptor.MaxLOD));
 
 	for (const TTuple<FName, FMultilayerProjector>& Pair : Descriptor.MultilayerProjectors)
 	{
diff --git a/Engine/Plugins/Experimental/Mutable/Source/CustomizableObject/Public/MuCO/CustomizableObjectInstanceDescriptor.h b/Engine/Plugins/Experimental/Mutable/Source/CustomizableObject/Public/MuCO/CustomizableObjectInstanceDescriptor.h
--- a/Engine/Plugins/Experimental/Mutable/Source/CustomizableObject/Public/MuCO/CustomizableObjectInstanceDescriptor.h
+++ b/Engine/Plugins/Experimental/Mutable/Source/CustomizableObject/Public/MuCO/CustomizableObjectInstanceDescriptor.h
@@ -61,6 +61,14 @@ struct CUSTOMIZABLEOBJECT_API FCustomizableObjectInstanceDescriptor
 	/** Update all parameters to be up to date with the Mutable Core parameters. */
 	void ReloadParameters();
     
+	int32 GetMinLod() const;
+
+	void SetMinLod(int32 InMinLOD);
+
+	int32 GetMaxLod() const;
+
+	void SetMaxLod(int32 InMaxLOD);
+
 	// ------------------------------------------------------------
 	// Parameters
 	// ------------------------------------------------------------
