[mutable] Fix bug with image layer optimization versus plain colours.
Debugger: more info on color constants, numeric parameters.
Removed some unused declarations.
#rnx

diff --git a/Engine/Plugins/Experimental/Mutable/Source/CustomizableObjectEditor/Private/MuCOE/SMutableCodeViewer.cpp b/Engine/Plugins/Experimental/Mutable/Source/CustomizableObjectEditor/Private/MuCOE/SMutableCodeViewer.cpp
--- a/Engine/Plugins/Experimental/Mutable/Source/CustomizableObjectEditor/Private/MuCOE/SMutableCodeViewer.cpp
+++ b/Engine/Plugins/Experimental/Mutable/Source/CustomizableObjectEditor/Private/MuCOE/SMutableCodeViewer.cpp
@@ -100,7 +100,7 @@ struct FSlateBrush;
 namespace MutableCodeTreeViewColumns
 {
 	static const FName OperationsColumnID("Operations");
-	static const FName AdditionalDataColumnID("Additional Data");
+	static const FName AdditionalDataColumnID("Flags");
 };
 
 /**
@@ -135,6 +135,14 @@ public:
 			// See if the operation type accepts additional information in the label
 			switch ( Type )
 			{
+			case mu::OP_TYPE::NU_PARAMETER:
+			{
+				mu::OP::ParameterArgs Args = Program.GetOpArgs<mu::OP::ParameterArgs>(RowItem->MutableOperation);
+				OpName += TEXT(" ");
+				OpName += ANSI_TO_TCHAR(Program.m_parameters[int32(Args.variable)].m_name.c_str());
+				break;
+			}
+
 			case mu::OP_TYPE::IM_SWIZZLE:
 			{
 				mu::OP::ImageSwizzleArgs Args = Program.GetOpArgs<mu::OP::ImageSwizzleArgs>(RowItem->MutableOperation);
@@ -402,10 +410,10 @@ private:
 	const FText EmptyText = FText(INVTEXT(" "));
 	
 	/** String printed on the UI when the operation is shown to be dynamic resource */
-	const FText DynamicResourceText = FText::FromString(FString("DYNAMIC RESOURCE"));
+	const FText DynamicResourceText = FText::FromString(FString("dyn"));
 
 	/** String printed on the UI when the operation is shown to be state constant */
-	const FText StateConstantText = FText::FromString(FString("STATE CONSTANT"));
+	const FText StateConstantText = FText::FromString(FString("const"));
 	
 	/** Color used on the extra data column when no extra data is shown */
 	const FSlateColor ExtraDataBackgroundBoxDefaultColor =  FSlateColor(TransparentColor);
@@ -2319,8 +2327,9 @@ void SMutableCodeViewer::Tick(const FGeometry& AllottedGeometry, const double In
 		float RedValue;
 		float GreenValue;
 		float BlueValue;
-		System->GetPrivate()->BuildColour(MutableModel, PreviewParameters.get(), SelectedOperationAddress, &RedValue, &GreenValue, &BlueValue);
-		PreviewColorViewer->SetColor(RedValue, GreenValue, BlueValue);
+		float AlphaValue;
+		System->GetPrivate()->BuildColour(MutableModel, PreviewParameters.get(), SelectedOperationAddress, &RedValue, &GreenValue, &BlueValue, &AlphaValue);
+		PreviewColorViewer->SetColor(RedValue, GreenValue, BlueValue, AlphaValue);
 		break;
 	}
 
diff --git a/Engine/Plugins/Experimental/Mutable/Source/CustomizableObjectEditor/Private/MuCOE/SMutableColorViewer.cpp b/Engine/Plugins/Experimental/Mutable/Source/CustomizableObjectEditor/Private/MuCOE/SMutableColorViewer.cpp
--- a/Engine/Plugins/Experimental/Mutable/Source/CustomizableObjectEditor/Private/MuCOE/SMutableColorViewer.cpp
+++ b/Engine/Plugins/Experimental/Mutable/Source/CustomizableObjectEditor/Private/MuCOE/SMutableColorViewer.cpp
@@ -25,6 +25,7 @@ void SMutableColorViewer::Construct(const FArguments& InArgs)
 	const FText RedColorValueTitle = LOCTEXT("RedColorValueTitle", "Red : ");
 	const FText GreenColorValueTitle = LOCTEXT("GreenColorValueTitle", "Green : ");
 	const FText BlueColorValueTitle = LOCTEXT("BlueColorValueTitle", "Blue : ");
+	const FText AlphaColorValueTitle = LOCTEXT("AlphaColorValueTitle", "Alpha : ");
 
 	const int32 IndentationSpace = 16;	
 	const int32 AfterTitleSpacing = 4;
@@ -99,25 +100,48 @@ void SMutableColorViewer::Construct(const FArguments& InArgs)
 
 				// Blue channel ----------------------------------------
 				+SVerticalBox::Slot()
-				.Padding(0, 1)
-				.AutoHeight()
-				[
-					SNew(SHorizontalBox)
+					.Padding(0, 1)
+					.AutoHeight()
+					[
+						SNew(SHorizontalBox)
 
-					+ SHorizontalBox::Slot()
+						+ SHorizontalBox::Slot()
 					.AutoWidth()
 					[
 						SNew(STextBlock).
 						Text(BlueColorValueTitle)
 					]
 
-					+ SHorizontalBox::Slot()
+				+ SHorizontalBox::Slot()
 					.AutoWidth()
 					[
 						SNew(STextBlock).
 						Text(this, &SMutableColorViewer::GetBlueValue)
 					]
-				]
+					]
+				// -----------------------------------------------------
+
+				// Alpha channel ----------------------------------------
+				+SVerticalBox::Slot()
+					.Padding(0, 1)
+					.AutoHeight()
+					[
+						SNew(SHorizontalBox)
+
+						+ SHorizontalBox::Slot()
+					.AutoWidth()
+					[
+						SNew(STextBlock).
+						Text(AlphaColorValueTitle)
+					]
+
+				+ SHorizontalBox::Slot()
+					.AutoWidth()
+					[
+						SNew(STextBlock).
+						Text(this, &SMutableColorViewer::GetAlphaValue)
+					]
+					]
 				// -----------------------------------------------------
 			]
 		
@@ -134,11 +158,12 @@ void SMutableColorViewer::Construct(const FArguments& InArgs)
 	];
 }
 
-void SMutableColorViewer::SetColor(const float& InRed, const float& InGreen, const float& InBlue)
+void SMutableColorViewer::SetColor(float InRed, float InGreen, float InBlue, float InAlpha)
 {
 	this->RedValue = InRed;
 	this->GreenValue = InGreen;
 	this->BlueValue = InBlue;
+	this->AlphaValue = InAlpha;
 
 	this->ColorPreview->SetColor( this->GetPreviewColor() );
 }
@@ -164,4 +189,9 @@ FText SMutableColorViewer::GetBlueValue() const
 	return FText::AsNumber(this->BlueValue);
 }
 
+FText SMutableColorViewer::GetAlphaValue() const
+{
+	return FText::AsNumber(this->AlphaValue);
+}
+
 #undef LOCTEXT_NAMESPACE
diff --git a/Engine/Plugins/Experimental/Mutable/Source/CustomizableObjectEditor/Private/MuCOE/SMutableColorViewer.h b/Engine/Plugins/Experimental/Mutable/Source/CustomizableObjectEditor/Private/MuCOE/SMutableColorViewer.h
--- a/Engine/Plugins/Experimental/Mutable/Source/CustomizableObjectEditor/Private/MuCOE/SMutableColorViewer.h
+++ b/Engine/Plugins/Experimental/Mutable/Source/CustomizableObjectEditor/Private/MuCOE/SMutableColorViewer.h
@@ -21,7 +21,7 @@ public:
 	void Construct(const FArguments& InArgs);
 
 	/** Set the Mutable Bool to be used for this widget */
-	void SetColor(const float& InRed, const float& InGreen, const float& InBlue);
+	void SetColor(float InRed, float InGreen, float InBlue, float InAlpha);
 
 private:
 
@@ -29,7 +29,7 @@ private:
 	float RedValue = 0.0f;
 	float GreenValue = 0.0f;
 	float BlueValue = 0.0f;
-	const float AlphaValue{1};
+	float AlphaValue = 1.0f;
 
 	/** Color box widget designed to serve as a preview of the color reported by mutable */
 	TSharedPtr<SMutableColorPreviewBox> ColorPreview;
@@ -38,6 +38,7 @@ private:
 	FText GetRedValue() const;
 	FText GetGreenValue() const;
 	FText GetBlueValue() const;
+	FText GetAlphaValue() const;
 
 	/*
 	* Get a color object that the SMutableColorPreviewBox is able to display 
diff --git a/Engine/Plugins/Experimental/Mutable/Source/MutableRuntime/Private/MuR/ImagePrivate.h b/Engine/Plugins/Experimental/Mutable/Source/MutableRuntime/Private/MuR/ImagePrivate.h
--- a/Engine/Plugins/Experimental/Mutable/Source/MutableRuntime/Private/MuR/ImagePrivate.h
+++ b/Engine/Plugins/Experimental/Mutable/Source/MutableRuntime/Private/MuR/ImagePrivate.h
@@ -165,12 +165,13 @@ namespace mu
 			uint8 r = uint8(FMath::Clamp(255.0f * c[0], 0.0f, 255.0f));
 			uint8 g = uint8(FMath::Clamp(255.0f * c[1], 0.0f, 255.0f));
 			uint8 b = uint8(FMath::Clamp(255.0f * c[2], 0.0f, 255.0f));
+			uint8 a = uint8(FMath::Clamp(255.0f * c[3], 0.0f, 255.0f));
 			for ( int p = 0; p < pixelCount; ++p )
             {
                 pData[0] = r;
                 pData[1] = g;
                 pData[2] = b;
-                pData[3] = 255;
+                pData[3] = a;
                 pData += 4;
             }
             break;
@@ -182,12 +183,13 @@ namespace mu
 			uint8 r = uint8(FMath::Clamp(255.0f * c[0], 0.0f, 255.0f));
 			uint8 g = uint8(FMath::Clamp(255.0f * c[1], 0.0f, 255.0f));
 			uint8 b = uint8(FMath::Clamp(255.0f * c[2], 0.0f, 255.0f));
+			uint8 a = uint8(FMath::Clamp(255.0f * c[3], 0.0f, 255.0f));
 			for ( int p = 0; p < pixelCount; ++p )
             {
                 pData[0] = b;
                 pData[1] = g;
                 pData[2] = r;
-                pData[3] = 255;
+                pData[3] = a;
                 pData += 4;
             }
             break;
diff --git a/Engine/Plugins/Experimental/Mutable/Source/MutableRuntime/Private/MuR/System.cpp b/Engine/Plugins/Experimental/Mutable/Source/MutableRuntime/Private/MuR/System.cpp
--- a/Engine/Plugins/Experimental/Mutable/Source/MutableRuntime/Private/MuR/System.cpp
+++ b/Engine/Plugins/Experimental/Mutable/Source/MutableRuntime/Private/MuR/System.cpp
@@ -697,7 +697,8 @@ namespace mu
 		OP::ADDRESS at,
 		float* pR,
 		float* pG,
-		float* pB)
+		float* pB,
+		float* pA)
 	{
 		FVector4f col;
 
@@ -710,6 +711,7 @@ namespace mu
 				if (pR) *pR = 0.0f;
 				if (pG) *pG = 0.0f;
 				if (pB) *pB = 0.0f;
+				if (pA) *pA = 1.0f;
 			}
 
 			col = m_memory->GetColour(FCacheAddress(at, 0, 0));
@@ -718,6 +720,7 @@ namespace mu
 		if (pR) *pR = col[0];
 		if (pG) *pG = col[1];
 		if (pB) *pB = col[2];
+		if (pA) *pA = col[3];
 	}
 
 	
diff --git a/Engine/Plugins/Experimental/Mutable/Source/MutableRuntime/Private/MuR/SystemPrivate.h b/Engine/Plugins/Experimental/Mutable/Source/MutableRuntime/Private/MuR/SystemPrivate.h
--- a/Engine/Plugins/Experimental/Mutable/Source/MutableRuntime/Private/MuR/SystemPrivate.h
+++ b/Engine/Plugins/Experimental/Mutable/Source/MutableRuntime/Private/MuR/SystemPrivate.h
@@ -917,7 +917,7 @@ namespace mu
 		MUTABLERUNTIME_API bool BuildBool(const TSharedPtr<const Model>&, const Parameters* pParams, OP::ADDRESS at) ;
 		MUTABLERUNTIME_API int BuildInt(const TSharedPtr<const Model>&, const Parameters* pParams, OP::ADDRESS at) ;
 		MUTABLERUNTIME_API float BuildScalar(const TSharedPtr<const Model>&, const Parameters* pParams, OP::ADDRESS at) ;
-		MUTABLERUNTIME_API void BuildColour(const TSharedPtr<const Model>&, const Parameters* pParams, OP::ADDRESS at, float* pR, float* pG, float* pB) ;
+		MUTABLERUNTIME_API void BuildColour(const TSharedPtr<const Model>&, const Parameters*, OP::ADDRESS, float* OutR, float* OutG, float* OutB, float* OutA) ;
 		MUTABLERUNTIME_API Ptr<const String> BuildString(const TSharedPtr<const Model>&, const Parameters* pParams, OP::ADDRESS at) ;
 		MUTABLERUNTIME_API Ptr<const Image> BuildImage(const TSharedPtr<const Model>&, const Parameters* pParams, OP::ADDRESS at, int32 MipsToSkip) ;
 		MUTABLERUNTIME_API Ptr<const Mesh> BuildMesh(const TSharedPtr<const Model>&, const Parameters* pParams, OP::ADDRESS at) ;
diff --git a/Engine/Plugins/Experimental/Mutable/Source/MutableTools/MutableTools.Build.cs b/Engine/Plugins/Experimental/Mutable/Source/MutableTools/MutableTools.Build.cs
--- a/Engine/Plugins/Experimental/Mutable/Source/MutableTools/MutableTools.Build.cs
+++ b/Engine/Plugins/Experimental/Mutable/Source/MutableTools/MutableTools.Build.cs
@@ -11,6 +11,7 @@ namespace UnrealBuildTool.Rules
 			ShortName = "MuT";
 
 			DefaultBuildSettings = BuildSettingsVersion.V2;
+			//bUseUnity = false;
 
 			// 
 			bUseRTTI = true;
diff --git a/Engine/Plugins/Experimental/Mutable/Source/MutableTools/Private/MuT/ASTOpImageSwizzle.cpp b/Engine/Plugins/Experimental/Mutable/Source/MutableTools/Private/MuT/ASTOpImageSwizzle.cpp
--- a/Engine/Plugins/Experimental/Mutable/Source/MutableTools/Private/MuT/ASTOpImageSwizzle.cpp
+++ b/Engine/Plugins/Experimental/Mutable/Source/MutableTools/Private/MuT/ASTOpImageSwizzle.cpp
@@ -254,22 +254,6 @@ namespace mu
 			switch (sourceType)
 			{
 
-			case OP_TYPE::IM_PLAINCOLOUR:
-			{
-				auto NewPlain = mu::Clone<ASTOpFixed>(channelSourceAt);
-				Ptr<ASTOpFixed> NewSwizzle = new ASTOpFixed;
-				NewSwizzle->op.type = OP_TYPE::CO_SWIZZLE;
-				for (int i = 0; i < MUTABLE_OP_MAX_SWIZZLE_CHANNELS; ++i)
-				{
-					NewSwizzle->SetChild(NewSwizzle->op.args.ColourSwizzle.sources[i], NewPlain->children[NewPlain->op.args.ImagePlainColour.colour]);
-					NewSwizzle->op.args.ColourSwizzle.sourceChannels[i] = SourceChannels[i];
-				}
-				NewPlain->SetChild(NewPlain->op.args.ImagePlainColour.colour, NewSwizzle);
-				NewPlain->op.args.ImagePlainColour.format = Format;
-				at = NewPlain;
-				break;
-			}
-
 			case OP_TYPE::IM_SWITCH:
 			{
 				// Move the swizzle down all the paths
@@ -759,6 +743,28 @@ namespace mu
 					at = NewFormat;
 				}
 			}
+
+			// Swizzle down plaincolours.
+			if (!at && sourceType == OP_TYPE::IM_PLAINCOLOUR)
+			{
+				Ptr<ASTOpFixed> NewPlain = mu::Clone<ASTOpFixed>(channelSourceAt);
+				Ptr<ASTOpFixed> NewSwizzle = new ASTOpFixed;
+				NewSwizzle->op.type = OP_TYPE::CO_SWIZZLE;
+				for (int c = 0; c < MUTABLE_OP_MAX_SWIZZLE_CHANNELS; ++c)
+				{
+					if (Sources[c])
+					{
+						const ASTOpFixed* TypedPlain = dynamic_cast<const ASTOpFixed*>(Sources[c].child().get());
+
+						NewSwizzle->SetChild(NewSwizzle->op.args.ColourSwizzle.sources[c], TypedPlain->children[TypedPlain->op.args.ImagePlainColour.colour]);
+					}
+					NewSwizzle->op.args.ColourSwizzle.sourceChannels[c] = SourceChannels[c];
+				}
+				NewPlain->SetChild(NewPlain->op.args.ImagePlainColour.colour, NewSwizzle);
+				NewPlain->op.args.ImagePlainColour.format = Format;
+				at = NewPlain;
+			}
+
 		}
 
 		// Swizzle of RGB from a source + A from a layer colour
diff --git a/Engine/Plugins/Experimental/Mutable/Source/MutableTools/Private/MuT/CodeGenerator_Color.cpp b/Engine/Plugins/Experimental/Mutable/Source/MutableTools/Private/MuT/CodeGenerator_Color.cpp
--- a/Engine/Plugins/Experimental/Mutable/Source/MutableTools/Private/MuT/CodeGenerator_Color.cpp
+++ b/Engine/Plugins/Experimental/Mutable/Source/MutableTools/Private/MuT/CodeGenerator_Color.cpp
@@ -127,6 +127,7 @@ namespace mu
 		op->op.args.ColourConstant.value[0] = node.m_value[0];
 		op->op.args.ColourConstant.value[1] = node.m_value[1];
 		op->op.args.ColourConstant.value[2] = node.m_value[2];
+		op->op.args.ColourConstant.value[3] = 1.0f;
 
 		result.op = op;
 	}
diff --git a/Engine/Plugins/Experimental/Mutable/Source/MutableTools/Private/MuT/CodeOptimiser.cpp b/Engine/Plugins/Experimental/Mutable/Source/MutableTools/Private/MuT/CodeOptimiser.cpp
--- a/Engine/Plugins/Experimental/Mutable/Source/MutableTools/Private/MuT/CodeOptimiser.cpp
+++ b/Engine/Plugins/Experimental/Mutable/Source/MutableTools/Private/MuT/CodeOptimiser.cpp
@@ -541,8 +541,8 @@ namespace mu
 			{
 				MUTABLE_CPUPROFILER_SCOPE(ConstantBool);
 
-				float r=0.0f, g=0.0f, b=0.0f;
-				pSystem->GetPrivate()->BuildColour( model, localParams.get(), at, &r,&g,&b );
+				float r=0.0f, g=0.0f, b = 0.0f, a = 0.0f;
+				pSystem->GetPrivate()->BuildColour( model, localParams.get(), at, &r,&g,&b,&a );
 
 				{
 					mu::Ptr<ASTOpFixed> constantOp = new ASTOpFixed();
@@ -550,7 +550,7 @@ namespace mu
 					constantOp->op.args.ColourConstant.value[0] = r;
 					constantOp->op.args.ColourConstant.value[1] = g;
 					constantOp->op.args.ColourConstant.value[2] = b;
-					constantOp->op.args.ColourConstant.value[3] = 1.0f;
+					constantOp->op.args.ColourConstant.value[3] = a;
 					m_result = constantOp;
 				}
 				break;
diff --git a/Engine/Plugins/Experimental/Mutable/Source/MutableTools/Private/MuT/CodeOptimiser_Parameters.cpp b/Engine/Plugins/Experimental/Mutable/Source/MutableTools/Private/MuT/CodeOptimiser_Parameters.cpp
--- a/Engine/Plugins/Experimental/Mutable/Source/MutableTools/Private/MuT/CodeOptimiser_Parameters.cpp
+++ b/Engine/Plugins/Experimental/Mutable/Source/MutableTools/Private/MuT/CodeOptimiser_Parameters.cpp
@@ -242,7 +242,8 @@ namespace mu
             whiteOp->op.type = OP_TYPE::CO_CONSTANT;
             whiteOp->op.args.ColourConstant.value[0] = 1;
             whiteOp->op.args.ColourConstant.value[1] = 1;
-            whiteOp->op.args.ColourConstant.value[2] = 1;
+			whiteOp->op.args.ColourConstant.value[2] = 1;
+			whiteOp->op.args.ColourConstant.value[3] = 1;
 
             Ptr<ASTOpFixed> wplainOp = new ASTOpFixed;
             wplainOp->op.type = OP_TYPE::IM_PLAINCOLOUR;
@@ -466,7 +467,8 @@ namespace mu
                             blackOp->op.type = OP_TYPE::CO_CONSTANT;
                             blackOp->op.args.ColourConstant.value[0] = 0;
                             blackOp->op.args.ColourConstant.value[1] = 0;
-                            blackOp->op.args.ColourConstant.value[2] = 0;
+							blackOp->op.args.ColourConstant.value[2] = 0;
+							blackOp->op.args.ColourConstant.value[3] = 1;
 
                             Ptr<ASTOpFixed> plainOp = new ASTOpFixed;
                             plainOp->op.type = OP_TYPE::IM_PLAINCOLOUR;
@@ -511,6 +513,7 @@ namespace mu
 //							blackOp.args.ColourConstant.value[0] = 0;
 //							blackOp.args.ColourConstant.value[1] = 0;
 //							blackOp.args.ColourConstant.value[2] = 0;
+//							blackOp.args.ColourConstant.value[3] = 1;
 
 //							OP plainOp;
 //							plainOp.type = OP_TYPE::IM_PLAINCOLOUR;
@@ -568,7 +571,8 @@ namespace mu
                             blackOp->op.type = OP_TYPE::CO_CONSTANT;
                             blackOp->op.args.ColourConstant.value[0] = 0;
                             blackOp->op.args.ColourConstant.value[1] = 0;
-                            blackOp->op.args.ColourConstant.value[2] = 0;
+							blackOp->op.args.ColourConstant.value[2] = 0;
+							blackOp->op.args.ColourConstant.value[3] = 1;
 
                             Ptr<ASTOpFixed> plainOp = new ASTOpFixed;
                             plainOp->op.type = OP_TYPE::IM_PLAINCOLOUR;
@@ -1035,6 +1039,7 @@ namespace mu
                         blackOp->op.args.ColourConstant.value[0] = 0;
                         blackOp->op.args.ColourConstant.value[1] = 0;
                         blackOp->op.args.ColourConstant.value[2] = 0;
+						blackOp->op.args.ColourConstant.value[3] = 1;
 
                         Ptr<ASTOpFixed> plainOp = new ASTOpFixed;
                         plainOp->op.type = OP_TYPE::IM_PLAINCOLOUR;
@@ -1090,6 +1095,7 @@ namespace mu
                         blackOp->op.args.ColourConstant.value[0] = 0;
                         blackOp->op.args.ColourConstant.value[1] = 0;
                         blackOp->op.args.ColourConstant.value[2] = 0;
+						blackOp->op.args.ColourConstant.value[3] = 1;
 
                         Ptr<ASTOpFixed> plainOp = new ASTOpFixed;
                         plainOp->op.type = OP_TYPE::IM_PLAINCOLOUR;
@@ -1116,6 +1122,7 @@ namespace mu
                         blackOp->op.args.ColourConstant.value[0] = 0;
                         blackOp->op.args.ColourConstant.value[1] = 0;
                         blackOp->op.args.ColourConstant.value[2] = 0;
+						blackOp->op.args.ColourConstant.value[3] = 1;
 
                         Ptr<ASTOpFixed> plainOp = new ASTOpFixed;
                         plainOp->op.type = OP_TYPE::IM_PLAINCOLOUR;
@@ -1199,6 +1206,7 @@ namespace mu
                         blackOp->op.args.ColourConstant.value[0] = 0;
                         blackOp->op.args.ColourConstant.value[1] = 0;
                         blackOp->op.args.ColourConstant.value[2] = 0;
+						blackOp->op.args.ColourConstant.value[3] = 1;
 
                         Ptr<ASTOpFixed> plainOp = new ASTOpFixed;
                         plainOp->op.type = OP_TYPE::IM_PLAINCOLOUR;
@@ -1244,6 +1252,7 @@ namespace mu
                         blackOp->op.args.ColourConstant.value[0] = 0;
                         blackOp->op.args.ColourConstant.value[1] = 0;
                         blackOp->op.args.ColourConstant.value[2] = 0;
+						blackOp->op.args.ColourConstant.value[3] = 1;
 
                         Ptr<ASTOpFixed> plainOp = new ASTOpFixed;
                         plainOp->op.type = OP_TYPE::IM_PLAINCOLOUR;
diff --git a/Engine/Plugins/Experimental/Mutable/Source/MutableTools/Public/MuT/Compiler.h b/Engine/Plugins/Experimental/Mutable/Source/MutableTools/Public/MuT/Compiler.h
--- a/Engine/Plugins/Experimental/Mutable/Source/MutableTools/Public/MuT/Compiler.h
+++ b/Engine/Plugins/Experimental/Mutable/Source/MutableTools/Public/MuT/Compiler.h
@@ -38,19 +38,6 @@ namespace mu
     class Node;
     class NodeTransformedObject;
 
-    //! Available GPU families
-    typedef enum
-    {
-        GPU_NONE			= 0,
-        GPU_GL4,
-        GPU_GLES2,
-
-        GPU_COUNT
-    } GPU_TYPE;
-
-    //! Return a readable string for a GPU.
-    extern const char* GetGPUName( GPU_TYPE );
-
     //! \brief Options used to compile the models with a compiler.
     class MUTABLETOOLS_API CompilerOptions : public RefCounted
     {
@@ -110,13 +97,6 @@ namespace mu
         //! rutime/Model.h).
         void SetTextureLayoutStrategy( TextureLayoutStrategy strategy );
 
-        //! Set the GPU type to optimise for, if any. Defaults to none.
-        void SetGPUType( GPU_TYPE );
-
-        //! If GPU optimsiation is enabled, set the maximum number of fragment textures the gpu can
-        //! use.
-        void SetGPUMaxFragmentTextures( int );
-
         //-----------------------------------------------------------------------------------------
         // Interface pattern
         //-----------------------------------------------------------------------------------------
