[mutable] Additional reuse of images for image layer and layer colour operations.
#rnx

diff --git a/Engine/Plugins/Experimental/Mutable/Source/MutableRuntime/Private/MuR/CodeRunnerTasked.cpp b/Engine/Plugins/Experimental/Mutable/Source/MutableRuntime/Private/MuR/CodeRunnerTasked.cpp
--- a/Engine/Plugins/Experimental/Mutable/Source/MutableRuntime/Private/MuR/CodeRunnerTasked.cpp
+++ b/Engine/Plugins/Experimental/Mutable/Source/MutableRuntime/Private/MuR/CodeRunnerTasked.cpp
@@ -435,8 +435,8 @@ namespace mu
 
 		bool bApplyColorBlendToAlpha = (Args.flags & OP::ImageLayerArgs::F_APPLY_TO_ALPHA) != 0;
 
-		//ImagePtr pNew = mu::CloneOrTakeOver<Image>(m_base.get());
-		ImagePtr pNew = m_base->Clone();
+		ImagePtr pNew = mu::CloneOrTakeOver<Image>(m_base.get());
+		//ImagePtr pNew = m_base->Clone();
 		check(pNew->GetDataSize() == m_base->GetDataSize());
 
 		bool bValid = pNew->GetSizeX() > 0 && pNew->GetSizeY() > 0;
@@ -694,7 +694,9 @@ namespace mu
 		// Does it apply to colour?
 		if (EBlendType(Args.blendType) != EBlendType::BT_NONE)
 		{
-			pNew = new Image(m_base->GetSizeX(), m_base->GetSizeY(), m_base->GetLODCount(), InitialFormat);
+			// TODO: It could be done "in-place"
+			pNew = mu::CloneOrTakeOver<mu::Image>(m_base.get());
+			//pNew = new Image(m_base->GetSizeX(), m_base->GetSizeY(), m_base->GetLODCount(), m_base->GetFormat());
 
 			if (Args.mask)
 			{			
