[mutable] Fix crash when generating non-power-of-two textures from mutable.
#rnx

diff --git a/Engine/Plugins/Experimental/Mutable/Source/CustomizableObject/Private/MuCO/CustomizableObjectInstance.cpp b/Engine/Plugins/Experimental/Mutable/Source/CustomizableObject/Private/MuCO/CustomizableObjectInstance.cpp
--- a/Engine/Plugins/Experimental/Mutable/Source/CustomizableObject/Private/MuCO/CustomizableObjectInstance.cpp
+++ b/Engine/Plugins/Experimental/Mutable/Source/CustomizableObject/Private/MuCO/CustomizableObjectInstance.cpp
@@ -2676,8 +2676,6 @@ FTexturePlatformData* UCustomizableInstancePrivateData::MutableCreateImagePlatfo
 		return nullptr;
 	}
 
-	FTexturePlatformData* PlatformData = new FTexturePlatformData();
-
 	int32 FirstLOD = 0;
 	for (int32 l = 0; l < OnlyLOD; ++l)
 	{
@@ -2713,7 +2711,14 @@ FTexturePlatformData* UCustomizableInstancePrivateData::MutableCreateImagePlatfo
 	mu::EImageFormat MutableFormat = MutableImage->GetFormat();
 
 	int32 MaxPossibleSize = int32(FMath::Pow(2.f, float(FullLODCount - 1)));
-	check(SizeX == MaxPossibleSize || SizeY == MaxPossibleSize || FullLODCount == 1);
+	
+	// This could happen with non-power-of-two images.
+	//check(SizeX == MaxPossibleSize || SizeY == MaxPossibleSize || FullLODCount == 1);
+	if (!(SizeX == MaxPossibleSize || SizeY == MaxPossibleSize || FullLODCount == 1))
+	{
+		UE_LOG(LogMutable, Warning, TEXT("Building instance: unsuported texture size %d x %d."), SizeX, SizeY);
+		//return nullptr;
+	}
 
 
 	EPixelFormat PlatformFormat = PF_Unknown;
@@ -2745,10 +2750,10 @@ FTexturePlatformData* UCustomizableInstancePrivateData::MutableCreateImagePlatfo
 	default:
 		// Cannot prepare texture if it's not in the right format, this can happen if mutable is in debug mode
 		UE_LOG(LogMutable, Warning, TEXT("Building instance: a texture was generated in an unsupported format, which cannot be converted to Unreal."));
-		delete PlatformData;
 		return nullptr; 
 	}
 
+	FTexturePlatformData* PlatformData = new FTexturePlatformData();
 	PlatformData->SizeX = SizeX;
 	PlatformData->SizeY = SizeY;
 	PlatformData->PixelFormat = PlatformFormat;
@@ -2759,6 +2764,8 @@ FTexturePlatformData* UCustomizableInstancePrivateData::MutableCreateImagePlatfo
 	if (!FMath::IsPowerOfTwo(SizeX) || !FMath::IsPowerOfTwo(SizeY))
 	{
 		EndLOD = FirstLOD + 1;
+		MipsToSkip = 0;
+		FullLODCount = 1;
 	}
 
 	for (int32 MipLevelUE = FirstLOD; MipLevelUE < EndLOD; ++MipLevelUE)
@@ -2869,7 +2876,7 @@ FTexturePlatformData* UCustomizableInstancePrivateData::MutableCreateImagePlatfo
 		}
 	}
 
-	check(OnlyLOD >= 0 || (BulkDataCount == MutableImage->GetLODCount()));
+	check( FullLODCount==1 || OnlyLOD >= 0 || (BulkDataCount == MutableImage->GetLODCount()));
 #endif
 
 	return PlatformData;
