[mutable] Optimize image layer operations with constant colors.
Fix bug with disabled layouts possibly causing a crash.
Prevent unnecessary mipmap operation split.
#rnx

diff --git a/Engine/Plugins/Experimental/Mutable/Source/CustomizableObjectEditor/Private/MuCOE/SMutableCodeViewer.cpp b/Engine/Plugins/Experimental/Mutable/Source/CustomizableObjectEditor/Private/MuCOE/SMutableCodeViewer.cpp
--- a/Engine/Plugins/Experimental/Mutable/Source/CustomizableObjectEditor/Private/MuCOE/SMutableCodeViewer.cpp
+++ b/Engine/Plugins/Experimental/Mutable/Source/CustomizableObjectEditor/Private/MuCOE/SMutableCodeViewer.cpp
@@ -494,8 +494,8 @@ void SMutableCodeViewer::Construct(const FArguments& InArgs, const TSharedPtr<mu
 	ToolbarBuilder.BeginSection("Export");
 
 	// Tree Sizes
-	constexpr float OperationsColumnWidth = 0.75f;
-	constexpr float ExtraDataColumnWidth = 0.25f;
+	constexpr float OperationsColumnWidth = 0.90f;
+	constexpr float ExtraDataColumnWidth = 0.10f;
 	
 	// Export
 	ToolbarBuilder.AddToolBarButton(
@@ -668,7 +668,7 @@ void SMutableCodeViewer::Construct(const FArguments& InArgs, const TSharedPtr<mu
 								.FillWidth(OperationsColumnWidth)
 								
 							+ SHeaderRow::Column(MutableCodeTreeViewColumns::AdditionalDataColumnID)
-								.DefaultLabel(LOCTEXT("AdditionalData", "Additional Data"))
+								.DefaultLabel(LOCTEXT("OperationFlags", "Flags"))
 								.FillWidth(ExtraDataColumnWidth)
 						)
 					]
diff --git a/Engine/Plugins/Experimental/Mutable/Source/MutableRuntime/Private/MuR/ImagePrivate.h b/Engine/Plugins/Experimental/Mutable/Source/MutableRuntime/Private/MuR/ImagePrivate.h
--- a/Engine/Plugins/Experimental/Mutable/Source/MutableRuntime/Private/MuR/ImagePrivate.h
+++ b/Engine/Plugins/Experimental/Mutable/Source/MutableRuntime/Private/MuR/ImagePrivate.h
@@ -217,7 +217,7 @@ namespace mu
     Ptr<T> CloneOrTakeOver( const T* source )
     {
         Ptr<T> result;
-        if (source->GetRefCount()==1)
+        if (source->IsUnique())
         {
             result = const_cast<T*>(source);
         }
diff --git a/Engine/Plugins/Experimental/Mutable/Source/MutableRuntime/Public/MuR/RefCounted.h b/Engine/Plugins/Experimental/Mutable/Source/MutableRuntime/Public/MuR/RefCounted.h
--- a/Engine/Plugins/Experimental/Mutable/Source/MutableRuntime/Public/MuR/RefCounted.h
+++ b/Engine/Plugins/Experimental/Mutable/Source/MutableRuntime/Public/MuR/RefCounted.h
@@ -26,22 +26,23 @@ namespace mu
 
 		inline void IncRef() const
 		{
-			m_refCount++;
+			std::atomic_fetch_add_explicit( &m_refCount, 1u, std::memory_order_relaxed );
 		}
 
 		inline void DecRef() const
 		{
 			LLM_SCOPE_BYNAME(TEXT("MutableRuntime"));
 
-			if (m_refCount.fetch_sub(1) == 1)
+			if (m_refCount.fetch_sub(1, std::memory_order_release) == 1)
 			{
+				std::atomic_thread_fence(std::memory_order_acquire);
 				delete this;
 			}
 		}
 
-		inline int32 GetRefCount() const
+		inline bool IsUnique() const
 		{
-			return m_refCount;
+			return m_refCount.load(std::memory_order_relaxed)==1;
 		}
 
 		RefCounted(const RefCounted&) = delete;
diff --git a/Engine/Plugins/Experimental/Mutable/Source/MutableTools/Private/MuT/AST.cpp b/Engine/Plugins/Experimental/Mutable/Source/MutableTools/Private/MuT/AST.cpp
--- a/Engine/Plugins/Experimental/Mutable/Source/MutableTools/Private/MuT/AST.cpp
+++ b/Engine/Plugins/Experimental/Mutable/Source/MutableTools/Private/MuT/AST.cpp
@@ -121,7 +121,7 @@ void ASTOp::RemoveChildren()
             if (c)
             {
                 // Are we clearing the last reference?
-                if (c.child()->GetRefCount()==1)
+                if (c.child()->IsUnique())
                 {
                     toDestroy.Add(c.child());
                     pending.Add(c.child().get());
diff --git a/Engine/Plugins/Experimental/Mutable/Source/MutableTools/Private/MuT/ASTOpImageLayer.cpp b/Engine/Plugins/Experimental/Mutable/Source/MutableTools/Private/MuT/ASTOpImageLayer.cpp
--- a/Engine/Plugins/Experimental/Mutable/Source/MutableTools/Private/MuT/ASTOpImageLayer.cpp
+++ b/Engine/Plugins/Experimental/Mutable/Source/MutableTools/Private/MuT/ASTOpImageLayer.cpp
@@ -2,16 +2,16 @@
 
 #include "MuT/ASTOpImageLayer.h"
 
-#include "Containers/Map.h"
-#include "HAL/PlatformMath.h"
+#include "MuT/StreamsPrivate.h"
+#include "MuT/ASTOpImagePatch.h"
+#include "MuT/ASTOpImageLayerColor.h"
+#include "MuT/ASTOpSwitch.h"
 #include "MuR/ModelPrivate.h"
 #include "MuR/RefCounted.h"
 #include "MuR/Types.h"
-#include "MuT/StreamsPrivate.h"
-#include "MuT/ASTOpImagePatch.h"
 
-#include <memory>
-#include <utility>
+#include "Containers/Map.h"
+#include "HAL/PlatformMath.h"
 
 
 namespace mu
@@ -176,10 +176,31 @@ namespace mu
 		Ptr<ASTOp> at;
 
 		auto baseAt = base.child();
+		auto blendAt = blend.child();
+		auto maskAt = mask.child();
+
+		// Convert to image layer color if blend is plain
+		if (!at && blendAt->GetOpType() == OP_TYPE::IM_PLAINCOLOUR)
+		{
+			bool bUseMaskFromBlendAlpha = (Flags & OP::ImageLayerArgs::F_USE_MASK_FROM_BLENDED);
+
+			const ASTOpFixed* BlendPlainColor = dynamic_cast<const ASTOpFixed*>(blendAt.get());
+
+			// \TODO: Add support for embedded alpha to image layer color.
+			if (!bUseMaskFromBlendAlpha)
+			{
+				Ptr<ASTOpImageLayerColor> NewLayerColor = new ASTOpImageLayerColor;
+				NewLayerColor->base = baseAt;
+				NewLayerColor->mask = maskAt;
+				NewLayerColor->blendType = blendType;
+				NewLayerColor->blendTypeAlpha = blendTypeAlpha;
+				NewLayerColor->color = BlendPlainColor->children[BlendPlainColor->op.args.ImagePlainColour.colour].child();
+				at = NewLayerColor;
+			}
+		}
 
 		// Plain masks optimization
-		auto maskAt = mask.child();
-		if (maskAt)
+		if (!at && maskAt)
 		{
 			FVector4f colour;
 			if (maskAt->IsImagePlainConstant(colour))
@@ -316,7 +337,7 @@ namespace mu
 
 		case OP_TYPE::IM_SWITCH:
 		{
-			// Disabled:
+			// Warning:
 			// It seems to cause data explosion in optimizer in some models. Because
 			// all switch branches become unique constants
 
@@ -368,6 +389,7 @@ namespace mu
 //                }
 
 //                at = nop;
+
 			break;
 		}
 
@@ -399,6 +421,81 @@ namespace mu
 
 		}
 
+		// If we failed to optimize so far, see if it is worth optimizing the blended only.
+		if (!at)
+		{
+			OP_TYPE BlendType = blendAt->GetOpType();
+			switch (BlendType)
+			{
+
+			case OP_TYPE::IM_SWITCH:
+			{
+				const ASTOpSwitch* BlendSwitch = dynamic_cast<const ASTOpSwitch*>(blendAt.get());
+
+				// If at least a switch option is a plain colour, sink the layer into the switch
+				bool bWorthSinking = false;
+				for (int32 v = 0; v < BlendSwitch->cases.Num(); ++v)
+				{
+					if (BlendSwitch->cases[v].branch)
+					{
+						// \TODO: Use the smarter query function to detect plain images?
+						if (BlendSwitch->cases[v].branch->GetOpType() == OP_TYPE::IM_PLAINCOLOUR)
+						{
+							bWorthSinking = true;
+							break;
+						}
+					}
+				}
+
+				if (bWorthSinking)
+				{
+					bool bMaskIsCompatibleSwitch = false;
+					const ASTOpSwitch* MaskSwitch = dynamic_cast<const ASTOpSwitch*>(maskAt.get());
+					if (maskAt && maskAt->GetOpType()== OP_TYPE::IM_SWITCH)
+					{
+						bMaskIsCompatibleSwitch = MaskSwitch->IsCompatibleWith(BlendSwitch);
+					}
+
+					Ptr<ASTOpSwitch> NewSwitch = mu::Clone<ASTOpSwitch>(BlendSwitch);
+
+					if (NewSwitch->def)
+					{
+						Ptr<ASTOpImageLayer> defOp = mu::Clone<ASTOpImageLayer>(this);
+						defOp->blend = BlendSwitch->def.child();
+						if (bMaskIsCompatibleSwitch)
+						{
+							defOp->mask = MaskSwitch->def.child();
+						}
+						NewSwitch->def = defOp;
+					}
+
+					for (int32 v = 0; v < NewSwitch->cases.Num(); ++v)
+					{
+						if (NewSwitch->cases[v].branch)
+						{
+							Ptr<ASTOpImageLayer> BranchOp = mu::Clone<ASTOpImageLayer>(this);
+							BranchOp->blend = BlendSwitch->cases[v].branch.child();
+							if (bMaskIsCompatibleSwitch)
+							{
+								BranchOp->mask = MaskSwitch->cases[v].branch.child();
+							}
+							NewSwitch->cases[v].branch = BranchOp;
+						}
+					}
+
+
+					at = NewSwitch;
+				}
+
+				break;
+			}
+
+			default:
+				break;
+
+			}
+		}
+
 		return at;
 	}
 
diff --git a/Engine/Plugins/Experimental/Mutable/Source/MutableTools/Private/MuT/ASTOpImageLayerColor.cpp b/Engine/Plugins/Experimental/Mutable/Source/MutableTools/Private/MuT/ASTOpImageLayerColor.cpp
--- a/Engine/Plugins/Experimental/Mutable/Source/MutableTools/Private/MuT/ASTOpImageLayerColor.cpp
+++ b/Engine/Plugins/Experimental/Mutable/Source/MutableTools/Private/MuT/ASTOpImageLayerColor.cpp
@@ -10,9 +10,6 @@
 #include "MuR/Types.h"
 #include "MuT/StreamsPrivate.h"
 
-#include <memory>
-#include <utility>
-
 
 namespace mu
 {
@@ -166,6 +163,81 @@ namespace mu
 	}
 
 
+	//---------------------------------------------------------------------------------------------
+	Ptr<ASTOp> ASTOpImageLayerColor::OptimiseSemantic(const FModelOptimizationOptions& options) const
+	{
+		Ptr<ASTOp> at;
+
+		// Plain masks optimization
+		if (!at && mask.child())
+		{
+			FVector4f colour;
+			if (mask.child()->IsImagePlainConstant(colour))
+			{
+				if (colour.IsNearlyZero3(UE_SMALL_NUMBER))
+				{
+					// If the mask is black, we can skip the entire operation
+					at = base.child();
+				}
+				else if (colour.Equals(FVector4f(1, 1, 1, 1), UE_SMALL_NUMBER))
+				{
+					// If the mask is white, we can remove it
+					Ptr<ASTOpImageLayerColor> nop = mu::Clone<ASTOpImageLayerColor>(this);
+					nop->mask = nullptr;
+					at = nop;
+				}
+			}
+		}
+
+		// Layer operations with constants that do nothing.
+		if (!at)
+		{
+			bool bRGBUnchanged = blendType == EBlendType::BT_NONE;
+			bool bAlphaUnchanged = blendTypeAlpha == EBlendType::BT_NONE;
+
+			FVector4f ColorConst(0,0,0,1);
+			if (!color || color.child()->GetOpType()==OP_TYPE::CO_CONSTANT)
+			{
+				if (color)
+				{
+					const ASTOpFixed* TypedColor = dynamic_cast<const ASTOpFixed*>(color.child().get());
+					const float* Value = TypedColor->op.args.ColourConstant.value;
+					ColorConst.Set(Value[0], Value[1], Value[2], Value[3]);
+				}
+
+				if (!bRGBUnchanged)
+				{
+					switch (blendType)
+					{ 
+					case EBlendType::BT_LIGHTEN: bRGBUnchanged = ColorConst.IsNearlyZero3(UE_SMALL_NUMBER); break;
+					case EBlendType::BT_MULTIPLY: bRGBUnchanged = ColorConst.Equals(FVector4f(1, 1, 1, 1)); break;
+					default: break;
+					}					
+				}
+
+				if (!bAlphaUnchanged)
+				{
+					// TODO: Update when alpha may come from alpha in the color.
+					switch (blendTypeAlpha)
+					{
+					case EBlendType::BT_LIGHTEN: bAlphaUnchanged = ColorConst.IsNearlyZero3(UE_SMALL_NUMBER); break;
+					case EBlendType::BT_MULTIPLY: bAlphaUnchanged = ColorConst.Equals(FVector4f(1, 1, 1, 1)); break;
+					default: break;
+					}
+				}
+			}
+
+			if (bRGBUnchanged && bAlphaUnchanged)
+			{
+				// Skip this operation.
+				at = base.child();
+			}
+		}
+
+		return at;
+	}
+
+
 	//-------------------------------------------------------------------------------------------------
 	Ptr<ASTOp> ASTOpImageLayerColor::OptimiseSink(const FModelOptimizationOptions& options, FOptimizeSinkContext& context) const
 	{
diff --git a/Engine/Plugins/Experimental/Mutable/Source/MutableTools/Private/MuT/ASTOpImageLayerColor.h b/Engine/Plugins/Experimental/Mutable/Source/MutableTools/Private/MuT/ASTOpImageLayerColor.h
--- a/Engine/Plugins/Experimental/Mutable/Source/MutableTools/Private/MuT/ASTOpImageLayerColor.h
+++ b/Engine/Plugins/Experimental/Mutable/Source/MutableTools/Private/MuT/ASTOpImageLayerColor.h
@@ -46,6 +46,7 @@ namespace mu
 		FImageDesc GetImageDesc(bool returnBestOption, FGetImageDescContext* context) const override;
 		void GetLayoutBlockSize(int* pBlockX, int* pBlockY) override;
 		Ptr<ImageSizeExpression> GetImageSizeExpression() const override;
+		Ptr<ASTOp> OptimiseSemantic(const FModelOptimizationOptions& options) const;
 		Ptr<ASTOp> OptimiseSink(const FModelOptimizationOptions&, FOptimizeSinkContext&) const;
 
 	};
diff --git a/Engine/Plugins/Experimental/Mutable/Source/MutableTools/Private/MuT/ASTOpImageMipmap.cpp b/Engine/Plugins/Experimental/Mutable/Source/MutableTools/Private/MuT/ASTOpImageMipmap.cpp
--- a/Engine/Plugins/Experimental/Mutable/Source/MutableTools/Private/MuT/ASTOpImageMipmap.cpp
+++ b/Engine/Plugins/Experimental/Mutable/Source/MutableTools/Private/MuT/ASTOpImageMipmap.cpp
@@ -268,10 +268,9 @@ namespace mu
 		// Before sinking, see if it is worth splitting into miptail and mip.
 		if (!m_root->bPreventSplitTail
 			// \TODO: Review this: It seems it fails for a BC1 8x4 texture having blockSize=4, and level=0
-				//&&
-				//m_root->Levels
-				//!=
-				//m_root->BlockLevels 
+			&&
+			// No need to split if there is no difference between the levels considering blocks
+			m_root->Levels != m_root->BlockLevels 
 			)
 		{
 			// the block mipmaps can be done before composition is done.
diff --git a/Engine/Plugins/Experimental/Mutable/Source/MutableTools/Private/MuT/ASTOpImagePixelFormat.cpp b/Engine/Plugins/Experimental/Mutable/Source/MutableTools/Private/MuT/ASTOpImagePixelFormat.cpp
--- a/Engine/Plugins/Experimental/Mutable/Source/MutableTools/Private/MuT/ASTOpImagePixelFormat.cpp
+++ b/Engine/Plugins/Experimental/Mutable/Source/MutableTools/Private/MuT/ASTOpImagePixelFormat.cpp
@@ -106,7 +106,7 @@ namespace mu
 		mu::Ptr<ASTOp> sourceAt = Source.child();
 
 		EImageFormat format = Format;
-		bool isCompressedFormat = IsCompressedFormat(Format);
+		bool bIsCompressedFormat = IsCompressedFormat(Format);
 		//bool isBlockFormat = GetImageFormatData( format ).m_pixelsPerBlockX!=0;
 
 		// The instruction can be sunk
@@ -126,7 +126,7 @@ namespace mu
 		case OP_TYPE::IM_DISPLACE:
 		{
 			// This op doesn't support compressed formats
-			if (!isCompressedFormat)
+			if (!bIsCompressedFormat)
 			{
 				mu::Ptr<ASTOpFixed> newOp = mu::Clone<ASTOpFixed>(sourceAt);
 
@@ -142,7 +142,7 @@ namespace mu
 		case OP_TYPE::IM_RASTERMESH:
 		{
 			// This op doesn't support compressed formats
-			if (!isCompressedFormat
+			if (!bIsCompressedFormat
 				&&
 				dynamic_cast<const ASTOpFixed*>(sourceAt.get())->op.args.ImageRasterMesh.image)
 			{
@@ -304,7 +304,7 @@ namespace mu
 		if (!at) return nullptr;
 
 		EImageFormat format = currentFormatOp->Format;
-		bool isCompressedFormat = IsCompressedFormat(format);
+		bool bIsCompressedFormat = IsCompressedFormat(format);
 		bool isBlockFormat = GetImageFormatData(format).m_pixelsPerBlockX != 0;
 
 		// Already visited?
@@ -320,7 +320,7 @@ namespace mu
 
 		case OP_TYPE::IM_CONDITIONAL:
 		{
-			// We move the mask creation down the two paths
+			// We move the op down the two paths
 			auto newOp = mu::Clone<ASTOpConditional>(at);
 			newOp->yes = Visit(newOp->yes.child(), currentFormatOp);
 			newOp->no = Visit(newOp->no.child(), currentFormatOp);
@@ -330,7 +330,7 @@ namespace mu
 
 		case OP_TYPE::IM_SWITCH:
 		{
-			// We move the mask creation down all the paths
+			// We move the op down all the paths
 			auto newOp = mu::Clone<ASTOpSwitch>(at);
 			newOp->def = Visit(newOp->def.child(), currentFormatOp);
 			for (auto& c : newOp->cases)
@@ -411,7 +411,7 @@ namespace mu
 
 			// If its a compressed format, only sink formats on mipmap operations that
 			// generate the tail. To avoid optimization loop.
-			if (!isCompressedFormat || typedSource->bOnlyTail)
+			if (!bIsCompressedFormat || typedSource->bOnlyTail)
 			{
 				auto newOp = mu::Clone<ASTOpImageMipmap>(typedSource);
 
@@ -426,7 +426,7 @@ namespace mu
 		case OP_TYPE::IM_INTERPOLATE:
 		{
 			// This op doesn't support compressed formats
-			if (!isCompressedFormat)
+			if (!bIsCompressedFormat)
 			{
 				// Move the format down all the paths
 				auto newOp = mu::Clone<ASTOpFixed>(at);
@@ -447,7 +447,7 @@ namespace mu
 		case OP_TYPE::IM_INTERPOLATE3:
 		{
 			// This op doesn't support compressed formats
-			if (!isCompressedFormat)
+			if (!bIsCompressedFormat)
 			{
 				// We move the format down all the paths
 				auto newOp = mu::Clone<ASTOpFixed>(at);
diff --git a/Engine/Plugins/Experimental/Mutable/Source/MutableTools/Private/MuT/CodeGenerator_Image.cpp b/Engine/Plugins/Experimental/Mutable/Source/MutableTools/Private/MuT/CodeGenerator_Image.cpp
--- a/Engine/Plugins/Experimental/Mutable/Source/MutableTools/Private/MuT/CodeGenerator_Image.cpp
+++ b/Engine/Plugins/Experimental/Mutable/Source/MutableTools/Private/MuT/CodeGenerator_Image.cpp
@@ -1432,7 +1432,11 @@ namespace mu
 
         pop->SetChild( pop->op.args.MeshProject.projector, projectorResult.op );
 
-		int32 LayoutBlockIndex = m_imageState.Last().m_pLayout->m_blocks.IndexOfByPredicate([&](const Layout::FBlock& Block) { return Block.m_id == m_imageState.Last().m_layoutBlockId; });
+		int32 LayoutBlockIndex = -1;
+		if (m_imageState.Last().m_pLayout)
+		{
+			LayoutBlockIndex = m_imageState.Last().m_pLayout->m_blocks.IndexOfByPredicate([&](const Layout::FBlock& Block) { return Block.m_id == m_imageState.Last().m_layoutBlockId; });
+		}
 		int32 GeneratedLayoutBlockId = -1;
 
         // Mesh
@@ -1452,7 +1456,11 @@ namespace mu
             GenerateMesh( MeshOptions, MeshResult, node.m_pMesh );
 
 			// Match the block id of the block we are generating with the id that resulted in the generated mesh
-			GeneratedLayoutBlockId = MeshResult.GeneratedLayouts[node.m_layout]->m_blocks[LayoutBlockIndex].m_id;
+			GeneratedLayoutBlockId = -1;
+			if (LayoutBlockIndex >= 0)
+			{
+				GeneratedLayoutBlockId = MeshResult.GeneratedLayouts[node.m_layout]->m_blocks[LayoutBlockIndex].m_id;
+			}
 
             pop->SetChild( pop->op.args.MeshProject.mesh, MeshResult.meshOp );
 
@@ -1477,7 +1485,7 @@ namespace mu
             else
             {
                 // Extract the mesh layout block
-                if ( m_imageState.Num() && m_imageState.Last().m_layoutBlockId >=0 )
+                if ( m_imageState.Num() && GeneratedLayoutBlockId>=0 )
                 {
                     Ptr<ASTOpMeshExtractLayoutBlocks> eop = new ASTOpMeshExtractLayoutBlocks();
                     eop->source = pop->children[pop->op.args.MeshProject.mesh].child();
