[mutable] Initial support for decompression of additional image formats.
#rnx

diff --git a/Engine/Plugins/Experimental/Mutable/Source/MutableRuntime/Private/MuR/BlockCompression/Miro/Miro.cpp b/Engine/Plugins/Experimental/Mutable/Source/MutableRuntime/Private/MuR/BlockCompression/Miro/Miro.cpp
--- a/Engine/Plugins/Experimental/Mutable/Source/MutableRuntime/Private/MuR/BlockCompression/Miro/Miro.cpp
+++ b/Engine/Plugins/Experimental/Mutable/Source/MutableRuntime/Private/MuR/BlockCompression/Miro/Miro.cpp
@@ -2065,7 +2065,7 @@ namespace impl
 
 		void stb_compress_dxt_block(unsigned char* dest, const unsigned char* src, int alpha, int mode)
 		{
-			// If this assert fails it means miro::initialize() has not been called.
+			// If this check fails it means miro::initialize() has not been called.
 			check(s_initialised_dxtc_decompress);
 
 			unsigned char data[16][4];
@@ -2121,7 +2121,7 @@ namespace impl
 		(void)x;
 		// if (!x)
 		// {
-		//     assert(x);
+		//     check(x);
 		// }
 	}
 
@@ -4711,4081 +4711,1494 @@ namespace impl
 	}
 
 
-
-
-
-	//-------------------------------------------------------------------------------------------------
-	// Based on ARMs reference implementation, under a special license
-	// See: [REDACTED_PROFILE_URL]/astc-encoder
-	//-------------------------------------------------------------------------------------------------
-	namespace arm
+	// Based on astc_dec released under Apache license
+	// See: [REDACTED_PROFILE_URL]/astc_dec
+	// Just modified for integration.
+	namespace astcdec
 	{
-		// Macro to silence warnings on ignored parameters.
-		// The presence of this macro should be a signal to look at refactoring.
-		//#define IGNORE(param) ((void)&param)
-
-#define astc_isnan(p) ((p)!=(p))
-
-// ASTC parameters
-#define MAX_TEXELS_PER_BLOCK 216
-#define MAX_WEIGHTS_PER_BLOCK 64
-#define MIN_WEIGHT_BITS_PER_BLOCK 24
-#define MAX_WEIGHT_BITS_PER_BLOCK 96
-#define PARTITION_BITS 10
-#define PARTITION_COUNT (1 << PARTITION_BITS)
-
-// the sum of weights for one texel.
-#define TEXEL_WEIGHT_SUM 16
-#define MAX_DECIMATION_MODES 87
-#define MAX_WEIGHT_MODES 2048
-
-// error reporting for codec internal errors.
-#define ASTC_CODEC_INTERNAL_ERROR check(false)
-
-
-/*
-   In ASTC, we don't necessarily provide a weight for every texel.
-   As such, for each block size, there are a number of patterns where some texels
-   have their weights computed as a weighted average of more than 1 weight.
-   As such, the codec uses a data structure that tells us: for each texel, which
-   weights it is a combination of for each weight, which texels it contributes to.
-   The decimation_table is this data structure.
-*/
-		struct decimation_table
-		{
-			int num_texels;
-			int num_weights;
-			uint8 texel_num_weights[MAX_TEXELS_PER_BLOCK];	// number of indices that go into the calculation for a texel
-			uint8 texel_weights_int[MAX_TEXELS_PER_BLOCK][4];	// the weight to assign to each weight
-			float texel_weights_float[MAX_TEXELS_PER_BLOCK][4];	// the weight to assign to each weight
-			uint8 texel_weights[MAX_TEXELS_PER_BLOCK][4];	// the weights that go into a texel calculation
-			uint8 weight_num_texels[MAX_WEIGHTS_PER_BLOCK];	// the number of texels that a given weight contributes to
-			uint8 weight_texel[MAX_WEIGHTS_PER_BLOCK][MAX_TEXELS_PER_BLOCK];	// the texels that the weight contributes to
-			uint8 weights_int[MAX_WEIGHTS_PER_BLOCK][MAX_TEXELS_PER_BLOCK];	// the weights that the weight contributes to a texel.
-			float weights_flt[MAX_WEIGHTS_PER_BLOCK][MAX_TEXELS_PER_BLOCK];	// the weights that the weight contributes to a texel.
-		};
-
-
-		struct block_mode
-		{
-			int8_t decimation_mode;
-			int8_t quantization_mode;
-			int8_t is_dual_plane;
-			int8_t permit_encode;
-			int8_t permit_decode;
-			float percentile;
-		};
 
+#define DE_LENGTH_OF_ARRAY(x) (sizeof(x)/sizeof(x[0]))
+#define DE_UNREF(x) (void)x
 
-		struct block_size_descriptor
-		{
-			int decimation_mode_count;
-			int decimation_mode_samples[MAX_DECIMATION_MODES];
-			int decimation_mode_maxprec_1plane[MAX_DECIMATION_MODES];
-			int decimation_mode_maxprec_2planes[MAX_DECIMATION_MODES];
-			float decimation_mode_percentile[MAX_DECIMATION_MODES];
-			int permit_encode[MAX_DECIMATION_MODES];
-			const decimation_table* decimation_tables[MAX_DECIMATION_MODES + 1];
-			block_mode block_modes[MAX_WEIGHT_MODES];
-		};
-
+		typedef uint8_t deUint8;
+		typedef int8_t deInt8;
+		typedef uint32_t deUint32;
+		typedef int32_t deInt32;
+		typedef uint16_t deUint16;
+		typedef int16_t deInt16;
+		typedef int64_t deInt64;
+		typedef uint64_t deUint64;
 
+#define DE_ASSERT check
 
-		struct symbolic_compressed_block
+		static bool inBounds(int v, int l, int h)
 		{
-			int error_block;			// 1 marks error block, 0 marks non-error-block.
-			int block_mode;				// 0 to 2047. Negative value marks constant-color block (-1: FP16, -2:UINT16)
-			int partition_count;		// 1 to 4; Zero marks a constant-color block.
-			int partition_index;		// 0 to 1023
-			int color_formats[4];		// color format for each endpoint color pair.
-			int color_formats_matched;	// color format for all endpoint pairs are matched.
-			int color_values[4][12];	// quantized endpoint color pairs.
-			int color_quantization_level;
-			uint8 plane1_weights[MAX_WEIGHTS_PER_BLOCK];	// quantized and decimated weights
-			uint8 plane2_weights[MAX_WEIGHTS_PER_BLOCK];
-			int plane2_color_component;	// color component for the secondary plane of weights
-			int constant_color[4];		// constant-color, as FP16 or UINT16. Used for constant-color blocks only.
-		};
+			return (v >= l) && (v < h);
+		}
 
-		struct physical_compressed_block
+		static bool inRange(int v, int l, int h)
 		{
-			uint8 data[16];
-		};
-
+			return (v >= l) && (v <= h);
+		}
 
-		static void initialize_decimation_table_2d(decimation_table* dt)
+		struct UVec4
 		{
-			int i, j;
-			int x, y;
-
-			int texels_per_block = 16;
-			int weights_per_block = 16;
-
-			int weightcount_of_texel[MAX_TEXELS_PER_BLOCK];
-			int grid_weights_of_texel[MAX_TEXELS_PER_BLOCK][4];
-			int weights_of_texel[MAX_TEXELS_PER_BLOCK][4];
-
-			int texelcount_of_weight[MAX_WEIGHTS_PER_BLOCK];
-			// stack is to big
-			//        int texels_of_weight[MAX_WEIGHTS_PER_BLOCK][MAX_TEXELS_PER_BLOCK];
-			//        int texelweights_of_weight[MAX_WEIGHTS_PER_BLOCK][MAX_TEXELS_PER_BLOCK];
-			TUniquePtr<int[]> texels_of_weight(new int[MAX_WEIGHTS_PER_BLOCK * MAX_TEXELS_PER_BLOCK]);
-			TUniquePtr<int[]> texelweights_of_weight(new int[MAX_WEIGHTS_PER_BLOCK * MAX_TEXELS_PER_BLOCK]);
-
-			for (i = 0; i < MAX_WEIGHTS_PER_BLOCK; i++)
-			{
-				texelcount_of_weight[i] = 0;
-				weightcount_of_texel[i] = 0;
-			}
-
-			for (y = 0; y < 4; y++)
-				for (x = 0; x < 4; x++)
-				{
-					int texel = y * 4 + x;
-
-					int x_weight = (((1024 + 4 / 2) / (4 - 1)) * x * (4 - 1) + 32) >> 6;
-					int y_weight = (((1024 + 4 / 2) / (4 - 1)) * y * (4 - 1) + 32) >> 6;
-
-					int x_weight_frac = x_weight & 0xF;
-					int y_weight_frac = y_weight & 0xF;
-					int x_weight_int = x_weight >> 4;
-					int y_weight_int = y_weight >> 4;
-					int qweight[4];
-					int weight[4];
-					qweight[0] = x_weight_int + y_weight_int * 4;
-					qweight[1] = qweight[0] + 1;
-					qweight[2] = qweight[0] + 4;
-					qweight[3] = qweight[2] + 1;
-
-					// truncated-precision bilinear interpolation.
-					int prod = x_weight_frac * y_weight_frac;
-
-					weight[3] = (prod + 8) >> 4;
-					weight[1] = x_weight_frac - weight[3];
-					weight[2] = y_weight_frac - weight[3];
-					weight[0] = 16 - x_weight_frac - y_weight_frac + weight[3];
-
-					for (i = 0; i < 4; i++)
-						if (weight[i] != 0)
-						{
-							grid_weights_of_texel[texel][weightcount_of_texel[texel]] = qweight[i];
-							weights_of_texel[texel][weightcount_of_texel[texel]] = weight[i];
-							weightcount_of_texel[texel]++;
-							texels_of_weight[qweight[i] * MAX_TEXELS_PER_BLOCK + texelcount_of_weight[qweight[i]]] = texel;
-							texelweights_of_weight[qweight[i] * MAX_TEXELS_PER_BLOCK + texelcount_of_weight[qweight[i]]] = weight[i];
-							texelcount_of_weight[qweight[i]]++;
-						}
-				}
+			uint32_t m_c[4];
 
-			for (i = 0; i < texels_per_block; i++)
+			UVec4()
 			{
-				dt->texel_num_weights[i] = (uint8)weightcount_of_texel[i];
-
-				// ensure that all 4 entries are actually initialized.
-				// This allows a branch-free implementation of compute_value_of_texel_flt()
-				for (j = 0; j < 4; j++)
-				{
-					dt->texel_weights_int[i][j] = 0;
-					dt->texel_weights_float[i][j] = 0.0f;
-					dt->texel_weights[i][j] = 0;
-				}
-
-				for (j = 0; j < weightcount_of_texel[i]; j++)
-				{
-					dt->texel_weights_int[i][j] = (uint8)weights_of_texel[i][j];
-					dt->texel_weights_float[i][j] = static_cast <float>(weights_of_texel[i][j]) * (1.0f / TEXEL_WEIGHT_SUM);
-					dt->texel_weights[i][j] = (uint8)grid_weights_of_texel[i][j];
-				}
+				m_c[0] = 0;
+				m_c[1] = 0;
+				m_c[2] = 0;
+				m_c[3] = 0;
 			}
 
-			for (i = 0; i < weights_per_block; i++)
+			UVec4(uint32_t x, uint32_t y, uint32_t z, uint32_t w)
 			{
-				dt->weight_num_texels[i] = (uint8)texelcount_of_weight[i];
-
-
-				for (j = 0; j < texelcount_of_weight[i]; j++)
-				{
-					dt->weight_texel[i][j] = (uint8)texels_of_weight[i * MAX_TEXELS_PER_BLOCK + j];
-					dt->weights_int[i][j] = (uint8)texelweights_of_weight[i * MAX_TEXELS_PER_BLOCK + j];
-					dt->weights_flt[i][j] = static_cast <float>(texelweights_of_weight[i * MAX_TEXELS_PER_BLOCK + j]);
-				}
+				m_c[0] = x;
+				m_c[1] = y;
+				m_c[2] = z;
+				m_c[3] = w;
 			}
 
-			dt->num_texels = texels_per_block;
-			dt->num_weights = weights_per_block;
-		}
-
-
-		// enumeration of all the quantization methods we support under this format.
-		enum quantization_method
-		{
-			QUANT_2 = 0,
-			QUANT_3 = 1,
-			QUANT_4 = 2,
-			QUANT_5 = 3,
-			QUANT_6 = 4,
-			QUANT_8 = 5,
-			QUANT_10 = 6,
-			QUANT_12 = 7,
-			QUANT_16 = 8,
-			QUANT_20 = 9,
-			QUANT_24 = 10,
-			QUANT_32 = 11,
-			QUANT_40 = 12,
-			QUANT_48 = 13,
-			QUANT_64 = 14,
-			QUANT_80 = 15,
-			QUANT_96 = 16,
-			QUANT_128 = 17,
-			QUANT_160 = 18,
-			QUANT_192 = 19,
-			QUANT_256 = 20
-		};
+			uint32_t x() const { return m_c[0]; }
+			uint32_t y() const { return m_c[1]; }
+			uint32_t z() const { return m_c[2]; }
+			uint32_t w() const { return m_c[3]; }
 
+			uint32_t& x() { return m_c[0]; }
+			uint32_t& y() { return m_c[1]; }
+			uint32_t& z() { return m_c[2]; }
+			uint32_t& w() { return m_c[3]; }
 
-		int compute_ise_bitcount(int items, quantization_method quant)
-		{
-			switch (quant)
-			{
-			case QUANT_2:
-				return items;
-			case QUANT_3:
-				return (8 * items + 4) / 5;
-			case QUANT_4:
-				return 2 * items;
-			case QUANT_5:
-				return (7 * items + 2) / 3;
-			case QUANT_6:
-				return (13 * items + 4) / 5;
-			case QUANT_8:
-				return 3 * items;
-			case QUANT_10:
-				return (10 * items + 2) / 3;
-			case QUANT_12:
-				return (18 * items + 4) / 5;
-			case QUANT_16:
-				return items * 4;
-			case QUANT_20:
-				return (13 * items + 2) / 3;
-			case QUANT_24:
-				return (23 * items + 4) / 5;
-			case QUANT_32:
-				return 5 * items;
-			case QUANT_40:
-				return (16 * items + 2) / 3;
-			case QUANT_48:
-				return (28 * items + 4) / 5;
-			case QUANT_64:
-				return 6 * items;
-			case QUANT_80:
-				return (19 * items + 2) / 3;
-			case QUANT_96:
-				return (33 * items + 4) / 5;
-			case QUANT_128:
-				return 7 * items;
-			case QUANT_160:
-				return (22 * items + 2) / 3;
-			case QUANT_192:
-				return (38 * items + 4) / 5;
-			case QUANT_256:
-				return 8 * items;
-			default:
-				return 100000;
-			}
-		}
-
-
-		const float percentile_table_4x4[2048] =
-		{
-			1.0000f, 1.0000f, 1.0000f, 1.0000f, 1.0000f, 1.0000f, 1.0000f, 1.0000f, 1.0000f, 1.0000f, 1.0000f, 1.0000f, 1.0000f, 1.0000f, 1.0000f, 1.0000f,
-			1.0000f, 1.0000f, 1.0000f, 0.8661f, 1.0000f, 1.0000f, 1.0000f, 1.0000f, 1.0000f, 1.0000f, 1.0000f, 1.0000f, 1.0000f, 1.0000f, 1.0000f, 1.0000f,
-			1.0000f, 1.0000f, 0.7732f, 0.8567f, 1.0000f, 1.0000f, 1.0000f, 1.0000f, 1.0000f, 1.0000f, 1.0000f, 1.0000f, 1.0000f, 1.0000f, 1.0000f, 1.0000f,
-			1.0000f, 1.0000f, 0.7818f, 0.8914f, 1.0000f, 1.0000f, 1.0000f, 1.0000f, 1.0000f, 1.0000f, 1.0000f, 1.0000f, 1.0000f, 1.0000f, 1.0000f, 1.0000f,
-			1.0000f, 1.0000f, 0.4578f, 0.5679f, 1.0000f, 1.0000f, 1.0000f, 1.0000f, 1.0000f, 1.0000f, 1.0000f, 1.0000f, 1.0000f, 1.0000f, 1.0000f, 1.0000f,
-			1.0000f, 0.4183f, 0.4961f, 0.5321f, 1.0000f, 1.0000f, 1.0000f, 1.0000f, 1.0000f, 1.0000f, 1.0000f, 1.0000f, 1.0000f, 1.0000f, 1.0000f, 1.0000f,
-			1.0000f, 1.0000f, 1.0000f, 1.0000f, 1.0000f, 1.0000f, 1.0000f, 1.0000f, 1.0000f, 1.0000f, 1.0000f, 1.0000f, 1.0000f, 1.0000f, 1.0000f, 1.0000f,
-			1.0000f, 1.0000f, 1.0000f, 1.0000f, 1.0000f, 1.0000f, 1.0000f, 1.0000f, 1.0000f, 1.0000f, 1.0000f, 1.0000f, 1.0000f, 1.0000f, 1.0000f, 1.0000f,
-			1.0000f, 1.0000f, 1.0000f, 1.0000f, 1.0000f, 1.0000f, 1.0000f, 1.0000f, 1.0000f, 1.0000f, 1.0000f, 1.0000f, 1.0000f, 1.0000f, 1.0000f, 1.0000f,
-			1.0000f, 1.0000f, 1.0000f, 1.0000f, 1.0000f, 1.0000f, 1.0000f, 1.0000f, 1.0000f, 1.0000f, 1.0000f, 1.0000f, 1.0000f, 1.0000f, 1.0000f, 1.0000f,
-			1.0000f, 1.0000f, 1.0000f, 1.0000f, 1.0000f, 1.0000f, 1.0000f, 1.0000f, 1.0000f, 1.0000f, 1.0000f, 1.0000f, 1.0000f, 1.0000f, 1.0000f, 1.0000f,
-			1.0000f, 1.0000f, 1.0000f, 1.0000f, 1.0000f, 1.0000f, 1.0000f, 1.0000f, 1.0000f, 1.0000f, 1.0000f, 1.0000f, 1.0000f, 1.0000f, 1.0000f, 1.0000f,
-			1.0000f, 1.0000f, 1.0000f, 1.0000f, 1.0000f, 1.0000f, 1.0000f, 1.0000f, 1.0000f, 1.0000f, 1.0000f, 1.0000f, 1.0000f, 1.0000f, 1.0000f, 1.0000f,
-			1.0000f, 1.0000f, 1.0000f, 1.0000f, 1.0000f, 1.0000f, 1.0000f, 1.0000f, 1.0000f, 1.0000f, 1.0000f, 1.0000f, 1.0000f, 1.0000f, 1.0000f, 1.0000f,
-			1.0000f, 1.0000f, 1.0000f, 1.0000f, 1.0000f, 1.0000f, 1.0000f, 1.0000f, 1.0000f, 1.0000f, 1.0000f, 1.0000f, 1.0000f, 1.0000f, 1.0000f, 1.0000f,
-			1.0000f, 1.0000f, 1.0000f, 1.0000f, 1.0000f, 1.0000f, 1.0000f, 1.0000f, 1.0000f, 1.0000f, 1.0000f, 1.0000f, 1.0000f, 1.0000f, 1.0000f, 1.0000f,
-			1.0000f, 1.0000f, 1.0000f, 1.0000f, 1.0000f, 1.0000f, 1.0000f, 1.0000f, 1.0000f, 1.0000f, 1.0000f, 1.0000f, 1.0000f, 1.0000f, 1.0000f, 1.0000f,
-			1.0000f, 1.0000f, 1.0000f, 1.0000f, 1.0000f, 1.0000f, 1.0000f, 1.0000f, 1.0000f, 1.0000f, 1.0000f, 1.0000f, 1.0000f, 1.0000f, 1.0000f, 1.0000f,
-			1.0000f, 1.0000f, 1.0000f, 1.0000f, 1.0000f, 1.0000f, 1.0000f, 1.0000f, 1.0000f, 1.0000f, 1.0000f, 1.0000f, 1.0000f, 1.0000f, 1.0000f, 1.0000f,
-			1.0000f, 1.0000f, 1.0000f, 1.0000f, 1.0000f, 1.0000f, 1.0000f, 1.0000f, 1.0000f, 1.0000f, 1.0000f, 1.0000f, 1.0000f, 1.0000f, 1.0000f, 1.0000f,
-			1.0000f, 1.0000f, 1.0000f, 1.0000f, 1.0000f, 1.0000f, 1.0000f, 1.0000f, 1.0000f, 1.0000f, 1.0000f, 1.0000f, 1.0000f, 1.0000f, 1.0000f, 1.0000f,
-			1.0000f, 1.0000f, 1.0000f, 1.0000f, 1.0000f, 1.0000f, 1.0000f, 1.0000f, 1.0000f, 1.0000f, 1.0000f, 1.0000f, 1.0000f, 1.0000f, 1.0000f, 0.9151f,
-			1.0000f, 1.0000f, 1.0000f, 1.0000f, 1.0000f, 1.0000f, 1.0000f, 1.0000f, 1.0000f, 1.0000f, 1.0000f, 1.0000f, 1.0000f, 1.0000f, 1.0000f, 1.0000f,
-			1.0000f, 1.0000f, 1.0000f, 1.0000f, 1.0000f, 1.0000f, 1.0000f, 1.0000f, 1.0000f, 1.0000f, 1.0000f, 1.0000f, 1.0000f, 1.0000f, 1.0000f, 1.0000f,
-			1.0000f, 1.0000f, 1.0000f, 1.0000f, 1.0000f, 1.0000f, 1.0000f, 1.0000f, 1.0000f, 1.0000f, 1.0000f, 1.0000f, 1.0000f, 1.0000f, 1.0000f, 1.0000f,
-			1.0000f, 1.0000f, 1.0000f, 1.0000f, 1.0000f, 1.0000f, 1.0000f, 1.0000f, 1.0000f, 1.0000f, 1.0000f, 1.0000f, 1.0000f, 1.0000f, 1.0000f, 1.0000f,
-			1.0000f, 1.0000f, 1.0000f, 1.0000f, 1.0000f, 1.0000f, 1.0000f, 1.0000f, 1.0000f, 1.0000f, 1.0000f, 1.0000f, 1.0000f, 1.0000f, 1.0000f, 0.9400f,
-			1.0000f, 1.0000f, 1.0000f, 1.0000f, 1.0000f, 1.0000f, 1.0000f, 1.0000f, 1.0000f, 1.0000f, 1.0000f, 1.0000f, 1.0000f, 1.0000f, 1.0000f, 0.9678f,
-			1.0000f, 1.0000f, 1.0000f, 1.0000f, 1.0000f, 1.0000f, 1.0000f, 1.0000f, 1.0000f, 1.0000f, 1.0000f, 1.0000f, 1.0000f, 1.0000f, 0.8111f, 0.8833f,
-			1.0000f, 1.0000f, 1.0000f, 1.0000f, 1.0000f, 1.0000f, 1.0000f, 1.0000f, 1.0000f, 1.0000f, 1.0000f, 1.0000f, 1.0000f, 1.0000f, 0.8299f, 0.8988f,
-			1.0000f, 1.0000f, 1.0000f, 1.0000f, 1.0000f, 1.0000f, 1.0000f, 1.0000f, 1.0000f, 1.0000f, 1.0000f, 1.0000f, 1.0000f, 1.0000f, 1.0000f, 1.0000f,
-			1.0000f, 1.0000f, 1.0000f, 1.0000f, 1.0000f, 1.0000f, 1.0000f, 1.0000f, 1.0000f, 1.0000f, 1.0000f, 1.0000f, 1.0000f, 1.0000f, 1.0000f, 1.0000f,
-			1.0000f, 0.9182f, 0.9692f, 0.9820f, 1.0000f, 1.0000f, 1.0000f, 1.0000f, 1.0000f, 1.0000f, 1.0000f, 1.0000f, 1.0000f, 1.0000f, 1.0000f, 1.0000f,
-			1.0000f, 0.9663f, 0.9911f, 0.8707f, 1.0000f, 1.0000f, 1.0000f, 1.0000f, 1.0000f, 1.0000f, 1.0000f, 1.0000f, 1.0000f, 1.0000f, 1.0000f, 1.0000f,
-			1.0000f, 0.9088f, 0.9374f, 0.8793f, 1.0000f, 1.0000f, 1.0000f, 1.0000f, 1.0000f, 1.0000f, 1.0000f, 1.0000f, 1.0000f, 1.0000f, 1.0000f, 1.0000f,
-			1.0000f, 0.8750f, 0.8952f, 0.7356f, 1.0000f, 1.0000f, 1.0000f, 1.0000f, 1.0000f, 1.0000f, 1.0000f, 1.0000f, 1.0000f, 1.0000f, 1.0000f, 1.0000f,
-			1.0000f, 0.2746f, 0.0000f, 0.0772f, 1.0000f, 1.0000f, 1.0000f, 1.0000f, 1.0000f, 1.0000f, 1.0000f, 1.0000f, 1.0000f, 1.0000f, 1.0000f, 1.0000f,
-			1.0000f, 0.1487f, 0.2193f, 0.3263f, 1.0000f, 1.0000f, 1.0000f, 1.0000f, 1.0000f, 1.0000f, 1.0000f, 1.0000f, 1.0000f, 1.0000f, 1.0000f, 1.0000f,
-			1.0000f, 1.0000f, 1.0000f, 1.0000f, 1.0000f, 1.0000f, 1.0000f, 1.0000f, 1.0000f, 1.0000f, 1.0000f, 1.0000f, 1.0000f, 1.0000f, 1.0000f, 1.0000f,
-			1.0000f, 1.0000f, 1.0000f, 1.0000f, 1.0000f, 1.0000f, 1.0000f, 1.0000f, 1.0000f, 1.0000f, 1.0000f, 1.0000f, 1.0000f, 1.0000f, 1.0000f, 1.0000f,
-			1.0000f, 1.0000f, 1.0000f, 1.0000f, 1.0000f, 1.0000f, 1.0000f, 1.0000f, 1.0000f, 1.0000f, 1.0000f, 1.0000f, 1.0000f, 1.0000f, 1.0000f, 1.0000f,
-			1.0000f, 1.0000f, 1.0000f, 1.0000f, 1.0000f, 1.0000f, 1.0000f, 1.0000f, 1.0000f, 1.0000f, 1.0000f, 1.0000f, 1.0000f, 1.0000f, 1.0000f, 1.0000f,
-			1.0000f, 1.0000f, 1.0000f, 1.0000f, 1.0000f, 1.0000f, 1.0000f, 1.0000f, 1.0000f, 1.0000f, 1.0000f, 1.0000f, 1.0000f, 1.0000f, 1.0000f, 1.0000f,
-			1.0000f, 1.0000f, 1.0000f, 1.0000f, 1.0000f, 1.0000f, 1.0000f, 1.0000f, 1.0000f, 1.0000f, 1.0000f, 1.0000f, 1.0000f, 1.0000f, 1.0000f, 1.0000f,
-			1.0000f, 1.0000f, 1.0000f, 1.0000f, 1.0000f, 1.0000f, 1.0000f, 1.0000f, 1.0000f, 1.0000f, 1.0000f, 1.0000f, 1.0000f, 1.0000f, 1.0000f, 1.0000f,
-			1.0000f, 1.0000f, 1.0000f, 1.0000f, 1.0000f, 1.0000f, 1.0000f, 1.0000f, 1.0000f, 1.0000f, 1.0000f, 1.0000f, 1.0000f, 1.0000f, 1.0000f, 1.0000f,
-			1.0000f, 1.0000f, 1.0000f, 1.0000f, 1.0000f, 1.0000f, 1.0000f, 1.0000f, 1.0000f, 1.0000f, 1.0000f, 1.0000f, 1.0000f, 1.0000f, 1.0000f, 1.0000f,
-			1.0000f, 1.0000f, 1.0000f, 1.0000f, 1.0000f, 1.0000f, 1.0000f, 1.0000f, 1.0000f, 1.0000f, 1.0000f, 1.0000f, 1.0000f, 1.0000f, 1.0000f, 1.0000f,
-			1.0000f, 1.0000f, 1.0000f, 1.0000f, 1.0000f, 1.0000f, 1.0000f, 1.0000f, 1.0000f, 1.0000f, 1.0000f, 1.0000f, 1.0000f, 1.0000f, 1.0000f, 1.0000f,
-			1.0000f, 1.0000f, 1.0000f, 1.0000f, 1.0000f, 1.0000f, 1.0000f, 1.0000f, 1.0000f, 1.0000f, 1.0000f, 1.0000f, 1.0000f, 1.0000f, 1.0000f, 1.0000f,
-			1.0000f, 1.0000f, 1.0000f, 1.0000f, 1.0000f, 1.0000f, 1.0000f, 1.0000f, 1.0000f, 1.0000f, 1.0000f, 1.0000f, 1.0000f, 1.0000f, 0.9917f, 1.0000f,
-			1.0000f, 1.0000f, 1.0000f, 1.0000f, 1.0000f, 1.0000f, 1.0000f, 1.0000f, 1.0000f, 1.0000f, 1.0000f, 1.0000f, 1.0000f, 1.0000f, 0.9995f, 0.9996f,
-			1.0000f, 1.0000f, 1.0000f, 1.0000f, 1.0000f, 1.0000f, 1.0000f, 1.0000f, 1.0000f, 1.0000f, 1.0000f, 1.0000f, 1.0000f, 0.9516f, 0.9838f, 0.9927f,
-			1.0000f, 1.0000f, 1.0000f, 1.0000f, 1.0000f, 1.0000f, 1.0000f, 1.0000f, 1.0000f, 1.0000f, 1.0000f, 1.0000f, 1.0000f, 0.9731f, 0.9906f, 0.9448f,
-			1.0000f, 1.0000f, 1.0000f, 1.0000f, 1.0000f, 1.0000f, 1.0000f, 1.0000f, 1.0000f, 1.0000f, 1.0000f, 1.0000f, 1.0000f, 1.0000f, 1.0000f, 1.0000f,
-			1.0000f, 1.0000f, 1.0000f, 1.0000f, 1.0000f, 1.0000f, 1.0000f, 1.0000f, 1.0000f, 1.0000f, 1.0000f, 1.0000f, 1.0000f, 1.0000f, 1.0000f, 1.0000f,
-			1.0000f, 1.0000f, 1.0000f, 1.0000f, 1.0000f, 1.0000f, 1.0000f, 1.0000f, 1.0000f, 1.0000f, 1.0000f, 1.0000f, 1.0000f, 1.0000f, 0.9900f, 0.9999f,
-			1.0000f, 1.0000f, 1.0000f, 1.0000f, 1.0000f, 1.0000f, 1.0000f, 1.0000f, 1.0000f, 1.0000f, 1.0000f, 1.0000f, 1.0000f, 1.0000f, 0.9992f, 0.9991f,
-			1.0000f, 1.0000f, 1.0000f, 1.0000f, 1.0000f, 1.0000f, 1.0000f, 1.0000f, 1.0000f, 1.0000f, 1.0000f, 1.0000f, 1.0000f, 0.9949f, 0.9987f, 0.9936f,
-			1.0000f, 1.0000f, 1.0000f, 1.0000f, 1.0000f, 1.0000f, 1.0000f, 1.0000f, 1.0000f, 1.0000f, 1.0000f, 1.0000f, 1.0000f, 0.9966f, 0.9985f, 0.9615f,
-			1.0000f, 1.0000f, 1.0000f, 1.0000f, 1.0000f, 1.0000f, 1.0000f, 1.0000f, 1.0000f, 1.0000f, 1.0000f, 1.0000f, 1.0000f, 0.9269f, 0.9577f, 0.9057f,
-			1.0000f, 1.0000f, 1.0000f, 1.0000f, 1.0000f, 1.0000f, 1.0000f, 1.0000f, 1.0000f, 1.0000f, 1.0000f, 1.0000f, 1.0000f, 0.9023f, 0.9241f, 0.7499f,
-			1.0000f, 1.0000f, 1.0000f, 1.0000f, 1.0000f, 1.0000f, 1.0000f, 1.0000f, 1.0000f, 1.0000f, 1.0000f, 1.0000f, 1.0000f, 1.0000f, 1.0000f, 1.0000f,
-			1.0000f, 1.0000f, 1.0000f, 1.0000f, 1.0000f, 1.0000f, 1.0000f, 1.0000f, 1.0000f, 1.0000f, 1.0000f, 1.0000f, 1.0000f, 1.0000f, 1.0000f, 1.0000f,
-			1.0000f, 1.0000f, 0.9757f, 0.9846f, 1.0000f, 1.0000f, 1.0000f, 1.0000f, 1.0000f, 1.0000f, 1.0000f, 1.0000f, 1.0000f, 1.0000f, 1.0000f, 1.0000f,
-			1.0000f, 0.9922f, 0.9932f, 0.9792f, 1.0000f, 1.0000f, 1.0000f, 1.0000f, 1.0000f, 1.0000f, 1.0000f, 1.0000f, 1.0000f, 1.0000f, 1.0000f, 1.0000f,
-			1.0000f, 0.9881f, 0.9494f, 0.8178f, 1.0000f, 1.0000f, 1.0000f, 1.0000f, 1.0000f, 1.0000f, 1.0000f, 1.0000f, 1.0000f, 1.0000f, 1.0000f, 1.0000f,
-			1.0000f, 0.9780f, 0.8518f, 0.6206f, 1.0000f, 1.0000f, 1.0000f, 1.0000f, 1.0000f, 1.0000f, 1.0000f, 1.0000f, 1.0000f, 1.0000f, 1.0000f, 1.0000f,
-			1.0000f, 0.9472f, 0.7191f, 0.7003f, 1.0000f, 1.0000f, 1.0000f, 1.0000f, 1.0000f, 1.0000f, 1.0000f, 1.0000f, 1.0000f, 1.0000f, 1.0000f, 1.0000f,
-			1.0000f, 0.8356f, 0.3772f, 0.9971f, 1.0000f, 1.0000f, 1.0000f, 1.0000f, 1.0000f, 1.0000f, 1.0000f, 1.0000f, 1.0000f, 1.0000f, 1.0000f, 1.0000f,
-			1.0000f, 1.0000f, 1.0000f, 1.0000f, 1.0000f, 1.0000f, 1.0000f, 1.0000f, 1.0000f, 1.0000f, 1.0000f, 1.0000f, 1.0000f, 1.0000f, 1.0000f, 1.0000f,
-			1.0000f, 1.0000f, 1.0000f, 1.0000f, 1.0000f, 1.0000f, 1.0000f, 1.0000f, 1.0000f, 1.0000f, 1.0000f, 1.0000f, 1.0000f, 1.0000f, 1.0000f, 1.0000f,
-			1.0000f, 1.0000f, 1.0000f, 1.0000f, 1.0000f, 1.0000f, 1.0000f, 1.0000f, 1.0000f, 1.0000f, 1.0000f, 1.0000f, 1.0000f, 1.0000f, 1.0000f, 1.0000f,
-			1.0000f, 1.0000f, 1.0000f, 1.0000f, 1.0000f, 1.0000f, 1.0000f, 1.0000f, 1.0000f, 1.0000f, 1.0000f, 1.0000f, 1.0000f, 1.0000f, 1.0000f, 1.0000f,
-			1.0000f, 1.0000f, 1.0000f, 1.0000f, 1.0000f, 1.0000f, 1.0000f, 1.0000f, 1.0000f, 1.0000f, 1.0000f, 1.0000f, 1.0000f, 1.0000f, 1.0000f, 1.0000f,
-			1.0000f, 1.0000f, 1.0000f, 1.0000f, 1.0000f, 1.0000f, 1.0000f, 1.0000f, 1.0000f, 1.0000f, 1.0000f, 1.0000f, 1.0000f, 1.0000f, 1.0000f, 1.0000f,
-			1.0000f, 1.0000f, 1.0000f, 1.0000f, 1.0000f, 1.0000f, 1.0000f, 1.0000f, 1.0000f, 1.0000f, 1.0000f, 1.0000f, 1.0000f, 1.0000f, 1.0000f, 1.0000f,
-			1.0000f, 1.0000f, 1.0000f, 1.0000f, 1.0000f, 1.0000f, 1.0000f, 1.0000f, 1.0000f, 1.0000f, 1.0000f, 1.0000f, 1.0000f, 1.0000f, 1.0000f, 1.0000f,
-			1.0000f, 1.0000f, 1.0000f, 1.0000f, 1.0000f, 1.0000f, 1.0000f, 1.0000f, 1.0000f, 1.0000f, 1.0000f, 1.0000f, 1.0000f, 1.0000f, 1.0000f, 1.0000f,
-			1.0000f, 1.0000f, 1.0000f, 1.0000f, 1.0000f, 1.0000f, 1.0000f, 1.0000f, 1.0000f, 1.0000f, 1.0000f, 1.0000f, 1.0000f, 1.0000f, 1.0000f, 1.0000f,
-			1.0000f, 1.0000f, 1.0000f, 1.0000f, 1.0000f, 1.0000f, 1.0000f, 1.0000f, 1.0000f, 1.0000f, 1.0000f, 1.0000f, 1.0000f, 1.0000f, 1.0000f, 1.0000f,
-			1.0000f, 1.0000f, 1.0000f, 1.0000f, 1.0000f, 1.0000f, 1.0000f, 1.0000f, 1.0000f, 1.0000f, 1.0000f, 1.0000f, 1.0000f, 1.0000f, 1.0000f, 0.9997f,
-			1.0000f, 1.0000f, 1.0000f, 1.0000f, 1.0000f, 1.0000f, 1.0000f, 1.0000f, 1.0000f, 1.0000f, 1.0000f, 1.0000f, 1.0000f, 1.0000f, 0.9968f, 0.9980f,
-			1.0000f, 1.0000f, 1.0000f, 1.0000f, 1.0000f, 1.0000f, 1.0000f, 1.0000f, 1.0000f, 1.0000f, 1.0000f, 1.0000f, 1.0000f, 1.0000f, 0.9975f, 0.9977f,
-			1.0000f, 1.0000f, 1.0000f, 1.0000f, 1.0000f, 1.0000f, 1.0000f, 1.0000f, 1.0000f, 1.0000f, 1.0000f, 1.0000f, 1.0000f, 1.0000f, 0.9769f, 0.9875f,
-			1.0000f, 1.0000f, 1.0000f, 1.0000f, 1.0000f, 1.0000f, 1.0000f, 1.0000f, 1.0000f, 1.0000f, 1.0000f, 1.0000f, 1.0000f, 0.9945f, 0.9941f, 0.9861f,
-			1.0000f, 1.0000f, 1.0000f, 1.0000f, 1.0000f, 1.0000f, 1.0000f, 1.0000f, 1.0000f, 1.0000f, 1.0000f, 1.0000f, 1.0000f, 1.0000f, 1.0000f, 1.0000f,
-			1.0000f, 1.0000f, 1.0000f, 1.0000f, 1.0000f, 1.0000f, 1.0000f, 1.0000f, 1.0000f, 1.0000f, 1.0000f, 1.0000f, 1.0000f, 1.0000f, 1.0000f, 1.0000f,
-			1.0000f, 1.0000f, 1.0000f, 1.0000f, 1.0000f, 1.0000f, 1.0000f, 1.0000f, 1.0000f, 1.0000f, 1.0000f, 1.0000f, 1.0000f, 1.0000f, 0.9956f, 0.9982f,
-			1.0000f, 1.0000f, 1.0000f, 1.0000f, 1.0000f, 1.0000f, 1.0000f, 1.0000f, 1.0000f, 1.0000f, 1.0000f, 1.0000f, 1.0000f, 1.0000f, 0.9979f, 0.9989f,
-			1.0000f, 1.0000f, 1.0000f, 1.0000f, 1.0000f, 1.0000f, 1.0000f, 1.0000f, 1.0000f, 1.0000f, 1.0000f, 1.0000f, 1.0000f, 1.0000f, 0.9868f, 0.9854f,
-			1.0000f, 1.0000f, 1.0000f, 1.0000f, 1.0000f, 1.0000f, 1.0000f, 1.0000f, 1.0000f, 1.0000f, 1.0000f, 1.0000f, 1.0000f, 0.9888f, 0.9811f, 0.9706f,
-			1.0000f, 1.0000f, 1.0000f, 1.0000f, 1.0000f, 1.0000f, 1.0000f, 1.0000f, 1.0000f, 1.0000f, 1.0000f, 1.0000f, 1.0000f, 0.9894f, 0.9425f, 0.8465f,
-			1.0000f, 1.0000f, 1.0000f, 1.0000f, 1.0000f, 1.0000f, 1.0000f, 1.0000f, 1.0000f, 1.0000f, 1.0000f, 1.0000f, 1.0000f, 0.9718f, 0.8614f, 0.6422f,
-			1.0000f, 1.0000f, 1.0000f, 1.0000f, 1.0000f, 1.0000f, 1.0000f, 1.0000f, 1.0000f, 1.0000f, 1.0000f, 1.0000f, 1.0000f, 1.0000f, 1.0000f, 1.0000f,
-			1.0000f, 1.0000f, 1.0000f, 1.0000f, 1.0000f, 1.0000f, 1.0000f, 1.0000f, 1.0000f, 1.0000f, 1.0000f, 1.0000f, 1.0000f, 1.0000f, 1.0000f, 1.0000f,
-			1.0000f, 0.9648f, 0.9212f, 0.8412f, 1.0000f, 1.0000f, 1.0000f, 1.0000f, 1.0000f, 1.0000f, 1.0000f, 1.0000f, 1.0000f, 1.0000f, 1.0000f, 1.0000f,
-			1.0000f, 0.9596f, 0.9557f, 0.9537f, 1.0000f, 1.0000f, 1.0000f, 1.0000f, 1.0000f, 1.0000f, 1.0000f, 1.0000f, 1.0000f, 1.0000f, 1.0000f, 1.0000f,
-			1.0000f, 0.6814f, 1.0000f, 1.0000f, 1.0000f, 1.0000f, 1.0000f, 1.0000f, 1.0000f, 1.0000f, 1.0000f, 1.0000f, 1.0000f, 1.0000f, 1.0000f, 1.0000f,
-			1.0000f, 0.7971f, 1.0000f, 1.0000f, 1.0000f, 1.0000f, 1.0000f, 1.0000f, 1.0000f, 1.0000f, 1.0000f, 1.0000f, 1.0000f, 1.0000f, 1.0000f, 1.0000f,
-			1.0000f, 1.0000f, 1.0000f, 1.0000f, 1.0000f, 1.0000f, 1.0000f, 1.0000f, 1.0000f, 1.0000f, 1.0000f, 1.0000f, 1.0000f, 1.0000f, 1.0000f, 1.0000f,
-			1.0000f, 1.0000f, 1.0000f, 1.0000f, 1.0000f, 1.0000f, 1.0000f, 1.0000f, 1.0000f, 1.0000f, 1.0000f, 1.0000f, 1.0000f, 1.0000f, 1.0000f, 1.0000f,
-			1.0000f, 1.0000f, 1.0000f, 1.0000f, 1.0000f, 1.0000f, 1.0000f, 1.0000f, 1.0000f, 1.0000f, 1.0000f, 1.0000f, 1.0000f, 1.0000f, 1.0000f, 1.0000f,
-			1.0000f, 1.0000f, 1.0000f, 1.0000f, 1.0000f, 1.0000f, 1.0000f, 1.0000f, 1.0000f, 1.0000f, 1.0000f, 1.0000f, 1.0000f, 1.0000f, 1.0000f, 1.0000f,
-			1.0000f, 1.0000f, 1.0000f, 1.0000f, 1.0000f, 1.0000f, 1.0000f, 1.0000f, 1.0000f, 1.0000f, 1.0000f, 1.0000f, 1.0000f, 1.0000f, 1.0000f, 1.0000f,
-			1.0000f, 1.0000f, 1.0000f, 1.0000f, 1.0000f, 1.0000f, 1.0000f, 1.0000f, 1.0000f, 1.0000f, 1.0000f, 1.0000f, 1.0000f, 1.0000f, 1.0000f, 1.0000f,
-			1.0000f, 1.0000f, 1.0000f, 1.0000f, 1.0000f, 1.0000f, 1.0000f, 1.0000f, 1.0000f, 1.0000f, 1.0000f, 1.0000f, 1.0000f, 1.0000f, 1.0000f, 1.0000f,
-			1.0000f, 1.0000f, 1.0000f, 1.0000f, 1.0000f, 1.0000f, 1.0000f, 1.0000f, 1.0000f, 1.0000f, 1.0000f, 1.0000f, 1.0000f, 1.0000f, 1.0000f, 1.0000f,
-			1.0000f, 1.0000f, 1.0000f, 1.0000f, 1.0000f, 1.0000f, 1.0000f, 1.0000f, 1.0000f, 1.0000f, 1.0000f, 1.0000f, 1.0000f, 1.0000f, 1.0000f, 1.0000f,
-			1.0000f, 1.0000f, 1.0000f, 1.0000f, 1.0000f, 1.0000f, 1.0000f, 1.0000f, 1.0000f, 1.0000f, 1.0000f, 1.0000f, 1.0000f, 1.0000f, 1.0000f, 1.0000f,
-			1.0000f, 1.0000f, 1.0000f, 1.0000f, 1.0000f, 1.0000f, 1.0000f, 1.0000f, 1.0000f, 1.0000f, 1.0000f, 1.0000f, 1.0000f, 1.0000f, 1.0000f, 1.0000f,
-			1.0000f, 1.0000f, 1.0000f, 1.0000f, 1.0000f, 1.0000f, 1.0000f, 1.0000f, 1.0000f, 1.0000f, 1.0000f, 1.0000f, 1.0000f, 1.0000f, 1.0000f, 1.0000f,
-			1.0000f, 1.0000f, 1.0000f, 1.0000f, 1.0000f, 1.0000f, 1.0000f, 1.0000f, 1.0000f, 1.0000f, 1.0000f, 1.0000f, 1.0000f, 0.9998f, 0.9999f, 1.0000f,
-			1.0000f, 1.0000f, 1.0000f, 1.0000f, 1.0000f, 1.0000f, 1.0000f, 1.0000f, 1.0000f, 1.0000f, 1.0000f, 1.0000f, 1.0000f, 0.9999f, 1.0000f, 0.9994f,
-			1.0000f, 1.0000f, 1.0000f, 1.0000f, 1.0000f, 1.0000f, 1.0000f, 1.0000f, 1.0000f, 1.0000f, 1.0000f, 1.0000f, 1.0000f, 0.9961f, 0.9984f, 0.9958f,
-			1.0000f, 1.0000f, 1.0000f, 1.0000f, 1.0000f, 1.0000f, 1.0000f, 1.0000f, 1.0000f, 1.0000f, 1.0000f, 1.0000f, 1.0000f, 0.9952f, 0.9994f, 0.9633f,
-			1.0000f, 1.0000f, 1.0000f, 1.0000f, 1.0000f, 1.0000f, 1.0000f, 1.0000f, 1.0000f, 1.0000f, 1.0000f, 1.0000f, 1.0000f, 0.9829f, 0.9120f, 0.8042f,
-			1.0000f, 1.0000f, 1.0000f, 1.0000f, 1.0000f, 1.0000f, 1.0000f, 1.0000f, 1.0000f, 1.0000f, 1.0000f, 1.0000f, 1.0000f, 0.9744f, 0.9296f, 0.9349f,
-			1.0000f, 1.0000f, 1.0000f, 1.0000f, 1.0000f, 1.0000f, 1.0000f, 1.0000f, 1.0000f, 1.0000f, 1.0000f, 1.0000f, 1.0000f, 1.0000f, 1.0000f, 1.0000f,
-			1.0000f, 1.0000f, 1.0000f, 1.0000f, 1.0000f, 1.0000f, 1.0000f, 1.0000f, 1.0000f, 1.0000f, 1.0000f, 1.0000f, 1.0000f, 1.0000f, 1.0000f, 1.0000f,
-			1.0000f, 1.0000f, 1.0000f, 1.0000f, 1.0000f, 1.0000f, 1.0000f, 1.0000f, 1.0000f, 1.0000f, 1.0000f, 1.0000f, 1.0000f, 0.9973f, 0.9993f, 0.9988f,
-			1.0000f, 1.0000f, 1.0000f, 1.0000f, 1.0000f, 1.0000f, 1.0000f, 1.0000f, 1.0000f, 1.0000f, 1.0000f, 1.0000f, 1.0000f, 0.9964f, 0.9998f, 0.9802f,
-			1.0000f, 1.0000f, 1.0000f, 1.0000f, 1.0000f, 1.0000f, 1.0000f, 1.0000f, 1.0000f, 1.0000f, 1.0000f, 1.0000f, 1.0000f, 0.8874f, 0.5963f, 0.9323f,
-			1.0000f, 1.0000f, 1.0000f, 1.0000f, 1.0000f, 1.0000f, 1.0000f, 1.0000f, 1.0000f, 1.0000f, 1.0000f, 1.0000f, 1.0000f, 0.8239f, 0.7625f, 1.0000f,
-			1.0000f, 1.0000f, 1.0000f, 1.0000f, 1.0000f, 1.0000f, 1.0000f, 1.0000f, 1.0000f, 1.0000f, 1.0000f, 1.0000f, 1.0000f, 0.6625f, 1.0000f, 1.0000f,
-			1.0000f, 1.0000f, 1.0000f, 1.0000f, 1.0000f, 1.0000f, 1.0000f, 1.0000f, 1.0000f, 1.0000f, 1.0000f, 1.0000f, 1.0000f, 0.7895f, 1.0000f, 1.0000f,
-			1.0000f, 1.0000f, 1.0000f, 1.0000f, 1.0000f, 1.0000f, 1.0000f, 1.0000f, 1.0000f, 1.0000f, 1.0000f, 1.0000f, 1.0000f, 1.0000f, 1.0000f, 1.0000f,
-			1.0000f, 1.0000f, 1.0000f, 1.0000f, 1.0000f, 1.0000f, 1.0000f, 1.0000f, 1.0000f, 1.0000f, 1.0000f, 1.0000f, 1.0000f, 1.0000f, 1.0000f, 1.0000f,
+			uint32_t operator[] (uint32_t idx) const { check(idx < 4);  return m_c[idx]; }
+			uint32_t& operator[] (uint32_t idx) { check(idx < 4);  return m_c[idx]; }
 		};
 
-
-
-		const float* get_2d_percentile_table()
-		{
-			return percentile_table_4x4;
-		}
-
-
-		// return 0 on invalid mode, 1 on valid mode.
-		static int decode_block_mode_2d(int blockmode, int* Nval, int* Mval, int* dual_weight_plane, int* quant_mode)
+		struct IVec4
 		{
-			int base_quant_mode = (blockmode >> 4) & 1;
-			int H = (blockmode >> 9) & 1;
-			int D = (blockmode >> 10) & 1;
+			int32_t m_c[4];
 
-			int A = (blockmode >> 5) & 0x3;
-
-			int N = 0, M = 0;
-
-			if ((blockmode & 3) != 0)
+			IVec4()
 			{
-				base_quant_mode |= (blockmode & 3) << 1;
-				int B = (blockmode >> 7) & 3;
-				switch ((blockmode >> 2) & 3)
-				{
-				case 0:
-					N = B + 4;
-					M = A + 2;
-					break;
-				case 1:
-					N = B + 8;
-					M = A + 2;
-					break;
-				case 2:
-					N = A + 2;
-					M = B + 8;
-					break;
-				case 3:
-					B &= 1;
-					if (blockmode & 0x100)
-					{
-						N = B + 2;
-						M = A + 2;
-					}
-					else
-					{
-						N = A + 2;
-						M = B + 6;
-					}
-					break;
-				}
+				m_c[0] = 0;
+				m_c[1] = 0;
+				m_c[2] = 0;
+				m_c[3] = 0;
 			}
-			else
+
+			IVec4(int32_t x, int32_t y, int32_t z, int32_t w)
 			{
-				base_quant_mode |= ((blockmode >> 2) & 3) << 1;
-				if (((blockmode >> 2) & 3) == 0)
-					return 0;
-				int B = (blockmode >> 9) & 3;
-				switch ((blockmode >> 7) & 3)
-				{
-				case 0:
-					N = 12;
-					M = A + 2;
-					break;
-				case 1:
-					N = A + 2;
-					M = 12;
-					break;
-				case 2:
-					N = A + 6;
-					M = B + 6;
-					D = 0;
-					H = 0;
-					break;
-				case 3:
-					switch ((blockmode >> 5) & 3)
-					{
-					case 0:
-						N = 6;
-						M = 10;
-						break;
-					case 1:
-						N = 10;
-						M = 6;
-						break;
-					case 2:
-					case 3:
-						return 0;
-					}
-					break;
-				}
+				m_c[0] = x;
+				m_c[1] = y;
+				m_c[2] = z;
+				m_c[3] = w;
 			}
 
-			int weight_count = N * M * (D + 1);
-			int qmode = (base_quant_mode - 2) + 6 * H;
-
-			int weightbits = compute_ise_bitcount(weight_count, (quantization_method)qmode);
-			if (weight_count > MAX_WEIGHTS_PER_BLOCK || weightbits < MIN_WEIGHT_BITS_PER_BLOCK || weightbits > MAX_WEIGHT_BITS_PER_BLOCK)
-				return 0;
-
-			*Nval = N;
-			*Mval = M;
-			*dual_weight_plane = D;
-			*quant_mode = qmode;
-			return 1;
-		}
-
-
-		void construct_block_size_descriptor_2d(block_size_descriptor* bsd)
-		{
-			int decimation_mode_index = 0;
-			int decimation_mode_count = 0;
+			int32_t x() const { return m_c[0]; }
+			int32_t y() const { return m_c[1]; }
+			int32_t z() const { return m_c[2]; }
+			int32_t w() const { return m_c[3]; }
 
-			int i;
-			int x_weights = 4;
-			int y_weights = 4;
+			int32_t& x() { return m_c[0]; }
+			int32_t& y() { return m_c[1]; }
+			int32_t& z() { return m_c[2]; }
+			int32_t& w() { return m_c[3]; }
 
-			// gather all the infill-modes that can be used with the current block size
-			//for (x_weights = 2; x_weights <= 12; x_weights++)
-			//    for (y_weights = 2; y_weights <= 12; y_weights++)
+			UVec4 asUint() const
 			{
-				//if (x_weights * y_weights > MAX_WEIGHTS_PER_BLOCK)
-				//    continue;
-				decimation_table* dt = new decimation_table;
-				decimation_mode_index = decimation_mode_count;
-				initialize_decimation_table_2d(dt);
-
-				int weight_count = x_weights * y_weights;
-
-				int maxprec_1plane = -1;
-				int maxprec_2planes = -1;
-				for (i = 0; i < 12; i++)
-				{
-					int bits_1plane = compute_ise_bitcount(weight_count, (quantization_method)i);
-					int bits_2planes = compute_ise_bitcount(2 * weight_count, (quantization_method)i);
-					if (bits_1plane >= MIN_WEIGHT_BITS_PER_BLOCK && bits_1plane <= MAX_WEIGHT_BITS_PER_BLOCK)
-						maxprec_1plane = i;
-					if (bits_2planes >= MIN_WEIGHT_BITS_PER_BLOCK && bits_2planes <= MAX_WEIGHT_BITS_PER_BLOCK)
-						maxprec_2planes = i;
-				}
-
-				// always false
-				//if (2 * x_weights * y_weights > MAX_WEIGHTS_PER_BLOCK)
-				//    maxprec_2planes = -1;
-
-				bsd->permit_encode[decimation_mode_count] = (x_weights <= 4 && y_weights <= 4);
+				return UVec4(FMath::Max(0, m_c[0]), FMath::Max(0, m_c[1]), FMath::Max(0, m_c[2]), FMath::Max(0, m_c[3]));
+			}
 
-				bsd->decimation_mode_samples[decimation_mode_count] = weight_count;
-				bsd->decimation_mode_maxprec_1plane[decimation_mode_count] = maxprec_1plane;
-				bsd->decimation_mode_maxprec_2planes[decimation_mode_count] = maxprec_2planes;
-				bsd->decimation_tables[decimation_mode_count] = dt;
+			int32_t operator[] (uint32_t idx) const { check(idx < 4);  return m_c[idx]; }
+			int32_t& operator[] (uint32_t idx) { check(idx < 4);  return m_c[idx]; }
+		};
 
-				decimation_mode_count++;
-			}
+		struct IVec3
+		{
+			int32_t m_c[3];
 
-			for (i = 0; i < MAX_DECIMATION_MODES; i++)
+			IVec3()
 			{
-				bsd->decimation_mode_percentile[i] = 1.0f;
+				m_c[0] = 0;
+				m_c[1] = 0;
+				m_c[2] = 0;
 			}
 
-			for (i = decimation_mode_count; i < MAX_DECIMATION_MODES; i++)
+			IVec3(int32_t x, int32_t y, int32_t z)
 			{
-				bsd->permit_encode[i] = 0;
-				bsd->decimation_mode_samples[i] = 0;
-				bsd->decimation_mode_maxprec_1plane[i] = -1;
-				bsd->decimation_mode_maxprec_2planes[i] = -1;
+				m_c[0] = x;
+				m_c[1] = y;
+				m_c[2] = z;
 			}
 
-			bsd->decimation_mode_count = decimation_mode_count;
-
-			const float* percentiles = get_2d_percentile_table();
+			int32_t x() const { return m_c[0]; }
+			int32_t y() const { return m_c[1]; }
+			int32_t z() const { return m_c[2]; }
 
-			// then construct the list of block formats
-			for (i = 0; i < 2048; i++)
-			{
-				//int x_weights, y_weights;
-				int is_dual_plane;
-				int quantization_mode;
-				int fail = 0;
-				int permit_encode = 1;
-
-				if (decode_block_mode_2d(i, &x_weights, &y_weights, &is_dual_plane, &quantization_mode))
-				{
-					if (x_weights > 4 || y_weights > 4)
-						permit_encode = 0;
-				}
-				else
-				{
-					fail = 1;
-					permit_encode = 0;
-				}
+			int32_t& x() { return m_c[0]; }
+			int32_t& y() { return m_c[1]; }
+			int32_t& z() { return m_c[2]; }
 
-				if (fail)
-				{
-					bsd->block_modes[i].decimation_mode = -1;
-					bsd->block_modes[i].quantization_mode = -1;
-					bsd->block_modes[i].is_dual_plane = -1;
-					bsd->block_modes[i].permit_encode = 0;
-					bsd->block_modes[i].permit_decode = 0;
-					bsd->block_modes[i].percentile = 1.0f;
-				}
-				else
-				{
-					int decimation_mode = decimation_mode_index;
-					bsd->block_modes[i].decimation_mode = (int8_t)decimation_mode;
-					bsd->block_modes[i].quantization_mode = (int8_t)quantization_mode;
-					bsd->block_modes[i].is_dual_plane = (int8_t)is_dual_plane;
-					bsd->block_modes[i].permit_encode = (int8_t)permit_encode;
-					bsd->block_modes[i].permit_decode = (int8_t)permit_encode;	// disallow decode of grid size larger than block size.
-					bsd->block_modes[i].percentile = percentiles[i];
-
-					if (bsd->decimation_mode_percentile[decimation_mode] > percentiles[i])
-						bsd->decimation_mode_percentile[decimation_mode] = percentiles[i];
-				}
+			int32_t operator[] (uint32_t idx) const { check(idx < 3);  return m_c[idx]; }
+			int32_t& operator[] (uint32_t idx) { check(idx < 3);  return m_c[idx]; }
+		};
 
-			}
+		static uint32_t deDivRoundUp32(uint32_t a, uint32_t b)
+		{
+			return (a + b - 1) / b;
 		}
 
-
-
-
-		static block_size_descriptor* s_bsd = nullptr;
-
-		// function to obtain a block size descriptor. If the descriptor does not exist,
-		// it is created as needed. Should not be called from within multi-threaded code.
-		const block_size_descriptor* get_block_size_descriptor()
+		static bool deInBounds32(uint32_t v, uint32_t l, uint32_t h)
 		{
-			if (s_bsd == nullptr)
-			{
-				s_bsd = new block_size_descriptor;
-				construct_block_size_descriptor_2d(s_bsd);
-			}
-			return s_bsd;
+			return (v >= l) && (v < h);
 		}
 
 
-
-		// routine to read up to 8 bits
-		static inline int read_bits(int bitcount, int bitoffset, const uint8* ptr)
+		// Common utilities
+		enum
+		{
+			MAX_BLOCK_WIDTH = 12,
+			MAX_BLOCK_HEIGHT = 12
+		};
+		inline deUint32 getBit(deUint32 src, int ndx)
 		{
-			int mask = (1 << bitcount) - 1;
-			ptr += bitoffset >> 3;
-			bitoffset &= 7;
-			int value = ptr[0] | (ptr[1] << 8);
-			value >>= bitoffset;
-			value &= mask;
-			return value;
+			DE_ASSERT(inBounds(ndx, 0, 32));
+			return (src >> ndx) & 1;
 		}
-
-		int bitrev8(int p)
+		inline deUint32 getBits(deUint32 src, int low, int high)
 		{
-			p = ((p & 0xF) << 4) | ((p >> 4) & 0xF);
-			p = ((p & 0x33) << 2) | ((p >> 2) & 0x33);
-			p = ((p & 0x55) << 1) | ((p >> 1) & 0x55);
-			return p;
+			const int numBits = (high - low) + 1;
+			DE_ASSERT(inRange(numBits, 1, 32));
+			if (numBits < 32)
+				return (deUint32)((src >> low) & ((1u << numBits) - 1));
+			else
+				return (deUint32)((src >> low) & 0xFFFFFFFFu);
+		}
+		inline bool isBitSet(deUint32 src, int ndx)
+		{
+			return getBit(src, ndx) != 0;
+		}
+		inline deUint32 reverseBits(deUint32 src, int numBits)
+		{
+			DE_ASSERT(inRange(numBits, 0, 32));
+			deUint32 result = 0;
+			for (int i = 0; i < numBits; i++)
+				result |= ((src >> i) & 1) << (numBits - 1 - i);
+			return result;
+		}
+		inline deUint32 bitReplicationScale(deUint32 src, int numSrcBits, int numDstBits)
+		{
+			DE_ASSERT(numSrcBits <= numDstBits);
+			DE_ASSERT((src & ((1 << numSrcBits) - 1)) == src);
+			deUint32 dst = 0;
+			for (int shift = numDstBits - numSrcBits; shift > -numSrcBits; shift -= numSrcBits)
+				dst |= shift >= 0 ? src << shift : src >> -shift;
+			return dst;
 		}
 
+		inline deInt32 signExtend(deInt32 src, int numSrcBits)
+		{
+			DE_ASSERT(inRange(numSrcBits, 2, 31));
+			const bool negative = (src & (1 << (numSrcBits - 1))) != 0;
+			return src | (negative ? ~((1 << numSrcBits) - 1) : 0);
+		}
 
+		//inline bool isFloat16InfOrNan (deFloat16 v)
+		//{
+		//	return getBits(v, 10, 14) == 31;
+		//}
 
-		void find_number_of_bits_trits_quints(int quantization_level, int* bits, int* trits, int* quints)
-		{
-			*bits = 0;
-			*trits = 0;
-			*quints = 0;
-			switch (quantization_level)
-			{
-			case QUANT_2:
-				*bits = 1;
-				break;
-			case QUANT_3:
-				*bits = 0;
-				*trits = 1;
-				break;
-			case QUANT_4:
-				*bits = 2;
-				break;
-			case QUANT_5:
-				*bits = 0;
-				*quints = 1;
-				break;
-			case QUANT_6:
-				*bits = 1;
-				*trits = 1;
-				break;
-			case QUANT_8:
-				*bits = 3;
-				break;
-			case QUANT_10:
-				*bits = 1;
-				*quints = 1;
-				break;
-			case QUANT_12:
-				*bits = 2;
-				*trits = 1;
-				break;
-			case QUANT_16:
-				*bits = 4;
-				break;
-			case QUANT_20:
-				*bits = 2;
-				*quints = 1;
-				break;
-			case QUANT_24:
-				*bits = 3;
-				*trits = 1;
-				break;
-			case QUANT_32:
-				*bits = 5;
-				break;
-			case QUANT_40:
-				*bits = 3;
-				*quints = 1;
-				break;
-			case QUANT_48:
-				*bits = 4;
-				*trits = 1;
-				break;
-			case QUANT_64:
-				*bits = 6;
-				break;
-			case QUANT_80:
-				*bits = 4;
-				*quints = 1;
-				break;
-			case QUANT_96:
-				*bits = 5;
-				*trits = 1;
-				break;
-			case QUANT_128:
-				*bits = 7;
-				break;
-			case QUANT_160:
-				*bits = 5;
-				*quints = 1;
-				break;
-			case QUANT_192:
-				*bits = 6;
-				*trits = 1;
-				break;
-			case QUANT_256:
-				*bits = 8;
-				break;
-			}
-		}
-
-
-		// unpacked trit quintuplets <low,_,_,_,high> for each packed-quint value
-		static const uint8 trits_of_integer[256][5] = {
-			{0, 0, 0, 0, 0},	{1, 0, 0, 0, 0},	{2, 0, 0, 0, 0},	{0, 0, 2, 0, 0},
-			{0, 1, 0, 0, 0},	{1, 1, 0, 0, 0},	{2, 1, 0, 0, 0},	{1, 0, 2, 0, 0},
-			{0, 2, 0, 0, 0},	{1, 2, 0, 0, 0},	{2, 2, 0, 0, 0},	{2, 0, 2, 0, 0},
-			{0, 2, 2, 0, 0},	{1, 2, 2, 0, 0},	{2, 2, 2, 0, 0},	{2, 0, 2, 0, 0},
-			{0, 0, 1, 0, 0},	{1, 0, 1, 0, 0},	{2, 0, 1, 0, 0},	{0, 1, 2, 0, 0},
-			{0, 1, 1, 0, 0},	{1, 1, 1, 0, 0},	{2, 1, 1, 0, 0},	{1, 1, 2, 0, 0},
-			{0, 2, 1, 0, 0},	{1, 2, 1, 0, 0},	{2, 2, 1, 0, 0},	{2, 1, 2, 0, 0},
-			{0, 0, 0, 2, 2},	{1, 0, 0, 2, 2},	{2, 0, 0, 2, 2},	{0, 0, 2, 2, 2},
-			{0, 0, 0, 1, 0},	{1, 0, 0, 1, 0},	{2, 0, 0, 1, 0},	{0, 0, 2, 1, 0},
-			{0, 1, 0, 1, 0},	{1, 1, 0, 1, 0},	{2, 1, 0, 1, 0},	{1, 0, 2, 1, 0},
-			{0, 2, 0, 1, 0},	{1, 2, 0, 1, 0},	{2, 2, 0, 1, 0},	{2, 0, 2, 1, 0},
-			{0, 2, 2, 1, 0},	{1, 2, 2, 1, 0},	{2, 2, 2, 1, 0},	{2, 0, 2, 1, 0},
-			{0, 0, 1, 1, 0},	{1, 0, 1, 1, 0},	{2, 0, 1, 1, 0},	{0, 1, 2, 1, 0},
-			{0, 1, 1, 1, 0},	{1, 1, 1, 1, 0},	{2, 1, 1, 1, 0},	{1, 1, 2, 1, 0},
-			{0, 2, 1, 1, 0},	{1, 2, 1, 1, 0},	{2, 2, 1, 1, 0},	{2, 1, 2, 1, 0},
-			{0, 1, 0, 2, 2},	{1, 1, 0, 2, 2},	{2, 1, 0, 2, 2},	{1, 0, 2, 2, 2},
-			{0, 0, 0, 2, 0},	{1, 0, 0, 2, 0},	{2, 0, 0, 2, 0},	{0, 0, 2, 2, 0},
-			{0, 1, 0, 2, 0},	{1, 1, 0, 2, 0},	{2, 1, 0, 2, 0},	{1, 0, 2, 2, 0},
-			{0, 2, 0, 2, 0},	{1, 2, 0, 2, 0},	{2, 2, 0, 2, 0},	{2, 0, 2, 2, 0},
-			{0, 2, 2, 2, 0},	{1, 2, 2, 2, 0},	{2, 2, 2, 2, 0},	{2, 0, 2, 2, 0},
-			{0, 0, 1, 2, 0},	{1, 0, 1, 2, 0},	{2, 0, 1, 2, 0},	{0, 1, 2, 2, 0},
-			{0, 1, 1, 2, 0},	{1, 1, 1, 2, 0},	{2, 1, 1, 2, 0},	{1, 1, 2, 2, 0},
-			{0, 2, 1, 2, 0},	{1, 2, 1, 2, 0},	{2, 2, 1, 2, 0},	{2, 1, 2, 2, 0},
-			{0, 2, 0, 2, 2},	{1, 2, 0, 2, 2},	{2, 2, 0, 2, 2},	{2, 0, 2, 2, 2},
-			{0, 0, 0, 0, 2},	{1, 0, 0, 0, 2},	{2, 0, 0, 0, 2},	{0, 0, 2, 0, 2},
-			{0, 1, 0, 0, 2},	{1, 1, 0, 0, 2},	{2, 1, 0, 0, 2},	{1, 0, 2, 0, 2},
-			{0, 2, 0, 0, 2},	{1, 2, 0, 0, 2},	{2, 2, 0, 0, 2},	{2, 0, 2, 0, 2},
-			{0, 2, 2, 0, 2},	{1, 2, 2, 0, 2},	{2, 2, 2, 0, 2},	{2, 0, 2, 0, 2},
-			{0, 0, 1, 0, 2},	{1, 0, 1, 0, 2},	{2, 0, 1, 0, 2},	{0, 1, 2, 0, 2},
-			{0, 1, 1, 0, 2},	{1, 1, 1, 0, 2},	{2, 1, 1, 0, 2},	{1, 1, 2, 0, 2},
-			{0, 2, 1, 0, 2},	{1, 2, 1, 0, 2},	{2, 2, 1, 0, 2},	{2, 1, 2, 0, 2},
-			{0, 2, 2, 2, 2},	{1, 2, 2, 2, 2},	{2, 2, 2, 2, 2},	{2, 0, 2, 2, 2},
-			{0, 0, 0, 0, 1},	{1, 0, 0, 0, 1},	{2, 0, 0, 0, 1},	{0, 0, 2, 0, 1},
-			{0, 1, 0, 0, 1},	{1, 1, 0, 0, 1},	{2, 1, 0, 0, 1},	{1, 0, 2, 0, 1},
-			{0, 2, 0, 0, 1},	{1, 2, 0, 0, 1},	{2, 2, 0, 0, 1},	{2, 0, 2, 0, 1},
-			{0, 2, 2, 0, 1},	{1, 2, 2, 0, 1},	{2, 2, 2, 0, 1},	{2, 0, 2, 0, 1},
-			{0, 0, 1, 0, 1},	{1, 0, 1, 0, 1},	{2, 0, 1, 0, 1},	{0, 1, 2, 0, 1},
-			{0, 1, 1, 0, 1},	{1, 1, 1, 0, 1},	{2, 1, 1, 0, 1},	{1, 1, 2, 0, 1},
-			{0, 2, 1, 0, 1},	{1, 2, 1, 0, 1},	{2, 2, 1, 0, 1},	{2, 1, 2, 0, 1},
-			{0, 0, 1, 2, 2},	{1, 0, 1, 2, 2},	{2, 0, 1, 2, 2},	{0, 1, 2, 2, 2},
-			{0, 0, 0, 1, 1},	{1, 0, 0, 1, 1},	{2, 0, 0, 1, 1},	{0, 0, 2, 1, 1},
-			{0, 1, 0, 1, 1},	{1, 1, 0, 1, 1},	{2, 1, 0, 1, 1},	{1, 0, 2, 1, 1},
-			{0, 2, 0, 1, 1},	{1, 2, 0, 1, 1},	{2, 2, 0, 1, 1},	{2, 0, 2, 1, 1},
-			{0, 2, 2, 1, 1},	{1, 2, 2, 1, 1},	{2, 2, 2, 1, 1},	{2, 0, 2, 1, 1},
-			{0, 0, 1, 1, 1},	{1, 0, 1, 1, 1},	{2, 0, 1, 1, 1},	{0, 1, 2, 1, 1},
-			{0, 1, 1, 1, 1},	{1, 1, 1, 1, 1},	{2, 1, 1, 1, 1},	{1, 1, 2, 1, 1},
-			{0, 2, 1, 1, 1},	{1, 2, 1, 1, 1},	{2, 2, 1, 1, 1},	{2, 1, 2, 1, 1},
-			{0, 1, 1, 2, 2},	{1, 1, 1, 2, 2},	{2, 1, 1, 2, 2},	{1, 1, 2, 2, 2},
-			{0, 0, 0, 2, 1},	{1, 0, 0, 2, 1},	{2, 0, 0, 2, 1},	{0, 0, 2, 2, 1},
-			{0, 1, 0, 2, 1},	{1, 1, 0, 2, 1},	{2, 1, 0, 2, 1},	{1, 0, 2, 2, 1},
-			{0, 2, 0, 2, 1},	{1, 2, 0, 2, 1},	{2, 2, 0, 2, 1},	{2, 0, 2, 2, 1},
-			{0, 2, 2, 2, 1},	{1, 2, 2, 2, 1},	{2, 2, 2, 2, 1},	{2, 0, 2, 2, 1},
-			{0, 0, 1, 2, 1},	{1, 0, 1, 2, 1},	{2, 0, 1, 2, 1},	{0, 1, 2, 2, 1},
-			{0, 1, 1, 2, 1},	{1, 1, 1, 2, 1},	{2, 1, 1, 2, 1},	{1, 1, 2, 2, 1},
-			{0, 2, 1, 2, 1},	{1, 2, 1, 2, 1},	{2, 2, 1, 2, 1},	{2, 1, 2, 2, 1},
-			{0, 2, 1, 2, 2},	{1, 2, 1, 2, 2},	{2, 2, 1, 2, 2},	{2, 1, 2, 2, 2},
-			{0, 0, 0, 1, 2},	{1, 0, 0, 1, 2},	{2, 0, 0, 1, 2},	{0, 0, 2, 1, 2},
-			{0, 1, 0, 1, 2},	{1, 1, 0, 1, 2},	{2, 1, 0, 1, 2},	{1, 0, 2, 1, 2},
-			{0, 2, 0, 1, 2},	{1, 2, 0, 1, 2},	{2, 2, 0, 1, 2},	{2, 0, 2, 1, 2},
-			{0, 2, 2, 1, 2},	{1, 2, 2, 1, 2},	{2, 2, 2, 1, 2},	{2, 0, 2, 1, 2},
-			{0, 0, 1, 1, 2},	{1, 0, 1, 1, 2},	{2, 0, 1, 1, 2},	{0, 1, 2, 1, 2},
-			{0, 1, 1, 1, 2},	{1, 1, 1, 1, 2},	{2, 1, 1, 1, 2},	{1, 1, 2, 1, 2},
-			{0, 2, 1, 1, 2},	{1, 2, 1, 1, 2},	{2, 2, 1, 1, 2},	{2, 1, 2, 1, 2},
-			{0, 2, 2, 2, 2},	{1, 2, 2, 2, 2},	{2, 2, 2, 2, 2},	{2, 1, 2, 2, 2},
+		enum ISEMode
+		{
+			ISEMODE_TRIT = 0,
+			ISEMODE_QUINT,
+			ISEMODE_PLAIN_BIT,
+			ISEMODE_LAST
 		};
-
-
-		static const uint8 quints_of_integer[128][3] = {
-			{0, 0, 0},	{1, 0, 0},	{2, 0, 0},	{3, 0, 0},
-			{4, 0, 0},	{0, 4, 0},	{4, 4, 0},	{4, 4, 4},
-			{0, 1, 0},	{1, 1, 0},	{2, 1, 0},	{3, 1, 0},
-			{4, 1, 0},	{1, 4, 0},	{4, 4, 1},	{4, 4, 4},
-			{0, 2, 0},	{1, 2, 0},	{2, 2, 0},	{3, 2, 0},
-			{4, 2, 0},	{2, 4, 0},	{4, 4, 2},	{4, 4, 4},
-			{0, 3, 0},	{1, 3, 0},	{2, 3, 0},	{3, 3, 0},
-			{4, 3, 0},	{3, 4, 0},	{4, 4, 3},	{4, 4, 4},
-			{0, 0, 1},	{1, 0, 1},	{2, 0, 1},	{3, 0, 1},
-			{4, 0, 1},	{0, 4, 1},	{4, 0, 4},	{0, 4, 4},
-			{0, 1, 1},	{1, 1, 1},	{2, 1, 1},	{3, 1, 1},
-			{4, 1, 1},	{1, 4, 1},	{4, 1, 4},	{1, 4, 4},
-			{0, 2, 1},	{1, 2, 1},	{2, 2, 1},	{3, 2, 1},
-			{4, 2, 1},	{2, 4, 1},	{4, 2, 4},	{2, 4, 4},
-			{0, 3, 1},	{1, 3, 1},	{2, 3, 1},	{3, 3, 1},
-			{4, 3, 1},	{3, 4, 1},	{4, 3, 4},	{3, 4, 4},
-			{0, 0, 2},	{1, 0, 2},	{2, 0, 2},	{3, 0, 2},
-			{4, 0, 2},	{0, 4, 2},	{2, 0, 4},	{3, 0, 4},
-			{0, 1, 2},	{1, 1, 2},	{2, 1, 2},	{3, 1, 2},
-			{4, 1, 2},	{1, 4, 2},	{2, 1, 4},	{3, 1, 4},
-			{0, 2, 2},	{1, 2, 2},	{2, 2, 2},	{3, 2, 2},
-			{4, 2, 2},	{2, 4, 2},	{2, 2, 4},	{3, 2, 4},
-			{0, 3, 2},	{1, 3, 2},	{2, 3, 2},	{3, 3, 2},
-			{4, 3, 2},	{3, 4, 2},	{2, 3, 4},	{3, 3, 4},
-			{0, 0, 3},	{1, 0, 3},	{2, 0, 3},	{3, 0, 3},
-			{4, 0, 3},	{0, 4, 3},	{0, 0, 4},	{1, 0, 4},
-			{0, 1, 3},	{1, 1, 3},	{2, 1, 3},	{3, 1, 3},
-			{4, 1, 3},	{1, 4, 3},	{0, 1, 4},	{1, 1, 4},
-			{0, 2, 3},	{1, 2, 3},	{2, 2, 3},	{3, 2, 3},
-			{4, 2, 3},	{2, 4, 3},	{0, 2, 4},	{1, 2, 4},
-			{0, 3, 3},	{1, 3, 3},	{2, 3, 3},	{3, 3, 3},
-			{4, 3, 3},	{3, 4, 3},	{0, 3, 4},	{1, 3, 4},
+		struct ISEParams
+		{
+			ISEMode		mode;
+			int			numBits;
+			ISEParams(ISEMode mode_, int numBits_) : mode(mode_), numBits(numBits_) {}
 		};
-
-		// quantization_mode_table[integercount/2][bits] gives
-		// us the quantization level for a given integer count and number of bits that
-		// the integer may fit into. This is needed for color decoding,
-		// and for the color encoding.
-		int8_t quantization_mode_table[17][128];
-
-		void build_quantization_mode_table(void)
+		inline int computeNumRequiredBits(const ISEParams& iseParams, int numValues)
 		{
-			int8_t i;
-			int j;
-			for (i = 0; i <= 16; i++)
-				for (j = 0; j < 128; j++)
-					quantization_mode_table[i][j] = -1;
-
-			for (i = 0; i < 21; i++)
-				for (j = 1; j <= 16; j++)
-				{
-					int p = compute_ise_bitcount(2 * j, (quantization_method)i);
-					if (p < 128)
-						quantization_mode_table[j][p] = i;
-				}
-			for (i = 0; i <= 16; i++)
+			switch (iseParams.mode)
 			{
-				int8_t largest_value_so_far = -1;
-				for (j = 0; j < 128; j++)
-				{
-					if (quantization_mode_table[i][j] > largest_value_so_far)
-						largest_value_so_far = quantization_mode_table[i][j];
-					else
-						quantization_mode_table[i][j] = largest_value_so_far;
-				}
+			case ISEMODE_TRIT:			return deDivRoundUp32(numValues * 8, 5) + numValues * iseParams.numBits;
+			case ISEMODE_QUINT:			return deDivRoundUp32(numValues * 7, 3) + numValues * iseParams.numBits;
+			case ISEMODE_PLAIN_BIT:		return numValues * iseParams.numBits;
+			default:
+				DE_ASSERT(false);
+				return -1;
 			}
 		}
-
-		void decode_ise(int quantization_level, int elements, const uint8* input_data, uint8* output_data, int bit_offset)
+		ISEParams computeMaximumRangeISEParams(int numAvailableBits, int numValuesInSequence)
 		{
-			int i;
-			// note: due to how the trit/quint-block unpacking is done in this function,
-			// we may write more temporary results than the number of outputs
-			// The maximum actual number of results is 64 bit, but we keep 4 additional elements
-			// of padding.
-			uint8 results[68];
-			uint8 tq_blocks[22];		// trit-blocks or quint-blocks
-
-			int bits, trits, quints;
-			find_number_of_bits_trits_quints(quantization_level, &bits, &trits, &quints);
-
-			int lcounter = 0;
-			int hcounter = 0;
-
-			// trit-blocks or quint-blocks must be zeroed out before we collect them in the loop below.
-			for (i = 0; i < 22; i++)
-				tq_blocks[i] = 0;
-
-			// collect bits for each element, as well as bits for any trit-blocks and quint-blocks.
-			for (i = 0; i < elements; i++)
+			int curBitsForTritMode = 6;
+			int curBitsForQuintMode = 5;
+			int curBitsForPlainBitMode = 8;
+			while (true)
 			{
-				results[i] = (uint8)read_bits(bits, bit_offset, input_data);
-				bit_offset += bits;
-				if (trits)
+				DE_ASSERT(curBitsForTritMode > 0 || curBitsForQuintMode > 0 || curBitsForPlainBitMode > 0);
+				const int tritRange = curBitsForTritMode > 0 ? (3 << curBitsForTritMode) - 1 : -1;
+				const int quintRange = curBitsForQuintMode > 0 ? (5 << curBitsForQuintMode) - 1 : -1;
+				const int plainBitRange = curBitsForPlainBitMode > 0 ? (1 << curBitsForPlainBitMode) - 1 : -1;
+				const int maxRange = FMath::Max(FMath::Max(tritRange, quintRange), plainBitRange);
+				if (maxRange == tritRange)
 				{
-					static const int bits_to_read[5] = { 2, 2, 1, 2, 1 };
-					static const int block_shift[5] = { 0, 2, 4, 5, 7 };
-					static const int next_lcounter[5] = { 1, 2, 3, 4, 0 };
-					static const int hcounter_incr[5] = { 0, 0, 0, 0, 1 };
-					DCHECK(lcounter < 5);
-					int tdata = read_bits(bits_to_read[lcounter], bit_offset, input_data);
-					bit_offset += bits_to_read[lcounter];
-					tq_blocks[hcounter] |= tdata << block_shift[lcounter];
-					hcounter += hcounter_incr[lcounter];
-					lcounter = next_lcounter[lcounter];
+					const ISEParams params(ISEMODE_TRIT, curBitsForTritMode);
+					if (computeNumRequiredBits(params, numValuesInSequence) <= numAvailableBits)
+						return ISEParams(ISEMODE_TRIT, curBitsForTritMode);
+					curBitsForTritMode--;
 				}
-				if (quints)
+				else if (maxRange == quintRange)
 				{
-					static const int bits_to_read[3] = { 3, 2, 2 };
-					static const int block_shift[3] = { 0, 3, 5 };
-					static const int next_lcounter[3] = { 1, 2, 0 };
-					static const int hcounter_incr[3] = { 0, 0, 1 };
-					DCHECK(lcounter < 3);
-					if (lcounter < 3)
-					{
-						int tdata = read_bits(bits_to_read[lcounter], bit_offset, input_data);
-						bit_offset += bits_to_read[lcounter];
-						tq_blocks[hcounter] |= tdata << block_shift[lcounter];
-						hcounter += hcounter_incr[lcounter];
-						lcounter = next_lcounter[lcounter];
-					}
+					const ISEParams params(ISEMODE_QUINT, curBitsForQuintMode);
+					if (computeNumRequiredBits(params, numValuesInSequence) <= numAvailableBits)
+						return ISEParams(ISEMODE_QUINT, curBitsForQuintMode);
+					curBitsForQuintMode--;
 				}
-			}
-
-
-			// unpack trit-blocks or quint-blocks as needed
-			if (trits)
-			{
-				int trit_blocks = (elements + 4) / 5;
-				for (i = 0; i < trit_blocks; i++)
+				else
 				{
-					const uint8* tritptr = trits_of_integer[tq_blocks[i]];
-					results[5 * i] |= tritptr[0] << bits;
-					results[5 * i + 1] |= tritptr[1] << bits;
-					results[5 * i + 2] |= tritptr[2] << bits;
-					results[5 * i + 3] |= tritptr[3] << bits;
-					results[5 * i + 4] |= tritptr[4] << bits;
+					const ISEParams params(ISEMODE_PLAIN_BIT, curBitsForPlainBitMode);
+					DE_ASSERT(maxRange == plainBitRange);
+					if (computeNumRequiredBits(params, numValuesInSequence) <= numAvailableBits)
+						return ISEParams(ISEMODE_PLAIN_BIT, curBitsForPlainBitMode);
+					curBitsForPlainBitMode--;
 				}
 			}
-
-			if (quints)
+		}
+		inline int computeNumColorEndpointValues(deUint32 endpointMode)
+		{
+			DE_ASSERT(endpointMode < 16);
+			return (endpointMode / 4 + 1) * 2;
+		}
+		// Decompression utilities
+		enum DecompressResult
+		{
+			DECOMPRESS_RESULT_VALID_BLOCK = 0,	//!< Decompressed valid block
+			DECOMPRESS_RESULT_ERROR,				//!< Encountered error while decompressing, error color written
+			DECOMPRESS_RESULT_LAST
+		};
+		// A helper for getting bits from a 128-bit block.
+		class Block128
+		{
+		private:
+			typedef deUint64 Word;
+			enum
 			{
-				int quint_blocks = (elements + 2) / 3;
-				for (i = 0; i < quint_blocks; i++)
+				WORD_BYTES = sizeof(Word),
+				WORD_BITS = 8 * WORD_BYTES,
+				NUM_WORDS = 128 / WORD_BITS
+			};
+			//DE_STATIC_ASSERT(128 % WORD_BITS == 0);
+		public:
+			Block128(const deUint8* src)
+			{
+				for (int wordNdx = 0; wordNdx < NUM_WORDS; wordNdx++)
 				{
-					const uint8* quintptr = quints_of_integer[tq_blocks[i]];
-					results[3 * i] |= quintptr[0] << bits;
-					results[3 * i + 1] |= quintptr[1] << bits;
-					results[3 * i + 2] |= quintptr[2] << bits;
+					m_words[wordNdx] = 0;
+					for (int byteNdx = 0; byteNdx < WORD_BYTES; byteNdx++)
+						m_words[wordNdx] |= (Word)src[wordNdx * WORD_BYTES + byteNdx] << (8 * byteNdx);
 				}
 			}
-
-			for (i = 0; i < elements; i++)
-				output_data[i] = results[i];
-		}
-
-
-		inline void physical_to_symbolic(physical_compressed_block pb, symbolic_compressed_block* res)
-		{
-			uint8 bswapped[16];
-			int i, j;
-
-			res->error_block = 0;
-
-			// get hold of the block-size descriptor and the decimation tables.
-			const block_size_descriptor* bsd = get_block_size_descriptor();
-			const decimation_table* const* ixtab2 = bsd->decimation_tables;
-
-			// extract header fields
-			int block_mode = read_bits(11, 0, pb.data);
-
-
-			if ((block_mode & 0x1FF) == 0x1FC)
+			deUint32 getBit(int ndx) const
 			{
-				// void-extent block!
-
-				// check what format the data has
-				if (block_mode & 0x200)
-					res->block_mode = -1;	// floating-point
+				DE_ASSERT(inBounds(ndx, 0, 128));
+				return (m_words[ndx / WORD_BITS] >> (ndx % WORD_BITS)) & 1;
+			}
+			deUint32 getBits(int low, int high) const
+			{
+				DE_ASSERT(inBounds(low, 0, 128));
+				DE_ASSERT(inBounds(high, 0, 128));
+				DE_ASSERT(inRange(high - low + 1, 0, 32));
+				if (high - low + 1 == 0)
+					return 0;
+				const int word0Ndx = low / WORD_BITS;
+				const int word1Ndx = high / WORD_BITS;
+				// \note "foo << bar << 1" done instead of "foo << (bar+1)" to avoid overflow, i.e. shift amount being too big.
+				if (word0Ndx == word1Ndx)
+					return (deUint32)((m_words[word0Ndx] & ((((Word)1 << high % WORD_BITS << 1) - 1))) >> ((Word)low % WORD_BITS));
 				else
-					res->block_mode = -2;	// unorm16.
-
-				res->partition_count = 0;
-				for (i = 0; i < 4; i++)
-				{
-					res->constant_color[i] = pb.data[2 * i + 8] | (pb.data[2 * i + 9] << 8);
-				}
-
-				// additionally, check that the void-extent
 				{
-					// 2D void-extent
-					int rsvbits = read_bits(2, 10, pb.data);
-					if (rsvbits != 3)
-						res->error_block = 1;
-
-					int vx_low_s = read_bits(8, 12, pb.data) | (read_bits(5, 12 + 8, pb.data) << 8);
-					int vx_high_s = read_bits(8, 25, pb.data) | (read_bits(5, 25 + 8, pb.data) << 8);
-					int vx_low_t = read_bits(8, 38, pb.data) | (read_bits(5, 38 + 8, pb.data) << 8);
-					int vx_high_t = read_bits(8, 51, pb.data) | (read_bits(5, 51 + 8, pb.data) << 8);
-
-					int all_ones = vx_low_s == 0x1FFF && vx_high_s == 0x1FFF && vx_low_t == 0x1FFF && vx_high_t == 0x1FFF;
-
-					if ((vx_low_s >= vx_high_s || vx_low_t >= vx_high_t) && !all_ones)
-						res->error_block = 1;
+					DE_ASSERT(word1Ndx == word0Ndx + 1);
+					return (deUint32)(m_words[word0Ndx] >> (low % WORD_BITS)) |
+						(deUint32)((m_words[word1Ndx] & (((Word)1 << high % WORD_BITS << 1) - 1)) << (high - low - high % WORD_BITS));
 				}
-
-				return;
 			}
-
-			if (bsd->block_modes[block_mode].permit_decode == 0)
+			bool isBitSet(int ndx) const
 			{
-				res->error_block = 1;
-				return;
+				DE_ASSERT(inBounds(ndx, 0, 128));
+				return getBit(ndx) != 0;
 			}
-
-			int weight_count = ixtab2[bsd->block_modes[block_mode].decimation_mode]->num_weights;
-			int weight_quantization_method = bsd->block_modes[block_mode].quantization_mode;
-			int is_dual_plane = bsd->block_modes[block_mode].is_dual_plane;
-
-			int real_weight_count = is_dual_plane ? 2 * weight_count : weight_count;
-
-			int partition_count = read_bits(2, 11, pb.data) + 1;
-
-			res->block_mode = block_mode;
-			res->partition_count = partition_count;
-
-			for (i = 0; i < 16; i++)
-				bswapped[i] = (uint8)bitrev8(pb.data[15 - i]);
-
-			int bits_for_weights = compute_ise_bitcount(real_weight_count,
-				(quantization_method)weight_quantization_method);
-
-			int below_weights_pos = 128 - bits_for_weights;
-
-			if (is_dual_plane)
+		private:
+			Word m_words[NUM_WORDS];
+		};
+		// A helper for sequential access into a Block128.
+		class BitAccessStream
+		{
+		public:
+			BitAccessStream(const Block128& src, int startNdxInSrc, int length, bool forward)
+				: m_src(src)
+				, m_startNdxInSrc(startNdxInSrc)
+				, m_length(length)
+				, m_forward(forward)
+				, m_ndx(0)
 			{
-				uint8 indices[64];
-				decode_ise(weight_quantization_method, real_weight_count, bswapped, indices, 0);
-				for (i = 0; i < weight_count; i++)
-				{
-					res->plane1_weights[i] = indices[2 * i];
-					res->plane2_weights[i] = indices[2 * i + 1];
-				}
-
-				//            if (s_debuglog)
-				//            {
-				//                i=0;
-				//                UE_LOG(LogMutableCore,Warning," decompression quantized r weights:");
-				//                for (int j=0;j<4;++j)
-				//                {
-				//                    UE_LOG(LogMutableCore,Warning," %.3d %.3d %.3d %.3d ",
-				//                            res->plane1_weights[i+0], res->plane1_weights[i+1], res->plane1_weights[i+2], res->plane1_weights[i+3] );
-				//                    i+=4;
-				//                }
-				////                i=0;
-				////                UE_LOG(LogMutableCore,Warning," decompression quantized g weights:");
-				////                for (int j=0;j<4;++j)
-				////                {
-				////                    UE_LOG(LogMutableCore,Warning," %.3d %.3d %.3d %.3d ",
-				////                            res->plane2_weights[i+0], res->plane2_weights[i+1], res->plane2_weights[i+2], res->plane2_weights[i+3] );
-				////                    i+=4;
-				////                }
-				//            }
 			}
-			else
+			// Get the next num bits. Bits at positions greater than or equal to m_length are zeros.
+			deUint32 getNext(int num)
 			{
-				decode_ise(weight_quantization_method, weight_count, bswapped, res->plane1_weights, 0);
+				if (num == 0 || m_ndx >= m_length)
+					return 0;
+				const int end = m_ndx + num;
+				const int numBitsFromSrc = FMath::Max(0, FMath::Min(m_length, end) - m_ndx);
+				const int low = m_ndx;
+				const int high = m_ndx + numBitsFromSrc - 1;
+				m_ndx += num;
+				return m_forward ? m_src.getBits(m_startNdxInSrc + low, m_startNdxInSrc + high)
+					: reverseBits(m_src.getBits(m_startNdxInSrc - high, m_startNdxInSrc - low), numBitsFromSrc);
 			}
-
-			if (is_dual_plane && partition_count == 4)
-				res->error_block = 1;
-
-
-
-			res->color_formats_matched = 0;
-
-			// then, determine the format of each endpoint pair
-			int color_formats[4];
-			int encoded_type_highpart_size = 0;
-			if (partition_count == 1)
+		private:
+			const Block128& m_src;
+			const int			m_startNdxInSrc;
+			const int			m_length;
+			const bool			m_forward;
+			int					m_ndx;
+		};
+		struct ISEDecodedResult
+		{
+			deUint32 m;
+			deUint32 tq; //!< Trit or quint value, depending on ISE mode.
+			deUint32 v;
+		};
+		// Data from an ASTC block's "block mode" part (i.e. bits [0,10]).
+		struct ASTCBlockMode
+		{
+			bool		isError;
+			// \note Following fields only relevant if !isError.
+			bool		isVoidExtent;
+			// \note Following fields only relevant if !isVoidExtent.
+			bool		isDualPlane;
+			int			weightGridWidth;
+			int			weightGridHeight;
+			ISEParams	weightISEParams;
+			ASTCBlockMode(void)
+				: isError(true)
+				, isVoidExtent(true)
+				, isDualPlane(true)
+				, weightGridWidth(-1)
+				, weightGridHeight(-1)
+				, weightISEParams(ISEMODE_LAST, -1)
 			{
-				color_formats[0] = read_bits(4, 13, pb.data);
-				res->partition_index = 0;
 			}
-			else
+		};
+		inline int computeNumWeights(const ASTCBlockMode& mode)
+		{
+			return mode.weightGridWidth * mode.weightGridHeight * (mode.isDualPlane ? 2 : 1);
+		}
+		struct ColorEndpointPair
+		{
+			UVec4 e0;
+			UVec4 e1;
+		};
+		struct TexelWeightPair
+		{
+			deUint32 w[2];
+		};
+		ASTCBlockMode getASTCBlockMode(deUint32 blockModeData)
+		{
+			ASTCBlockMode blockMode;
+			blockMode.isError = true; // \note Set to false later, if not error.
+			blockMode.isVoidExtent = getBits(blockModeData, 0, 8) == 0x1fc;
+			if (!blockMode.isVoidExtent)
 			{
-				encoded_type_highpart_size = (3 * partition_count) - 4;
-				below_weights_pos -= encoded_type_highpart_size;
-				int encoded_type = read_bits(6, 13 + PARTITION_BITS, pb.data) | (read_bits(encoded_type_highpart_size, below_weights_pos, pb.data) << 6);
-				int baseclass = encoded_type & 0x3;
-				if (baseclass == 0)
+				if ((getBits(blockModeData, 0, 1) == 0 && getBits(blockModeData, 6, 8) == 7) || getBits(blockModeData, 0, 3) == 0)
+					return blockMode; // Invalid ("reserved").
+				deUint32 r = (deUint32)-1; // \note Set in the following branches.
+				if (getBits(blockModeData, 0, 1) == 0)
 				{
-					for (i = 0; i < partition_count; i++)
+					const deUint32 r0 = getBit(blockModeData, 4);
+					const deUint32 r1 = getBit(blockModeData, 2);
+					const deUint32 r2 = getBit(blockModeData, 3);
+					const deUint32 i78 = getBits(blockModeData, 7, 8);
+					r = (r2 << 2) | (r1 << 1) | (r0 << 0);
+					if (i78 == 3)
+					{
+						const bool i5 = isBitSet(blockModeData, 5);
+						blockMode.weightGridWidth = i5 ? 10 : 6;
+						blockMode.weightGridHeight = i5 ? 6 : 10;
+					}
+					else
 					{
-						color_formats[i] = (encoded_type >> 2) & 0xF;
+						const deUint32 a = getBits(blockModeData, 5, 6);
+						switch (i78)
+						{
+						case 0:		blockMode.weightGridWidth = 12;		blockMode.weightGridHeight = a + 2;									break;
+						case 1:		blockMode.weightGridWidth = a + 2;	blockMode.weightGridHeight = 12;									break;
+						case 2:		blockMode.weightGridWidth = a + 6;	blockMode.weightGridHeight = getBits(blockModeData, 9, 10) + 6;		break;
+						default: DE_ASSERT(false);
+						}
 					}
-					below_weights_pos += encoded_type_highpart_size;
-					res->color_formats_matched = 1;
-					encoded_type_highpart_size = 0;
 				}
 				else
 				{
-					int bitpos = 2;
-					baseclass--;
-					for (i = 0; i < partition_count; i++)
+					const deUint32 r0 = getBit(blockModeData, 4);
+					const deUint32 r1 = getBit(blockModeData, 0);
+					const deUint32 r2 = getBit(blockModeData, 1);
+					const deUint32 i23 = getBits(blockModeData, 2, 3);
+					const deUint32 a = getBits(blockModeData, 5, 6);
+					r = (r2 << 2) | (r1 << 1) | (r0 << 0);
+					if (i23 == 3)
 					{
-						color_formats[i] = (((encoded_type >> bitpos) & 1) + baseclass) << 2;
-						bitpos++;
+						const deUint32	b = getBit(blockModeData, 7);
+						const bool		i8 = isBitSet(blockModeData, 8);
+						blockMode.weightGridWidth = i8 ? b + 2 : a + 2;
+						blockMode.weightGridHeight = i8 ? a + 2 : b + 6;
 					}
-					for (i = 0; i < partition_count; i++)
+					else
 					{
-						color_formats[i] |= (encoded_type >> bitpos) & 3;
-						bitpos += 2;
+						const deUint32 b = getBits(blockModeData, 7, 8);
+						switch (i23)
+						{
+						case 0:		blockMode.weightGridWidth = b + 4;	blockMode.weightGridHeight = a + 2;	break;
+						case 1:		blockMode.weightGridWidth = b + 8;	blockMode.weightGridHeight = a + 2;	break;
+						case 2:		blockMode.weightGridWidth = a + 2;	blockMode.weightGridHeight = b + 8;	break;
+						default: DE_ASSERT(false);
+						}
+					}
+				}
+				const bool	zeroDH = getBits(blockModeData, 0, 1) == 0 && getBits(blockModeData, 7, 8) == 2;
+				const bool	h = zeroDH ? 0 : isBitSet(blockModeData, 9);
+				blockMode.isDualPlane = zeroDH ? 0 : isBitSet(blockModeData, 10);
+				{
+					ISEMode& m = blockMode.weightISEParams.mode;
+					int& b = blockMode.weightISEParams.numBits;
+					m = ISEMODE_PLAIN_BIT;
+					b = 0;
+					if (h)
+					{
+						switch (r)
+						{
+						case 2:							m = ISEMODE_QUINT;	b = 1;	break;
+						case 3:		m = ISEMODE_TRIT;						b = 2;	break;
+						case 4:												b = 4;	break;
+						case 5:							m = ISEMODE_QUINT;	b = 2;	break;
+						case 6:		m = ISEMODE_TRIT;						b = 3;	break;
+						case 7:												b = 5;	break;
+						default:	DE_ASSERT(false);
+						}
+					}
+					else
+					{
+						switch (r)
+						{
+						case 2:												b = 1;	break;
+						case 3:		m = ISEMODE_TRIT;								break;
+						case 4:												b = 2;	break;
+						case 5:							m = ISEMODE_QUINT;			break;
+						case 6:		m = ISEMODE_TRIT;						b = 1;	break;
+						case 7:												b = 3;	break;
+						default:	DE_ASSERT(false);
+						}
 					}
 				}
-				res->partition_index = read_bits(6, 13, pb.data) | (read_bits(PARTITION_BITS - 6, 19, pb.data) << 6);
-
-			}
-			for (i = 0; i < partition_count; i++)
-				res->color_formats[i] = color_formats[i];
-
-
-			// then, determine the number of integers we need to unpack for the endpoint pairs
-			int color_integer_count = 0;
-			for (i = 0; i < partition_count; i++)
-			{
-				int endpoint_class = color_formats[i] >> 2;
-				color_integer_count += (endpoint_class + 1) * 2;
-			}
-
-			if (color_integer_count > 18)
-				res->error_block = 1;
-
-			// then, determine the color endpoint format to use for these integers
-			static const int color_bits_arr[5] = { -1, 115 - 4, 113 - 4 - PARTITION_BITS, 113 - 4 - PARTITION_BITS, 113 - 4 - PARTITION_BITS };
-			int color_bits = color_bits_arr[partition_count] - bits_for_weights - encoded_type_highpart_size;
-			if (is_dual_plane)
-				color_bits -= 2;
-			if (color_bits < 0)
-				color_bits = 0;
-
-			int color_quantization_level = quantization_mode_table[color_integer_count >> 1][color_bits];
-			res->color_quantization_level = color_quantization_level;
-			if (color_quantization_level < 4)
-				res->error_block = 1;
-
-
-			// then unpack the integer-bits
-			uint8 values_to_decode[32];
-			decode_ise(color_quantization_level, color_integer_count, pb.data, values_to_decode, (partition_count == 1 ? 17 : 19 + PARTITION_BITS));
-
-			// and distribute them over the endpoint types
-			int valuecount_to_decode = 0;
-
-			for (i = 0; i < partition_count; i++)
-			{
-				int vals = 2 * (color_formats[i] >> 2) + 2;
-				for (j = 0; j < vals; j++)
-					res->color_values[i][j] = values_to_decode[j + valuecount_to_decode];
-				valuecount_to_decode += vals;
 			}
-
-			// get hold of color component for second-plane in the case of dual plane of weights.
-			if (is_dual_plane)
-				res->plane2_color_component = read_bits(2, below_weights_pos - 2, pb.data);
-
+			blockMode.isError = false;
+			return blockMode;
 		}
-
-
-
-		// data structure representing one block of an image.
-		// it is expanded to float prior to processing to save some computation time
-		// on conversions to/from uint8 (this also allows us to handle HDR textures easily)
-		struct imageblock
+		inline void setASTCErrorColorBlock(void* dst, int blockWidth, int blockHeight, bool isSRGB)
 		{
-			float orig_data[MAX_TEXELS_PER_BLOCK * 4];  // original input data
-			float work_data[MAX_TEXELS_PER_BLOCK * 4];  // the data that we will compress, either linear or LNS (0..65535 in both cases)
-
-			uint8 rgb_lns[MAX_TEXELS_PER_BLOCK];      // 1 if RGB data are being treated as LNS
-			uint8 alpha_lns[MAX_TEXELS_PER_BLOCK];    // 1 if Alpha data are being treated as LNS
-		};
-
-		enum astc_decode_mode
-		{
-			DECODE_LDR_SRGB,
-			DECODE_LDR,
-			DECODE_HDR
-		};
-
-
-		template < typename vtype > class vtype4
-		{
-		public:
-			vtype x, y, z, w;
-			vtype4() {}
-			vtype4(vtype p, vtype q, vtype r, vtype s) :x(p), y(q), z(r), w(s) {}
-			vtype4(const vtype4& p) :x(p.x), y(p.y), z(p.z), w(p.w) {}
-		};
-		typedef unsigned short ushort;
-		typedef vtype4 < ushort > ushort4;
-		typedef vtype4 < int >int4;
-		static inline int4 operator-(int4 p, int4 q)
-		{
-			return int4(p.x - q.x, p.y - q.y, p.z - q.z, p.w - q.w);
-		}
-		static inline int4 operator>>(int4 p, int q)
-		{
-			return int4(p.x >> q, p.y >> q, p.z >> q, p.w >> q);
-		}
-		static inline int4 operator*(int4 p, int4 q)
-		{
-			return int4(p.x * q.x, p.y * q.y, p.z * q.z, p.w * q.w);
-		}
-		static inline int4 operator+(int4 p, int4 q)
-		{
-			return int4(p.x + q.x, p.y + q.y, p.z + q.z, p.w + q.w);
-		}
-		static inline int4 operator<<(int4 p, int q)
-		{
-			return int4(p.x << q, p.y << q, p.z << q, p.w << q);
-		}
-		static inline int4 operator|(int4 p, int4 q)
-		{
-			return int4(p.x | q.x, p.y | q.y, p.z | q.z, p.w | q.w);
-		}
-
-
-
-		int compute_value_of_texel_int(int texel_to_get, const decimation_table* it, const int* weights)
-		{
-			int i;
-			int summed_value = 8;
-			int weights_to_evaluate = it->texel_num_weights[texel_to_get];
-			for (i = 0; i < weights_to_evaluate; i++)
+			if (isSRGB)
 			{
-				summed_value += weights[it->texel_weights[texel_to_get][i]] * it->texel_weights_int[texel_to_get][i];
-			}
-			return summed_value >> 4;
-		}
-
-
-		ushort4 lerp_color_int(astc_decode_mode decode_mode, ushort4 color0, ushort4 color1, int weight, int plane2_weight, int plane2_color_component	// -1 in 1-plane mode
-		)
-		{
-			int4 ecolor0 = int4(color0.x, color0.y, color0.z, color0.w);
-			int4 ecolor1 = int4(color1.x, color1.y, color1.z, color1.w);
-
-			int4 eweight1 = int4(weight, weight, weight, weight);
-			switch (plane2_color_component)
-			{
-			case 0:
-				eweight1.x = plane2_weight;
-				break;
-			case 1:
-				eweight1.y = plane2_weight;
-				break;
-			case 2:
-				eweight1.z = plane2_weight;
-				break;
-			case 3:
-				eweight1.w = plane2_weight;
-				break;
-			default:
-				break;
+				deUint8* const dstU = (deUint8*)dst;
+				for (int i = 0; i < blockWidth * blockHeight; i++)
+				{
+					dstU[4 * i + 0] = 0xff;
+					dstU[4 * i + 1] = 0;
+					dstU[4 * i + 2] = 0xff;
+					dstU[4 * i + 3] = 0xff;
+				}
 			}
-
-			int4 eweight0 = int4(64, 64, 64, 64) - eweight1;
-
-			if (decode_mode == DECODE_LDR_SRGB)
+			else
 			{
-				ecolor0 = ecolor0 >> 8;
-				ecolor1 = ecolor1 >> 8;
+				float* const dstF = (float*)dst;
+				for (int i = 0; i < blockWidth * blockHeight; i++)
+				{
+					dstF[4 * i + 0] = 1.0f;
+					dstF[4 * i + 1] = 0.0f;
+					dstF[4 * i + 2] = 1.0f;
+					dstF[4 * i + 3] = 1.0f;
+				}
 			}
-			int4 color = (ecolor0 * eweight0) + (ecolor1 * eweight1) + int4(32, 32, 32, 32);
-			color = color >> 6;
-			if (decode_mode == DECODE_LDR_SRGB)
-				color = color | (color << 8);
-
-			ushort4 rcolor = ushort4((ushort)color.x, (ushort)color.y, (ushort)color.z, (ushort)color.w);
-			return rcolor;
 		}
-
-		// conversion functions between the LNS representation and the FP16 representation.
-
-		float float_to_lns(float p)
+		DecompressResult decodeVoidExtentBlock(void* dst, const Block128& blockData, int blockWidth, int blockHeight, bool isSRGB, bool isLDRMode)
 		{
-
-			if (astc_isnan(p) || p <= 1.0f / 67108864.0f)
+			const deUint32	minSExtent = blockData.getBits(12, 24);
+			const deUint32	maxSExtent = blockData.getBits(25, 37);
+			const deUint32	minTExtent = blockData.getBits(38, 50);
+			const deUint32	maxTExtent = blockData.getBits(51, 63);
+			const bool		allExtentsAllOnes = minSExtent == 0x1fff && maxSExtent == 0x1fff && minTExtent == 0x1fff && maxTExtent == 0x1fff;
+			const bool		isHDRBlock = blockData.isBitSet(9);
+			if ((isLDRMode && isHDRBlock) || (!allExtentsAllOnes && (minSExtent >= maxSExtent || minTExtent >= maxTExtent)))
 			{
-				// underflow or NaN value, return 0.
-				// We count underflow if the input value is smaller than 2^-26.
-				return 0;
+				setASTCErrorColorBlock(dst, blockWidth, blockHeight, isSRGB);
+				return DECOMPRESS_RESULT_ERROR;
 			}
-
-			if (fabs(p) >= 65536.0f)
+			const deUint32 rgba[4] =
 			{
-				// overflow, return a +INF value
-				return 65535;
-			}
-
-			int expo;
-			float normfrac = frexp(p, &expo);
-			float p1;
-			if (expo < -13)
+				blockData.getBits(64,  79),
+				blockData.getBits(80,  95),
+				blockData.getBits(96,  111),
+				blockData.getBits(112, 127)
+			};
+			if (isSRGB)
 			{
-				// input number is smaller than 2^-14. In this case, multiply by 2^25.
-				p1 = p * 33554432.0f;
-				expo = 0;
+				deUint8* const dstU = (deUint8*)dst;
+				for (int i = 0; i < blockWidth * blockHeight; i++)
+					for (int c = 0; c < 4; c++)
+						dstU[i * 4 + c] = (deUint8)((rgba[c] & 0xff00) >> 8);
 			}
 			else
 			{
-				expo += 14;
-				p1 = (normfrac - 0.5f) * 4096.0f;
-			}
-
-			if (p1 < 384.0f)
-				p1 *= 4.0f / 3.0f;
-			else if (p1 <= 1408.0f)
-				p1 += 128.0f;
-			else
-				p1 = (p1 + 512.0f) * (4.0f / 5.0f);
-
-			p1 += expo * 2048.0f;
-			return p1 + 1.0f;
-		}
-
-
-		// helper function to initialize the work-data from the orig-data
-		void imageblock_initialize_work_from_orig(imageblock* pb, int pixelcount)
-		{
-			int i;
-			float* fptr = pb->orig_data;
-			float* wptr = pb->work_data;
-
-			for (i = 0; i < pixelcount; i++)
-			{
-				if (pb->rgb_lns[i])
-				{
-					wptr[0] = float_to_lns(fptr[0]);
-					wptr[1] = float_to_lns(fptr[1]);
-					wptr[2] = float_to_lns(fptr[2]);
-				}
-				else
-				{
-					wptr[0] = fptr[0] * 65535.0f;
-					wptr[1] = fptr[1] * 65535.0f;
-					wptr[2] = fptr[2] * 65535.0f;
-				}
-
-				if (pb->alpha_lns[i])
+				float* const dstF = (float*)dst;
+				if (isHDRBlock)
 				{
-					wptr[3] = float_to_lns(fptr[3]);
+					// rg - REMOVING HDR SUPPORT FOR NOW
+#if 0
+					for (int c = 0; c < 4; c++)
+					{
+						if (isFloat16InfOrNan((deFloat16)rgba[c]))
+							throw InternalError("Infinity or NaN color component in HDR void extent block in ASTC texture (behavior undefined by ASTC specification)");
+					}
+					for (int i = 0; i < blockWidth * blockHeight; i++)
+						for (int c = 0; c < 4; c++)
+							dstF[i * 4 + c] = deFloat16To32((deFloat16)rgba[c]);
+#endif
 				}
 				else
 				{
-					wptr[3] = fptr[3] * 65535.0f;
+					for (int i = 0; i < blockWidth * blockHeight; i++)
+						for (int c = 0; c < 4; c++)
+							dstF[i * 4 + c] = rgba[c] == 65535 ? 1.0f : (float)rgba[c] / 65536.0f;
 				}
-				fptr += 4;
-				wptr += 4;
-			}
-		}
-
-
-		/*	sized soft-float types. These are mapped to the sized integer types of C99, instead of C's
-			floating-point types; this is because the library needs to maintain exact, bit-level control on all
-			operations on these data types. */
-		typedef uint16 sf16;
-		typedef uint32 sf32;
-
-		typedef union if32_
-		{
-			uint32 u;
-			int32 s;
-			float f;
-		} if32;
-
-		/******************************************
-		  helper functions and their lookup tables
-		 ******************************************/
-		 /* count leading zeros functions. Only used when the input is nonzero. */
-
-#if defined(__GNUC__) && (defined(__i386) || defined(__amd64))
-#elif defined(__arm__) && defined(__ARMCC_VERSION)
-#elif defined(__arm__) && defined(__GNUC__)
-#else
-	/* table used for the slow default versions. */
-		static const uint8 clz_table[256] =
-		{
-			8, 7, 6, 6, 5, 5, 5, 5, 4, 4, 4, 4, 4, 4, 4, 4,
-			3, 3, 3, 3, 3, 3, 3, 3, 3, 3, 3, 3, 3, 3, 3, 3,
-			2, 2, 2, 2, 2, 2, 2, 2, 2, 2, 2, 2, 2, 2, 2, 2,
-			2, 2, 2, 2, 2, 2, 2, 2, 2, 2, 2, 2, 2, 2, 2, 2,
-			1, 1, 1, 1, 1, 1, 1, 1, 1, 1, 1, 1, 1, 1, 1, 1,
-			1, 1, 1, 1, 1, 1, 1, 1, 1, 1, 1, 1, 1, 1, 1, 1,
-			1, 1, 1, 1, 1, 1, 1, 1, 1, 1, 1, 1, 1, 1, 1, 1,
-			1, 1, 1, 1, 1, 1, 1, 1, 1, 1, 1, 1, 1, 1, 1, 1,
-			0, 0, 0, 0, 0, 0, 0, 0, 0, 0, 0, 0, 0, 0, 0, 0,
-			0, 0, 0, 0, 0, 0, 0, 0, 0, 0, 0, 0, 0, 0, 0, 0,
-			0, 0, 0, 0, 0, 0, 0, 0, 0, 0, 0, 0, 0, 0, 0, 0,
-			0, 0, 0, 0, 0, 0, 0, 0, 0, 0, 0, 0, 0, 0, 0, 0,
-			0, 0, 0, 0, 0, 0, 0, 0, 0, 0, 0, 0, 0, 0, 0, 0,
-			0, 0, 0, 0, 0, 0, 0, 0, 0, 0, 0, 0, 0, 0, 0, 0,
-			0, 0, 0, 0, 0, 0, 0, 0, 0, 0, 0, 0, 0, 0, 0, 0,
-			0, 0, 0, 0, 0, 0, 0, 0, 0, 0, 0, 0, 0, 0, 0, 0
-		};
-#endif
-
-		/*
-		   32-bit count-leading-zeros function: use the Assembly instruction whenever possible. */
-		inline uint32 clz32(uint32 inp)
-		{
-#if defined(__GNUC__) && (defined(__i386) || defined(__amd64))
-			uint32 bsr;
-			__asm__("bsrl %1, %0": "=r"(bsr) : "r"(inp | 1));
-			return 31 - bsr;
-#else
-#if defined(__arm__) && defined(__ARMCC_VERSION)
-			return __clz(inp);			/* armcc builtin */
-#else
-#if defined(__arm__) && defined(__GNUC__)
-			uint32 lz;
-			__asm__("clz %0, %1": "=r"(lz) : "r"(inp));
-			return lz;
-#else
-			/* slow default version */
-			uint32 summa = 24;
-			if (inp >= UINT32_C(0x10000))
-			{
-				inp >>= 16;
-				summa -= 16;
-			}
-			if (inp >= UINT32_C(0x100))
-			{
-				inp >>= 8;
-				summa -= 8;
 			}
-			return summa + clz_table[inp];
-#endif
-#endif
-#endif
-		}
-
-
-		/* convert from FP16 to FP32. */
-		sf32 sf16_to_sf32(sf16 inp)
-		{
-			uint32 inpx = inp;
-
-			/*
-				This table contains, for every FP16 sign/exponent value combination,
-				the difference between the input FP16 value and the value obtained
-				by shifting the correct FP32 result right by 13 bits.
-				This table allows us to handle every case except denormals and NaN
-				with just 1 table lookup, 2 shifts and 1 add.
-			*/
-
-#define WITH_MB(a) INT32_C((a) | (1 << 31))
-			static const int32 tbl[64] =
-			{
-				WITH_MB(0x00000), INT32_C(0x1C000), INT32_C(0x1C000), INT32_C(0x1C000), INT32_C(0x1C000), INT32_C(0x1C000), INT32_C(0x1C000), INT32_C(0x1C000),
-				INT32_C(0x1C000), INT32_C(0x1C000), INT32_C(0x1C000), INT32_C(0x1C000), INT32_C(0x1C000), INT32_C(0x1C000), INT32_C(0x1C000), INT32_C(0x1C000),
-				INT32_C(0x1C000), INT32_C(0x1C000), INT32_C(0x1C000), INT32_C(0x1C000), INT32_C(0x1C000), INT32_C(0x1C000), INT32_C(0x1C000), INT32_C(0x1C000),
-				INT32_C(0x1C000), INT32_C(0x1C000), INT32_C(0x1C000), INT32_C(0x1C000), INT32_C(0x1C000), INT32_C(0x1C000), INT32_C(0x1C000), WITH_MB(0x38000),
-				WITH_MB(0x38000), INT32_C(0x54000), INT32_C(0x54000), INT32_C(0x54000), INT32_C(0x54000), INT32_C(0x54000), INT32_C(0x54000), INT32_C(0x54000),
-				INT32_C(0x54000), INT32_C(0x54000), INT32_C(0x54000), INT32_C(0x54000), INT32_C(0x54000), INT32_C(0x54000), INT32_C(0x54000), INT32_C(0x54000),
-				INT32_C(0x54000), INT32_C(0x54000), INT32_C(0x54000), INT32_C(0x54000), INT32_C(0x54000), INT32_C(0x54000), INT32_C(0x54000), INT32_C(0x54000),
-				INT32_C(0x54000), INT32_C(0x54000), INT32_C(0x54000), INT32_C(0x54000), INT32_C(0x54000), INT32_C(0x54000), INT32_C(0x54000), WITH_MB(0x70000)
-			};
-
-			int32 res = tbl[inpx >> 10];
-			res += inpx;
-
-			/* the normal cases: the MSB of 'res' is not set. */
-			if (res >= 0)				/* signed compare */
-				return res << 13;
-
-			/* Infinity and Zero: the bottom 10 bits of 'res' are clear. */
-			if ((res & UINT32_C(0x3FF)) == 0)
-				return res << 13;
-
-			/* NaN: the exponent field of 'inp' is not zero; NaNs must be quietened. */
-			if ((inpx & 0x7C00) != 0)
-				return (res << 13) | UINT32_C(0x400000);
-
-			/* the remaining cases are Denormals. */
-			{
-				uint32 sign = (inpx & UINT32_C(0x8000)) << 16;
-				uint32 mskval = inpx & UINT32_C(0x7FFF);
-				uint32 leadingzeroes = clz32(mskval);
-				mskval <<= leadingzeroes;
-				return (mskval >> 8) + ((0x85 - leadingzeroes) << 23) + sign;
-			}
-		}
-
-
-		float sf16_to_float(sf16 p)
-		{
-			if32 i;
-			i.u = sf16_to_sf32(p);
-			return i.f;
+			return DECOMPRESS_RESULT_VALID_BLOCK;
 		}
-
-		// conversion function from 16-bit LDR value to FP16.
-		// note: for LDR interpolation, it is impossible to get a denormal result;
-		// this simplifies the conversion.
-		// FALSE; we can receive a very small UNORM16 through the constant-block.
-		uint16 unorm16_to_sf16(uint16 p)
-		{
-			if (p == 0xFFFF)
-				return 0x3C00;			// value of 1.0 .
-			if (p < 4)
-				return p << 8;
-
-			int lz = clz32(p) - 16;
-			p <<= (lz + 1);
-			p >>= 6;
-			p |= (14 - lz) << 10;
-			return p;
-		}
-
-
-		/*
-			Partition table representation:
-			For each block size, we have 3 tables, each with 1024 partitionings;
-			these three tables correspond to 2, 3 and 4 partitions respectively.
-			For each partitioning, we have:
-			* a 4-entry table indicating how many texels there are in each of the 4 partitions.
-			  This may be from 0 to a very large value.
-			* a table indicating the partition index of each of the texels in the block.
-			  Each index may be 0, 1, 2 or 3.
-			* Each element in the table is an uint8 indicating partition index (0, 1, 2 or 3)
-		*/
-
-		struct partition_info
-		{
-			int partition_count;
-			uint8 texels_per_partition[4];
-			uint8 partition_of_texel[MAX_TEXELS_PER_BLOCK];
-			uint8 texels_of_partition[4][MAX_TEXELS_PER_BLOCK];
-
-			uint64 coverage_bitmaps[4];	// used for the purposes of k-means partition search.
-		};
-
-
-		static partition_info* s_partition_table;
-
-
-		uint32 hash52(uint32 inp)
-		{
-			inp ^= inp >> 15;
-
-			inp *= 0xEEDE0891;			// (2^4+1)*(2^7+1)*(2^17-1)
-			inp ^= inp >> 5;
-			inp += inp << 16;
-			inp ^= inp >> 7;
-			inp ^= inp >> 3;
-			inp ^= inp << 6;
-			inp ^= inp >> 17;
-			return inp;
-		}
-
-		int select_partition(int seed, int x, int y, int z, int partitioncount, int small_block)
+		void decodeColorEndpointModes(deUint32* endpointModesDst, const Block128& blockData, int numPartitions, int extraCemBitsStart)
 		{
-			if (small_block)
-			{
-				x <<= 1;
-				y <<= 1;
-				z <<= 1;
-			}
-
-			seed += (partitioncount - 1) * 1024;
-
-			uint32 rnum = hash52(seed);
-
-			uint8 seed1 = rnum & 0xF;
-			uint8 seed2 = (rnum >> 4) & 0xF;
-			uint8 seed3 = (rnum >> 8) & 0xF;
-			uint8 seed4 = (rnum >> 12) & 0xF;
-			uint8 seed5 = (rnum >> 16) & 0xF;
-			uint8 seed6 = (rnum >> 20) & 0xF;
-			uint8 seed7 = (rnum >> 24) & 0xF;
-			uint8 seed8 = (rnum >> 28) & 0xF;
-			uint8 seed9 = (rnum >> 18) & 0xF;
-			uint8 seed10 = (rnum >> 22) & 0xF;
-			uint8 seed11 = (rnum >> 26) & 0xF;
-			uint8 seed12 = ((rnum >> 30) | (rnum << 2)) & 0xF;
-
-			// squaring all the seeds in order to bias their distribution
-			// towards lower values.
-			seed1 *= seed1;
-			seed2 *= seed2;
-			seed3 *= seed3;
-			seed4 *= seed4;
-			seed5 *= seed5;
-			seed6 *= seed6;
-			seed7 *= seed7;
-			seed8 *= seed8;
-			seed9 *= seed9;
-			seed10 *= seed10;
-			seed11 *= seed11;
-			seed12 *= seed12;
-
-
-			int sh1, sh2, sh3;
-			if (seed & 1)
-			{
-				sh1 = (seed & 2 ? 4 : 5);
-				sh2 = (partitioncount == 3 ? 6 : 5);
-			}
+			if (numPartitions == 1)
+				endpointModesDst[0] = blockData.getBits(13, 16);
 			else
 			{
-				sh1 = (partitioncount == 3 ? 6 : 5);
-				sh2 = (seed & 2 ? 4 : 5);
-			}
-			sh3 = (seed & 0x10) ? sh1 : sh2;
-
-			seed1 >>= sh1;
-			seed2 >>= sh2;
-			seed3 >>= sh1;
-			seed4 >>= sh2;
-			seed5 >>= sh1;
-			seed6 >>= sh2;
-			seed7 >>= sh1;
-			seed8 >>= sh2;
-
-			seed9 >>= sh3;
-			seed10 >>= sh3;
-			seed11 >>= sh3;
-			seed12 >>= sh3;
-
-
-
-			int a = seed1 * x + seed2 * y + seed11 * z + (rnum >> 14);
-			int b = seed3 * x + seed4 * y + seed12 * z + (rnum >> 10);
-			int c = seed5 * x + seed6 * y + seed9 * z + (rnum >> 6);
-			int d = seed7 * x + seed8 * y + seed10 * z + (rnum >> 2);
-
-
-			// apply the saw
-			a &= 0x3F;
-			b &= 0x3F;
-			c &= 0x3F;
-			d &= 0x3F;
-
-			// remove some of the components if we are to output < 4 partitions.
-			if (partitioncount <= 3)
-				d = 0;
-			if (partitioncount <= 2)
-				c = 0;
-			if (partitioncount <= 1)
-				b = 0;
-
-			int partition;
-			if (a >= b && a >= c && a >= d)
-				partition = 0;
-			else if (b >= c && b >= d)
-				partition = 1;
-			else if (c >= d)
-				partition = 2;
-			else
-				partition = 3;
-			return partition;
-		}
-
-
-
-		void generate_one_partition_table(int partition_count, int partition_index, partition_info* pt)
-		{
-			int small_block = 1;
-
-			uint8* partition_of_texel = pt->partition_of_texel;
-			int x, y, z, i;
-
-
-			for (z = 0; z < 1; z++)
-				for (y = 0; y < 4; y++)
-					for (x = 0; x < 4; x++)
+				const deUint32 highLevelSelector = blockData.getBits(23, 24);
+				if (highLevelSelector == 0)
+				{
+					const deUint32 mode = blockData.getBits(25, 28);
+					for (int i = 0; i < numPartitions; i++)
+						endpointModesDst[i] = mode;
+				}
+				else
+				{
+					for (int partNdx = 0; partNdx < numPartitions; partNdx++)
 					{
-						uint8 part = (uint8)select_partition(partition_index, x, y, z, partition_count, small_block);
-						*partition_of_texel++ = part;
+						const deUint32 cemClass = highLevelSelector - (blockData.isBitSet(25 + partNdx) ? 0 : 1);
+						const deUint32 lowBit0Ndx = numPartitions + 2 * partNdx;
+						const deUint32 lowBit1Ndx = numPartitions + 2 * partNdx + 1;
+						const deUint32 lowBit0 = blockData.getBit(lowBit0Ndx < 4 ? 25 + lowBit0Ndx : extraCemBitsStart + lowBit0Ndx - 4);
+						const deUint32 lowBit1 = blockData.getBit(lowBit1Ndx < 4 ? 25 + lowBit1Ndx : extraCemBitsStart + lowBit1Ndx - 4);
+						endpointModesDst[partNdx] = (cemClass << 2) | (lowBit1 << 1) | lowBit0;
 					}
-
-
-			int texels_per_block = 16;
-
-			int counts[4];
-			for (i = 0; i < 4; i++)
-				counts[i] = 0;
-
-			for (i = 0; i < texels_per_block; i++)
-			{
-				int partition = pt->partition_of_texel[i];
-				pt->texels_of_partition[partition][counts[partition]++] = (uint8)i;
-			}
-
-			for (i = 0; i < 4; i++)
-				pt->texels_per_partition[i] = (uint8)counts[i];
-
-			if (counts[0] == 0)
-				pt->partition_count = 0;
-			else if (counts[1] == 0)
-				pt->partition_count = 1;
-			else if (counts[2] == 0)
-				pt->partition_count = 2;
-			else if (counts[3] == 0)
-				pt->partition_count = 3;
-			else
-				pt->partition_count = 4;
-
-
-			for (i = 0; i < 4; i++)
-				pt->coverage_bitmaps[i] = 0ULL;
-
-			int texels_to_process = 16;
-			for (i = 0; i < texels_to_process; i++)
-			{
-				pt->coverage_bitmaps[pt->partition_of_texel[i]] |= 1ULL << i;
-			}
-		}
-
-
-		static void generate_partition_tables()
-		{
-			partition_info* one_partition = new partition_info;
-
-			generate_one_partition_table(1, 0, one_partition);
-
-			s_partition_table = one_partition;
-		}
-
-
-		const partition_info* get_partition_table()
-		{
-			if (s_partition_table == nullptr)
-				generate_partition_tables();
-
-			return s_partition_table;
-		}
-
-
-
-		const uint8 color_unquantization_tables[21][256] = {
-			{
-			 0, 255,
-			 },
-			{
-			 0, 128, 255,
-			 },
-			{
-			 0, 85, 170, 255,
-			 },
-			{
-			 0, 64, 128, 192, 255,
-			 },
-			{
-			 0, 255, 51, 204, 102, 153,
-			 },
-			{
-			 0, 36, 73, 109, 146, 182, 219, 255,
-			 },
-			{
-			 0, 255, 28, 227, 56, 199, 84, 171, 113, 142,
-			 },
-			{
-			 0, 255, 69, 186, 23, 232, 92, 163, 46, 209, 116, 139,
-			 },
-			{
-			 0, 17, 34, 51, 68, 85, 102, 119, 136, 153, 170, 187, 204, 221, 238, 255,
-			 },
-			{
-			 0, 255, 67, 188, 13, 242, 80, 175, 27, 228, 94, 161, 40, 215, 107, 148,
-			 54, 201, 121, 134,
-			 },
-			{
-			 0, 255, 33, 222, 66, 189, 99, 156, 11, 244, 44, 211, 77, 178, 110, 145,
-			 22, 233, 55, 200, 88, 167, 121, 134,
-			 },
-			{
-			 0, 8, 16, 24, 33, 41, 49, 57, 66, 74, 82, 90, 99, 107, 115, 123,
-			 132, 140, 148, 156, 165, 173, 181, 189, 198, 206, 214, 222, 231, 239, 247, 255,
-			 },
-			{
-			 0, 255, 32, 223, 65, 190, 97, 158, 6, 249, 39, 216, 71, 184, 104, 151,
-			 13, 242, 45, 210, 78, 177, 110, 145, 19, 236, 52, 203, 84, 171, 117, 138,
-			 26, 229, 58, 197, 91, 164, 123, 132,
-			 },
-			{
-			 0, 255, 16, 239, 32, 223, 48, 207, 65, 190, 81, 174, 97, 158, 113, 142,
-			 5, 250, 21, 234, 38, 217, 54, 201, 70, 185, 86, 169, 103, 152, 119, 136,
-			 11, 244, 27, 228, 43, 212, 59, 196, 76, 179, 92, 163, 108, 147, 124, 131,
-			 },
-			{
-			 0, 4, 8, 12, 16, 20, 24, 28, 32, 36, 40, 44, 48, 52, 56, 60,
-			 65, 69, 73, 77, 81, 85, 89, 93, 97, 101, 105, 109, 113, 117, 121, 125,
-			 130, 134, 138, 142, 146, 150, 154, 158, 162, 166, 170, 174, 178, 182, 186, 190,
-			 195, 199, 203, 207, 211, 215, 219, 223, 227, 231, 235, 239, 243, 247, 251, 255,
-			 },
-			{
-			 0, 255, 16, 239, 32, 223, 48, 207, 64, 191, 80, 175, 96, 159, 112, 143,
-			 3, 252, 19, 236, 35, 220, 51, 204, 67, 188, 83, 172, 100, 155, 116, 139,
-			 6, 249, 22, 233, 38, 217, 54, 201, 71, 184, 87, 168, 103, 152, 119, 136,
-			 9, 246, 25, 230, 42, 213, 58, 197, 74, 181, 90, 165, 106, 149, 122, 133,
-			 13, 242, 29, 226, 45, 210, 61, 194, 77, 178, 93, 162, 109, 146, 125, 130,
-			 },
-			{
-			 0, 255, 8, 247, 16, 239, 24, 231, 32, 223, 40, 215, 48, 207, 56, 199,
-			 64, 191, 72, 183, 80, 175, 88, 167, 96, 159, 104, 151, 112, 143, 120, 135,
-			 2, 253, 10, 245, 18, 237, 26, 229, 35, 220, 43, 212, 51, 204, 59, 196,
-			 67, 188, 75, 180, 83, 172, 91, 164, 99, 156, 107, 148, 115, 140, 123, 132,
-			 5, 250, 13, 242, 21, 234, 29, 226, 37, 218, 45, 210, 53, 202, 61, 194,
-			 70, 185, 78, 177, 86, 169, 94, 161, 102, 153, 110, 145, 118, 137, 126, 129,
-			 },
-			{
-			 0, 2, 4, 6, 8, 10, 12, 14, 16, 18, 20, 22, 24, 26, 28, 30,
-			 32, 34, 36, 38, 40, 42, 44, 46, 48, 50, 52, 54, 56, 58, 60, 62,
-			 64, 66, 68, 70, 72, 74, 76, 78, 80, 82, 84, 86, 88, 90, 92, 94,
-			 96, 98, 100, 102, 104, 106, 108, 110, 112, 114, 116, 118, 120, 122, 124, 126,
-			 129, 131, 133, 135, 137, 139, 141, 143, 145, 147, 149, 151, 153, 155, 157, 159,
-			 161, 163, 165, 167, 169, 171, 173, 175, 177, 179, 181, 183, 185, 187, 189, 191,
-			 193, 195, 197, 199, 201, 203, 205, 207, 209, 211, 213, 215, 217, 219, 221, 223,
-			 225, 227, 229, 231, 233, 235, 237, 239, 241, 243, 245, 247, 249, 251, 253, 255,
-			 },
-			{
-			 0, 255, 8, 247, 16, 239, 24, 231, 32, 223, 40, 215, 48, 207, 56, 199,
-			 64, 191, 72, 183, 80, 175, 88, 167, 96, 159, 104, 151, 112, 143, 120, 135,
-			 1, 254, 9, 246, 17, 238, 25, 230, 33, 222, 41, 214, 49, 206, 57, 198,
-			 65, 190, 73, 182, 81, 174, 89, 166, 97, 158, 105, 150, 113, 142, 121, 134,
-			 3, 252, 11, 244, 19, 236, 27, 228, 35, 220, 43, 212, 51, 204, 59, 196,
-			 67, 188, 75, 180, 83, 172, 91, 164, 99, 156, 107, 148, 115, 140, 123, 132,
-			 4, 251, 12, 243, 20, 235, 28, 227, 36, 219, 44, 211, 52, 203, 60, 195,
-			 68, 187, 76, 179, 84, 171, 92, 163, 100, 155, 108, 147, 116, 139, 124, 131,
-			 6, 249, 14, 241, 22, 233, 30, 225, 38, 217, 46, 209, 54, 201, 62, 193,
-			 70, 185, 78, 177, 86, 169, 94, 161, 102, 153, 110, 145, 118, 137, 126, 129,
-			 },
-			{
-			 0, 255, 4, 251, 8, 247, 12, 243, 16, 239, 20, 235, 24, 231, 28, 227,
-			 32, 223, 36, 219, 40, 215, 44, 211, 48, 207, 52, 203, 56, 199, 60, 195,
-			 64, 191, 68, 187, 72, 183, 76, 179, 80, 175, 84, 171, 88, 167, 92, 163,
-			 96, 159, 100, 155, 104, 151, 108, 147, 112, 143, 116, 139, 120, 135, 124, 131,
-			 1, 254, 5, 250, 9, 246, 13, 242, 17, 238, 21, 234, 25, 230, 29, 226,
-			 33, 222, 37, 218, 41, 214, 45, 210, 49, 206, 53, 202, 57, 198, 61, 194,
-			 65, 190, 69, 186, 73, 182, 77, 178, 81, 174, 85, 170, 89, 166, 93, 162,
-			 97, 158, 101, 154, 105, 150, 109, 146, 113, 142, 117, 138, 121, 134, 125, 130,
-			 2, 253, 6, 249, 10, 245, 14, 241, 18, 237, 22, 233, 26, 229, 30, 225,
-			 34, 221, 38, 217, 42, 213, 46, 209, 50, 205, 54, 201, 58, 197, 62, 193,
-			 66, 189, 70, 185, 74, 181, 78, 177, 82, 173, 86, 169, 90, 165, 94, 161,
-			 98, 157, 102, 153, 106, 149, 110, 145, 114, 141, 118, 137, 122, 133, 126, 129,
-			 },
-			{
-			 0, 1, 2, 3, 4, 5, 6, 7, 8, 9, 10, 11, 12, 13, 14, 15,
-			 16, 17, 18, 19, 20, 21, 22, 23, 24, 25, 26, 27, 28, 29, 30, 31,
-			 32, 33, 34, 35, 36, 37, 38, 39, 40, 41, 42, 43, 44, 45, 46, 47,
-			 48, 49, 50, 51, 52, 53, 54, 55, 56, 57, 58, 59, 60, 61, 62, 63,
-			 64, 65, 66, 67, 68, 69, 70, 71, 72, 73, 74, 75, 76, 77, 78, 79,
-			 80, 81, 82, 83, 84, 85, 86, 87, 88, 89, 90, 91, 92, 93, 94, 95,
-			 96, 97, 98, 99, 100, 101, 102, 103, 104, 105, 106, 107, 108, 109, 110, 111,
-			 112, 113, 114, 115, 116, 117, 118, 119, 120, 121, 122, 123, 124, 125, 126, 127,
-			 128, 129, 130, 131, 132, 133, 134, 135, 136, 137, 138, 139, 140, 141, 142, 143,
-			 144, 145, 146, 147, 148, 149, 150, 151, 152, 153, 154, 155, 156, 157, 158, 159,
-			 160, 161, 162, 163, 164, 165, 166, 167, 168, 169, 170, 171, 172, 173, 174, 175,
-			 176, 177, 178, 179, 180, 181, 182, 183, 184, 185, 186, 187, 188, 189, 190, 191,
-			 192, 193, 194, 195, 196, 197, 198, 199, 200, 201, 202, 203, 204, 205, 206, 207,
-			 208, 209, 210, 211, 212, 213, 214, 215, 216, 217, 218, 219, 220, 221, 222, 223,
-			 224, 225, 226, 227, 228, 229, 230, 231, 232, 233, 234, 235, 236, 237, 238, 239,
-			 240, 241, 242, 243, 244, 245, 246, 247, 248, 249, 250, 251, 252, 253, 254, 255,
-			 },
-		};
-
-
-		enum endpoint_formats
-		{
-			FMT_LUMINANCE = 0,
-			FMT_LUMINANCE_DELTA = 1,
-			FMT_HDR_LUMINANCE_LARGE_RANGE = 2,
-			FMT_HDR_LUMINANCE_SMALL_RANGE = 3,
-			FMT_LUMINANCE_ALPHA = 4,
-			FMT_LUMINANCE_ALPHA_DELTA = 5,
-			FMT_RGB_SCALE = 6,
-			FMT_HDR_RGB_SCALE = 7,
-			FMT_RGB = 8,
-			FMT_RGB_DELTA = 9,
-			FMT_RGB_SCALE_ALPHA = 10,
-			FMT_HDR_RGB = 11,
-			FMT_RGBA = 12,
-			FMT_RGBA_DELTA = 13,
-			FMT_HDR_RGB_LDR_ALPHA = 14,
-			FMT_HDR_RGBA = 15,
-		};
-
-
-		int rgb_delta_unpack(const int input[6], int quantization_level, ushort4* output0, ushort4* output1)
-		{
-			// unquantize the color endpoints
-			int r0 = color_unquantization_tables[quantization_level][input[0]];
-			int g0 = color_unquantization_tables[quantization_level][input[2]];
-			int b0 = color_unquantization_tables[quantization_level][input[4]];
-
-			int r1 = color_unquantization_tables[quantization_level][input[1]];
-			int g1 = color_unquantization_tables[quantization_level][input[3]];
-			int b1 = color_unquantization_tables[quantization_level][input[5]];
-
-			// perform the bit-transfer procedure
-			r0 |= (r1 & 0x80) << 1;
-			g0 |= (g1 & 0x80) << 1;
-			b0 |= (b1 & 0x80) << 1;
-			r1 &= 0x7F;
-			g1 &= 0x7F;
-			b1 &= 0x7F;
-			if (r1 & 0x40)
-				r1 -= 0x80;
-			if (g1 & 0x40)
-				g1 -= 0x80;
-			if (b1 & 0x40)
-				b1 -= 0x80;
-
-			r0 >>= 1;
-			g0 >>= 1;
-			b0 >>= 1;
-			r1 >>= 1;
-			g1 >>= 1;
-			b1 >>= 1;
-
-			int rgbsum = r1 + g1 + b1;
-
-			r1 += r0;
-			g1 += g0;
-			b1 += b0;
-
-
-			int retval;
-
-			int r0e, g0e, b0e;
-			int r1e, g1e, b1e;
-
-			if (rgbsum >= 0)
-			{
-				r0e = r0;
-				g0e = g0;
-				b0e = b0;
-
-				r1e = r1;
-				g1e = g1;
-				b1e = b1;
-
-				retval = 0;
-			}
-			else
-			{
-				r0e = (r1 + b1) >> 1;
-				g0e = (g1 + b1) >> 1;
-				b0e = b1;
-
-				r1e = (r0 + b0) >> 1;
-				g1e = (g0 + b0) >> 1;
-				b1e = b0;
-
-				retval = 1;
-			}
-
-			if (r0e < 0)
-				r0e = 0;
-			else if (r0e > 255)
-				r0e = 255;
-
-			if (g0e < 0)
-				g0e = 0;
-			else if (g0e > 255)
-				g0e = 255;
-
-			if (b0e < 0)
-				b0e = 0;
-			else if (b0e > 255)
-				b0e = 255;
-
-			if (r1e < 0)
-				r1e = 0;
-			else if (r1e > 255)
-				r1e = 255;
-
-			if (g1e < 0)
-				g1e = 0;
-			else if (g1e > 255)
-				g1e = 255;
-
-			if (b1e < 0)
-				b1e = 0;
-			else if (b1e > 255)
-				b1e = 255;
-
-			output0->x = (ushort)r0e;
-			output0->y = (ushort)g0e;
-			output0->z = (ushort)b0e;
-			output0->w = 0xFF;
-
-			output1->x = (ushort)r1e;
-			output1->y = (ushort)g1e;
-			output1->z = (ushort)b1e;
-			output1->w = 0xFF;
-
-			return retval;
-		}
-
-
-		int rgb_unpack(const int input[6], int quantization_level, ushort4* output0, ushort4* output1)
-		{
-
-			int ri0b = color_unquantization_tables[quantization_level][input[0]];
-			int ri1b = color_unquantization_tables[quantization_level][input[1]];
-			int gi0b = color_unquantization_tables[quantization_level][input[2]];
-			int gi1b = color_unquantization_tables[quantization_level][input[3]];
-			int bi0b = color_unquantization_tables[quantization_level][input[4]];
-			int bi1b = color_unquantization_tables[quantization_level][input[5]];
-
-			if (ri0b + gi0b + bi0b > ri1b + gi1b + bi1b)
-			{
-				// blue-contraction
-				ri0b = (ri0b + bi0b) >> 1;
-				gi0b = (gi0b + bi0b) >> 1;
-				ri1b = (ri1b + bi1b) >> 1;
-				gi1b = (gi1b + bi1b) >> 1;
-
-				output0->x = (ushort)ri1b;
-				output0->y = (ushort)gi1b;
-				output0->z = (ushort)bi1b;
-				output0->w = 255;
-
-				output1->x = (ushort)ri0b;
-				output1->y = (ushort)gi0b;
-				output1->z = (ushort)bi0b;
-				output1->w = 255;
-				return 1;
-			}
-			else
-			{
-				output0->x = (ushort)ri0b;
-				output0->y = (ushort)gi0b;
-				output0->z = (ushort)bi0b;
-				output0->w = 255;
-
-				output1->x = (ushort)ri1b;
-				output1->y = (ushort)gi1b;
-				output1->z = (ushort)bi1b;
-				output1->w = 255;
-				return 0;
-			}
-		}
-
-
-		void rgba_unpack(const int input[8], int quantization_level, ushort4* output0, ushort4* output1)
-		{
-			int order = rgb_unpack(input, quantization_level, output0, output1);
-			if (order == 0)
-			{
-				output0->w = color_unquantization_tables[quantization_level][input[6]];
-				output1->w = color_unquantization_tables[quantization_level][input[7]];
-			}
-			else
-			{
-				output0->w = color_unquantization_tables[quantization_level][input[7]];
-				output1->w = color_unquantization_tables[quantization_level][input[6]];
-			}
-		}
-
-
-		void rgba_delta_unpack(const int input[8], int quantization_level, ushort4* output0, ushort4* output1)
-		{
-			int a0 = color_unquantization_tables[quantization_level][input[6]];
-			int a1 = color_unquantization_tables[quantization_level][input[7]];
-			a0 |= (a1 & 0x80) << 1;
-			a1 &= 0x7F;
-			if (a1 & 0x40)
-				a1 -= 0x80;
-			a0 >>= 1;
-			a1 >>= 1;
-			a1 += a0;
-
-			if (a1 < 0)
-				a1 = 0;
-			else if (a1 > 255)
-				a1 = 255;
-
-			int order = rgb_delta_unpack(input, quantization_level, output0, output1);
-			if (order == 0)
-			{
-				output0->w = (ushort)a0;
-				output1->w = (ushort)a1;
-			}
-			else
-			{
-				output0->w = (ushort)a1;
-				output1->w = (ushort)a0;
+				}
 			}
 		}
-
-
-		void rgb_scale_unpack(const int input[4], int quantization_level, ushort4* output0, ushort4* output1)
-		{
-			int ir = color_unquantization_tables[quantization_level][input[0]];
-			int ig = color_unquantization_tables[quantization_level][input[1]];
-			int ib = color_unquantization_tables[quantization_level][input[2]];
-
-			int iscale = color_unquantization_tables[quantization_level][input[3]];
-
-			*output1 = ushort4((ushort)ir, (ushort)ig, (ushort)ib, 255);
-			*output0 = ushort4((ushort)((ir * iscale) >> 8), (ushort)((ig * iscale) >> 8), (ushort)((ib * iscale) >> 8), 255);
-		}
-
-
-		void rgb_scale_alpha_unpack(const int input[6], int quantization_level, ushort4* output0, ushort4* output1)
-		{
-			rgb_scale_unpack(input, quantization_level, output0, output1);
-			output0->w = color_unquantization_tables[quantization_level][input[4]];
-			output1->w = color_unquantization_tables[quantization_level][input[5]];
-
-		}
-
-
-		void luminance_unpack(const int input[2], int quantization_level, ushort4* output0, ushort4* output1)
-		{
-			uint8 lum0 = color_unquantization_tables[quantization_level][input[0]];
-			uint8 lum1 = color_unquantization_tables[quantization_level][input[1]];
-			*output0 = ushort4(lum0, lum0, lum0, 255);
-			*output1 = ushort4(lum1, lum1, lum1, 255);
-		}
-
-
-		void luminance_delta_unpack(const int input[2], int quantization_level, ushort4* output0, ushort4* output1)
-		{
-			ushort v0 = color_unquantization_tables[quantization_level][input[0]];
-			ushort v1 = color_unquantization_tables[quantization_level][input[1]];
-			ushort l0 = (v0 >> 2) | (v1 & 0xC0);
-			ushort l1 = l0 + (v1 & 0x3F);
-
-			if (l1 > 255)
-				l1 = 255;
-
-			*output0 = ushort4(l0, l0, l0, 255);
-			*output1 = ushort4(l1, l1, l1, 255);
-		}
-
-
-		void luminance_alpha_unpack(const int input[4], int quantization_level, ushort4* output0, ushort4* output1)
+		int computeNumColorEndpointValues(const deUint32* endpointModes, int numPartitions)
 		{
-			ushort lum0 = color_unquantization_tables[quantization_level][input[0]];
-			ushort lum1 = color_unquantization_tables[quantization_level][input[1]];
-			ushort alpha0 = color_unquantization_tables[quantization_level][input[2]];
-			ushort alpha1 = color_unquantization_tables[quantization_level][input[3]];
-			*output0 = ushort4(lum0, lum0, lum0, alpha0);
-			*output1 = ushort4(lum1, lum1, lum1, alpha1);
+			int result = 0;
+			for (int i = 0; i < numPartitions; i++)
+				result += computeNumColorEndpointValues(endpointModes[i]);
+			return result;
 		}
-
-
-		void luminance_alpha_delta_unpack(const int input[4], int quantization_level, ushort4* output0, ushort4* output1)
-		{
-			ushort lum0 = color_unquantization_tables[quantization_level][input[0]];
-			ushort lum1 = color_unquantization_tables[quantization_level][input[1]];
-			ushort alpha0 = color_unquantization_tables[quantization_level][input[2]];
-			ushort alpha1 = color_unquantization_tables[quantization_level][input[3]];
-
-			lum0 |= (lum1 & 0x80) << 1;
-			alpha0 |= (alpha1 & 0x80) << 1;
-			lum1 &= 0x7F;
-			alpha1 &= 0x7F;
-			if (lum1 & 0x40)
-				lum1 -= 0x80;
-			if (alpha1 & 0x40)
-				alpha1 -= 0x80;
-
-			lum0 >>= 1;
-			lum1 >>= 1;
-			alpha0 >>= 1;
-			alpha1 >>= 1;
-			lum1 += lum0;
-			alpha1 += alpha0;
-
-
-			if (lum1 > 255)
-				lum1 = 255;
-
-
-			if (alpha1 > 255)
-				alpha1 = 255;
-
-			*output0 = ushort4(lum0, lum0, lum0, alpha0);
-			*output1 = ushort4(lum1, lum1, lum1, alpha1);
+		void decodeISETritBlock(ISEDecodedResult* dst, int numValues, BitAccessStream& data, int numBits)
+		{
+			DE_ASSERT(inRange(numValues, 1, 5));
+			deUint32 m[5];
+			m[0] = data.getNext(numBits);
+			deUint32 T01 = data.getNext(2);
+			m[1] = data.getNext(numBits);
+			deUint32 T23 = data.getNext(2);
+			m[2] = data.getNext(numBits);
+			deUint32 T4 = data.getNext(1);
+			m[3] = data.getNext(numBits);
+			deUint32 T56 = data.getNext(2);
+			m[4] = data.getNext(numBits);
+			deUint32 T7 = data.getNext(1);
+			switch (numValues)
+			{
+				// \note Fall-throughs.
+			case 1: T23 = 0;
+			case 2: T4 = 0;
+			case 3: T56 = 0;
+			case 4: T7 = 0;
+			case 5: break;
+			default:
+				DE_ASSERT(false);
+			}
+			const deUint32 T = (T7 << 7) | (T56 << 5) | (T4 << 4) | (T23 << 2) | (T01 << 0);
+			static const deUint32 tritsFromT[256][5] =
+			{
+				{ 0,0,0,0,0 }, { 1,0,0,0,0 }, { 2,0,0,0,0 }, { 0,0,2,0,0 }, { 0,1,0,0,0 }, { 1,1,0,0,0 }, { 2,1,0,0,0 }, { 1,0,2,0,0 }, { 0,2,0,0,0 }, { 1,2,0,0,0 }, { 2,2,0,0,0 }, { 2,0,2,0,0 }, { 0,2,2,0,0 }, { 1,2,2,0,0 }, { 2,2,2,0,0 }, { 2,0,2,0,0 },
+				{ 0,0,1,0,0 }, { 1,0,1,0,0 }, { 2,0,1,0,0 }, { 0,1,2,0,0 }, { 0,1,1,0,0 }, { 1,1,1,0,0 }, { 2,1,1,0,0 }, { 1,1,2,0,0 }, { 0,2,1,0,0 }, { 1,2,1,0,0 }, { 2,2,1,0,0 }, { 2,1,2,0,0 }, { 0,0,0,2,2 }, { 1,0,0,2,2 }, { 2,0,0,2,2 }, { 0,0,2,2,2 },
+				{ 0,0,0,1,0 }, { 1,0,0,1,0 }, { 2,0,0,1,0 }, { 0,0,2,1,0 }, { 0,1,0,1,0 }, { 1,1,0,1,0 }, { 2,1,0,1,0 }, { 1,0,2,1,0 }, { 0,2,0,1,0 }, { 1,2,0,1,0 }, { 2,2,0,1,0 }, { 2,0,2,1,0 }, { 0,2,2,1,0 }, { 1,2,2,1,0 }, { 2,2,2,1,0 }, { 2,0,2,1,0 },
+				{ 0,0,1,1,0 }, { 1,0,1,1,0 }, { 2,0,1,1,0 }, { 0,1,2,1,0 }, { 0,1,1,1,0 }, { 1,1,1,1,0 }, { 2,1,1,1,0 }, { 1,1,2,1,0 }, { 0,2,1,1,0 }, { 1,2,1,1,0 }, { 2,2,1,1,0 }, { 2,1,2,1,0 }, { 0,1,0,2,2 }, { 1,1,0,2,2 }, { 2,1,0,2,2 }, { 1,0,2,2,2 },
+				{ 0,0,0,2,0 }, { 1,0,0,2,0 }, { 2,0,0,2,0 }, { 0,0,2,2,0 }, { 0,1,0,2,0 }, { 1,1,0,2,0 }, { 2,1,0,2,0 }, { 1,0,2,2,0 }, { 0,2,0,2,0 }, { 1,2,0,2,0 }, { 2,2,0,2,0 }, { 2,0,2,2,0 }, { 0,2,2,2,0 }, { 1,2,2,2,0 }, { 2,2,2,2,0 }, { 2,0,2,2,0 },
+				{ 0,0,1,2,0 }, { 1,0,1,2,0 }, { 2,0,1,2,0 }, { 0,1,2,2,0 }, { 0,1,1,2,0 }, { 1,1,1,2,0 }, { 2,1,1,2,0 }, { 1,1,2,2,0 }, { 0,2,1,2,0 }, { 1,2,1,2,0 }, { 2,2,1,2,0 }, { 2,1,2,2,0 }, { 0,2,0,2,2 }, { 1,2,0,2,2 }, { 2,2,0,2,2 }, { 2,0,2,2,2 },
+				{ 0,0,0,0,2 }, { 1,0,0,0,2 }, { 2,0,0,0,2 }, { 0,0,2,0,2 }, { 0,1,0,0,2 }, { 1,1,0,0,2 }, { 2,1,0,0,2 }, { 1,0,2,0,2 }, { 0,2,0,0,2 }, { 1,2,0,0,2 }, { 2,2,0,0,2 }, { 2,0,2,0,2 }, { 0,2,2,0,2 }, { 1,2,2,0,2 }, { 2,2,2,0,2 }, { 2,0,2,0,2 },
+				{ 0,0,1,0,2 }, { 1,0,1,0,2 }, { 2,0,1,0,2 }, { 0,1,2,0,2 }, { 0,1,1,0,2 }, { 1,1,1,0,2 }, { 2,1,1,0,2 }, { 1,1,2,0,2 }, { 0,2,1,0,2 }, { 1,2,1,0,2 }, { 2,2,1,0,2 }, { 2,1,2,0,2 }, { 0,2,2,2,2 }, { 1,2,2,2,2 }, { 2,2,2,2,2 }, { 2,0,2,2,2 },
+				{ 0,0,0,0,1 }, { 1,0,0,0,1 }, { 2,0,0,0,1 }, { 0,0,2,0,1 }, { 0,1,0,0,1 }, { 1,1,0,0,1 }, { 2,1,0,0,1 }, { 1,0,2,0,1 }, { 0,2,0,0,1 }, { 1,2,0,0,1 }, { 2,2,0,0,1 }, { 2,0,2,0,1 }, { 0,2,2,0,1 }, { 1,2,2,0,1 }, { 2,2,2,0,1 }, { 2,0,2,0,1 },
+				{ 0,0,1,0,1 }, { 1,0,1,0,1 }, { 2,0,1,0,1 }, { 0,1,2,0,1 }, { 0,1,1,0,1 }, { 1,1,1,0,1 }, { 2,1,1,0,1 }, { 1,1,2,0,1 }, { 0,2,1,0,1 }, { 1,2,1,0,1 }, { 2,2,1,0,1 }, { 2,1,2,0,1 }, { 0,0,1,2,2 }, { 1,0,1,2,2 }, { 2,0,1,2,2 }, { 0,1,2,2,2 },
+				{ 0,0,0,1,1 }, { 1,0,0,1,1 }, { 2,0,0,1,1 }, { 0,0,2,1,1 }, { 0,1,0,1,1 }, { 1,1,0,1,1 }, { 2,1,0,1,1 }, { 1,0,2,1,1 }, { 0,2,0,1,1 }, { 1,2,0,1,1 }, { 2,2,0,1,1 }, { 2,0,2,1,1 }, { 0,2,2,1,1 }, { 1,2,2,1,1 }, { 2,2,2,1,1 }, { 2,0,2,1,1 },
+				{ 0,0,1,1,1 }, { 1,0,1,1,1 }, { 2,0,1,1,1 }, { 0,1,2,1,1 }, { 0,1,1,1,1 }, { 1,1,1,1,1 }, { 2,1,1,1,1 }, { 1,1,2,1,1 }, { 0,2,1,1,1 }, { 1,2,1,1,1 }, { 2,2,1,1,1 }, { 2,1,2,1,1 }, { 0,1,1,2,2 }, { 1,1,1,2,2 }, { 2,1,1,2,2 }, { 1,1,2,2,2 },
+				{ 0,0,0,2,1 }, { 1,0,0,2,1 }, { 2,0,0,2,1 }, { 0,0,2,2,1 }, { 0,1,0,2,1 }, { 1,1,0,2,1 }, { 2,1,0,2,1 }, { 1,0,2,2,1 }, { 0,2,0,2,1 }, { 1,2,0,2,1 }, { 2,2,0,2,1 }, { 2,0,2,2,1 }, { 0,2,2,2,1 }, { 1,2,2,2,1 }, { 2,2,2,2,1 }, { 2,0,2,2,1 },
+				{ 0,0,1,2,1 }, { 1,0,1,2,1 }, { 2,0,1,2,1 }, { 0,1,2,2,1 }, { 0,1,1,2,1 }, { 1,1,1,2,1 }, { 2,1,1,2,1 }, { 1,1,2,2,1 }, { 0,2,1,2,1 }, { 1,2,1,2,1 }, { 2,2,1,2,1 }, { 2,1,2,2,1 }, { 0,2,1,2,2 }, { 1,2,1,2,2 }, { 2,2,1,2,2 }, { 2,1,2,2,2 },
+				{ 0,0,0,1,2 }, { 1,0,0,1,2 }, { 2,0,0,1,2 }, { 0,0,2,1,2 }, { 0,1,0,1,2 }, { 1,1,0,1,2 }, { 2,1,0,1,2 }, { 1,0,2,1,2 }, { 0,2,0,1,2 }, { 1,2,0,1,2 }, { 2,2,0,1,2 }, { 2,0,2,1,2 }, { 0,2,2,1,2 }, { 1,2,2,1,2 }, { 2,2,2,1,2 }, { 2,0,2,1,2 },
+				{ 0,0,1,1,2 }, { 1,0,1,1,2 }, { 2,0,1,1,2 }, { 0,1,2,1,2 }, { 0,1,1,1,2 }, { 1,1,1,1,2 }, { 2,1,1,1,2 }, { 1,1,2,1,2 }, { 0,2,1,1,2 }, { 1,2,1,1,2 }, { 2,2,1,1,2 }, { 2,1,2,1,2 }, { 0,2,2,2,2 }, { 1,2,2,2,2 }, { 2,2,2,2,2 }, { 2,1,2,2,2 }
+			};
+			const deUint32(&trits)[5] = tritsFromT[T];
+			for (int i = 0; i < numValues; i++)
+			{
+				dst[i].m = m[i];
+				dst[i].tq = trits[i];
+				dst[i].v = (trits[i] << numBits) + m[i];
+			}
+		}
+		void decodeISEQuintBlock(ISEDecodedResult* dst, int numValues, BitAccessStream& data, int numBits)
+		{
+			DE_ASSERT(inRange(numValues, 1, 3));
+			deUint32 m[3];
+			m[0] = data.getNext(numBits);
+			deUint32 Q012 = data.getNext(3);
+			m[1] = data.getNext(numBits);
+			deUint32 Q34 = data.getNext(2);
+			m[2] = data.getNext(numBits);
+			deUint32 Q56 = data.getNext(2);
+			switch (numValues)
+			{
+				// \note Fall-throughs.
+			case 1: Q34 = 0;
+			case 2: Q56 = 0;
+			case 3: break;
+			default:
+				DE_ASSERT(false);
+			}
+			const deUint32 Q = (Q56 << 5) | (Q34 << 3) | (Q012 << 0);
+			static const deUint32 quintsFromQ[256][3] =
+			{
+				{ 0,0,0 }, { 1,0,0 }, { 2,0,0 }, { 3,0,0 }, { 4,0,0 }, { 0,4,0 }, { 4,4,0 }, { 4,4,4 }, { 0,1,0 }, { 1,1,0 }, { 2,1,0 }, { 3,1,0 }, { 4,1,0 }, { 1,4,0 }, { 4,4,1 }, { 4,4,4 },
+				{ 0,2,0 }, { 1,2,0 }, { 2,2,0 }, { 3,2,0 }, { 4,2,0 }, { 2,4,0 }, { 4,4,2 }, { 4,4,4 }, { 0,3,0 }, { 1,3,0 }, { 2,3,0 }, { 3,3,0 }, { 4,3,0 }, { 3,4,0 }, { 4,4,3 }, { 4,4,4 },
+				{ 0,0,1 }, { 1,0,1 }, { 2,0,1 }, { 3,0,1 }, { 4,0,1 }, { 0,4,1 }, { 4,0,4 }, { 0,4,4 }, { 0,1,1 }, { 1,1,1 }, { 2,1,1 }, { 3,1,1 }, { 4,1,1 }, { 1,4,1 }, { 4,1,4 }, { 1,4,4 },
+				{ 0,2,1 }, { 1,2,1 }, { 2,2,1 }, { 3,2,1 }, { 4,2,1 }, { 2,4,1 }, { 4,2,4 }, { 2,4,4 }, { 0,3,1 }, { 1,3,1 }, { 2,3,1 }, { 3,3,1 }, { 4,3,1 }, { 3,4,1 }, { 4,3,4 }, { 3,4,4 },
+				{ 0,0,2 }, { 1,0,2 }, { 2,0,2 }, { 3,0,2 }, { 4,0,2 }, { 0,4,2 }, { 2,0,4 }, { 3,0,4 }, { 0,1,2 }, { 1,1,2 }, { 2,1,2 }, { 3,1,2 }, { 4,1,2 }, { 1,4,2 }, { 2,1,4 }, { 3,1,4 },
+				{ 0,2,2 }, { 1,2,2 }, { 2,2,2 }, { 3,2,2 }, { 4,2,2 }, { 2,4,2 }, { 2,2,4 }, { 3,2,4 }, { 0,3,2 }, { 1,3,2 }, { 2,3,2 }, { 3,3,2 }, { 4,3,2 }, { 3,4,2 }, { 2,3,4 }, { 3,3,4 },
+				{ 0,0,3 }, { 1,0,3 }, { 2,0,3 }, { 3,0,3 }, { 4,0,3 }, { 0,4,3 }, { 0,0,4 }, { 1,0,4 }, { 0,1,3 }, { 1,1,3 }, { 2,1,3 }, { 3,1,3 }, { 4,1,3 }, { 1,4,3 }, { 0,1,4 }, { 1,1,4 },
+				{ 0,2,3 }, { 1,2,3 }, { 2,2,3 }, { 3,2,3 }, { 4,2,3 }, { 2,4,3 }, { 0,2,4 }, { 1,2,4 }, { 0,3,3 }, { 1,3,3 }, { 2,3,3 }, { 3,3,3 }, { 4,3,3 }, { 3,4,3 }, { 0,3,4 }, { 1,3,4 }
+			};
+			const deUint32(&quints)[3] = quintsFromQ[Q];
+			for (int i = 0; i < numValues; i++)
+			{
+				dst[i].m = m[i];
+				dst[i].tq = quints[i];
+				dst[i].v = (quints[i] << numBits) + m[i];
+			}
 		}
-
-
-		// RGB-offset format
-		void hdr_rgbo_unpack3(const int input[4], int quantization_level, ushort4* output0, ushort4* output1)
+		inline void decodeISEBitBlock(ISEDecodedResult* dst, BitAccessStream& data, int numBits)
 		{
-			int v0 = color_unquantization_tables[quantization_level][input[0]];
-			int v1 = color_unquantization_tables[quantization_level][input[1]];
-			int v2 = color_unquantization_tables[quantization_level][input[2]];
-			int v3 = color_unquantization_tables[quantization_level][input[3]];
-
-			int modeval = ((v0 & 0xC0) >> 6) | (((v1 & 0x80) >> 7) << 2) | (((v2 & 0x80) >> 7) << 3);
-
-			int majcomp;
-			int mode;
-			if ((modeval & 0xC) != 0xC)
+			dst[0].m = data.getNext(numBits);
+			dst[0].v = dst[0].m;
+		}
+		void decodeISE(ISEDecodedResult* dst, int numValues, BitAccessStream& data, const ISEParams& params)
+		{
+			if (params.mode == ISEMODE_TRIT)
 			{
-				majcomp = modeval >> 2;
-				mode = modeval & 3;
+				const int numBlocks = deDivRoundUp32(numValues, 5);
+				for (int blockNdx = 0; blockNdx < numBlocks; blockNdx++)
+				{
+					const int numValuesInBlock = blockNdx == numBlocks - 1 ? numValues - 5 * (numBlocks - 1) : 5;
+					decodeISETritBlock(&dst[5 * blockNdx], numValuesInBlock, data, params.numBits);
+				}
 			}
-			else if (modeval != 0xF)
+			else if (params.mode == ISEMODE_QUINT)
 			{
-				majcomp = modeval & 3;
-				mode = 4;
+				const int numBlocks = deDivRoundUp32(numValues, 3);
+				for (int blockNdx = 0; blockNdx < numBlocks; blockNdx++)
+				{
+					const int numValuesInBlock = blockNdx == numBlocks - 1 ? numValues - 3 * (numBlocks - 1) : 3;
+					decodeISEQuintBlock(&dst[3 * blockNdx], numValuesInBlock, data, params.numBits);
+				}
 			}
 			else
 			{
-				majcomp = 0;
-				mode = 5;
-			}
-
-			int red = v0 & 0x3F;
-			int green = v1 & 0x1F;
-			int blue = v2 & 0x1F;
-			int scale = v3 & 0x1F;
-
-			int bit0 = (v1 >> 6) & 1;
-			int bit1 = (v1 >> 5) & 1;
-			int bit2 = (v2 >> 6) & 1;
-			int bit3 = (v2 >> 5) & 1;
-			int bit4 = (v3 >> 7) & 1;
-			int bit5 = (v3 >> 6) & 1;
-			int bit6 = (v3 >> 5) & 1;
-
-			int ohcomp = 1 << mode;
-
-			if (ohcomp & 0x30)
-				green |= bit0 << 6;
-			if (ohcomp & 0x3A)
-				green |= bit1 << 5;
-			if (ohcomp & 0x30)
-				blue |= bit2 << 6;
-			if (ohcomp & 0x3A)
-				blue |= bit3 << 5;
-
-			if (ohcomp & 0x3D)
-				scale |= bit6 << 5;
-			if (ohcomp & 0x2D)
-				scale |= bit5 << 6;
-			if (ohcomp & 0x04)
-				scale |= bit4 << 7;
-
-			if (ohcomp & 0x3B)
-				red |= bit4 << 6;
-			if (ohcomp & 0x04)
-				red |= bit3 << 6;
-
-			if (ohcomp & 0x10)
-				red |= bit5 << 7;
-			if (ohcomp & 0x0F)
-				red |= bit2 << 7;
-
-			if (ohcomp & 0x05)
-				red |= bit1 << 8;
-			if (ohcomp & 0x0A)
-				red |= bit0 << 8;
-
-			if (ohcomp & 0x05)
-				red |= bit0 << 9;
-			if (ohcomp & 0x02)
-				red |= bit6 << 9;
-
-			if (ohcomp & 0x01)
-				red |= bit3 << 10;
-			if (ohcomp & 0x02)
-				red |= bit5 << 10;
-
-
-			// expand to 12 bits.
-			static const int shamts[6] = { 1, 1, 2, 3, 4, 5 };
-			int shamt = shamts[mode];
-			red <<= shamt;
-			green <<= shamt;
-			blue <<= shamt;
-			scale <<= shamt;
-
-			// on modes 0 to 4, the values stored for "green" and "blue" are differentials,
-			// not absolute values.
-			if (mode != 5)
-			{
-				green = red - green;
-				blue = red - blue;
-			}
-
-			// switch around components.
-			int temp;
-			switch (majcomp)
-			{
-			case 1:
-				temp = red;
-				red = green;
-				green = temp;
-				break;
-			case 2:
-				temp = red;
-				red = blue;
-				blue = temp;
-				break;
-			default:
-				break;
+				DE_ASSERT(params.mode == ISEMODE_PLAIN_BIT);
+				for (int i = 0; i < numValues; i++)
+					decodeISEBitBlock(&dst[i], data, params.numBits);
 			}
-
-
-			int red0 = red - scale;
-			int green0 = green - scale;
-			int blue0 = blue - scale;
-
-			// clamp to [0,0xFFF].
-			if (red < 0)
-				red = 0;
-			if (green < 0)
-				green = 0;
-			if (blue < 0)
-				blue = 0;
-
-			if (red0 < 0)
-				red0 = 0;
-			if (green0 < 0)
-				green0 = 0;
-			if (blue0 < 0)
-				blue0 = 0;
-
-			*output0 = ushort4(ushort(red0 << 4), ushort(green0 << 4), ushort(blue0 << 4), 0x7800);
-			*output1 = ushort4(ushort(red << 4), ushort(green << 4), ushort(blue << 4), 0x7800);
 		}
-
-
-		void hdr_rgb_unpack3(const int input[6], int quantization_level, ushort4* output0, ushort4* output1)
+		void unquantizeColorEndpoints(deUint32* dst, const ISEDecodedResult* iseResults, int numEndpoints, const ISEParams& iseParams)
 		{
-
-			int v0 = color_unquantization_tables[quantization_level][input[0]];
-			int v1 = color_unquantization_tables[quantization_level][input[1]];
-			int v2 = color_unquantization_tables[quantization_level][input[2]];
-			int v3 = color_unquantization_tables[quantization_level][input[3]];
-			int v4 = color_unquantization_tables[quantization_level][input[4]];
-			int v5 = color_unquantization_tables[quantization_level][input[5]];
-
-			// extract all the fixed-placement bitfields
-			int modeval = ((v1 & 0x80) >> 7) | (((v2 & 0x80) >> 7) << 1) | (((v3 & 0x80) >> 7) << 2);
-
-			int majcomp = ((v4 & 0x80) >> 7) | (((v5 & 0x80) >> 7) << 1);
-
-			if (majcomp == 3)
+			if (iseParams.mode == ISEMODE_TRIT || iseParams.mode == ISEMODE_QUINT)
 			{
-				*output0 = ushort4(ushort(v0 << 8), ushort(v2 << 8), ushort(v4 & 0x7F) << 9, 0x7800);
-				*output1 = ushort4(ushort(v1 << 8), ushort(v3 << 8), ushort(v5 & 0x7F) << 9, 0x7800);
-				return;
-			}
-
-			int a = v0 | ((v1 & 0x40) << 2);
-			int b0 = v2 & 0x3f;
-			int b1 = v3 & 0x3f;
-			int c = v1 & 0x3f;
-			int d0 = v4 & 0x7f;
-			int d1 = v5 & 0x7f;
-
-			// get hold of the number of bits in 'd0' and 'd1'
-			static const int dbits_tab[8] = { 7, 6, 7, 6, 5, 6, 5, 6 };
-			int dbits = dbits_tab[modeval];
-
-			// extract six variable-placement bits
-			int bit0 = (v2 >> 6) & 1;
-			int bit1 = (v3 >> 6) & 1;
-
-			int bit2 = (v4 >> 6) & 1;
-			int bit3 = (v5 >> 6) & 1;
-			int bit4 = (v4 >> 5) & 1;
-			int bit5 = (v5 >> 5) & 1;
-
-
-			// and prepend the variable-placement bits depending on mode.
-			int ohmod = 1 << modeval;	// one-hot-mode
-			if (ohmod & 0xA4)
-				a |= bit0 << 9;
-			if (ohmod & 0x8)
-				a |= bit2 << 9;
-			if (ohmod & 0x50)
-				a |= bit4 << 9;
-
-			if (ohmod & 0x50)
-				a |= bit5 << 10;
-			if (ohmod & 0xA0)
-				a |= bit1 << 10;
-
-			if (ohmod & 0xC0)
-				a |= bit2 << 11;
-
-			if (ohmod & 0x4)
-				c |= bit1 << 6;
-			if (ohmod & 0xE8)
-				c |= bit3 << 6;
-
-			if (ohmod & 0x20)
-				c |= bit2 << 7;
-
-
-			if (ohmod & 0x5B)
-				b0 |= bit0 << 6;
-			if (ohmod & 0x5B)
-				b1 |= bit1 << 6;
-
-			if (ohmod & 0x12)
-				b0 |= bit2 << 7;
-			if (ohmod & 0x12)
-				b1 |= bit3 << 7;
-
-			if (ohmod & 0xAF)
-				d0 |= bit4 << 5;
-			if (ohmod & 0xAF)
-				d1 |= bit5 << 5;
-			if (ohmod & 0x5)
-				d0 |= bit2 << 6;
-			if (ohmod & 0x5)
-				d1 |= bit3 << 6;
-
-			// sign-extend 'd0' and 'd1'
-			// note: this code assumes that signed right-shift actually sign-fills, not zero-fills.
-			int32 d0x = d0;
-			int32 d1x = d1;
-			int sx_shamt = 32 - dbits;
-			d0x <<= sx_shamt;
-			d0x >>= sx_shamt;
-			d1x <<= sx_shamt;
-			d1x >>= sx_shamt;
-			d0 = d0x;
-			d1 = d1x;
-
-			// expand all values to 12 bits, with left-shift as needed.
-			int val_shamt = (modeval >> 1) ^ 3;
-			a <<= val_shamt;
-			b0 <<= val_shamt;
-			b1 <<= val_shamt;
-			c <<= val_shamt;
-			d0 <<= val_shamt;
-			d1 <<= val_shamt;
-
-			// then compute the actual color values.
-			int red1 = a;
-			int green1 = a - b0;
-			int blue1 = a - b1;
-			int red0 = a - c;
-			int green0 = a - b0 - c - d0;
-			int blue0 = a - b1 - c - d1;
-
-			// clamp the color components to [0,2^12 - 1]
-			if (red0 < 0)
-				red0 = 0;
-			else if (red0 > 0xFFF)
-				red0 = 0xFFF;
-
-			if (green0 < 0)
-				green0 = 0;
-			else if (green0 > 0xFFF)
-				green0 = 0xFFF;
-
-			if (blue0 < 0)
-				blue0 = 0;
-			else if (blue0 > 0xFFF)
-				blue0 = 0xFFF;
-
-			if (red1 < 0)
-				red1 = 0;
-			else if (red1 > 0xFFF)
-				red1 = 0xFFF;
-
-			if (green1 < 0)
-				green1 = 0;
-			else if (green1 > 0xFFF)
-				green1 = 0xFFF;
-
-			if (blue1 < 0)
-				blue1 = 0;
-			else if (blue1 > 0xFFF)
-				blue1 = 0xFFF;
-
-
-			// switch around the color components
-			int temp0, temp1;
-			switch (majcomp)
-			{
-			case 1:					// switch around red and green
-				temp0 = red0;
-				temp1 = red1;
-				red0 = green0;
-				red1 = green1;
-				green0 = temp0;
-				green1 = temp1;
-				break;
-			case 2:					// switch around red and blue
-				temp0 = red0;
-				temp1 = red1;
-				red0 = blue0;
-				red1 = blue1;
-				blue0 = temp0;
-				blue1 = temp1;
-				break;
-			case 0:					// no switch
-				break;
-			}
-
-			*output0 = ushort4(ushort(red0 << 4), ushort(green0 << 4), ushort(blue0 << 4), 0x7800);
-			*output1 = ushort4(ushort(red1 << 4), ushort(green1 << 4), ushort(blue1 << 4), 0x7800);
-		}
-
-
-		void hdr_rgb_ldr_alpha_unpack3(const int input[8], int quantization_level, ushort4* output0, ushort4* output1)
-		{
-			hdr_rgb_unpack3(input, quantization_level, output0, output1);
-
-			uint8 v6 = color_unquantization_tables[quantization_level][input[6]];
-			uint8 v7 = color_unquantization_tables[quantization_level][input[7]];
-			output0->w = v6;
-			output1->w = v7;
-		}
-
-
-		void hdr_luminance_small_range_unpack(const int input[2], int quantization_level, ushort4* output0, ushort4* output1)
-		{
-			int v0 = color_unquantization_tables[quantization_level][input[0]];
-			int v1 = color_unquantization_tables[quantization_level][input[1]];
-
-			int y0, y1;
-			if (v0 & 0x80)
-			{
-				y0 = ((v1 & 0xE0) << 4) | ((v0 & 0x7F) << 2);
-				y1 = (v1 & 0x1F) << 2;
+				const int rangeCase = iseParams.numBits * 2 - (iseParams.mode == ISEMODE_TRIT ? 2 : 1);
+				DE_ASSERT(inRange(rangeCase, 0, 10));
+				static const deUint32	Ca[11] = { 204, 113, 93, 54, 44, 26, 22, 13, 11, 6, 5 };
+				const deUint32			C = Ca[rangeCase];
+				for (int endpointNdx = 0; endpointNdx < numEndpoints; endpointNdx++)
+				{
+					const deUint32 a = getBit(iseResults[endpointNdx].m, 0);
+					const deUint32 b = getBit(iseResults[endpointNdx].m, 1);
+					const deUint32 c = getBit(iseResults[endpointNdx].m, 2);
+					const deUint32 d = getBit(iseResults[endpointNdx].m, 3);
+					const deUint32 e = getBit(iseResults[endpointNdx].m, 4);
+					const deUint32 f = getBit(iseResults[endpointNdx].m, 5);
+					const deUint32 A = a == 0 ? 0 : (1 << 9) - 1;
+					const deUint32 B = rangeCase == 0 ? 0
+						: rangeCase == 1 ? 0
+						: rangeCase == 2 ? (b << 8) | (b << 4) | (b << 2) | (b << 1)
+						: rangeCase == 3 ? (b << 8) | (b << 3) | (b << 2)
+						: rangeCase == 4 ? (c << 8) | (b << 7) | (c << 3) | (b << 2) | (c << 1) | (b << 0)
+						: rangeCase == 5 ? (c << 8) | (b << 7) | (c << 2) | (b << 1) | (c << 0)
+						: rangeCase == 6 ? (d << 8) | (c << 7) | (b << 6) | (d << 2) | (c << 1) | (b << 0)
+						: rangeCase == 7 ? (d << 8) | (c << 7) | (b << 6) | (d << 1) | (c << 0)
+						: rangeCase == 8 ? (e << 8) | (d << 7) | (c << 6) | (b << 5) | (e << 1) | (d << 0)
+						: rangeCase == 9 ? (e << 8) | (d << 7) | (c << 6) | (b << 5) | (e << 0)
+						: rangeCase == 10 ? (f << 8) | (e << 7) | (d << 6) | (c << 5) | (b << 4) | (f << 0)
+						: (deUint32)-1;
+					DE_ASSERT(B != (deUint32)-1);
+					dst[endpointNdx] = (((iseResults[endpointNdx].tq * C + B) ^ A) >> 2) | (A & 0x80);
+				}
 			}
 			else
 			{
-				y0 = ((v1 & 0xF0) << 4) | ((v0 & 0x7F) << 1);
-				y1 = (v1 & 0xF) << 1;
+				DE_ASSERT(iseParams.mode == ISEMODE_PLAIN_BIT);
+				for (int endpointNdx = 0; endpointNdx < numEndpoints; endpointNdx++)
+					dst[endpointNdx] = bitReplicationScale(iseResults[endpointNdx].v, iseParams.numBits, 8);
+			}
+		}
+		inline void bitTransferSigned(deInt32& a, deInt32& b)
+		{
+			b >>= 1;
+			b |= a & 0x80;
+			a >>= 1;
+			a &= 0x3f;
+			if (isBitSet(a, 5))
+				a -= 0x40;
+		}
+		inline UVec4 clampedRGBA(const IVec4& rgba)
+		{
+			return UVec4(FMath::Clamp(rgba.x(), 0, 0xff),
+				FMath::Clamp(rgba.y(), 0, 0xff),
+				FMath::Clamp(rgba.z(), 0, 0xff),
+				FMath::Clamp(rgba.w(), 0, 0xff));
+		}
+		inline IVec4 blueContract(int r, int g, int b, int a)
+		{
+			return IVec4((r + b) >> 1, (g + b) >> 1, b, a);
+		}
+		inline bool isColorEndpointModeHDR(deUint32 mode)
+		{
+			return mode == 2 ||
+				mode == 3 ||
+				mode == 7 ||
+				mode == 11 ||
+				mode == 14 ||
+				mode == 15;
+		}
+		void decodeHDREndpointMode7(UVec4& e0, UVec4& e1, deUint32 v0, deUint32 v1, deUint32 v2, deUint32 v3)
+		{
+			const deUint32 m10 = getBit(v1, 7) | (getBit(v2, 7) << 1);
+			const deUint32 m23 = getBits(v0, 6, 7);
+			const deUint32 majComp = m10 != 3 ? m10
+				: m23 != 3 ? m23
+				: 0;
+			const deUint32 mode = m10 != 3 ? m23
+				: m23 != 3 ? 4
+				: 5;
+			deInt32			red = (deInt32)getBits(v0, 0, 5);
+			deInt32			green = (deInt32)getBits(v1, 0, 4);
+			deInt32			blue = (deInt32)getBits(v2, 0, 4);
+			deInt32			scale = (deInt32)getBits(v3, 0, 4);
+			{
+#define SHOR(DST_VAR, SHIFT, BIT_VAR) (DST_VAR) |= (BIT_VAR) << (SHIFT)
+#define ASSIGN_X_BITS(V0,S0, V1,S1, V2,S2, V3,S3, V4,S4, V5,S5, V6,S6) do { SHOR(V0,S0,x0); SHOR(V1,S1,x1); SHOR(V2,S2,x2); SHOR(V3,S3,x3); SHOR(V4,S4,x4); SHOR(V5,S5,x5); SHOR(V6,S6,x6); } while (false)
+				const deUint32	x0 = getBit(v1, 6);
+				const deUint32	x1 = getBit(v1, 5);
+				const deUint32	x2 = getBit(v2, 6);
+				const deUint32	x3 = getBit(v2, 5);
+				const deUint32	x4 = getBit(v3, 7);
+				const deUint32	x5 = getBit(v3, 6);
+				const deUint32	x6 = getBit(v3, 5);
+				deInt32& R = red;
+				deInt32& G = green;
+				deInt32& B = blue;
+				deInt32& S = scale;
+				switch (mode)
+				{
+				case 0: ASSIGN_X_BITS(R, 9, R, 8, R, 7, R, 10, R, 6, S, 6, S, 5); break;
+				case 1: ASSIGN_X_BITS(R, 8, G, 5, R, 7, B, 5, R, 6, R, 10, R, 9); break;
+				case 2: ASSIGN_X_BITS(R, 9, R, 8, R, 7, R, 6, S, 7, S, 6, S, 5); break;
+				case 3: ASSIGN_X_BITS(R, 8, G, 5, R, 7, B, 5, R, 6, S, 6, S, 5); break;
+				case 4: ASSIGN_X_BITS(G, 6, G, 5, B, 6, B, 5, R, 6, R, 7, S, 5); break;
+				case 5: ASSIGN_X_BITS(G, 6, G, 5, B, 6, B, 5, R, 6, S, 6, S, 5); break;
+				default:
+					DE_ASSERT(false);
+				}
+#undef ASSIGN_X_BITS
+#undef SHOR
+			}
+			static const int shiftAmounts[] = { 1, 1, 2, 3, 4, 5 };
+			DE_ASSERT(mode < DE_LENGTH_OF_ARRAY(shiftAmounts));
+			red <<= shiftAmounts[mode];
+			green <<= shiftAmounts[mode];
+			blue <<= shiftAmounts[mode];
+			scale <<= shiftAmounts[mode];
+			if (mode != 5)
+			{
+				green = red - green;
+				blue = red - blue;
 			}
-
-			y1 += y0;
-			if (y1 > 0xFFF)
-				y1 = 0xFFF;
-
-			*output0 = ushort4(ushort(y0 << 4), ushort(y0 << 4), ushort(y0 << 4), 0x7800);
-			*output1 = ushort4(ushort(y1 << 4), ushort(y1 << 4), ushort(y1 << 4), 0x7800);
+			if (majComp == 1)
+				std::swap(red, green);
+			else if (majComp == 2)
+				std::swap(red, blue);
+			e0 = UVec4(FMath::Clamp(red - scale, 0, 0xfff),
+				FMath::Clamp(green - scale, 0, 0xfff),
+				FMath::Clamp(blue - scale, 0, 0xfff),
+				0x780);
+			e1 = UVec4(FMath::Clamp(red, 0, 0xfff),
+				FMath::Clamp(green, 0, 0xfff),
+				FMath::Clamp(blue, 0, 0xfff),
+				0x780);
 		}
-
-
-		void hdr_luminance_large_range_unpack(const int input[2], int quantization_level, ushort4* output0, ushort4* output1)
+		void decodeHDREndpointMode11(UVec4& e0, UVec4& e1, deUint32 v0, deUint32 v1, deUint32 v2, deUint32 v3, deUint32 v4, deUint32 v5)
 		{
-			int v0 = color_unquantization_tables[quantization_level][input[0]];
-			int v1 = color_unquantization_tables[quantization_level][input[1]];
-
-			int y0, y1;
-			if (v1 >= v0)
+			const deUint32 major = (getBit(v5, 7) << 1) | getBit(v4, 7);
+			if (major == 3)
 			{
-				y0 = v0 << 4;
-				y1 = v1 << 4;
+				e0 = UVec4(v0 << 4, v2 << 4, getBits(v4, 0, 6) << 5, 0x780);
+				e1 = UVec4(v1 << 4, v3 << 4, getBits(v5, 0, 6) << 5, 0x780);
 			}
 			else
 			{
-				y0 = (v1 << 4) + 8;
-				y1 = (v0 << 4) - 8;
+				const deUint32 mode = (getBit(v3, 7) << 2) | (getBit(v2, 7) << 1) | getBit(v1, 7);
+				deInt32 a = (deInt32)((getBit(v1, 6) << 8) | v0);
+				deInt32 c = (deInt32)(getBits(v1, 0, 5));
+				deInt32 b0 = (deInt32)(getBits(v2, 0, 5));
+				deInt32 b1 = (deInt32)(getBits(v3, 0, 5));
+				deInt32 d0 = (deInt32)(getBits(v4, 0, 4));
+				deInt32 d1 = (deInt32)(getBits(v5, 0, 4));
+				{
+#define SHOR(DST_VAR, SHIFT, BIT_VAR) (DST_VAR) |= (BIT_VAR) << (SHIFT)
+#define ASSIGN_X_BITS(V0,S0, V1,S1, V2,S2, V3,S3, V4,S4, V5,S5) do { SHOR(V0,S0,x0); SHOR(V1,S1,x1); SHOR(V2,S2,x2); SHOR(V3,S3,x3); SHOR(V4,S4,x4); SHOR(V5,S5,x5); } while (false)
+					const deUint32 x0 = getBit(v2, 6);
+					const deUint32 x1 = getBit(v3, 6);
+					const deUint32 x2 = getBit(v4, 6);
+					const deUint32 x3 = getBit(v5, 6);
+					const deUint32 x4 = getBit(v4, 5);
+					const deUint32 x5 = getBit(v5, 5);
+					switch (mode)
+					{
+					case 0: ASSIGN_X_BITS(b0, 6, b1, 6, d0, 6, d1, 6, d0, 5, d1, 5); break;
+					case 1: ASSIGN_X_BITS(b0, 6, b1, 6, b0, 7, b1, 7, d0, 5, d1, 5); break;
+					case 2: ASSIGN_X_BITS(a, 9, c, 6, d0, 6, d1, 6, d0, 5, d1, 5); break;
+					case 3: ASSIGN_X_BITS(b0, 6, b1, 6, a, 9, c, 6, d0, 5, d1, 5); break;
+					case 4: ASSIGN_X_BITS(b0, 6, b1, 6, b0, 7, b1, 7, a, 9, a, 10); break;
+					case 5: ASSIGN_X_BITS(a, 9, a, 10, c, 7, c, 6, d0, 5, d1, 5); break;
+					case 6: ASSIGN_X_BITS(b0, 6, b1, 6, a, 11, c, 6, a, 9, a, 10); break;
+					case 7: ASSIGN_X_BITS(a, 9, a, 10, a, 11, c, 6, d0, 5, d1, 5); break;
+					default:
+						DE_ASSERT(false);
+					}
+#undef ASSIGN_X_BITS
+#undef SHOR
+				}
+				static const int numDBits[] = { 7, 6, 7, 6, 5, 6, 5, 6 };
+				DE_ASSERT(mode < DE_LENGTH_OF_ARRAY(numDBits));
+				d0 = signExtend(d0, numDBits[mode]);
+				d1 = signExtend(d1, numDBits[mode]);
+				const int shiftAmount = (mode >> 1) ^ 3;
+				a <<= shiftAmount;
+				c <<= shiftAmount;
+				b0 <<= shiftAmount;
+				b1 <<= shiftAmount;
+				d0 <<= shiftAmount;
+				d1 <<= shiftAmount;
+				e0 = UVec4(FMath::Clamp(a - c, 0, 0xfff),
+					FMath::Clamp(a - b0 - c - d0, 0, 0xfff),
+					FMath::Clamp(a - b1 - c - d1, 0, 0xfff),
+					0x780);
+				e1 = UVec4(FMath::Clamp(a, 0, 0xfff),
+					FMath::Clamp(a - b0, 0, 0xfff),
+					FMath::Clamp(a - b1, 0, 0xfff),
+					0x780);
+				if (major == 1)
+				{
+					std::swap(e0.x(), e0.y());
+					std::swap(e1.x(), e1.y());
+				}
+				else if (major == 2)
+				{
+					std::swap(e0.x(), e0.z());
+					std::swap(e1.x(), e1.z());
+				}
 			}
-			*output0 = ushort4(ushort(y0 << 4), ushort(y0 << 4), ushort(y0 << 4), 0x7800);
-			*output1 = ushort4(ushort(y1 << 4), ushort(y1 << 4), ushort(y1 << 4), 0x7800);
 		}
-
-
-		void hdr_alpha_unpack(const int input[2], int quantization_level, int* a0, int* a1)
+		void decodeHDREndpointMode15(UVec4& e0, UVec4& e1, deUint32 v0, deUint32 v1, deUint32 v2, deUint32 v3, deUint32 v4, deUint32 v5, deUint32 v6In, deUint32 v7In)
 		{
-
-			int v6 = color_unquantization_tables[quantization_level][input[0]];
-			int v7 = color_unquantization_tables[quantization_level][input[1]];
-
-			int selector = ((v6 >> 7) & 1) | ((v7 >> 6) & 2);
-			v6 &= 0x7F;
-			v7 &= 0x7F;
-			if (selector == 3)
+			decodeHDREndpointMode11(e0, e1, v0, v1, v2, v3, v4, v5);
+			const deUint32	mode = (getBit(v7In, 7) << 1) | getBit(v6In, 7);
+			deInt32			v6 = (deInt32)getBits(v6In, 0, 6);
+			deInt32			v7 = (deInt32)getBits(v7In, 0, 6);
+			if (mode == 3)
 			{
-				*a0 = v6 << 5;
-				*a1 = v7 << 5;
+				e0.w() = v6 << 5;
+				e1.w() = v7 << 5;
 			}
 			else
 			{
-				v6 |= (v7 << (selector + 1)) & 0x780;
-				v7 &= (0x3f >> selector);
-				v7 ^= 32 >> selector;
-				v7 -= 32 >> selector;
-				v6 <<= (4 - selector);
-				v7 <<= (4 - selector);
+				v6 |= (v7 << (mode + 1)) & 0x780;
+				v7 &= (0x3f >> mode);
+				v7 ^= 0x20 >> mode;
+				v7 -= 0x20 >> mode;
+				v6 <<= 4 - mode;
+				v7 <<= 4 - mode;
 				v7 += v6;
-
-				if (v7 < 0)
-					v7 = 0;
-				else if (v7 > 0xFFF)
-					v7 = 0xFFF;
-
-				*a0 = v6;
-				*a1 = v7;
+				v7 = FMath::Clamp(v7, 0, 0xfff);
+				e0.w() = v6;
+				e1.w() = v7;
 			}
-
-			*a0 <<= 4;
-			*a1 <<= 4;
-		}
-
-
-		void hdr_rgb_hdr_alpha_unpack3(const int input[8], int quantization_level, ushort4* output0, ushort4* output1)
-		{
-			hdr_rgb_unpack3(input, quantization_level, output0, output1);
-
-			int alpha0, alpha1;
-			hdr_alpha_unpack(input + 6, quantization_level, &alpha0, &alpha1);
-
-			output0->w = ushort(alpha0);
-			output1->w = ushort(alpha1);
 		}
-
-
-		void unpack_color_endpoints(astc_decode_mode decode_mode, int format, int quantization_level, const int* input, int* rgb_hdr, int* alpha_hdr, int* nan_endpoint, ushort4* output0, ushort4* output1)
+		void decodeColorEndpoints(ColorEndpointPair* dst, const deUint32* unquantizedEndpoints, const deUint32* endpointModes, int numPartitions)
 		{
-			*nan_endpoint = 0;
-
-			switch (format)
-			{
-			case FMT_LUMINANCE:
-				*rgb_hdr = 0;
-				*alpha_hdr = 0;
-				luminance_unpack(input, quantization_level, output0, output1);
-				break;
-
-			case FMT_LUMINANCE_DELTA:
-				*rgb_hdr = 0;
-				*alpha_hdr = 0;
-				luminance_delta_unpack(input, quantization_level, output0, output1);
-				break;
-
-			case FMT_HDR_LUMINANCE_SMALL_RANGE:
-				*rgb_hdr = 1;
-				*alpha_hdr = -1;
-				hdr_luminance_small_range_unpack(input, quantization_level, output0, output1);
-				break;
-
-			case FMT_HDR_LUMINANCE_LARGE_RANGE:
-				*rgb_hdr = 1;
-				*alpha_hdr = -1;
-				hdr_luminance_large_range_unpack(input, quantization_level, output0, output1);
-				break;
-
-			case FMT_LUMINANCE_ALPHA:
-				*rgb_hdr = 0;
-				*alpha_hdr = 0;
-				luminance_alpha_unpack(input, quantization_level, output0, output1);
-				break;
-
-			case FMT_LUMINANCE_ALPHA_DELTA:
-				*rgb_hdr = 0;
-				*alpha_hdr = 0;
-				luminance_alpha_delta_unpack(input, quantization_level, output0, output1);
-				break;
-
-			case FMT_RGB_SCALE:
-				*rgb_hdr = 0;
-				*alpha_hdr = 0;
-				rgb_scale_unpack(input, quantization_level, output0, output1);
-				break;
-
-			case FMT_RGB_SCALE_ALPHA:
-				*rgb_hdr = 0;
-				*alpha_hdr = 0;
-				rgb_scale_alpha_unpack(input, quantization_level, output0, output1);
-				break;
-
-			case FMT_HDR_RGB_SCALE:
-				*rgb_hdr = 1;
-				*alpha_hdr = -1;
-				hdr_rgbo_unpack3(input, quantization_level, output0, output1);
-				break;
-
-			case FMT_RGB:
-				*rgb_hdr = 0;
-				*alpha_hdr = 0;
-				rgb_unpack(input, quantization_level, output0, output1);
-				break;
-
-			case FMT_RGB_DELTA:
-				*rgb_hdr = 0;
-				*alpha_hdr = 0;
-				rgb_delta_unpack(input, quantization_level, output0, output1);
-				break;
-
-			case FMT_HDR_RGB:
-				*rgb_hdr = 1;
-				*alpha_hdr = -1;
-				hdr_rgb_unpack3(input, quantization_level, output0, output1);
-				break;
-
-			case FMT_RGBA:
-				*rgb_hdr = 0;
-				*alpha_hdr = 0;
-				rgba_unpack(input, quantization_level, output0, output1);
-				break;
-
-			case FMT_RGBA_DELTA:
-				*rgb_hdr = 0;
-				*alpha_hdr = 0;
-				rgba_delta_unpack(input, quantization_level, output0, output1);
-				break;
-
-			case FMT_HDR_RGB_LDR_ALPHA:
-				*rgb_hdr = 1;
-				*alpha_hdr = 0;
-				hdr_rgb_ldr_alpha_unpack3(input, quantization_level, output0, output1);
-				break;
-
-			case FMT_HDR_RGBA:
-				*rgb_hdr = 1;
-				*alpha_hdr = 1;
-				hdr_rgb_hdr_alpha_unpack3(input, quantization_level, output0, output1);
-				break;
-
-			default:
-				ASTC_CODEC_INTERNAL_ERROR;
-			}
-
-
-			if (*alpha_hdr == -1)
-			{
-				output0->w = 0x00FF;
-				output1->w = 0x00FF;
-				*alpha_hdr = 0;
-			}
-
-
-			switch (decode_mode)
+			int unquantizedNdx = 0;
+			for (int partitionNdx = 0; partitionNdx < numPartitions; partitionNdx++)
 			{
-			case DECODE_LDR_SRGB:
-				if (*rgb_hdr == 1)
-				{
-					output0->x = 0xFF00;
-					output0->y = 0x0000;
-					output0->z = 0xFF00;
-					output0->w = 0xFF00;
-					output1->x = 0xFF00;
-					output1->y = 0x0000;
-					output1->z = 0xFF00;
-					output1->w = 0xFF00;
-				}
-				else
+				const deUint32		endpointMode = endpointModes[partitionNdx];
+				const deUint32* v = &unquantizedEndpoints[unquantizedNdx];
+				UVec4& e0 = dst[partitionNdx].e0;
+				UVec4& e1 = dst[partitionNdx].e1;
+				unquantizedNdx += computeNumColorEndpointValues(endpointMode);
+				switch (endpointMode)
 				{
-					output0->x *= 257;
-					output0->y *= 257;
-					output0->z *= 257;
-					output0->w *= 257;
-					output1->x *= 257;
-					output1->y *= 257;
-					output1->z *= 257;
-					output1->w *= 257;
-				}
-				*rgb_hdr = 0;
-				*alpha_hdr = 0;
-				break;
-
-			case DECODE_LDR:
-				if (*rgb_hdr == 1)
+				case 0:
+					e0 = UVec4(v[0], v[0], v[0], 0xff);
+					e1 = UVec4(v[1], v[1], v[1], 0xff);
+					break;
+				case 1:
 				{
-					output0->x = 0xFFFF;
-					output0->y = 0xFFFF;
-					output0->z = 0xFFFF;
-					output0->w = 0xFFFF;
-					output1->x = 0xFFFF;
-					output1->y = 0xFFFF;
-					output1->z = 0xFFFF;
-					output1->w = 0xFFFF;
-					*nan_endpoint = 1;
+					const deUint32 L0 = (v[0] >> 2) | (getBits(v[1], 6, 7) << 6);
+					const deUint32 L1 = FMath::Min(0xffu, L0 + getBits(v[1], 0, 5));
+					e0 = UVec4(L0, L0, L0, 0xff);
+					e1 = UVec4(L1, L1, L1, 0xff);
+					break;
 				}
-				else
+				case 2:
 				{
-					output0->x *= 257;
-					output0->y *= 257;
-					output0->z *= 257;
-					output0->w *= 257;
-					output1->x *= 257;
-					output1->y *= 257;
-					output1->z *= 257;
-					output1->w *= 257;
+					const deUint32 v1Gr = v[1] >= v[0];
+					const deUint32 y0 = v1Gr ? v[0] << 4 : (v[1] << 4) + 8;
+					const deUint32 y1 = v1Gr ? v[1] << 4 : (v[0] << 4) - 8;
+					e0 = UVec4(y0, y0, y0, 0x780);
+					e1 = UVec4(y1, y1, y1, 0x780);
+					break;
 				}
-				*rgb_hdr = 0;
-				*alpha_hdr = 0;
-				break;
-
-			case DECODE_HDR:
-
-				if (*rgb_hdr == 0)
+				case 3:
 				{
-					output0->x *= 257;
-					output0->y *= 257;
-					output0->z *= 257;
-					output1->x *= 257;
-					output1->y *= 257;
-					output1->z *= 257;
+					const bool		m = isBitSet(v[0], 7);
+					const deUint32	y0 = m ? (getBits(v[1], 5, 7) << 9) | (getBits(v[0], 0, 6) << 2)
+						: (getBits(v[1], 4, 7) << 8) | (getBits(v[0], 0, 6) << 1);
+					const deUint32	d = m ? getBits(v[1], 0, 4) << 2
+						: getBits(v[1], 0, 3) << 1;
+					const deUint32	y1 = FMath::Min(0xfffu, y0 + d);
+					e0 = UVec4(y0, y0, y0, 0x780);
+					e1 = UVec4(y1, y1, y1, 0x780);
+					break;
 				}
-				if (*alpha_hdr == 0)
+				case 4:
+					e0 = UVec4(v[0], v[0], v[0], v[2]);
+					e1 = UVec4(v[1], v[1], v[1], v[3]);
+					break;
+				case 5:
 				{
-					output0->w *= 257;
-					output1->w *= 257;
+					deInt32 v0 = (deInt32)v[0];
+					deInt32 v1 = (deInt32)v[1];
+					deInt32 v2 = (deInt32)v[2];
+					deInt32 v3 = (deInt32)v[3];
+					bitTransferSigned(v1, v0);
+					bitTransferSigned(v3, v2);
+					e0 = clampedRGBA(IVec4(v0, v0, v0, v2));
+					e1 = clampedRGBA(IVec4(v0 + v1, v0 + v1, v0 + v1, v2 + v3));
+					break;
 				}
-				break;
-			}
-		}
-
-
-		struct quantization_and_transfer_table
-		{
-			quantization_method method;
-			uint8 unquantized_value[32];	// 0..64
-			float unquantized_value_flt[32];	// 0..1
-			uint8 prev_quantized_value[32];
-			uint8 next_quantized_value[32];
-			uint8 closest_quantized_weight[1025];
-		};
-
-		const quantization_and_transfer_table quant_and_xfer_tables[12] = {
-			// quantization method 0, range 0..1
-			{
-			 QUANT_2,
-			 {0, 64,},
-			 {
-			  0.000000, 1.000000,},
-			 {0, 0,},
-			 {1, 1,},
-			 {
-			  0, 0, 0, 0, 0, 0, 0, 0, 0, 0, 0, 0, 0, 0, 0, 0,
-			  0, 0, 0, 0, 0, 0, 0, 0, 0, 0, 0, 0, 0, 0, 0, 0,
-			  0, 0, 0, 0, 0, 0, 0, 0, 0, 0, 0, 0, 0, 0, 0, 0,
-			  0, 0, 0, 0, 0, 0, 0, 0, 0, 0, 0, 0, 0, 0, 0, 0,
-			  0, 0, 0, 0, 0, 0, 0, 0, 0, 0, 0, 0, 0, 0, 0, 0,
-			  0, 0, 0, 0, 0, 0, 0, 0, 0, 0, 0, 0, 0, 0, 0, 0,
-			  0, 0, 0, 0, 0, 0, 0, 0, 0, 0, 0, 0, 0, 0, 0, 0,
-			  0, 0, 0, 0, 0, 0, 0, 0, 0, 0, 0, 0, 0, 0, 0, 0,
-			  0, 0, 0, 0, 0, 0, 0, 0, 0, 0, 0, 0, 0, 0, 0, 0,
-			  0, 0, 0, 0, 0, 0, 0, 0, 0, 0, 0, 0, 0, 0, 0, 0,
-			  0, 0, 0, 0, 0, 0, 0, 0, 0, 0, 0, 0, 0, 0, 0, 0,
-			  0, 0, 0, 0, 0, 0, 0, 0, 0, 0, 0, 0, 0, 0, 0, 0,
-			  0, 0, 0, 0, 0, 0, 0, 0, 0, 0, 0, 0, 0, 0, 0, 0,
-			  0, 0, 0, 0, 0, 0, 0, 0, 0, 0, 0, 0, 0, 0, 0, 0,
-			  0, 0, 0, 0, 0, 0, 0, 0, 0, 0, 0, 0, 0, 0, 0, 0,
-			  0, 0, 0, 0, 0, 0, 0, 0, 0, 0, 0, 0, 0, 0, 0, 0,
-			  0, 0, 0, 0, 0, 0, 0, 0, 0, 0, 0, 0, 0, 0, 0, 0,
-			  0, 0, 0, 0, 0, 0, 0, 0, 0, 0, 0, 0, 0, 0, 0, 0,
-			  0, 0, 0, 0, 0, 0, 0, 0, 0, 0, 0, 0, 0, 0, 0, 0,
-			  0, 0, 0, 0, 0, 0, 0, 0, 0, 0, 0, 0, 0, 0, 0, 0,
-			  0, 0, 0, 0, 0, 0, 0, 0, 0, 0, 0, 0, 0, 0, 0, 0,
-			  0, 0, 0, 0, 0, 0, 0, 0, 0, 0, 0, 0, 0, 0, 0, 0,
-			  0, 0, 0, 0, 0, 0, 0, 0, 0, 0, 0, 0, 0, 0, 0, 0,
-			  0, 0, 0, 0, 0, 0, 0, 0, 0, 0, 0, 0, 0, 0, 0, 0,
-			  0, 0, 0, 0, 0, 0, 0, 0, 0, 0, 0, 0, 0, 0, 0, 0,
-			  0, 0, 0, 0, 0, 0, 0, 0, 0, 0, 0, 0, 0, 0, 0, 0,
-			  0, 0, 0, 0, 0, 0, 0, 0, 0, 0, 0, 0, 0, 0, 0, 0,
-			  0, 0, 0, 0, 0, 0, 0, 0, 0, 0, 0, 0, 0, 0, 0, 0,
-			  0, 0, 0, 0, 0, 0, 0, 0, 0, 0, 0, 0, 0, 0, 0, 0,
-			  0, 0, 0, 0, 0, 0, 0, 0, 0, 0, 0, 0, 0, 0, 0, 0,
-			  0, 0, 0, 0, 0, 0, 0, 0, 0, 0, 0, 0, 0, 0, 0, 0,
-			  0, 0, 0, 0, 0, 0, 0, 0, 0, 0, 0, 0, 0, 0, 0, 0,
-			  0, 1, 1, 1, 1, 1, 1, 1, 1, 1, 1, 1, 1, 1, 1, 1,
-			  1, 1, 1, 1, 1, 1, 1, 1, 1, 1, 1, 1, 1, 1, 1, 1,
-			  1, 1, 1, 1, 1, 1, 1, 1, 1, 1, 1, 1, 1, 1, 1, 1,
-			  1, 1, 1, 1, 1, 1, 1, 1, 1, 1, 1, 1, 1, 1, 1, 1,
-			  1, 1, 1, 1, 1, 1, 1, 1, 1, 1, 1, 1, 1, 1, 1, 1,
-			  1, 1, 1, 1, 1, 1, 1, 1, 1, 1, 1, 1, 1, 1, 1, 1,
-			  1, 1, 1, 1, 1, 1, 1, 1, 1, 1, 1, 1, 1, 1, 1, 1,
-			  1, 1, 1, 1, 1, 1, 1, 1, 1, 1, 1, 1, 1, 1, 1, 1,
-			  1, 1, 1, 1, 1, 1, 1, 1, 1, 1, 1, 1, 1, 1, 1, 1,
-			  1, 1, 1, 1, 1, 1, 1, 1, 1, 1, 1, 1, 1, 1, 1, 1,
-			  1, 1, 1, 1, 1, 1, 1, 1, 1, 1, 1, 1, 1, 1, 1, 1,
-			  1, 1, 1, 1, 1, 1, 1, 1, 1, 1, 1, 1, 1, 1, 1, 1,
-			  1, 1, 1, 1, 1, 1, 1, 1, 1, 1, 1, 1, 1, 1, 1, 1,
-			  1, 1, 1, 1, 1, 1, 1, 1, 1, 1, 1, 1, 1, 1, 1, 1,
-			  1, 1, 1, 1, 1, 1, 1, 1, 1, 1, 1, 1, 1, 1, 1, 1,
-			  1, 1, 1, 1, 1, 1, 1, 1, 1, 1, 1, 1, 1, 1, 1, 1,
-			  1, 1, 1, 1, 1, 1, 1, 1, 1, 1, 1, 1, 1, 1, 1, 1,
-			  1, 1, 1, 1, 1, 1, 1, 1, 1, 1, 1, 1, 1, 1, 1, 1,
-			  1, 1, 1, 1, 1, 1, 1, 1, 1, 1, 1, 1, 1, 1, 1, 1,
-			  1, 1, 1, 1, 1, 1, 1, 1, 1, 1, 1, 1, 1, 1, 1, 1,
-			  1, 1, 1, 1, 1, 1, 1, 1, 1, 1, 1, 1, 1, 1, 1, 1,
-			  1, 1, 1, 1, 1, 1, 1, 1, 1, 1, 1, 1, 1, 1, 1, 1,
-			  1, 1, 1, 1, 1, 1, 1, 1, 1, 1, 1, 1, 1, 1, 1, 1,
-			  1, 1, 1, 1, 1, 1, 1, 1, 1, 1, 1, 1, 1, 1, 1, 1,
-			  1, 1, 1, 1, 1, 1, 1, 1, 1, 1, 1, 1, 1, 1, 1, 1,
-			  1, 1, 1, 1, 1, 1, 1, 1, 1, 1, 1, 1, 1, 1, 1, 1,
-			  1, 1, 1, 1, 1, 1, 1, 1, 1, 1, 1, 1, 1, 1, 1, 1,
-			  1, 1, 1, 1, 1, 1, 1, 1, 1, 1, 1, 1, 1, 1, 1, 1,
-			  1, 1, 1, 1, 1, 1, 1, 1, 1, 1, 1, 1, 1, 1, 1, 1,
-			  1, 1, 1, 1, 1, 1, 1, 1, 1, 1, 1, 1, 1, 1, 1, 1,
-			  1, 1, 1, 1, 1, 1, 1, 1, 1, 1, 1, 1, 1, 1, 1, 1,
-			  1, 1, 1, 1, 1, 1, 1, 1, 1, 1, 1, 1, 1, 1, 1, 1,
-			  1,
-			  },
-			 },
-
-
-
-
-			 // quantization method 1, range 0..2
-			 {
-			  QUANT_3,
-			  {0, 32, 64,},
-			  {
-			   0.000000, 0.500000, 1.000000,},
-			  {0, 0, 1,},
-			  {1, 2, 2,},
-			  {
-			   0, 0, 0, 0, 0, 0, 0, 0, 0, 0, 0, 0, 0, 0, 0, 0,
-			   0, 0, 0, 0, 0, 0, 0, 0, 0, 0, 0, 0, 0, 0, 0, 0,
-			   0, 0, 0, 0, 0, 0, 0, 0, 0, 0, 0, 0, 0, 0, 0, 0,
-			   0, 0, 0, 0, 0, 0, 0, 0, 0, 0, 0, 0, 0, 0, 0, 0,
-			   0, 0, 0, 0, 0, 0, 0, 0, 0, 0, 0, 0, 0, 0, 0, 0,
-			   0, 0, 0, 0, 0, 0, 0, 0, 0, 0, 0, 0, 0, 0, 0, 0,
-			   0, 0, 0, 0, 0, 0, 0, 0, 0, 0, 0, 0, 0, 0, 0, 0,
-			   0, 0, 0, 0, 0, 0, 0, 0, 0, 0, 0, 0, 0, 0, 0, 0,
-			   0, 0, 0, 0, 0, 0, 0, 0, 0, 0, 0, 0, 0, 0, 0, 0,
-			   0, 0, 0, 0, 0, 0, 0, 0, 0, 0, 0, 0, 0, 0, 0, 0,
-			   0, 0, 0, 0, 0, 0, 0, 0, 0, 0, 0, 0, 0, 0, 0, 0,
-			   0, 0, 0, 0, 0, 0, 0, 0, 0, 0, 0, 0, 0, 0, 0, 0,
-			   0, 0, 0, 0, 0, 0, 0, 0, 0, 0, 0, 0, 0, 0, 0, 0,
-			   0, 0, 0, 0, 0, 0, 0, 0, 0, 0, 0, 0, 0, 0, 0, 0,
-			   0, 0, 0, 0, 0, 0, 0, 0, 0, 0, 0, 0, 0, 0, 0, 0,
-			   0, 0, 0, 0, 0, 0, 0, 0, 0, 0, 0, 0, 0, 0, 0, 0,
-			   0, 1, 1, 1, 1, 1, 1, 1, 1, 1, 1, 1, 1, 1, 1, 1,
-			   1, 1, 1, 1, 1, 1, 1, 1, 1, 1, 1, 1, 1, 1, 1, 1,
-			   1, 1, 1, 1, 1, 1, 1, 1, 1, 1, 1, 1, 1, 1, 1, 1,
-			   1, 1, 1, 1, 1, 1, 1, 1, 1, 1, 1, 1, 1, 1, 1, 1,
-			   1, 1, 1, 1, 1, 1, 1, 1, 1, 1, 1, 1, 1, 1, 1, 1,
-			   1, 1, 1, 1, 1, 1, 1, 1, 1, 1, 1, 1, 1, 1, 1, 1,
-			   1, 1, 1, 1, 1, 1, 1, 1, 1, 1, 1, 1, 1, 1, 1, 1,
-			   1, 1, 1, 1, 1, 1, 1, 1, 1, 1, 1, 1, 1, 1, 1, 1,
-			   1, 1, 1, 1, 1, 1, 1, 1, 1, 1, 1, 1, 1, 1, 1, 1,
-			   1, 1, 1, 1, 1, 1, 1, 1, 1, 1, 1, 1, 1, 1, 1, 1,
-			   1, 1, 1, 1, 1, 1, 1, 1, 1, 1, 1, 1, 1, 1, 1, 1,
-			   1, 1, 1, 1, 1, 1, 1, 1, 1, 1, 1, 1, 1, 1, 1, 1,
-			   1, 1, 1, 1, 1, 1, 1, 1, 1, 1, 1, 1, 1, 1, 1, 1,
-			   1, 1, 1, 1, 1, 1, 1, 1, 1, 1, 1, 1, 1, 1, 1, 1,
-			   1, 1, 1, 1, 1, 1, 1, 1, 1, 1, 1, 1, 1, 1, 1, 1,
-			   1, 1, 1, 1, 1, 1, 1, 1, 1, 1, 1, 1, 1, 1, 1, 1,
-			   1, 1, 1, 1, 1, 1, 1, 1, 1, 1, 1, 1, 1, 1, 1, 1,
-			   1, 1, 1, 1, 1, 1, 1, 1, 1, 1, 1, 1, 1, 1, 1, 1,
-			   1, 1, 1, 1, 1, 1, 1, 1, 1, 1, 1, 1, 1, 1, 1, 1,
-			   1, 1, 1, 1, 1, 1, 1, 1, 1, 1, 1, 1, 1, 1, 1, 1,
-			   1, 1, 1, 1, 1, 1, 1, 1, 1, 1, 1, 1, 1, 1, 1, 1,
-			   1, 1, 1, 1, 1, 1, 1, 1, 1, 1, 1, 1, 1, 1, 1, 1,
-			   1, 1, 1, 1, 1, 1, 1, 1, 1, 1, 1, 1, 1, 1, 1, 1,
-			   1, 1, 1, 1, 1, 1, 1, 1, 1, 1, 1, 1, 1, 1, 1, 1,
-			   1, 1, 1, 1, 1, 1, 1, 1, 1, 1, 1, 1, 1, 1, 1, 1,
-			   1, 1, 1, 1, 1, 1, 1, 1, 1, 1, 1, 1, 1, 1, 1, 1,
-			   1, 1, 1, 1, 1, 1, 1, 1, 1, 1, 1, 1, 1, 1, 1, 1,
-			   1, 1, 1, 1, 1, 1, 1, 1, 1, 1, 1, 1, 1, 1, 1, 1,
-			   1, 1, 1, 1, 1, 1, 1, 1, 1, 1, 1, 1, 1, 1, 1, 1,
-			   1, 1, 1, 1, 1, 1, 1, 1, 1, 1, 1, 1, 1, 1, 1, 1,
-			   1, 1, 1, 1, 1, 1, 1, 1, 1, 1, 1, 1, 1, 1, 1, 1,
-			   1, 1, 1, 1, 1, 1, 1, 1, 1, 1, 1, 1, 1, 1, 1, 1,
-			   1, 2, 2, 2, 2, 2, 2, 2, 2, 2, 2, 2, 2, 2, 2, 2,
-			   2, 2, 2, 2, 2, 2, 2, 2, 2, 2, 2, 2, 2, 2, 2, 2,
-			   2, 2, 2, 2, 2, 2, 2, 2, 2, 2, 2, 2, 2, 2, 2, 2,
-			   2, 2, 2, 2, 2, 2, 2, 2, 2, 2, 2, 2, 2, 2, 2, 2,
-			   2, 2, 2, 2, 2, 2, 2, 2, 2, 2, 2, 2, 2, 2, 2, 2,
-			   2, 2, 2, 2, 2, 2, 2, 2, 2, 2, 2, 2, 2, 2, 2, 2,
-			   2, 2, 2, 2, 2, 2, 2, 2, 2, 2, 2, 2, 2, 2, 2, 2,
-			   2, 2, 2, 2, 2, 2, 2, 2, 2, 2, 2, 2, 2, 2, 2, 2,
-			   2, 2, 2, 2, 2, 2, 2, 2, 2, 2, 2, 2, 2, 2, 2, 2,
-			   2, 2, 2, 2, 2, 2, 2, 2, 2, 2, 2, 2, 2, 2, 2, 2,
-			   2, 2, 2, 2, 2, 2, 2, 2, 2, 2, 2, 2, 2, 2, 2, 2,
-			   2, 2, 2, 2, 2, 2, 2, 2, 2, 2, 2, 2, 2, 2, 2, 2,
-			   2, 2, 2, 2, 2, 2, 2, 2, 2, 2, 2, 2, 2, 2, 2, 2,
-			   2, 2, 2, 2, 2, 2, 2, 2, 2, 2, 2, 2, 2, 2, 2, 2,
-			   2, 2, 2, 2, 2, 2, 2, 2, 2, 2, 2, 2, 2, 2, 2, 2,
-			   2, 2, 2, 2, 2, 2, 2, 2, 2, 2, 2, 2, 2, 2, 2, 2,
-			   2,
-			   },
-			  },
-
-
-
-
-			  // quantization method 2, range 0..3
-			  {
-			   QUANT_4,
-			   {0, 21, 43, 64,},
-			   {
-				0.000000, 0.328125, 0.671875, 1.000000,},
-			   {0, 0, 1, 2,},
-			   {1, 2, 3, 3,},
-			   {
-				0, 0, 0, 0, 0, 0, 0, 0, 0, 0, 0, 0, 0, 0, 0, 0,
-				0, 0, 0, 0, 0, 0, 0, 0, 0, 0, 0, 0, 0, 0, 0, 0,
-				0, 0, 0, 0, 0, 0, 0, 0, 0, 0, 0, 0, 0, 0, 0, 0,
-				0, 0, 0, 0, 0, 0, 0, 0, 0, 0, 0, 0, 0, 0, 0, 0,
-				0, 0, 0, 0, 0, 0, 0, 0, 0, 0, 0, 0, 0, 0, 0, 0,
-				0, 0, 0, 0, 0, 0, 0, 0, 0, 0, 0, 0, 0, 0, 0, 0,
-				0, 0, 0, 0, 0, 0, 0, 0, 0, 0, 0, 0, 0, 0, 0, 0,
-				0, 0, 0, 0, 0, 0, 0, 0, 0, 0, 0, 0, 0, 0, 0, 0,
-				0, 0, 0, 0, 0, 0, 0, 0, 0, 0, 0, 0, 0, 0, 0, 0,
-				0, 0, 0, 0, 0, 0, 0, 0, 0, 0, 0, 0, 0, 0, 0, 0,
-				0, 0, 0, 0, 0, 0, 0, 0, 0, 1, 1, 1, 1, 1, 1, 1,
-				1, 1, 1, 1, 1, 1, 1, 1, 1, 1, 1, 1, 1, 1, 1, 1,
-				1, 1, 1, 1, 1, 1, 1, 1, 1, 1, 1, 1, 1, 1, 1, 1,
-				1, 1, 1, 1, 1, 1, 1, 1, 1, 1, 1, 1, 1, 1, 1, 1,
-				1, 1, 1, 1, 1, 1, 1, 1, 1, 1, 1, 1, 1, 1, 1, 1,
-				1, 1, 1, 1, 1, 1, 1, 1, 1, 1, 1, 1, 1, 1, 1, 1,
-				1, 1, 1, 1, 1, 1, 1, 1, 1, 1, 1, 1, 1, 1, 1, 1,
-				1, 1, 1, 1, 1, 1, 1, 1, 1, 1, 1, 1, 1, 1, 1, 1,
-				1, 1, 1, 1, 1, 1, 1, 1, 1, 1, 1, 1, 1, 1, 1, 1,
-				1, 1, 1, 1, 1, 1, 1, 1, 1, 1, 1, 1, 1, 1, 1, 1,
-				1, 1, 1, 1, 1, 1, 1, 1, 1, 1, 1, 1, 1, 1, 1, 1,
-				1, 1, 1, 1, 1, 1, 1, 1, 1, 1, 1, 1, 1, 1, 1, 1,
-				1, 1, 1, 1, 1, 1, 1, 1, 1, 1, 1, 1, 1, 1, 1, 1,
-				1, 1, 1, 1, 1, 1, 1, 1, 1, 1, 1, 1, 1, 1, 1, 1,
-				1, 1, 1, 1, 1, 1, 1, 1, 1, 1, 1, 1, 1, 1, 1, 1,
-				1, 1, 1, 1, 1, 1, 1, 1, 1, 1, 1, 1, 1, 1, 1, 1,
-				1, 1, 1, 1, 1, 1, 1, 1, 1, 1, 1, 1, 1, 1, 1, 1,
-				1, 1, 1, 1, 1, 1, 1, 1, 1, 1, 1, 1, 1, 1, 1, 1,
-				1, 1, 1, 1, 1, 1, 1, 1, 1, 1, 1, 1, 1, 1, 1, 1,
-				1, 1, 1, 1, 1, 1, 1, 1, 1, 1, 1, 1, 1, 1, 1, 1,
-				1, 1, 1, 1, 1, 1, 1, 1, 1, 1, 1, 1, 1, 1, 1, 1,
-				1, 1, 1, 1, 1, 1, 1, 1, 1, 1, 1, 1, 1, 1, 1, 1,
-				1, 2, 2, 2, 2, 2, 2, 2, 2, 2, 2, 2, 2, 2, 2, 2,
-				2, 2, 2, 2, 2, 2, 2, 2, 2, 2, 2, 2, 2, 2, 2, 2,
-				2, 2, 2, 2, 2, 2, 2, 2, 2, 2, 2, 2, 2, 2, 2, 2,
-				2, 2, 2, 2, 2, 2, 2, 2, 2, 2, 2, 2, 2, 2, 2, 2,
-				2, 2, 2, 2, 2, 2, 2, 2, 2, 2, 2, 2, 2, 2, 2, 2,
-				2, 2, 2, 2, 2, 2, 2, 2, 2, 2, 2, 2, 2, 2, 2, 2,
-				2, 2, 2, 2, 2, 2, 2, 2, 2, 2, 2, 2, 2, 2, 2, 2,
-				2, 2, 2, 2, 2, 2, 2, 2, 2, 2, 2, 2, 2, 2, 2, 2,
-				2, 2, 2, 2, 2, 2, 2, 2, 2, 2, 2, 2, 2, 2, 2, 2,
-				2, 2, 2, 2, 2, 2, 2, 2, 2, 2, 2, 2, 2, 2, 2, 2,
-				2, 2, 2, 2, 2, 2, 2, 2, 2, 2, 2, 2, 2, 2, 2, 2,
-				2, 2, 2, 2, 2, 2, 2, 2, 2, 2, 2, 2, 2, 2, 2, 2,
-				2, 2, 2, 2, 2, 2, 2, 2, 2, 2, 2, 2, 2, 2, 2, 2,
-				2, 2, 2, 2, 2, 2, 2, 2, 2, 2, 2, 2, 2, 2, 2, 2,
-				2, 2, 2, 2, 2, 2, 2, 2, 2, 2, 2, 2, 2, 2, 2, 2,
-				2, 2, 2, 2, 2, 2, 2, 2, 2, 2, 2, 2, 2, 2, 2, 2,
-				2, 2, 2, 2, 2, 2, 2, 2, 2, 2, 2, 2, 2, 2, 2, 2,
-				2, 2, 2, 2, 2, 2, 2, 2, 2, 2, 2, 2, 2, 2, 2, 2,
-				2, 2, 2, 2, 2, 2, 2, 2, 2, 2, 2, 2, 2, 2, 2, 2,
-				2, 2, 2, 2, 2, 2, 2, 2, 2, 2, 2, 2, 2, 2, 2, 2,
-				2, 2, 2, 2, 2, 2, 2, 2, 2, 2, 2, 2, 2, 2, 2, 2,
-				2, 2, 2, 2, 2, 2, 2, 2, 2, 3, 3, 3, 3, 3, 3, 3,
-				3, 3, 3, 3, 3, 3, 3, 3, 3, 3, 3, 3, 3, 3, 3, 3,
-				3, 3, 3, 3, 3, 3, 3, 3, 3, 3, 3, 3, 3, 3, 3, 3,
-				3, 3, 3, 3, 3, 3, 3, 3, 3, 3, 3, 3, 3, 3, 3, 3,
-				3, 3, 3, 3, 3, 3, 3, 3, 3, 3, 3, 3, 3, 3, 3, 3,
-				3, 3, 3, 3, 3, 3, 3, 3, 3, 3, 3, 3, 3, 3, 3, 3,
-				3, 3, 3, 3, 3, 3, 3, 3, 3, 3, 3, 3, 3, 3, 3, 3,
-				3, 3, 3, 3, 3, 3, 3, 3, 3, 3, 3, 3, 3, 3, 3, 3,
-				3, 3, 3, 3, 3, 3, 3, 3, 3, 3, 3, 3, 3, 3, 3, 3,
-				3, 3, 3, 3, 3, 3, 3, 3, 3, 3, 3, 3, 3, 3, 3, 3,
-				3, 3, 3, 3, 3, 3, 3, 3, 3, 3, 3, 3, 3, 3, 3, 3,
-				3,
-				},
-			   },
-
-
-
-
-			   // quantization method 3, range 0..4
-			   {
-				QUANT_5,
-				{0, 16, 32, 48, 64,},
-				{
-				 0.000000, 0.250000, 0.500000, 0.750000,
-				 1.000000,},
-				{0, 0, 1, 2, 3,},
-				{1, 2, 3, 4, 4,},
-				{
-				 0, 0, 0, 0, 0, 0, 0, 0, 0, 0, 0, 0, 0, 0, 0, 0,
-				 0, 0, 0, 0, 0, 0, 0, 0, 0, 0, 0, 0, 0, 0, 0, 0,
-				 0, 0, 0, 0, 0, 0, 0, 0, 0, 0, 0, 0, 0, 0, 0, 0,
-				 0, 0, 0, 0, 0, 0, 0, 0, 0, 0, 0, 0, 0, 0, 0, 0,
-				 0, 0, 0, 0, 0, 0, 0, 0, 0, 0, 0, 0, 0, 0, 0, 0,
-				 0, 0, 0, 0, 0, 0, 0, 0, 0, 0, 0, 0, 0, 0, 0, 0,
-				 0, 0, 0, 0, 0, 0, 0, 0, 0, 0, 0, 0, 0, 0, 0, 0,
-				 0, 0, 0, 0, 0, 0, 0, 0, 0, 0, 0, 0, 0, 0, 0, 0,
-				 0, 1, 1, 1, 1, 1, 1, 1, 1, 1, 1, 1, 1, 1, 1, 1,
-				 1, 1, 1, 1, 1, 1, 1, 1, 1, 1, 1, 1, 1, 1, 1, 1,
-				 1, 1, 1, 1, 1, 1, 1, 1, 1, 1, 1, 1, 1, 1, 1, 1,
-				 1, 1, 1, 1, 1, 1, 1, 1, 1, 1, 1, 1, 1, 1, 1, 1,
-				 1, 1, 1, 1, 1, 1, 1, 1, 1, 1, 1, 1, 1, 1, 1, 1,
-				 1, 1, 1, 1, 1, 1, 1, 1, 1, 1, 1, 1, 1, 1, 1, 1,
-				 1, 1, 1, 1, 1, 1, 1, 1, 1, 1, 1, 1, 1, 1, 1, 1,
-				 1, 1, 1, 1, 1, 1, 1, 1, 1, 1, 1, 1, 1, 1, 1, 1,
-				 1, 1, 1, 1, 1, 1, 1, 1, 1, 1, 1, 1, 1, 1, 1, 1,
-				 1, 1, 1, 1, 1, 1, 1, 1, 1, 1, 1, 1, 1, 1, 1, 1,
-				 1, 1, 1, 1, 1, 1, 1, 1, 1, 1, 1, 1, 1, 1, 1, 1,
-				 1, 1, 1, 1, 1, 1, 1, 1, 1, 1, 1, 1, 1, 1, 1, 1,
-				 1, 1, 1, 1, 1, 1, 1, 1, 1, 1, 1, 1, 1, 1, 1, 1,
-				 1, 1, 1, 1, 1, 1, 1, 1, 1, 1, 1, 1, 1, 1, 1, 1,
-				 1, 1, 1, 1, 1, 1, 1, 1, 1, 1, 1, 1, 1, 1, 1, 1,
-				 1, 1, 1, 1, 1, 1, 1, 1, 1, 1, 1, 1, 1, 1, 1, 1,
-				 1, 2, 2, 2, 2, 2, 2, 2, 2, 2, 2, 2, 2, 2, 2, 2,
-				 2, 2, 2, 2, 2, 2, 2, 2, 2, 2, 2, 2, 2, 2, 2, 2,
-				 2, 2, 2, 2, 2, 2, 2, 2, 2, 2, 2, 2, 2, 2, 2, 2,
-				 2, 2, 2, 2, 2, 2, 2, 2, 2, 2, 2, 2, 2, 2, 2, 2,
-				 2, 2, 2, 2, 2, 2, 2, 2, 2, 2, 2, 2, 2, 2, 2, 2,
-				 2, 2, 2, 2, 2, 2, 2, 2, 2, 2, 2, 2, 2, 2, 2, 2,
-				 2, 2, 2, 2, 2, 2, 2, 2, 2, 2, 2, 2, 2, 2, 2, 2,
-				 2, 2, 2, 2, 2, 2, 2, 2, 2, 2, 2, 2, 2, 2, 2, 2,
-				 2, 2, 2, 2, 2, 2, 2, 2, 2, 2, 2, 2, 2, 2, 2, 2,
-				 2, 2, 2, 2, 2, 2, 2, 2, 2, 2, 2, 2, 2, 2, 2, 2,
-				 2, 2, 2, 2, 2, 2, 2, 2, 2, 2, 2, 2, 2, 2, 2, 2,
-				 2, 2, 2, 2, 2, 2, 2, 2, 2, 2, 2, 2, 2, 2, 2, 2,
-				 2, 2, 2, 2, 2, 2, 2, 2, 2, 2, 2, 2, 2, 2, 2, 2,
-				 2, 2, 2, 2, 2, 2, 2, 2, 2, 2, 2, 2, 2, 2, 2, 2,
-				 2, 2, 2, 2, 2, 2, 2, 2, 2, 2, 2, 2, 2, 2, 2, 2,
-				 2, 2, 2, 2, 2, 2, 2, 2, 2, 2, 2, 2, 2, 2, 2, 2,
-				 2, 3, 3, 3, 3, 3, 3, 3, 3, 3, 3, 3, 3, 3, 3, 3,
-				 3, 3, 3, 3, 3, 3, 3, 3, 3, 3, 3, 3, 3, 3, 3, 3,
-				 3, 3, 3, 3, 3, 3, 3, 3, 3, 3, 3, 3, 3, 3, 3, 3,
-				 3, 3, 3, 3, 3, 3, 3, 3, 3, 3, 3, 3, 3, 3, 3, 3,
-				 3, 3, 3, 3, 3, 3, 3, 3, 3, 3, 3, 3, 3, 3, 3, 3,
-				 3, 3, 3, 3, 3, 3, 3, 3, 3, 3, 3, 3, 3, 3, 3, 3,
-				 3, 3, 3, 3, 3, 3, 3, 3, 3, 3, 3, 3, 3, 3, 3, 3,
-				 3, 3, 3, 3, 3, 3, 3, 3, 3, 3, 3, 3, 3, 3, 3, 3,
-				 3, 3, 3, 3, 3, 3, 3, 3, 3, 3, 3, 3, 3, 3, 3, 3,
-				 3, 3, 3, 3, 3, 3, 3, 3, 3, 3, 3, 3, 3, 3, 3, 3,
-				 3, 3, 3, 3, 3, 3, 3, 3, 3, 3, 3, 3, 3, 3, 3, 3,
-				 3, 3, 3, 3, 3, 3, 3, 3, 3, 3, 3, 3, 3, 3, 3, 3,
-				 3, 3, 3, 3, 3, 3, 3, 3, 3, 3, 3, 3, 3, 3, 3, 3,
-				 3, 3, 3, 3, 3, 3, 3, 3, 3, 3, 3, 3, 3, 3, 3, 3,
-				 3, 3, 3, 3, 3, 3, 3, 3, 3, 3, 3, 3, 3, 3, 3, 3,
-				 3, 3, 3, 3, 3, 3, 3, 3, 3, 3, 3, 3, 3, 3, 3, 3,
-				 3, 4, 4, 4, 4, 4, 4, 4, 4, 4, 4, 4, 4, 4, 4, 4,
-				 4, 4, 4, 4, 4, 4, 4, 4, 4, 4, 4, 4, 4, 4, 4, 4,
-				 4, 4, 4, 4, 4, 4, 4, 4, 4, 4, 4, 4, 4, 4, 4, 4,
-				 4, 4, 4, 4, 4, 4, 4, 4, 4, 4, 4, 4, 4, 4, 4, 4,
-				 4, 4, 4, 4, 4, 4, 4, 4, 4, 4, 4, 4, 4, 4, 4, 4,
-				 4, 4, 4, 4, 4, 4, 4, 4, 4, 4, 4, 4, 4, 4, 4, 4,
-				 4, 4, 4, 4, 4, 4, 4, 4, 4, 4, 4, 4, 4, 4, 4, 4,
-				 4, 4, 4, 4, 4, 4, 4, 4, 4, 4, 4, 4, 4, 4, 4, 4,
-				 4,
-				 },
-				},
-
-
-
-
-				// quantization method 4, range 0..5
-				{
-				 QUANT_6,
-				 {0, 64, 12, 52, 25, 39,},
-				 {
-				  0.000000, 1.000000, 0.187500, 0.812500,
-				  0.390625, 0.609375,},
-				 {0, 3, 0, 5, 2, 4,},
-				 {2, 1, 4, 1, 5, 3,},
-				 {
-				  0, 0, 0, 0, 0, 0, 0, 0, 0, 0, 0, 0, 0, 0, 0, 0,
-				  0, 0, 0, 0, 0, 0, 0, 0, 0, 0, 0, 0, 0, 0, 0, 0,
-				  0, 0, 0, 0, 0, 0, 0, 0, 0, 0, 0, 0, 0, 0, 0, 0,
-				  0, 0, 0, 0, 0, 0, 0, 0, 0, 0, 0, 0, 0, 0, 0, 0,
-				  0, 0, 0, 0, 0, 0, 0, 0, 0, 0, 0, 0, 0, 0, 0, 0,
-				  0, 0, 0, 0, 0, 0, 0, 0, 0, 0, 0, 0, 0, 0, 0, 0,
-				  0, 2, 2, 2, 2, 2, 2, 2, 2, 2, 2, 2, 2, 2, 2, 2,
-				  2, 2, 2, 2, 2, 2, 2, 2, 2, 2, 2, 2, 2, 2, 2, 2,
-				  2, 2, 2, 2, 2, 2, 2, 2, 2, 2, 2, 2, 2, 2, 2, 2,
-				  2, 2, 2, 2, 2, 2, 2, 2, 2, 2, 2, 2, 2, 2, 2, 2,
-				  2, 2, 2, 2, 2, 2, 2, 2, 2, 2, 2, 2, 2, 2, 2, 2,
-				  2, 2, 2, 2, 2, 2, 2, 2, 2, 2, 2, 2, 2, 2, 2, 2,
-				  2, 2, 2, 2, 2, 2, 2, 2, 2, 2, 2, 2, 2, 2, 2, 2,
-				  2, 2, 2, 2, 2, 2, 2, 2, 2, 2, 2, 2, 2, 2, 2, 2,
-				  2, 2, 2, 2, 2, 2, 2, 2, 2, 2, 2, 2, 2, 2, 2, 2,
-				  2, 2, 2, 2, 2, 2, 2, 2, 2, 2, 2, 2, 2, 2, 2, 2,
-				  2, 2, 2, 2, 2, 2, 2, 2, 2, 2, 2, 2, 2, 2, 2, 2,
-				  2, 2, 2, 2, 2, 2, 2, 2, 2, 2, 2, 2, 2, 2, 2, 2,
-				  2, 2, 2, 2, 2, 2, 2, 2, 2, 4, 4, 4, 4, 4, 4, 4,
-				  4, 4, 4, 4, 4, 4, 4, 4, 4, 4, 4, 4, 4, 4, 4, 4,
-				  4, 4, 4, 4, 4, 4, 4, 4, 4, 4, 4, 4, 4, 4, 4, 4,
-				  4, 4, 4, 4, 4, 4, 4, 4, 4, 4, 4, 4, 4, 4, 4, 4,
-				  4, 4, 4, 4, 4, 4, 4, 4, 4, 4, 4, 4, 4, 4, 4, 4,
-				  4, 4, 4, 4, 4, 4, 4, 4, 4, 4, 4, 4, 4, 4, 4, 4,
-				  4, 4, 4, 4, 4, 4, 4, 4, 4, 4, 4, 4, 4, 4, 4, 4,
-				  4, 4, 4, 4, 4, 4, 4, 4, 4, 4, 4, 4, 4, 4, 4, 4,
-				  4, 4, 4, 4, 4, 4, 4, 4, 4, 4, 4, 4, 4, 4, 4, 4,
-				  4, 4, 4, 4, 4, 4, 4, 4, 4, 4, 4, 4, 4, 4, 4, 4,
-				  4, 4, 4, 4, 4, 4, 4, 4, 4, 4, 4, 4, 4, 4, 4, 4,
-				  4, 4, 4, 4, 4, 4, 4, 4, 4, 4, 4, 4, 4, 4, 4, 4,
-				  4, 4, 4, 4, 4, 4, 4, 4, 4, 4, 4, 4, 4, 4, 4, 4,
-				  4, 4, 4, 4, 4, 4, 4, 4, 4, 4, 4, 4, 4, 4, 4, 4,
-				  4, 5, 5, 5, 5, 5, 5, 5, 5, 5, 5, 5, 5, 5, 5, 5,
-				  5, 5, 5, 5, 5, 5, 5, 5, 5, 5, 5, 5, 5, 5, 5, 5,
-				  5, 5, 5, 5, 5, 5, 5, 5, 5, 5, 5, 5, 5, 5, 5, 5,
-				  5, 5, 5, 5, 5, 5, 5, 5, 5, 5, 5, 5, 5, 5, 5, 5,
-				  5, 5, 5, 5, 5, 5, 5, 5, 5, 5, 5, 5, 5, 5, 5, 5,
-				  5, 5, 5, 5, 5, 5, 5, 5, 5, 5, 5, 5, 5, 5, 5, 5,
-				  5, 5, 5, 5, 5, 5, 5, 5, 5, 5, 5, 5, 5, 5, 5, 5,
-				  5, 5, 5, 5, 5, 5, 5, 5, 5, 5, 5, 5, 5, 5, 5, 5,
-				  5, 5, 5, 5, 5, 5, 5, 5, 5, 5, 5, 5, 5, 5, 5, 5,
-				  5, 5, 5, 5, 5, 5, 5, 5, 5, 5, 5, 5, 5, 5, 5, 5,
-				  5, 5, 5, 5, 5, 5, 5, 5, 5, 5, 5, 5, 5, 5, 5, 5,
-				  5, 5, 5, 5, 5, 5, 5, 5, 5, 5, 5, 5, 5, 5, 5, 5,
-				  5, 5, 5, 5, 5, 5, 5, 5, 5, 5, 5, 5, 5, 5, 5, 5,
-				  5, 5, 5, 5, 5, 5, 5, 5, 5, 3, 3, 3, 3, 3, 3, 3,
-				  3, 3, 3, 3, 3, 3, 3, 3, 3, 3, 3, 3, 3, 3, 3, 3,
-				  3, 3, 3, 3, 3, 3, 3, 3, 3, 3, 3, 3, 3, 3, 3, 3,
-				  3, 3, 3, 3, 3, 3, 3, 3, 3, 3, 3, 3, 3, 3, 3, 3,
-				  3, 3, 3, 3, 3, 3, 3, 3, 3, 3, 3, 3, 3, 3, 3, 3,
-				  3, 3, 3, 3, 3, 3, 3, 3, 3, 3, 3, 3, 3, 3, 3, 3,
-				  3, 3, 3, 3, 3, 3, 3, 3, 3, 3, 3, 3, 3, 3, 3, 3,
-				  3, 3, 3, 3, 3, 3, 3, 3, 3, 3, 3, 3, 3, 3, 3, 3,
-				  3, 3, 3, 3, 3, 3, 3, 3, 3, 3, 3, 3, 3, 3, 3, 3,
-				  3, 3, 3, 3, 3, 3, 3, 3, 3, 3, 3, 3, 3, 3, 3, 3,
-				  3, 3, 3, 3, 3, 3, 3, 3, 3, 3, 3, 3, 3, 3, 3, 3,
-				  3, 3, 3, 3, 3, 3, 3, 3, 3, 3, 3, 3, 3, 3, 3, 3,
-				  3, 3, 3, 3, 3, 3, 3, 3, 3, 3, 3, 3, 3, 3, 3, 3,
-				  3, 1, 1, 1, 1, 1, 1, 1, 1, 1, 1, 1, 1, 1, 1, 1,
-				  1, 1, 1, 1, 1, 1, 1, 1, 1, 1, 1, 1, 1, 1, 1, 1,
-				  1, 1, 1, 1, 1, 1, 1, 1, 1, 1, 1, 1, 1, 1, 1, 1,
-				  1, 1, 1, 1, 1, 1, 1, 1, 1, 1, 1, 1, 1, 1, 1, 1,
-				  1, 1, 1, 1, 1, 1, 1, 1, 1, 1, 1, 1, 1, 1, 1, 1,
-				  1, 1, 1, 1, 1, 1, 1, 1, 1, 1, 1, 1, 1, 1, 1, 1,
-				  1,
-				  },
-				 },
-
-
-
-
-				 // quantization method 5, range 0..7
-				 {
-				  QUANT_8,
-				  {0, 9, 18, 27, 37, 46, 55, 64,},
-				  {
-				   0.000000, 0.140625, 0.281250, 0.421875,
-				   0.578125, 0.718750, 0.859375, 1.000000,},
-				  {0, 0, 1, 2, 3, 4, 5, 6,},
-				  {1, 2, 3, 4, 5, 6, 7, 7,},
-				  {
-				   0, 0, 0, 0, 0, 0, 0, 0, 0, 0, 0, 0, 0, 0, 0, 0,
-				   0, 0, 0, 0, 0, 0, 0, 0, 0, 0, 0, 0, 0, 0, 0, 0,
-				   0, 0, 0, 0, 0, 0, 0, 0, 0, 0, 0, 0, 0, 0, 0, 0,
-				   0, 0, 0, 0, 0, 0, 0, 0, 0, 0, 0, 0, 0, 0, 0, 0,
-				   0, 0, 0, 0, 0, 0, 0, 0, 0, 1, 1, 1, 1, 1, 1, 1,
-				   1, 1, 1, 1, 1, 1, 1, 1, 1, 1, 1, 1, 1, 1, 1, 1,
-				   1, 1, 1, 1, 1, 1, 1, 1, 1, 1, 1, 1, 1, 1, 1, 1,
-				   1, 1, 1, 1, 1, 1, 1, 1, 1, 1, 1, 1, 1, 1, 1, 1,
-				   1, 1, 1, 1, 1, 1, 1, 1, 1, 1, 1, 1, 1, 1, 1, 1,
-				   1, 1, 1, 1, 1, 1, 1, 1, 1, 1, 1, 1, 1, 1, 1, 1,
-				   1, 1, 1, 1, 1, 1, 1, 1, 1, 1, 1, 1, 1, 1, 1, 1,
-				   1, 1, 1, 1, 1, 1, 1, 1, 1, 1, 1, 1, 1, 1, 1, 1,
-				   1, 1, 1, 1, 1, 1, 1, 1, 1, 1, 1, 1, 1, 1, 1, 1,
-				   1, 1, 1, 1, 1, 1, 1, 1, 1, 2, 2, 2, 2, 2, 2, 2,
-				   2, 2, 2, 2, 2, 2, 2, 2, 2, 2, 2, 2, 2, 2, 2, 2,
-				   2, 2, 2, 2, 2, 2, 2, 2, 2, 2, 2, 2, 2, 2, 2, 2,
-				   2, 2, 2, 2, 2, 2, 2, 2, 2, 2, 2, 2, 2, 2, 2, 2,
-				   2, 2, 2, 2, 2, 2, 2, 2, 2, 2, 2, 2, 2, 2, 2, 2,
-				   2, 2, 2, 2, 2, 2, 2, 2, 2, 2, 2, 2, 2, 2, 2, 2,
-				   2, 2, 2, 2, 2, 2, 2, 2, 2, 2, 2, 2, 2, 2, 2, 2,
-				   2, 2, 2, 2, 2, 2, 2, 2, 2, 2, 2, 2, 2, 2, 2, 2,
-				   2, 2, 2, 2, 2, 2, 2, 2, 2, 2, 2, 2, 2, 2, 2, 2,
-				   2, 2, 2, 2, 2, 2, 2, 2, 2, 3, 3, 3, 3, 3, 3, 3,
-				   3, 3, 3, 3, 3, 3, 3, 3, 3, 3, 3, 3, 3, 3, 3, 3,
-				   3, 3, 3, 3, 3, 3, 3, 3, 3, 3, 3, 3, 3, 3, 3, 3,
-				   3, 3, 3, 3, 3, 3, 3, 3, 3, 3, 3, 3, 3, 3, 3, 3,
-				   3, 3, 3, 3, 3, 3, 3, 3, 3, 3, 3, 3, 3, 3, 3, 3,
-				   3, 3, 3, 3, 3, 3, 3, 3, 3, 3, 3, 3, 3, 3, 3, 3,
-				   3, 3, 3, 3, 3, 3, 3, 3, 3, 3, 3, 3, 3, 3, 3, 3,
-				   3, 3, 3, 3, 3, 3, 3, 3, 3, 3, 3, 3, 3, 3, 3, 3,
-				   3, 3, 3, 3, 3, 3, 3, 3, 3, 3, 3, 3, 3, 3, 3, 3,
-				   3, 3, 3, 3, 3, 3, 3, 3, 3, 3, 3, 3, 3, 3, 3, 3,
-				   3, 4, 4, 4, 4, 4, 4, 4, 4, 4, 4, 4, 4, 4, 4, 4,
-				   4, 4, 4, 4, 4, 4, 4, 4, 4, 4, 4, 4, 4, 4, 4, 4,
-				   4, 4, 4, 4, 4, 4, 4, 4, 4, 4, 4, 4, 4, 4, 4, 4,
-				   4, 4, 4, 4, 4, 4, 4, 4, 4, 4, 4, 4, 4, 4, 4, 4,
-				   4, 4, 4, 4, 4, 4, 4, 4, 4, 4, 4, 4, 4, 4, 4, 4,
-				   4, 4, 4, 4, 4, 4, 4, 4, 4, 4, 4, 4, 4, 4, 4, 4,
-				   4, 4, 4, 4, 4, 4, 4, 4, 4, 4, 4, 4, 4, 4, 4, 4,
-				   4, 4, 4, 4, 4, 4, 4, 4, 4, 4, 4, 4, 4, 4, 4, 4,
-				   4, 4, 4, 4, 4, 4, 4, 4, 4, 4, 4, 4, 4, 4, 4, 4,
-				   4, 4, 4, 4, 4, 4, 4, 4, 4, 5, 5, 5, 5, 5, 5, 5,
-				   5, 5, 5, 5, 5, 5, 5, 5, 5, 5, 5, 5, 5, 5, 5, 5,
-				   5, 5, 5, 5, 5, 5, 5, 5, 5, 5, 5, 5, 5, 5, 5, 5,
-				   5, 5, 5, 5, 5, 5, 5, 5, 5, 5, 5, 5, 5, 5, 5, 5,
-				   5, 5, 5, 5, 5, 5, 5, 5, 5, 5, 5, 5, 5, 5, 5, 5,
-				   5, 5, 5, 5, 5, 5, 5, 5, 5, 5, 5, 5, 5, 5, 5, 5,
-				   5, 5, 5, 5, 5, 5, 5, 5, 5, 5, 5, 5, 5, 5, 5, 5,
-				   5, 5, 5, 5, 5, 5, 5, 5, 5, 5, 5, 5, 5, 5, 5, 5,
-				   5, 5, 5, 5, 5, 5, 5, 5, 5, 5, 5, 5, 5, 5, 5, 5,
-				   5, 5, 5, 5, 5, 5, 5, 5, 5, 6, 6, 6, 6, 6, 6, 6,
-				   6, 6, 6, 6, 6, 6, 6, 6, 6, 6, 6, 6, 6, 6, 6, 6,
-				   6, 6, 6, 6, 6, 6, 6, 6, 6, 6, 6, 6, 6, 6, 6, 6,
-				   6, 6, 6, 6, 6, 6, 6, 6, 6, 6, 6, 6, 6, 6, 6, 6,
-				   6, 6, 6, 6, 6, 6, 6, 6, 6, 6, 6, 6, 6, 6, 6, 6,
-				   6, 6, 6, 6, 6, 6, 6, 6, 6, 6, 6, 6, 6, 6, 6, 6,
-				   6, 6, 6, 6, 6, 6, 6, 6, 6, 6, 6, 6, 6, 6, 6, 6,
-				   6, 6, 6, 6, 6, 6, 6, 6, 6, 6, 6, 6, 6, 6, 6, 6,
-				   6, 6, 6, 6, 6, 6, 6, 6, 6, 6, 6, 6, 6, 6, 6, 6,
-				   6, 6, 6, 6, 6, 6, 6, 6, 6, 7, 7, 7, 7, 7, 7, 7,
-				   7, 7, 7, 7, 7, 7, 7, 7, 7, 7, 7, 7, 7, 7, 7, 7,
-				   7, 7, 7, 7, 7, 7, 7, 7, 7, 7, 7, 7, 7, 7, 7, 7,
-				   7, 7, 7, 7, 7, 7, 7, 7, 7, 7, 7, 7, 7, 7, 7, 7,
-				   7, 7, 7, 7, 7, 7, 7, 7, 7, 7, 7, 7, 7, 7, 7, 7,
-				   7,
-				   },
-				  },
-
-
-
-
-				  // quantization method 6, range 0..9
-				  {
-				   QUANT_10,
-				   {0, 64, 7, 57, 14, 50, 21, 43, 28, 36,},
-				   {
-					0.000000, 1.000000, 0.109375, 0.890625,
-					0.218750, 0.781250, 0.328125, 0.671875,
-					0.437500, 0.562500,},
-				   {0, 3, 0, 5, 2, 7, 4, 9, 6, 8,},
-				   {2, 1, 4, 1, 6, 3, 8, 5, 9, 7,},
-				   {
-					0, 0, 0, 0, 0, 0, 0, 0, 0, 0, 0, 0, 0, 0, 0, 0,
-					0, 0, 0, 0, 0, 0, 0, 0, 0, 0, 0, 0, 0, 0, 0, 0,
-					0, 0, 0, 0, 0, 0, 0, 0, 0, 0, 0, 0, 0, 0, 0, 0,
-					0, 0, 0, 0, 0, 0, 0, 0, 0, 2, 2, 2, 2, 2, 2, 2,
-					2, 2, 2, 2, 2, 2, 2, 2, 2, 2, 2, 2, 2, 2, 2, 2,
-					2, 2, 2, 2, 2, 2, 2, 2, 2, 2, 2, 2, 2, 2, 2, 2,
-					2, 2, 2, 2, 2, 2, 2, 2, 2, 2, 2, 2, 2, 2, 2, 2,
-					2, 2, 2, 2, 2, 2, 2, 2, 2, 2, 2, 2, 2, 2, 2, 2,
-					2, 2, 2, 2, 2, 2, 2, 2, 2, 2, 2, 2, 2, 2, 2, 2,
-					2, 2, 2, 2, 2, 2, 2, 2, 2, 2, 2, 2, 2, 2, 2, 2,
-					2, 2, 2, 2, 2, 2, 2, 2, 2, 4, 4, 4, 4, 4, 4, 4,
-					4, 4, 4, 4, 4, 4, 4, 4, 4, 4, 4, 4, 4, 4, 4, 4,
-					4, 4, 4, 4, 4, 4, 4, 4, 4, 4, 4, 4, 4, 4, 4, 4,
-					4, 4, 4, 4, 4, 4, 4, 4, 4, 4, 4, 4, 4, 4, 4, 4,
-					4, 4, 4, 4, 4, 4, 4, 4, 4, 4, 4, 4, 4, 4, 4, 4,
-					4, 4, 4, 4, 4, 4, 4, 4, 4, 4, 4, 4, 4, 4, 4, 4,
-					4, 4, 4, 4, 4, 4, 4, 4, 4, 4, 4, 4, 4, 4, 4, 4,
-					4, 4, 4, 4, 4, 4, 4, 4, 4, 6, 6, 6, 6, 6, 6, 6,
-					6, 6, 6, 6, 6, 6, 6, 6, 6, 6, 6, 6, 6, 6, 6, 6,
-					6, 6, 6, 6, 6, 6, 6, 6, 6, 6, 6, 6, 6, 6, 6, 6,
-					6, 6, 6, 6, 6, 6, 6, 6, 6, 6, 6, 6, 6, 6, 6, 6,
-					6, 6, 6, 6, 6, 6, 6, 6, 6, 6, 6, 6, 6, 6, 6, 6,
-					6, 6, 6, 6, 6, 6, 6, 6, 6, 6, 6, 6, 6, 6, 6, 6,
-					6, 6, 6, 6, 6, 6, 6, 6, 6, 6, 6, 6, 6, 6, 6, 6,
-					6, 6, 6, 6, 6, 6, 6, 6, 6, 8, 8, 8, 8, 8, 8, 8,
-					8, 8, 8, 8, 8, 8, 8, 8, 8, 8, 8, 8, 8, 8, 8, 8,
-					8, 8, 8, 8, 8, 8, 8, 8, 8, 8, 8, 8, 8, 8, 8, 8,
-					8, 8, 8, 8, 8, 8, 8, 8, 8, 8, 8, 8, 8, 8, 8, 8,
-					8, 8, 8, 8, 8, 8, 8, 8, 8, 8, 8, 8, 8, 8, 8, 8,
-					8, 8, 8, 8, 8, 8, 8, 8, 8, 8, 8, 8, 8, 8, 8, 8,
-					8, 8, 8, 8, 8, 8, 8, 8, 8, 8, 8, 8, 8, 8, 8, 8,
-					8, 8, 8, 8, 8, 8, 8, 8, 8, 8, 8, 8, 8, 8, 8, 8,
-					8, 9, 9, 9, 9, 9, 9, 9, 9, 9, 9, 9, 9, 9, 9, 9,
-					9, 9, 9, 9, 9, 9, 9, 9, 9, 9, 9, 9, 9, 9, 9, 9,
-					9, 9, 9, 9, 9, 9, 9, 9, 9, 9, 9, 9, 9, 9, 9, 9,
-					9, 9, 9, 9, 9, 9, 9, 9, 9, 9, 9, 9, 9, 9, 9, 9,
-					9, 9, 9, 9, 9, 9, 9, 9, 9, 9, 9, 9, 9, 9, 9, 9,
-					9, 9, 9, 9, 9, 9, 9, 9, 9, 9, 9, 9, 9, 9, 9, 9,
-					9, 9, 9, 9, 9, 9, 9, 9, 9, 9, 9, 9, 9, 9, 9, 9,
-					9, 9, 9, 9, 9, 9, 9, 9, 9, 7, 7, 7, 7, 7, 7, 7,
-					7, 7, 7, 7, 7, 7, 7, 7, 7, 7, 7, 7, 7, 7, 7, 7,
-					7, 7, 7, 7, 7, 7, 7, 7, 7, 7, 7, 7, 7, 7, 7, 7,
-					7, 7, 7, 7, 7, 7, 7, 7, 7, 7, 7, 7, 7, 7, 7, 7,
-					7, 7, 7, 7, 7, 7, 7, 7, 7, 7, 7, 7, 7, 7, 7, 7,
-					7, 7, 7, 7, 7, 7, 7, 7, 7, 7, 7, 7, 7, 7, 7, 7,
-					7, 7, 7, 7, 7, 7, 7, 7, 7, 7, 7, 7, 7, 7, 7, 7,
-					7, 7, 7, 7, 7, 7, 7, 7, 7, 5, 5, 5, 5, 5, 5, 5,
-					5, 5, 5, 5, 5, 5, 5, 5, 5, 5, 5, 5, 5, 5, 5, 5,
-					5, 5, 5, 5, 5, 5, 5, 5, 5, 5, 5, 5, 5, 5, 5, 5,
-					5, 5, 5, 5, 5, 5, 5, 5, 5, 5, 5, 5, 5, 5, 5, 5,
-					5, 5, 5, 5, 5, 5, 5, 5, 5, 5, 5, 5, 5, 5, 5, 5,
-					5, 5, 5, 5, 5, 5, 5, 5, 5, 5, 5, 5, 5, 5, 5, 5,
-					5, 5, 5, 5, 5, 5, 5, 5, 5, 5, 5, 5, 5, 5, 5, 5,
-					5, 5, 5, 5, 5, 5, 5, 5, 5, 3, 3, 3, 3, 3, 3, 3,
-					3, 3, 3, 3, 3, 3, 3, 3, 3, 3, 3, 3, 3, 3, 3, 3,
-					3, 3, 3, 3, 3, 3, 3, 3, 3, 3, 3, 3, 3, 3, 3, 3,
-					3, 3, 3, 3, 3, 3, 3, 3, 3, 3, 3, 3, 3, 3, 3, 3,
-					3, 3, 3, 3, 3, 3, 3, 3, 3, 3, 3, 3, 3, 3, 3, 3,
-					3, 3, 3, 3, 3, 3, 3, 3, 3, 3, 3, 3, 3, 3, 3, 3,
-					3, 3, 3, 3, 3, 3, 3, 3, 3, 3, 3, 3, 3, 3, 3, 3,
-					3, 3, 3, 3, 3, 3, 3, 3, 3, 1, 1, 1, 1, 1, 1, 1,
-					1, 1, 1, 1, 1, 1, 1, 1, 1, 1, 1, 1, 1, 1, 1, 1,
-					1, 1, 1, 1, 1, 1, 1, 1, 1, 1, 1, 1, 1, 1, 1, 1,
-					1, 1, 1, 1, 1, 1, 1, 1, 1, 1, 1, 1, 1, 1, 1, 1,
-					1,
-					},
-				   },
-
-
-
-
-				   // quantization method 7, range 0..11
-				   {
-					QUANT_12,
-					{0, 64, 17, 47, 5, 59, 23, 41, 11, 53, 28, 36,},
-					{
-					 0.000000, 1.000000, 0.265625, 0.734375,
-					 0.078125, 0.921875, 0.359375, 0.640625,
-					 0.171875, 0.828125, 0.437500, 0.562500,},
-					{0, 5, 8, 7, 0, 9, 2, 11, 4, 3, 6, 10,},
-					{4, 1, 6, 9, 8, 1, 10, 3, 2, 5, 11, 7,},
+				case 6:
+					e0 = UVec4((v[0] * v[3]) >> 8, (v[1] * v[3]) >> 8, (v[2] * v[3]) >> 8, 0xff);
+					e1 = UVec4(v[0], v[1], v[2], 0xff);
+					break;
+				case 7:
+					decodeHDREndpointMode7(e0, e1, v[0], v[1], v[2], v[3]);
+					break;
+				case 8:
+					if (v[1] + v[3] + v[5] >= v[0] + v[2] + v[4])
 					{
-					 0, 0, 0, 0, 0, 0, 0, 0, 0, 0, 0, 0, 0, 0, 0, 0,
-					 0, 0, 0, 0, 0, 0, 0, 0, 0, 0, 0, 0, 0, 0, 0, 0,
-					 0, 0, 0, 0, 0, 0, 0, 0, 0, 4, 4, 4, 4, 4, 4, 4,
-					 4, 4, 4, 4, 4, 4, 4, 4, 4, 4, 4, 4, 4, 4, 4, 4,
-					 4, 4, 4, 4, 4, 4, 4, 4, 4, 4, 4, 4, 4, 4, 4, 4,
-					 4, 4, 4, 4, 4, 4, 4, 4, 4, 4, 4, 4, 4, 4, 4, 4,
-					 4, 4, 4, 4, 4, 4, 4, 4, 4, 4, 4, 4, 4, 4, 4, 4,
-					 4, 4, 4, 4, 4, 4, 4, 4, 4, 4, 4, 4, 4, 4, 4, 4,
-					 4, 8, 8, 8, 8, 8, 8, 8, 8, 8, 8, 8, 8, 8, 8, 8,
-					 8, 8, 8, 8, 8, 8, 8, 8, 8, 8, 8, 8, 8, 8, 8, 8,
-					 8, 8, 8, 8, 8, 8, 8, 8, 8, 8, 8, 8, 8, 8, 8, 8,
-					 8, 8, 8, 8, 8, 8, 8, 8, 8, 8, 8, 8, 8, 8, 8, 8,
-					 8, 8, 8, 8, 8, 8, 8, 8, 8, 8, 8, 8, 8, 8, 8, 8,
-					 8, 8, 8, 8, 8, 8, 8, 8, 8, 8, 8, 8, 8, 8, 8, 8,
-					 8, 2, 2, 2, 2, 2, 2, 2, 2, 2, 2, 2, 2, 2, 2, 2,
-					 2, 2, 2, 2, 2, 2, 2, 2, 2, 2, 2, 2, 2, 2, 2, 2,
-					 2, 2, 2, 2, 2, 2, 2, 2, 2, 2, 2, 2, 2, 2, 2, 2,
-					 2, 2, 2, 2, 2, 2, 2, 2, 2, 2, 2, 2, 2, 2, 2, 2,
-					 2, 2, 2, 2, 2, 2, 2, 2, 2, 2, 2, 2, 2, 2, 2, 2,
-					 2, 2, 2, 2, 2, 2, 2, 2, 2, 2, 2, 2, 2, 2, 2, 2,
-					 2, 6, 6, 6, 6, 6, 6, 6, 6, 6, 6, 6, 6, 6, 6, 6,
-					 6, 6, 6, 6, 6, 6, 6, 6, 6, 6, 6, 6, 6, 6, 6, 6,
-					 6, 6, 6, 6, 6, 6, 6, 6, 6, 6, 6, 6, 6, 6, 6, 6,
-					 6, 6, 6, 6, 6, 6, 6, 6, 6, 6, 6, 6, 6, 6, 6, 6,
-					 6, 6, 6, 6, 6, 6, 6, 6, 6, 6, 6, 6, 6, 6, 6, 6,
-					 6, 6, 6, 6, 6, 6, 6, 6, 6, 10, 10, 10, 10, 10, 10, 10,
-					 10, 10, 10, 10, 10, 10, 10, 10, 10, 10, 10, 10, 10, 10, 10, 10,
-					 10, 10, 10, 10, 10, 10, 10, 10, 10, 10, 10, 10, 10, 10, 10, 10,
-					 10, 10, 10, 10, 10, 10, 10, 10, 10, 10, 10, 10, 10, 10, 10, 10,
-					 10, 10, 10, 10, 10, 10, 10, 10, 10, 10, 10, 10, 10, 10, 10, 10,
-					 10, 10, 10, 10, 10, 10, 10, 10, 10, 10, 10, 10, 10, 10, 10, 10,
-					 10, 10, 10, 10, 10, 10, 10, 10, 10, 10, 10, 10, 10, 10, 10, 10,
-					 10, 11, 11, 11, 11, 11, 11, 11, 11, 11, 11, 11, 11, 11, 11, 11,
-					 11, 11, 11, 11, 11, 11, 11, 11, 11, 11, 11, 11, 11, 11, 11, 11,
-					 11, 11, 11, 11, 11, 11, 11, 11, 11, 11, 11, 11, 11, 11, 11, 11,
-					 11, 11, 11, 11, 11, 11, 11, 11, 11, 11, 11, 11, 11, 11, 11, 11,
-					 11, 11, 11, 11, 11, 11, 11, 11, 11, 11, 11, 11, 11, 11, 11, 11,
-					 11, 11, 11, 11, 11, 11, 11, 11, 11, 11, 11, 11, 11, 11, 11, 11,
-					 11, 11, 11, 11, 11, 11, 11, 11, 11, 7, 7, 7, 7, 7, 7, 7,
-					 7, 7, 7, 7, 7, 7, 7, 7, 7, 7, 7, 7, 7, 7, 7, 7,
-					 7, 7, 7, 7, 7, 7, 7, 7, 7, 7, 7, 7, 7, 7, 7, 7,
-					 7, 7, 7, 7, 7, 7, 7, 7, 7, 7, 7, 7, 7, 7, 7, 7,
-					 7, 7, 7, 7, 7, 7, 7, 7, 7, 7, 7, 7, 7, 7, 7, 7,
-					 7, 7, 7, 7, 7, 7, 7, 7, 7, 7, 7, 7, 7, 7, 7, 7,
-					 7, 3, 3, 3, 3, 3, 3, 3, 3, 3, 3, 3, 3, 3, 3, 3,
-					 3, 3, 3, 3, 3, 3, 3, 3, 3, 3, 3, 3, 3, 3, 3, 3,
-					 3, 3, 3, 3, 3, 3, 3, 3, 3, 3, 3, 3, 3, 3, 3, 3,
-					 3, 3, 3, 3, 3, 3, 3, 3, 3, 3, 3, 3, 3, 3, 3, 3,
-					 3, 3, 3, 3, 3, 3, 3, 3, 3, 3, 3, 3, 3, 3, 3, 3,
-					 3, 3, 3, 3, 3, 3, 3, 3, 3, 3, 3, 3, 3, 3, 3, 3,
-					 3, 9, 9, 9, 9, 9, 9, 9, 9, 9, 9, 9, 9, 9, 9, 9,
-					 9, 9, 9, 9, 9, 9, 9, 9, 9, 9, 9, 9, 9, 9, 9, 9,
-					 9, 9, 9, 9, 9, 9, 9, 9, 9, 9, 9, 9, 9, 9, 9, 9,
-					 9, 9, 9, 9, 9, 9, 9, 9, 9, 9, 9, 9, 9, 9, 9, 9,
-					 9, 9, 9, 9, 9, 9, 9, 9, 9, 9, 9, 9, 9, 9, 9, 9,
-					 9, 9, 9, 9, 9, 9, 9, 9, 9, 9, 9, 9, 9, 9, 9, 9,
-					 9, 5, 5, 5, 5, 5, 5, 5, 5, 5, 5, 5, 5, 5, 5, 5,
-					 5, 5, 5, 5, 5, 5, 5, 5, 5, 5, 5, 5, 5, 5, 5, 5,
-					 5, 5, 5, 5, 5, 5, 5, 5, 5, 5, 5, 5, 5, 5, 5, 5,
-					 5, 5, 5, 5, 5, 5, 5, 5, 5, 5, 5, 5, 5, 5, 5, 5,
-					 5, 5, 5, 5, 5, 5, 5, 5, 5, 5, 5, 5, 5, 5, 5, 5,
-					 5, 5, 5, 5, 5, 5, 5, 5, 5, 1, 1, 1, 1, 1, 1, 1,
-					 1, 1, 1, 1, 1, 1, 1, 1, 1, 1, 1, 1, 1, 1, 1, 1,
-					 1, 1, 1, 1, 1, 1, 1, 1, 1, 1, 1, 1, 1, 1, 1, 1,
-					 1,
-					 },
-					},
-
-
-
-
-					// quantization method 8, range 0..15
+						e0 = UVec4(v[0], v[2], v[4], 0xff);
+						e1 = UVec4(v[1], v[3], v[5], 0xff);
+					}
+					else
 					{
-					 QUANT_16,
-					 {0, 4, 8, 12, 17, 21, 25, 29, 35, 39, 43, 47, 52, 56, 60, 64,},
-					 {
-					  0.000000, 0.062500, 0.125000, 0.187500,
-					  0.265625, 0.328125, 0.390625, 0.453125,
-					  0.546875, 0.609375, 0.671875, 0.734375,
-					  0.812500, 0.875000, 0.937500, 1.000000,},
-					 {0, 0, 1, 2, 3, 4, 5, 6, 7, 8, 9, 10, 11, 12, 13, 14,},
-					 {1, 2, 3, 4, 5, 6, 7, 8, 9, 10, 11, 12, 13, 14, 15, 15,},
-					 {
-					  0, 0, 0, 0, 0, 0, 0, 0, 0, 0, 0, 0, 0, 0, 0, 0,
-					  0, 0, 0, 0, 0, 0, 0, 0, 0, 0, 0, 0, 0, 0, 0, 0,
-					  0, 1, 1, 1, 1, 1, 1, 1, 1, 1, 1, 1, 1, 1, 1, 1,
-					  1, 1, 1, 1, 1, 1, 1, 1, 1, 1, 1, 1, 1, 1, 1, 1,
-					  1, 1, 1, 1, 1, 1, 1, 1, 1, 1, 1, 1, 1, 1, 1, 1,
-					  1, 1, 1, 1, 1, 1, 1, 1, 1, 1, 1, 1, 1, 1, 1, 1,
-					  1, 2, 2, 2, 2, 2, 2, 2, 2, 2, 2, 2, 2, 2, 2, 2,
-					  2, 2, 2, 2, 2, 2, 2, 2, 2, 2, 2, 2, 2, 2, 2, 2,
-					  2, 2, 2, 2, 2, 2, 2, 2, 2, 2, 2, 2, 2, 2, 2, 2,
-					  2, 2, 2, 2, 2, 2, 2, 2, 2, 2, 2, 2, 2, 2, 2, 2,
-					  2, 3, 3, 3, 3, 3, 3, 3, 3, 3, 3, 3, 3, 3, 3, 3,
-					  3, 3, 3, 3, 3, 3, 3, 3, 3, 3, 3, 3, 3, 3, 3, 3,
-					  3, 3, 3, 3, 3, 3, 3, 3, 3, 3, 3, 3, 3, 3, 3, 3,
-					  3, 3, 3, 3, 3, 3, 3, 3, 3, 3, 3, 3, 3, 3, 3, 3,
-					  3, 3, 3, 3, 3, 3, 3, 3, 3, 4, 4, 4, 4, 4, 4, 4,
-					  4, 4, 4, 4, 4, 4, 4, 4, 4, 4, 4, 4, 4, 4, 4, 4,
-					  4, 4, 4, 4, 4, 4, 4, 4, 4, 4, 4, 4, 4, 4, 4, 4,
-					  4, 4, 4, 4, 4, 4, 4, 4, 4, 4, 4, 4, 4, 4, 4, 4,
-					  4, 4, 4, 4, 4, 4, 4, 4, 4, 4, 4, 4, 4, 4, 4, 4,
-					  4, 5, 5, 5, 5, 5, 5, 5, 5, 5, 5, 5, 5, 5, 5, 5,
-					  5, 5, 5, 5, 5, 5, 5, 5, 5, 5, 5, 5, 5, 5, 5, 5,
-					  5, 5, 5, 5, 5, 5, 5, 5, 5, 5, 5, 5, 5, 5, 5, 5,
-					  5, 5, 5, 5, 5, 5, 5, 5, 5, 5, 5, 5, 5, 5, 5, 5,
-					  5, 6, 6, 6, 6, 6, 6, 6, 6, 6, 6, 6, 6, 6, 6, 6,
-					  6, 6, 6, 6, 6, 6, 6, 6, 6, 6, 6, 6, 6, 6, 6, 6,
-					  6, 6, 6, 6, 6, 6, 6, 6, 6, 6, 6, 6, 6, 6, 6, 6,
-					  6, 6, 6, 6, 6, 6, 6, 6, 6, 6, 6, 6, 6, 6, 6, 6,
-					  6, 7, 7, 7, 7, 7, 7, 7, 7, 7, 7, 7, 7, 7, 7, 7,
-					  7, 7, 7, 7, 7, 7, 7, 7, 7, 7, 7, 7, 7, 7, 7, 7,
-					  7, 7, 7, 7, 7, 7, 7, 7, 7, 7, 7, 7, 7, 7, 7, 7,
-					  7, 7, 7, 7, 7, 7, 7, 7, 7, 7, 7, 7, 7, 7, 7, 7,
-					  7, 7, 7, 7, 7, 7, 7, 7, 7, 7, 7, 7, 7, 7, 7, 7,
-					  7, 8, 8, 8, 8, 8, 8, 8, 8, 8, 8, 8, 8, 8, 8, 8,
-					  8, 8, 8, 8, 8, 8, 8, 8, 8, 8, 8, 8, 8, 8, 8, 8,
-					  8, 8, 8, 8, 8, 8, 8, 8, 8, 8, 8, 8, 8, 8, 8, 8,
-					  8, 8, 8, 8, 8, 8, 8, 8, 8, 8, 8, 8, 8, 8, 8, 8,
-					  8, 8, 8, 8, 8, 8, 8, 8, 8, 8, 8, 8, 8, 8, 8, 8,
-					  8, 9, 9, 9, 9, 9, 9, 9, 9, 9, 9, 9, 9, 9, 9, 9,
-					  9, 9, 9, 9, 9, 9, 9, 9, 9, 9, 9, 9, 9, 9, 9, 9,
-					  9, 9, 9, 9, 9, 9, 9, 9, 9, 9, 9, 9, 9, 9, 9, 9,
-					  9, 9, 9, 9, 9, 9, 9, 9, 9, 9, 9, 9, 9, 9, 9, 9,
-					  9, 10, 10, 10, 10, 10, 10, 10, 10, 10, 10, 10, 10, 10, 10, 10,
-					  10, 10, 10, 10, 10, 10, 10, 10, 10, 10, 10, 10, 10, 10, 10, 10,
-					  10, 10, 10, 10, 10, 10, 10, 10, 10, 10, 10, 10, 10, 10, 10, 10,
-					  10, 10, 10, 10, 10, 10, 10, 10, 10, 10, 10, 10, 10, 10, 10, 10,
-					  10, 11, 11, 11, 11, 11, 11, 11, 11, 11, 11, 11, 11, 11, 11, 11,
-					  11, 11, 11, 11, 11, 11, 11, 11, 11, 11, 11, 11, 11, 11, 11, 11,
-					  11, 11, 11, 11, 11, 11, 11, 11, 11, 11, 11, 11, 11, 11, 11, 11,
-					  11, 11, 11, 11, 11, 11, 11, 11, 11, 11, 11, 11, 11, 11, 11, 11,
-					  11, 11, 11, 11, 11, 11, 11, 11, 11, 12, 12, 12, 12, 12, 12, 12,
-					  12, 12, 12, 12, 12, 12, 12, 12, 12, 12, 12, 12, 12, 12, 12, 12,
-					  12, 12, 12, 12, 12, 12, 12, 12, 12, 12, 12, 12, 12, 12, 12, 12,
-					  12, 12, 12, 12, 12, 12, 12, 12, 12, 12, 12, 12, 12, 12, 12, 12,
-					  12, 12, 12, 12, 12, 12, 12, 12, 12, 12, 12, 12, 12, 12, 12, 12,
-					  12, 13, 13, 13, 13, 13, 13, 13, 13, 13, 13, 13, 13, 13, 13, 13,
-					  13, 13, 13, 13, 13, 13, 13, 13, 13, 13, 13, 13, 13, 13, 13, 13,
-					  13, 13, 13, 13, 13, 13, 13, 13, 13, 13, 13, 13, 13, 13, 13, 13,
-					  13, 13, 13, 13, 13, 13, 13, 13, 13, 13, 13, 13, 13, 13, 13, 13,
-					  13, 14, 14, 14, 14, 14, 14, 14, 14, 14, 14, 14, 14, 14, 14, 14,
-					  14, 14, 14, 14, 14, 14, 14, 14, 14, 14, 14, 14, 14, 14, 14, 14,
-					  14, 14, 14, 14, 14, 14, 14, 14, 14, 14, 14, 14, 14, 14, 14, 14,
-					  14, 14, 14, 14, 14, 14, 14, 14, 14, 14, 14, 14, 14, 14, 14, 14,
-					  14, 15, 15, 15, 15, 15, 15, 15, 15, 15, 15, 15, 15, 15, 15, 15,
-					  15, 15, 15, 15, 15, 15, 15, 15, 15, 15, 15, 15, 15, 15, 15, 15,
-					  15,
-					  },
-					 },
-
-
-
-
-					 // quantization method 9, range 0..19
-					 {
-					  QUANT_20,
-					  {0, 64, 16, 48, 3, 61, 19, 45, 6, 58, 23, 41, 9, 55, 26, 38, 13, 51, 29, 35,},
-					  {
-					   0.000000, 1.000000, 0.250000, 0.750000,
-					   0.046875, 0.953125, 0.296875, 0.703125,
-					   0.093750, 0.906250, 0.359375, 0.640625,
-					   0.140625, 0.859375, 0.406250, 0.593750,
-					   0.203125, 0.796875, 0.453125, 0.546875,},
-					  {0, 5, 16, 7, 0, 9, 2, 11, 4, 13, 6, 15, 8, 17, 10, 19, 12, 3, 14, 18,},
-					  {4, 1, 6, 17, 8, 1, 10, 3, 12, 5, 14, 7, 16, 9, 18, 11, 2, 13, 19, 15,},
-					  {
-					   0, 0, 0, 0, 0, 0, 0, 0, 0, 0, 0, 0, 0, 0, 0, 0,
-					   0, 0, 0, 0, 0, 0, 0, 0, 0, 4, 4, 4, 4, 4, 4, 4,
-					   4, 4, 4, 4, 4, 4, 4, 4, 4, 4, 4, 4, 4, 4, 4, 4,
-					   4, 4, 4, 4, 4, 4, 4, 4, 4, 4, 4, 4, 4, 4, 4, 4,
-					   4, 4, 4, 4, 4, 4, 4, 4, 4, 8, 8, 8, 8, 8, 8, 8,
-					   8, 8, 8, 8, 8, 8, 8, 8, 8, 8, 8, 8, 8, 8, 8, 8,
-					   8, 8, 8, 8, 8, 8, 8, 8, 8, 8, 8, 8, 8, 8, 8, 8,
-					   8, 8, 8, 8, 8, 8, 8, 8, 8, 12, 12, 12, 12, 12, 12, 12,
-					   12, 12, 12, 12, 12, 12, 12, 12, 12, 12, 12, 12, 12, 12, 12, 12,
-					   12, 12, 12, 12, 12, 12, 12, 12, 12, 12, 12, 12, 12, 12, 12, 12,
-					   12, 12, 12, 12, 12, 12, 12, 12, 12, 12, 12, 12, 12, 12, 12, 12,
-					   12, 16, 16, 16, 16, 16, 16, 16, 16, 16, 16, 16, 16, 16, 16, 16,
-					   16, 16, 16, 16, 16, 16, 16, 16, 16, 16, 16, 16, 16, 16, 16, 16,
-					   16, 16, 16, 16, 16, 16, 16, 16, 16, 16, 16, 16, 16, 16, 16, 16,
-					   16, 16, 16, 16, 16, 16, 16, 16, 16, 2, 2, 2, 2, 2, 2, 2,
-					   2, 2, 2, 2, 2, 2, 2, 2, 2, 2, 2, 2, 2, 2, 2, 2,
-					   2, 2, 2, 2, 2, 2, 2, 2, 2, 2, 2, 2, 2, 2, 2, 2,
-					   2, 2, 2, 2, 2, 2, 2, 2, 2, 6, 6, 6, 6, 6, 6, 6,
-					   6, 6, 6, 6, 6, 6, 6, 6, 6, 6, 6, 6, 6, 6, 6, 6,
-					   6, 6, 6, 6, 6, 6, 6, 6, 6, 6, 6, 6, 6, 6, 6, 6,
-					   6, 6, 6, 6, 6, 6, 6, 6, 6, 6, 6, 6, 6, 6, 6, 6,
-					   6, 10, 10, 10, 10, 10, 10, 10, 10, 10, 10, 10, 10, 10, 10, 10,
-					   10, 10, 10, 10, 10, 10, 10, 10, 10, 10, 10, 10, 10, 10, 10, 10,
-					   10, 10, 10, 10, 10, 10, 10, 10, 10, 10, 10, 10, 10, 10, 10, 10,
-					   10, 10, 10, 10, 10, 10, 10, 10, 10, 14, 14, 14, 14, 14, 14, 14,
-					   14, 14, 14, 14, 14, 14, 14, 14, 14, 14, 14, 14, 14, 14, 14, 14,
-					   14, 14, 14, 14, 14, 14, 14, 14, 14, 14, 14, 14, 14, 14, 14, 14,
-					   14, 14, 14, 14, 14, 14, 14, 14, 14, 18, 18, 18, 18, 18, 18, 18,
-					   18, 18, 18, 18, 18, 18, 18, 18, 18, 18, 18, 18, 18, 18, 18, 18,
-					   18, 18, 18, 18, 18, 18, 18, 18, 18, 18, 18, 18, 18, 18, 18, 18,
-					   18, 18, 18, 18, 18, 18, 18, 18, 18, 18, 18, 18, 18, 18, 18, 18,
-					   18, 18, 18, 18, 18, 18, 18, 18, 18, 18, 18, 18, 18, 18, 18, 18,
-					   18, 19, 19, 19, 19, 19, 19, 19, 19, 19, 19, 19, 19, 19, 19, 19,
-					   19, 19, 19, 19, 19, 19, 19, 19, 19, 19, 19, 19, 19, 19, 19, 19,
-					   19, 19, 19, 19, 19, 19, 19, 19, 19, 19, 19, 19, 19, 19, 19, 19,
-					   19, 19, 19, 19, 19, 19, 19, 19, 19, 19, 19, 19, 19, 19, 19, 19,
-					   19, 19, 19, 19, 19, 19, 19, 19, 19, 15, 15, 15, 15, 15, 15, 15,
-					   15, 15, 15, 15, 15, 15, 15, 15, 15, 15, 15, 15, 15, 15, 15, 15,
-					   15, 15, 15, 15, 15, 15, 15, 15, 15, 15, 15, 15, 15, 15, 15, 15,
-					   15, 15, 15, 15, 15, 15, 15, 15, 15, 11, 11, 11, 11, 11, 11, 11,
-					   11, 11, 11, 11, 11, 11, 11, 11, 11, 11, 11, 11, 11, 11, 11, 11,
-					   11, 11, 11, 11, 11, 11, 11, 11, 11, 11, 11, 11, 11, 11, 11, 11,
-					   11, 11, 11, 11, 11, 11, 11, 11, 11, 11, 11, 11, 11, 11, 11, 11,
-					   11, 7, 7, 7, 7, 7, 7, 7, 7, 7, 7, 7, 7, 7, 7, 7,
-					   7, 7, 7, 7, 7, 7, 7, 7, 7, 7, 7, 7, 7, 7, 7, 7,
-					   7, 7, 7, 7, 7, 7, 7, 7, 7, 7, 7, 7, 7, 7, 7, 7,
-					   7, 7, 7, 7, 7, 7, 7, 7, 7, 3, 3, 3, 3, 3, 3, 3,
-					   3, 3, 3, 3, 3, 3, 3, 3, 3, 3, 3, 3, 3, 3, 3, 3,
-					   3, 3, 3, 3, 3, 3, 3, 3, 3, 3, 3, 3, 3, 3, 3, 3,
-					   3, 3, 3, 3, 3, 3, 3, 3, 3, 17, 17, 17, 17, 17, 17, 17,
-					   17, 17, 17, 17, 17, 17, 17, 17, 17, 17, 17, 17, 17, 17, 17, 17,
-					   17, 17, 17, 17, 17, 17, 17, 17, 17, 17, 17, 17, 17, 17, 17, 17,
-					   17, 17, 17, 17, 17, 17, 17, 17, 17, 17, 17, 17, 17, 17, 17, 17,
-					   17, 13, 13, 13, 13, 13, 13, 13, 13, 13, 13, 13, 13, 13, 13, 13,
-					   13, 13, 13, 13, 13, 13, 13, 13, 13, 13, 13, 13, 13, 13, 13, 13,
-					   13, 13, 13, 13, 13, 13, 13, 13, 13, 13, 13, 13, 13, 13, 13, 13,
-					   13, 13, 13, 13, 13, 13, 13, 13, 13, 9, 9, 9, 9, 9, 9, 9,
-					   9, 9, 9, 9, 9, 9, 9, 9, 9, 9, 9, 9, 9, 9, 9, 9,
-					   9, 9, 9, 9, 9, 9, 9, 9, 9, 9, 9, 9, 9, 9, 9, 9,
-					   9, 9, 9, 9, 9, 9, 9, 9, 9, 5, 5, 5, 5, 5, 5, 5,
-					   5, 5, 5, 5, 5, 5, 5, 5, 5, 5, 5, 5, 5, 5, 5, 5,
-					   5, 5, 5, 5, 5, 5, 5, 5, 5, 5, 5, 5, 5, 5, 5, 5,
-					   5, 5, 5, 5, 5, 5, 5, 5, 5, 1, 1, 1, 1, 1, 1, 1,
-					   1, 1, 1, 1, 1, 1, 1, 1, 1, 1, 1, 1, 1, 1, 1, 1,
-					   1,
-					   },
-					  },
-
-
-
-
-					  // quantization method 10, range 0..23
-					  {
-					   QUANT_24,
-					   {0, 64, 8, 56, 16, 48, 24, 40, 2, 62, 11, 53, 19, 45, 27, 37, 5, 59, 13, 51, 22, 42, 30, 34,},
-					   {
-						0.000000, 1.000000, 0.125000, 0.875000,
-						0.250000, 0.750000, 0.375000, 0.625000,
-						0.031250, 0.968750, 0.171875, 0.828125,
-						0.296875, 0.703125, 0.421875, 0.578125,
-						0.078125, 0.921875, 0.203125, 0.796875,
-						0.343750, 0.656250, 0.468750, 0.531250,},
-					   {0, 9, 16, 11, 18, 13, 20, 15, 0, 17, 2, 19, 4, 21, 6, 23, 8, 3, 10, 5, 12, 7, 14, 22,},
-					   {8, 1, 10, 17, 12, 19, 14, 21, 16, 1, 18, 3, 20, 5, 22, 7, 2, 9, 4, 11, 6, 13, 23, 15,},
-					   {
-						0, 0, 0, 0, 0, 0, 0, 0, 0, 0, 0, 0, 0, 0, 0, 0,
-						0, 8, 8, 8, 8, 8, 8, 8, 8, 8, 8, 8, 8, 8, 8, 8,
-						8, 8, 8, 8, 8, 8, 8, 8, 8, 8, 8, 8, 8, 8, 8, 8,
-						8, 8, 8, 8, 8, 8, 8, 8, 8, 16, 16, 16, 16, 16, 16, 16,
-						16, 16, 16, 16, 16, 16, 16, 16, 16, 16, 16, 16, 16, 16, 16, 16,
-						16, 16, 16, 16, 16, 16, 16, 16, 16, 16, 16, 16, 16, 16, 16, 16,
-						16, 16, 16, 16, 16, 16, 16, 16, 16, 2, 2, 2, 2, 2, 2, 2,
-						2, 2, 2, 2, 2, 2, 2, 2, 2, 2, 2, 2, 2, 2, 2, 2,
-						2, 2, 2, 2, 2, 2, 2, 2, 2, 2, 2, 2, 2, 2, 2, 2,
-						2, 2, 2, 2, 2, 2, 2, 2, 2, 10, 10, 10, 10, 10, 10, 10,
-						10, 10, 10, 10, 10, 10, 10, 10, 10, 10, 10, 10, 10, 10, 10, 10,
-						10, 10, 10, 10, 10, 10, 10, 10, 10, 10, 10, 10, 10, 10, 10, 10,
-						10, 18, 18, 18, 18, 18, 18, 18, 18, 18, 18, 18, 18, 18, 18, 18,
-						18, 18, 18, 18, 18, 18, 18, 18, 18, 18, 18, 18, 18, 18, 18, 18,
-						18, 18, 18, 18, 18, 18, 18, 18, 18, 4, 4, 4, 4, 4, 4, 4,
-						4, 4, 4, 4, 4, 4, 4, 4, 4, 4, 4, 4, 4, 4, 4, 4,
-						4, 4, 4, 4, 4, 4, 4, 4, 4, 4, 4, 4, 4, 4, 4, 4,
-						4, 4, 4, 4, 4, 4, 4, 4, 4, 12, 12, 12, 12, 12, 12, 12,
-						12, 12, 12, 12, 12, 12, 12, 12, 12, 12, 12, 12, 12, 12, 12, 12,
-						12, 12, 12, 12, 12, 12, 12, 12, 12, 12, 12, 12, 12, 12, 12, 12,
-						12, 12, 12, 12, 12, 12, 12, 12, 12, 20, 20, 20, 20, 20, 20, 20,
-						20, 20, 20, 20, 20, 20, 20, 20, 20, 20, 20, 20, 20, 20, 20, 20,
-						20, 20, 20, 20, 20, 20, 20, 20, 20, 20, 20, 20, 20, 20, 20, 20,
-						20, 6, 6, 6, 6, 6, 6, 6, 6, 6, 6, 6, 6, 6, 6, 6,
-						6, 6, 6, 6, 6, 6, 6, 6, 6, 6, 6, 6, 6, 6, 6, 6,
-						6, 6, 6, 6, 6, 6, 6, 6, 6, 14, 14, 14, 14, 14, 14, 14,
-						14, 14, 14, 14, 14, 14, 14, 14, 14, 14, 14, 14, 14, 14, 14, 14,
-						14, 14, 14, 14, 14, 14, 14, 14, 14, 14, 14, 14, 14, 14, 14, 14,
-						14, 14, 14, 14, 14, 14, 14, 14, 14, 22, 22, 22, 22, 22, 22, 22,
-						22, 22, 22, 22, 22, 22, 22, 22, 22, 22, 22, 22, 22, 22, 22, 22,
-						22, 22, 22, 22, 22, 22, 22, 22, 22, 22, 22, 22, 22, 22, 22, 22,
-						22, 22, 22, 22, 22, 22, 22, 22, 22, 22, 22, 22, 22, 22, 22, 22,
-						22, 23, 23, 23, 23, 23, 23, 23, 23, 23, 23, 23, 23, 23, 23, 23,
-						23, 23, 23, 23, 23, 23, 23, 23, 23, 23, 23, 23, 23, 23, 23, 23,
-						23, 23, 23, 23, 23, 23, 23, 23, 23, 23, 23, 23, 23, 23, 23, 23,
-						23, 23, 23, 23, 23, 23, 23, 23, 23, 15, 15, 15, 15, 15, 15, 15,
-						15, 15, 15, 15, 15, 15, 15, 15, 15, 15, 15, 15, 15, 15, 15, 15,
-						15, 15, 15, 15, 15, 15, 15, 15, 15, 15, 15, 15, 15, 15, 15, 15,
-						15, 15, 15, 15, 15, 15, 15, 15, 15, 7, 7, 7, 7, 7, 7, 7,
-						7, 7, 7, 7, 7, 7, 7, 7, 7, 7, 7, 7, 7, 7, 7, 7,
-						7, 7, 7, 7, 7, 7, 7, 7, 7, 7, 7, 7, 7, 7, 7, 7,
-						7, 21, 21, 21, 21, 21, 21, 21, 21, 21, 21, 21, 21, 21, 21, 21,
-						21, 21, 21, 21, 21, 21, 21, 21, 21, 21, 21, 21, 21, 21, 21, 21,
-						21, 21, 21, 21, 21, 21, 21, 21, 21, 13, 13, 13, 13, 13, 13, 13,
-						13, 13, 13, 13, 13, 13, 13, 13, 13, 13, 13, 13, 13, 13, 13, 13,
-						13, 13, 13, 13, 13, 13, 13, 13, 13, 13, 13, 13, 13, 13, 13, 13,
-						13, 13, 13, 13, 13, 13, 13, 13, 13, 5, 5, 5, 5, 5, 5, 5,
-						5, 5, 5, 5, 5, 5, 5, 5, 5, 5, 5, 5, 5, 5, 5, 5,
-						5, 5, 5, 5, 5, 5, 5, 5, 5, 5, 5, 5, 5, 5, 5, 5,
-						5, 5, 5, 5, 5, 5, 5, 5, 5, 19, 19, 19, 19, 19, 19, 19,
-						19, 19, 19, 19, 19, 19, 19, 19, 19, 19, 19, 19, 19, 19, 19, 19,
-						19, 19, 19, 19, 19, 19, 19, 19, 19, 19, 19, 19, 19, 19, 19, 19,
-						19, 11, 11, 11, 11, 11, 11, 11, 11, 11, 11, 11, 11, 11, 11, 11,
-						11, 11, 11, 11, 11, 11, 11, 11, 11, 11, 11, 11, 11, 11, 11, 11,
-						11, 11, 11, 11, 11, 11, 11, 11, 11, 3, 3, 3, 3, 3, 3, 3,
-						3, 3, 3, 3, 3, 3, 3, 3, 3, 3, 3, 3, 3, 3, 3, 3,
-						3, 3, 3, 3, 3, 3, 3, 3, 3, 3, 3, 3, 3, 3, 3, 3,
-						3, 3, 3, 3, 3, 3, 3, 3, 3, 17, 17, 17, 17, 17, 17, 17,
-						17, 17, 17, 17, 17, 17, 17, 17, 17, 17, 17, 17, 17, 17, 17, 17,
-						17, 17, 17, 17, 17, 17, 17, 17, 17, 17, 17, 17, 17, 17, 17, 17,
-						17, 17, 17, 17, 17, 17, 17, 17, 17, 9, 9, 9, 9, 9, 9, 9,
-						9, 9, 9, 9, 9, 9, 9, 9, 9, 9, 9, 9, 9, 9, 9, 9,
-						9, 9, 9, 9, 9, 9, 9, 9, 9, 9, 9, 9, 9, 9, 9, 9,
-						9, 1, 1, 1, 1, 1, 1, 1, 1, 1, 1, 1, 1, 1, 1, 1,
-						1,
-						},
-					   },
-
-
-
-
-					   // quantization method 11, range 0..31
-					   {
-						QUANT_32,
-						{0, 2, 4, 6, 8, 10, 12, 14, 16, 18, 20, 22, 24, 26, 28, 30, 34, 36, 38, 40, 42, 44, 46, 48, 50, 52, 54, 56, 58, 60, 62, 64,},
-						{
-						 0.000000, 0.031250, 0.062500, 0.093750,
-						 0.125000, 0.156250, 0.187500, 0.218750,
-						 0.250000, 0.281250, 0.312500, 0.343750,
-						 0.375000, 0.406250, 0.437500, 0.468750,
-						 0.531250, 0.562500, 0.593750, 0.625000,
-						 0.656250, 0.687500, 0.718750, 0.750000,
-						 0.781250, 0.812500, 0.843750, 0.875000,
-						 0.906250, 0.937500, 0.968750, 1.000000,},
-						{0, 0, 1, 2, 3, 4, 5, 6, 7, 8, 9, 10, 11, 12, 13, 14, 15, 16, 17, 18, 19, 20, 21, 22, 23, 24, 25, 26, 27, 28, 29, 30,},
-						{1, 2, 3, 4, 5, 6, 7, 8, 9, 10, 11, 12, 13, 14, 15, 16, 17, 18, 19, 20, 21, 22, 23, 24, 25, 26, 27, 28, 29, 30, 31, 31,},
-						{
-						 0, 0, 0, 0, 0, 0, 0, 0, 0, 0, 0, 0, 0, 0, 0, 0,
-						 0, 1, 1, 1, 1, 1, 1, 1, 1, 1, 1, 1, 1, 1, 1, 1,
-						 1, 1, 1, 1, 1, 1, 1, 1, 1, 1, 1, 1, 1, 1, 1, 1,
-						 1, 2, 2, 2, 2, 2, 2, 2, 2, 2, 2, 2, 2, 2, 2, 2,
-						 2, 2, 2, 2, 2, 2, 2, 2, 2, 2, 2, 2, 2, 2, 2, 2,
-						 2, 3, 3, 3, 3, 3, 3, 3, 3, 3, 3, 3, 3, 3, 3, 3,
-						 3, 3, 3, 3, 3, 3, 3, 3, 3, 3, 3, 3, 3, 3, 3, 3,
-						 3, 4, 4, 4, 4, 4, 4, 4, 4, 4, 4, 4, 4, 4, 4, 4,
-						 4, 4, 4, 4, 4, 4, 4, 4, 4, 4, 4, 4, 4, 4, 4, 4,
-						 4, 5, 5, 5, 5, 5, 5, 5, 5, 5, 5, 5, 5, 5, 5, 5,
-						 5, 5, 5, 5, 5, 5, 5, 5, 5, 5, 5, 5, 5, 5, 5, 5,
-						 5, 6, 6, 6, 6, 6, 6, 6, 6, 6, 6, 6, 6, 6, 6, 6,
-						 6, 6, 6, 6, 6, 6, 6, 6, 6, 6, 6, 6, 6, 6, 6, 6,
-						 6, 7, 7, 7, 7, 7, 7, 7, 7, 7, 7, 7, 7, 7, 7, 7,
-						 7, 7, 7, 7, 7, 7, 7, 7, 7, 7, 7, 7, 7, 7, 7, 7,
-						 7, 8, 8, 8, 8, 8, 8, 8, 8, 8, 8, 8, 8, 8, 8, 8,
-						 8, 8, 8, 8, 8, 8, 8, 8, 8, 8, 8, 8, 8, 8, 8, 8,
-						 8, 9, 9, 9, 9, 9, 9, 9, 9, 9, 9, 9, 9, 9, 9, 9,
-						 9, 9, 9, 9, 9, 9, 9, 9, 9, 9, 9, 9, 9, 9, 9, 9,
-						 9, 10, 10, 10, 10, 10, 10, 10, 10, 10, 10, 10, 10, 10, 10, 10,
-						 10, 10, 10, 10, 10, 10, 10, 10, 10, 10, 10, 10, 10, 10, 10, 10,
-						 10, 11, 11, 11, 11, 11, 11, 11, 11, 11, 11, 11, 11, 11, 11, 11,
-						 11, 11, 11, 11, 11, 11, 11, 11, 11, 11, 11, 11, 11, 11, 11, 11,
-						 11, 12, 12, 12, 12, 12, 12, 12, 12, 12, 12, 12, 12, 12, 12, 12,
-						 12, 12, 12, 12, 12, 12, 12, 12, 12, 12, 12, 12, 12, 12, 12, 12,
-						 12, 13, 13, 13, 13, 13, 13, 13, 13, 13, 13, 13, 13, 13, 13, 13,
-						 13, 13, 13, 13, 13, 13, 13, 13, 13, 13, 13, 13, 13, 13, 13, 13,
-						 13, 14, 14, 14, 14, 14, 14, 14, 14, 14, 14, 14, 14, 14, 14, 14,
-						 14, 14, 14, 14, 14, 14, 14, 14, 14, 14, 14, 14, 14, 14, 14, 14,
-						 14, 15, 15, 15, 15, 15, 15, 15, 15, 15, 15, 15, 15, 15, 15, 15,
-						 15, 15, 15, 15, 15, 15, 15, 15, 15, 15, 15, 15, 15, 15, 15, 15,
-						 15, 15, 15, 15, 15, 15, 15, 15, 15, 15, 15, 15, 15, 15, 15, 15,
-						 15, 16, 16, 16, 16, 16, 16, 16, 16, 16, 16, 16, 16, 16, 16, 16,
-						 16, 16, 16, 16, 16, 16, 16, 16, 16, 16, 16, 16, 16, 16, 16, 16,
-						 16, 16, 16, 16, 16, 16, 16, 16, 16, 16, 16, 16, 16, 16, 16, 16,
-						 16, 17, 17, 17, 17, 17, 17, 17, 17, 17, 17, 17, 17, 17, 17, 17,
-						 17, 17, 17, 17, 17, 17, 17, 17, 17, 17, 17, 17, 17, 17, 17, 17,
-						 17, 18, 18, 18, 18, 18, 18, 18, 18, 18, 18, 18, 18, 18, 18, 18,
-						 18, 18, 18, 18, 18, 18, 18, 18, 18, 18, 18, 18, 18, 18, 18, 18,
-						 18, 19, 19, 19, 19, 19, 19, 19, 19, 19, 19, 19, 19, 19, 19, 19,
-						 19, 19, 19, 19, 19, 19, 19, 19, 19, 19, 19, 19, 19, 19, 19, 19,
-						 19, 20, 20, 20, 20, 20, 20, 20, 20, 20, 20, 20, 20, 20, 20, 20,
-						 20, 20, 20, 20, 20, 20, 20, 20, 20, 20, 20, 20, 20, 20, 20, 20,
-						 20, 21, 21, 21, 21, 21, 21, 21, 21, 21, 21, 21, 21, 21, 21, 21,
-						 21, 21, 21, 21, 21, 21, 21, 21, 21, 21, 21, 21, 21, 21, 21, 21,
-						 21, 22, 22, 22, 22, 22, 22, 22, 22, 22, 22, 22, 22, 22, 22, 22,
-						 22, 22, 22, 22, 22, 22, 22, 22, 22, 22, 22, 22, 22, 22, 22, 22,
-						 22, 23, 23, 23, 23, 23, 23, 23, 23, 23, 23, 23, 23, 23, 23, 23,
-						 23, 23, 23, 23, 23, 23, 23, 23, 23, 23, 23, 23, 23, 23, 23, 23,
-						 23, 24, 24, 24, 24, 24, 24, 24, 24, 24, 24, 24, 24, 24, 24, 24,
-						 24, 24, 24, 24, 24, 24, 24, 24, 24, 24, 24, 24, 24, 24, 24, 24,
-						 24, 25, 25, 25, 25, 25, 25, 25, 25, 25, 25, 25, 25, 25, 25, 25,
-						 25, 25, 25, 25, 25, 25, 25, 25, 25, 25, 25, 25, 25, 25, 25, 25,
-						 25, 26, 26, 26, 26, 26, 26, 26, 26, 26, 26, 26, 26, 26, 26, 26,
-						 26, 26, 26, 26, 26, 26, 26, 26, 26, 26, 26, 26, 26, 26, 26, 26,
-						 26, 27, 27, 27, 27, 27, 27, 27, 27, 27, 27, 27, 27, 27, 27, 27,
-						 27, 27, 27, 27, 27, 27, 27, 27, 27, 27, 27, 27, 27, 27, 27, 27,
-						 27, 28, 28, 28, 28, 28, 28, 28, 28, 28, 28, 28, 28, 28, 28, 28,
-						 28, 28, 28, 28, 28, 28, 28, 28, 28, 28, 28, 28, 28, 28, 28, 28,
-						 28, 29, 29, 29, 29, 29, 29, 29, 29, 29, 29, 29, 29, 29, 29, 29,
-						 29, 29, 29, 29, 29, 29, 29, 29, 29, 29, 29, 29, 29, 29, 29, 29,
-						 29, 30, 30, 30, 30, 30, 30, 30, 30, 30, 30, 30, 30, 30, 30, 30,
-						 30, 30, 30, 30, 30, 30, 30, 30, 30, 30, 30, 30, 30, 30, 30, 30,
-						 30, 31, 31, 31, 31, 31, 31, 31, 31, 31, 31, 31, 31, 31, 31, 31,
-						 31,
-						 },
-						},
-
-		};
-
-
-
-
-		uint16 lns_to_sf16(uint16 p)
-		{
-
-			uint16 mc = p & 0x7FF;
-			uint16 ec = p >> 11;
-			uint16 mt;
-			if (mc < 512)
-				mt = 3 * mc;
-			else if (mc < 1536)
-				mt = 4 * mc - 512;
-			else
-				mt = 5 * mc - 2048;
-
-			uint16 res = (ec << 10) | (mt >> 3);
-			if (res >= 0x7BFF)
-				res = 0x7BFF;
-			return res;
-		}
-
-
-		// helper function to initialize the orig-data from the work-data
-		void imageblock_initialize_orig_from_work(imageblock* pb, int pixelcount)
-		{
-			int i;
-			float* fptr = pb->orig_data;
-			float* wptr = pb->work_data;
-
-			for (i = 0; i < pixelcount; i++)
-			{
-				if (pb->rgb_lns[i])
-				{
-					fptr[0] = sf16_to_float(lns_to_sf16((uint16)wptr[0]));
-					fptr[1] = sf16_to_float(lns_to_sf16((uint16)wptr[1]));
-					fptr[2] = sf16_to_float(lns_to_sf16((uint16)wptr[2]));
-				}
-				else
+						e0 = blueContract(v[1], v[3], v[5], 0xff).asUint();
+						e1 = blueContract(v[0], v[2], v[4], 0xff).asUint();
+					}
+					break;
+				case 9:
 				{
-					fptr[0] = sf16_to_float(unorm16_to_sf16((uint16)wptr[0]));
-					fptr[1] = sf16_to_float(unorm16_to_sf16((uint16)wptr[1]));
-					fptr[2] = sf16_to_float(unorm16_to_sf16((uint16)wptr[2]));
+					deInt32 v0 = (deInt32)v[0];
+					deInt32 v1 = (deInt32)v[1];
+					deInt32 v2 = (deInt32)v[2];
+					deInt32 v3 = (deInt32)v[3];
+					deInt32 v4 = (deInt32)v[4];
+					deInt32 v5 = (deInt32)v[5];
+					bitTransferSigned(v1, v0);
+					bitTransferSigned(v3, v2);
+					bitTransferSigned(v5, v4);
+					if (v1 + v3 + v5 >= 0)
+					{
+						e0 = clampedRGBA(IVec4(v0, v2, v4, 0xff));
+						e1 = clampedRGBA(IVec4(v0 + v1, v2 + v3, v4 + v5, 0xff));
+					}
+					else
+					{
+						e0 = clampedRGBA(blueContract(v0 + v1, v2 + v3, v4 + v5, 0xff));
+						e1 = clampedRGBA(blueContract(v0, v2, v4, 0xff));
+					}
+					break;
 				}
-
-				if (pb->alpha_lns[i])
+				case 10:
+					e0 = UVec4((v[0] * v[3]) >> 8, (v[1] * v[3]) >> 8, (v[2] * v[3]) >> 8, v[4]);
+					e1 = UVec4(v[0], v[1], v[2], v[5]);
+					break;
+				case 11:
+					decodeHDREndpointMode11(e0, e1, v[0], v[1], v[2], v[3], v[4], v[5]);
+					break;
+				case 12:
+					if (v[1] + v[3] + v[5] >= v[0] + v[2] + v[4])
+					{
+						e0 = UVec4(v[0], v[2], v[4], v[6]);
+						e1 = UVec4(v[1], v[3], v[5], v[7]);
+					}
+					else
+					{
+						e0 = clampedRGBA(blueContract(v[1], v[3], v[5], v[7]));
+						e1 = clampedRGBA(blueContract(v[0], v[2], v[4], v[6]));
+					}
+					break;
+				case 13:
 				{
-					fptr[3] = sf16_to_float(lns_to_sf16((uint16)wptr[3]));
+					deInt32 v0 = (deInt32)v[0];
+					deInt32 v1 = (deInt32)v[1];
+					deInt32 v2 = (deInt32)v[2];
+					deInt32 v3 = (deInt32)v[3];
+					deInt32 v4 = (deInt32)v[4];
+					deInt32 v5 = (deInt32)v[5];
+					deInt32 v6 = (deInt32)v[6];
+					deInt32 v7 = (deInt32)v[7];
+					bitTransferSigned(v1, v0);
+					bitTransferSigned(v3, v2);
+					bitTransferSigned(v5, v4);
+					bitTransferSigned(v7, v6);
+					if (v1 + v3 + v5 >= 0)
+					{
+						e0 = clampedRGBA(IVec4(v0, v2, v4, v6));
+						e1 = clampedRGBA(IVec4(v0 + v1, v2 + v3, v4 + v5, v6 + v7));
+					}
+					else
+					{
+						e0 = clampedRGBA(blueContract(v0 + v1, v2 + v3, v4 + v5, v6 + v7));
+						e1 = clampedRGBA(blueContract(v0, v2, v4, v6));
+					}
+					break;
 				}
-				else
-				{
-					fptr[3] = sf16_to_float(unorm16_to_sf16((uint16)wptr[3]));
+				case 14:
+					decodeHDREndpointMode11(e0, e1, v[0], v[1], v[2], v[3], v[4], v[5]);
+					e0.w() = v[6];
+					e1.w() = v[7];
+					break;
+				case 15:
+					decodeHDREndpointMode15(e0, e1, v[0], v[1], v[2], v[3], v[4], v[5], v[6], v[7]);
+					break;
+				default:
+					DE_ASSERT(false);
 				}
-
-				fptr += 4;
-				wptr += 4;
 			}
 		}
-
-
-		inline void decompress_symbolic_block(astc_decode_mode decode_mode,
-			const symbolic_compressed_block* scb,
-			imageblock* blk)
+		void computeColorEndpoints(ColorEndpointPair* dst, const Block128& blockData, const deUint32* endpointModes, int numPartitions, int numColorEndpointValues, const ISEParams& iseParams, int numBitsAvailable)
 		{
-			int i;
-
-			//        if (s_debuglog)
-			//        {
-			//            UE_LOG(LogMutableCore,Warning," uncompressing interesting block" );
-			//        }
-
-					// if we detected an error-block, blow up immediately.
-			if (scb->error_block)
+			const int			colorEndpointDataStart = numPartitions == 1 ? 17 : 29;
+			ISEDecodedResult	colorEndpointData[18];
 			{
-				//check(false);
-				return;
+				BitAccessStream dataStream(blockData, colorEndpointDataStart, numBitsAvailable, true);
+				decodeISE(&colorEndpointData[0], numColorEndpointValues, dataStream, iseParams);
 			}
-
-
-			if (scb->block_mode < 0)
 			{
-				float red = 0, green = 0, blue = 0, alpha = 0;
-				int use_lns = 0;
-
-				if (scb->block_mode == -2)
+				deUint32 unquantizedEndpoints[18];
+				unquantizeColorEndpoints(&unquantizedEndpoints[0], &colorEndpointData[0], numColorEndpointValues, iseParams);
+				decodeColorEndpoints(dst, &unquantizedEndpoints[0], &endpointModes[0], numPartitions);
+			}
+		}
+		void unquantizeWeights(deUint32 dst[64], const ISEDecodedResult* weightGrid, const ASTCBlockMode& blockMode)
+		{
+			const int			numWeights = computeNumWeights(blockMode);
+			const ISEParams& iseParams = blockMode.weightISEParams;
+			if (iseParams.mode == ISEMODE_TRIT || iseParams.mode == ISEMODE_QUINT)
+			{
+				const int rangeCase = iseParams.numBits * 2 + (iseParams.mode == ISEMODE_QUINT ? 1 : 0);
+				if (rangeCase == 0 || rangeCase == 1)
 				{
-					// For sRGB decoding, we should return only the top 8 bits.
-					int mask = (decode_mode == DECODE_LDR_SRGB) ? 0xFF00 : 0xFFFF;
-
-					red = sf16_to_float(unorm16_to_sf16(uint16(scb->constant_color[0] & mask)));
-					green = sf16_to_float(unorm16_to_sf16(uint16(scb->constant_color[1] & mask)));
-					blue = sf16_to_float(unorm16_to_sf16(uint16(scb->constant_color[2] & mask)));
-					alpha = sf16_to_float(unorm16_to_sf16(uint16(scb->constant_color[3] & mask)));
-					use_lns = 0;
+					static const deUint32 map0[3] = { 0, 32, 63 };
+					static const deUint32 map1[5] = { 0, 16, 32, 47, 63 };
+					const deUint32* const map = rangeCase == 0 ? &map0[0] : &map1[0];
+					for (int i = 0; i < numWeights; i++)
+					{
+						DE_ASSERT(weightGrid[i].v < (rangeCase == 0 ? 3u : 5u));
+						dst[i] = map[weightGrid[i].v];
+					}
 				}
 				else
 				{
-					switch (decode_mode)
+					DE_ASSERT(rangeCase <= 6);
+					static const deUint32	Ca[5] = { 50, 28, 23, 13, 11 };
+					const deUint32			C = Ca[rangeCase - 2];
+					for (int weightNdx = 0; weightNdx < numWeights; weightNdx++)
 					{
-					case DECODE_LDR_SRGB:
-						red = 1.0f;
-						green = 0.0f;
-						blue = 1.0f;
-						alpha = 1.0f;
-						use_lns = 0;
-						break;
-					case DECODE_LDR:
-						red = 0.0f;
-						green = 0.0f;
-						blue = 0.0f;
-						alpha = 0.0f;
-						use_lns = 0;
-						break;
-					case DECODE_HDR:
-						// constant-color block; unpack from FP16 to FP32.
-						red = sf16_to_float((sf16)scb->constant_color[0]);
-						green = sf16_to_float((sf16)scb->constant_color[1]);
-						blue = sf16_to_float((sf16)scb->constant_color[2]);
-						alpha = sf16_to_float((sf16)scb->constant_color[3]);
-						use_lns = 1;
-						break;
+						const deUint32 a = getBit(weightGrid[weightNdx].m, 0);
+						const deUint32 b = getBit(weightGrid[weightNdx].m, 1);
+						const deUint32 c = getBit(weightGrid[weightNdx].m, 2);
+						const deUint32 A = a == 0 ? 0 : (1 << 7) - 1;
+						const deUint32 B = rangeCase == 2 ? 0
+							: rangeCase == 3 ? 0
+							: rangeCase == 4 ? (b << 6) | (b << 2) | (b << 0)
+							: rangeCase == 5 ? (b << 6) | (b << 1)
+							: rangeCase == 6 ? (c << 6) | (b << 5) | (c << 1) | (b << 0)
+							: (deUint32)-1;
+						dst[weightNdx] = (((weightGrid[weightNdx].tq * C + B) ^ A) >> 2) | (A & 0x20);
 					}
 				}
-
-				for (i = 0; i < 16; i++)
-				{
-					blk->orig_data[4 * i] = red;
-					blk->orig_data[4 * i + 1] = green;
-					blk->orig_data[4 * i + 2] = blue;
-					blk->orig_data[4 * i + 3] = alpha;
-					blk->rgb_lns[i] = (uint8)use_lns;
-					blk->alpha_lns[i] = (uint8)use_lns;
-				}
-
-
-				imageblock_initialize_work_from_orig(blk, 16);
-				return;
 			}
-
-
-			// get the appropriate partition-table entry
-			int partition_count = scb->partition_count;
-			check(partition_count == 1);
-			const partition_info* pt = get_partition_table();
-			pt += scb->partition_index;
-
-			// get the appropriate block descriptor
-			const block_size_descriptor* bsd = get_block_size_descriptor();
-			const decimation_table* const* ixtab2 = bsd->decimation_tables;
-
-
-			const decimation_table* it = ixtab2[bsd->block_modes[scb->block_mode].decimation_mode];
-
-			int is_dual_plane = bsd->block_modes[scb->block_mode].is_dual_plane;
-
-			int weight_quantization_level = bsd->block_modes[scb->block_mode].quantization_mode;
-
-
-			// decode the color endpoints
-			ushort4 color_endpoint0[4];
-			ushort4 color_endpoint1[4];
-			int rgb_hdr_endpoint[4] = { 0 };
-			int alpha_hdr_endpoint[4] = { 0 };
-			int nan_endpoint[4] = { 0 };
-
-			for (i = 0; i < partition_count; i++)
-				unpack_color_endpoints(decode_mode,
-					scb->color_formats[i],
-					scb->color_quantization_level, scb->color_values[i], &(rgb_hdr_endpoint[i]), &(alpha_hdr_endpoint[i]), &(nan_endpoint[i]), &(color_endpoint0[i]), &(color_endpoint1[i]));
-
-
-
-
-
-			// first unquantize the weights
-			int uq_plane1_weights[MAX_WEIGHTS_PER_BLOCK];
-			int uq_plane2_weights[MAX_WEIGHTS_PER_BLOCK];
-			int weight_count = it->num_weights;
-
-
-			const quantization_and_transfer_table* qat = &(quant_and_xfer_tables[weight_quantization_level]);
-
-			for (i = 0; i < weight_count; i++)
+			else
 			{
-				uq_plane1_weights[i] = qat->unquantized_value[scb->plane1_weights[i]];
+				DE_ASSERT(iseParams.mode == ISEMODE_PLAIN_BIT);
+				for (int weightNdx = 0; weightNdx < numWeights; weightNdx++)
+					dst[weightNdx] = bitReplicationScale(weightGrid[weightNdx].v, iseParams.numBits, 6);
 			}
-			if (is_dual_plane)
+			for (int weightNdx = 0; weightNdx < numWeights; weightNdx++)
+				dst[weightNdx] += dst[weightNdx] > 32 ? 1 : 0;
+			// Initialize nonexistent weights to poison values
+			for (int weightNdx = numWeights; weightNdx < 64; weightNdx++)
+				dst[weightNdx] = ~0u;
+		}
+		void interpolateWeights(TexelWeightPair* dst, const deUint32(&unquantizedWeights)[64], int blockWidth, int blockHeight, const ASTCBlockMode& blockMode)
+		{
+			const int		numWeightsPerTexel = blockMode.isDualPlane ? 2 : 1;
+			const deUint32	scaleX = (1024 + blockWidth / 2) / (blockWidth - 1);
+			const deUint32	scaleY = (1024 + blockHeight / 2) / (blockHeight - 1);
+			DE_ASSERT(blockMode.weightGridWidth * blockMode.weightGridHeight * numWeightsPerTexel <= DE_LENGTH_OF_ARRAY(unquantizedWeights));
+			for (int texelY = 0; texelY < blockHeight; texelY++)
 			{
-				for (i = 0; i < weight_count; i++)
-					uq_plane2_weights[i] = qat->unquantized_value[scb->plane2_weights[i]];
+				for (int texelX = 0; texelX < blockWidth; texelX++)
+				{
+					const deUint32 gX = (scaleX * texelX * (blockMode.weightGridWidth - 1) + 32) >> 6;
+					const deUint32 gY = (scaleY * texelY * (blockMode.weightGridHeight - 1) + 32) >> 6;
+					const deUint32 jX = gX >> 4;
+					const deUint32 jY = gY >> 4;
+					const deUint32 fX = gX & 0xf;
+					const deUint32 fY = gY & 0xf;
+					const deUint32 w11 = (fX * fY + 8) >> 4;
+					const deUint32 w10 = fY - w11;
+					const deUint32 w01 = fX - w11;
+					const deUint32 w00 = 16 - fX - fY + w11;
+					const deUint32 i00 = jY * blockMode.weightGridWidth + jX;
+					const deUint32 i01 = i00 + 1;
+					const deUint32 i10 = i00 + blockMode.weightGridWidth;
+					const deUint32 i11 = i00 + blockMode.weightGridWidth + 1;
+					// These addresses can be out of bounds, but respective weights will be 0 then.
+					DE_ASSERT(deInBounds32(i00, 0, blockMode.weightGridWidth * blockMode.weightGridHeight) || w00 == 0);
+					DE_ASSERT(deInBounds32(i01, 0, blockMode.weightGridWidth * blockMode.weightGridHeight) || w01 == 0);
+					DE_ASSERT(deInBounds32(i10, 0, blockMode.weightGridWidth * blockMode.weightGridHeight) || w10 == 0);
+					DE_ASSERT(deInBounds32(i11, 0, blockMode.weightGridWidth * blockMode.weightGridHeight) || w11 == 0);
+					for (int texelWeightNdx = 0; texelWeightNdx < numWeightsPerTexel; texelWeightNdx++)
+					{
+						// & 0x3f clamps address to bounds of unquantizedWeights
+						const deUint32 p00 = unquantizedWeights[(i00 * numWeightsPerTexel + texelWeightNdx) & 0x3f];
+						const deUint32 p01 = unquantizedWeights[(i01 * numWeightsPerTexel + texelWeightNdx) & 0x3f];
+						const deUint32 p10 = unquantizedWeights[(i10 * numWeightsPerTexel + texelWeightNdx) & 0x3f];
+						const deUint32 p11 = unquantizedWeights[(i11 * numWeightsPerTexel + texelWeightNdx) & 0x3f];
+						dst[texelY * blockWidth + texelX].w[texelWeightNdx] = (p00 * w00 + p01 * w01 + p10 * w10 + p11 * w11 + 8) >> 4;
+					}
+				}
 			}
-
-
-			// then undecimate them.
-			int weights[MAX_TEXELS_PER_BLOCK] = { 0 };
-			int plane2_weights[MAX_TEXELS_PER_BLOCK] = { 0 };
-
-
-			int texels_per_block = 16;
-			for (i = 0; i < texels_per_block; i++)
-				weights[i] = compute_value_of_texel_int(i, it, uq_plane1_weights);
-
-			if (is_dual_plane)
-				for (i = 0; i < texels_per_block; i++)
-					plane2_weights[i] = compute_value_of_texel_int(i, it, uq_plane2_weights);
-
-
-			int plane2_color_component = scb->plane2_color_component;
-
-
-			// now that we have endpoint colors and weights, we can unpack actual colors for
-			// each texel.
-			for (i = 0; i < texels_per_block; i++)
+		}
+		void computeTexelWeights(TexelWeightPair* dst, const Block128& blockData, int blockWidth, int blockHeight, const ASTCBlockMode& blockMode)
+		{
+			ISEDecodedResult weightGrid[64];
 			{
-				int partition = pt->partition_of_texel[i];
-
-				ushort4 color = lerp_color_int(decode_mode,
-					color_endpoint0[partition],
-					color_endpoint1[partition],
-					weights[i],
-					plane2_weights[i],
-					is_dual_plane ? plane2_color_component : -1);
-
-				blk->rgb_lns[i] = (uint8)rgb_hdr_endpoint[partition];
-				blk->alpha_lns[i] = (uint8)alpha_hdr_endpoint[partition];
-
-				blk->work_data[4 * i] = color.x;
-				blk->work_data[4 * i + 1] = color.y;
-				blk->work_data[4 * i + 2] = color.z;
-				blk->work_data[4 * i + 3] = color.w;
+				BitAccessStream dataStream(blockData, 127, computeNumRequiredBits(blockMode.weightISEParams, computeNumWeights(blockMode)), false);
+				decodeISE(&weightGrid[0], computeNumWeights(blockMode), dataStream, blockMode.weightISEParams);
+			}
+			{
+				deUint32 unquantizedWeights[64];
+				unquantizeWeights(&unquantizedWeights[0], &weightGrid[0], blockMode);
+				interpolateWeights(dst, unquantizedWeights, blockWidth, blockHeight, blockMode);
+			}
+		}
+		inline deUint32 hash52(deUint32 v)
+		{
+			deUint32 p = v;
+			p ^= p >> 15;	p -= p << 17;	p += p << 7;	p += p << 4;
+			p ^= p >> 5;	p += p << 16;	p ^= p >> 7;	p ^= p >> 3;
+			p ^= p << 6;	p ^= p >> 17;
+			return p;
+		}
+		int computeTexelPartition(deUint32 seedIn, deUint32 xIn, deUint32 yIn, deUint32 zIn, int numPartitions, bool smallBlock)
+		{
+			DE_ASSERT(zIn == 0);
+			const deUint32	x = smallBlock ? xIn << 1 : xIn;
+			const deUint32	y = smallBlock ? yIn << 1 : yIn;
+			const deUint32	z = smallBlock ? zIn << 1 : zIn;
+			const deUint32	seed = seedIn + 1024 * (numPartitions - 1);
+			const deUint32	rnum = hash52(seed);
+			deUint8			seed1 = (deUint8)(rnum & 0xf);
+			deUint8			seed2 = (deUint8)((rnum >> 4) & 0xf);
+			deUint8			seed3 = (deUint8)((rnum >> 8) & 0xf);
+			deUint8			seed4 = (deUint8)((rnum >> 12) & 0xf);
+			deUint8			seed5 = (deUint8)((rnum >> 16) & 0xf);
+			deUint8			seed6 = (deUint8)((rnum >> 20) & 0xf);
+			deUint8			seed7 = (deUint8)((rnum >> 24) & 0xf);
+			deUint8			seed8 = (deUint8)((rnum >> 28) & 0xf);
+			deUint8			seed9 = (deUint8)((rnum >> 18) & 0xf);
+			deUint8			seed10 = (deUint8)((rnum >> 22) & 0xf);
+			deUint8			seed11 = (deUint8)((rnum >> 26) & 0xf);
+			deUint8			seed12 = (deUint8)(((rnum >> 30) | (rnum << 2)) & 0xf);
+			seed1 = (deUint8)(seed1 * seed1);
+			seed2 = (deUint8)(seed2 * seed2);
+			seed3 = (deUint8)(seed3 * seed3);
+			seed4 = (deUint8)(seed4 * seed4);
+			seed5 = (deUint8)(seed5 * seed5);
+			seed6 = (deUint8)(seed6 * seed6);
+			seed7 = (deUint8)(seed7 * seed7);
+			seed8 = (deUint8)(seed8 * seed8);
+			seed9 = (deUint8)(seed9 * seed9);
+			seed10 = (deUint8)(seed10 * seed10);
+			seed11 = (deUint8)(seed11 * seed11);
+			seed12 = (deUint8)(seed12 * seed12);
+			const int shA = (seed & 2) != 0 ? 4 : 5;
+			const int shB = numPartitions == 3 ? 6 : 5;
+			const int sh1 = (seed & 1) != 0 ? shA : shB;
+			const int sh2 = (seed & 1) != 0 ? shB : shA;
+			const int sh3 = (seed & 0x10) != 0 ? sh1 : sh2;
+			seed1 = (deUint8)(seed1 >> sh1);
+			seed2 = (deUint8)(seed2 >> sh2);
+			seed3 = (deUint8)(seed3 >> sh1);
+			seed4 = (deUint8)(seed4 >> sh2);
+			seed5 = (deUint8)(seed5 >> sh1);
+			seed6 = (deUint8)(seed6 >> sh2);
+			seed7 = (deUint8)(seed7 >> sh1);
+			seed8 = (deUint8)(seed8 >> sh2);
+			seed9 = (deUint8)(seed9 >> sh3);
+			seed10 = (deUint8)(seed10 >> sh3);
+			seed11 = (deUint8)(seed11 >> sh3);
+			seed12 = (deUint8)(seed12 >> sh3);
+			const int a = 0x3f & (seed1 * x + seed2 * y + seed11 * z + (rnum >> 14));
+			const int b = 0x3f & (seed3 * x + seed4 * y + seed12 * z + (rnum >> 10));
+			const int c = numPartitions >= 3 ? 0x3f & (seed5 * x + seed6 * y + seed9 * z + (rnum >> 6)) : 0;
+			const int d = numPartitions >= 4 ? 0x3f & (seed7 * x + seed8 * y + seed10 * z + (rnum >> 2)) : 0;
+			return a >= b && a >= c && a >= d ? 0
+				: b >= c && b >= d ? 1
+				: c >= d ? 2
+				: 3;
+		}
+		DecompressResult setTexelColors(void* dst, ColorEndpointPair* colorEndpoints, TexelWeightPair* texelWeights, int ccs, deUint32 partitionIndexSeed,
+			int numPartitions, int blockWidth, int blockHeight, bool isSRGB, bool isLDRMode, const deUint32* colorEndpointModes)
+		{
+			const bool			smallBlock = blockWidth * blockHeight < 31;
+			DecompressResult	result = DECOMPRESS_RESULT_VALID_BLOCK;
+			bool				isHDREndpoint[4];
+			for (int i = 0; i < numPartitions; i++)
+			{
+				isHDREndpoint[i] = isColorEndpointModeHDR(colorEndpointModes[i]);
+
+				// rg - REMOVING HDR SUPPORT FOR NOW
+				if (isHDREndpoint[i])
+					return DECOMPRESS_RESULT_ERROR;
+			}
+
+			for (int texelY = 0; texelY < blockHeight; texelY++)
+				for (int texelX = 0; texelX < blockWidth; texelX++)
+				{
+					const int				texelNdx = texelY * blockWidth + texelX;
+					const int				colorEndpointNdx = numPartitions == 1 ? 0 : computeTexelPartition(partitionIndexSeed, texelX, texelY, 0, numPartitions, smallBlock);
+					DE_ASSERT(colorEndpointNdx < numPartitions);
+					const UVec4& e0 = colorEndpoints[colorEndpointNdx].e0;
+					const UVec4& e1 = colorEndpoints[colorEndpointNdx].e1;
+					const TexelWeightPair& weight = texelWeights[texelNdx];
+					if (isLDRMode && isHDREndpoint[colorEndpointNdx])
+					{
+						if (isSRGB)
+						{
+							((deUint8*)dst)[texelNdx * 4 + 0] = 0xff;
+							((deUint8*)dst)[texelNdx * 4 + 1] = 0;
+							((deUint8*)dst)[texelNdx * 4 + 2] = 0xff;
+							((deUint8*)dst)[texelNdx * 4 + 3] = 0xff;
+						}
+						else
+						{
+							((float*)dst)[texelNdx * 4 + 0] = 1.0f;
+							((float*)dst)[texelNdx * 4 + 1] = 0;
+							((float*)dst)[texelNdx * 4 + 2] = 1.0f;
+							((float*)dst)[texelNdx * 4 + 3] = 1.0f;
+						}
+						result = DECOMPRESS_RESULT_ERROR;
+					}
+					else
+					{
+						for (int channelNdx = 0; channelNdx < 4; channelNdx++)
+						{
+							if (!isHDREndpoint[colorEndpointNdx] || (channelNdx == 3 && colorEndpointModes[colorEndpointNdx] == 14)) // \note Alpha for mode 14 is treated the same as LDR.
+							{
+								const deUint32 c0 = (e0[channelNdx] << 8) | (isSRGB ? 0x80 : e0[channelNdx]);
+								const deUint32 c1 = (e1[channelNdx] << 8) | (isSRGB ? 0x80 : e1[channelNdx]);
+								const deUint32 w = weight.w[ccs == channelNdx ? 1 : 0];
+								const deUint32 c = (c0 * (64 - w) + c1 * w + 32) / 64;
+								if (isSRGB)
+									((deUint8*)dst)[texelNdx * 4 + channelNdx] = (deUint8)((c & 0xff00) >> 8);
+								else
+									((float*)dst)[texelNdx * 4 + channelNdx] = c == 65535 ? 1.0f : (float)c / 65536.0f;
+							}
+							else
+							{
+								//DE_STATIC_ASSERT((meta::TypesSame<deFloat16, deUint16>::Value));
+								// rg - REMOVING HDR SUPPORT FOR NOW
+#if 0
+								const deUint32		c0 = e0[channelNdx] << 4;
+								const deUint32		c1 = e1[channelNdx] << 4;
+								const deUint32		w = weight.w[ccs == channelNdx ? 1 : 0];
+								const deUint32		c = (c0 * (64 - w) + c1 * w + 32) / 64;
+								const deUint32		e = getBits(c, 11, 15);
+								const deUint32		m = getBits(c, 0, 10);
+								const deUint32		mt = m < 512 ? 3 * m
+									: m >= 1536 ? 5 * m - 2048
+									: 4 * m - 512;
+								const deFloat16		cf = (deFloat16)((e << 10) + (mt >> 3));
+								((float*)dst)[texelNdx * 4 + channelNdx] = deFloat16To32(isFloat16InfOrNan(cf) ? 0x7bff : cf);
+#endif
+							}
+						}
+					}
+				}
+			return result;
+		}
+		DecompressResult decompressBlock(void* dst, const Block128& blockData, int blockWidth, int blockHeight, bool isSRGB, bool isLDR)
+		{
+			DE_ASSERT(isLDR || !isSRGB);
+			// Decode block mode.
+			const ASTCBlockMode blockMode = getASTCBlockMode(blockData.getBits(0, 10));
+			// Check for block mode errors.
+			if (blockMode.isError)
+			{
+				setASTCErrorColorBlock(dst, blockWidth, blockHeight, isSRGB);
+				return DECOMPRESS_RESULT_ERROR;
+			}
+			// Separate path for void-extent.
+			if (blockMode.isVoidExtent)
+				return decodeVoidExtentBlock(dst, blockData, blockWidth, blockHeight, isSRGB, isLDR);
+			// Compute weight grid values.
+			const int numWeights = computeNumWeights(blockMode);
+			const int numWeightDataBits = computeNumRequiredBits(blockMode.weightISEParams, numWeights);
+			const int numPartitions = (int)blockData.getBits(11, 12) + 1;
+			// Check for errors in weight grid, partition and dual-plane parameters.
+			if (numWeights > 64 ||
+				numWeightDataBits > 96 ||
+				numWeightDataBits < 24 ||
+				blockMode.weightGridWidth > blockWidth ||
+				blockMode.weightGridHeight > blockHeight ||
+				(numPartitions == 4 && blockMode.isDualPlane))
+			{
+				setASTCErrorColorBlock(dst, blockWidth, blockHeight, isSRGB);
+				return DECOMPRESS_RESULT_ERROR;
+			}
+			// Compute number of bits available for color endpoint data.
+			const bool	isSingleUniqueCem = numPartitions == 1 || blockData.getBits(23, 24) == 0;
+			const int	numConfigDataBits = (numPartitions == 1 ? 17 : isSingleUniqueCem ? 29 : 25 + 3 * numPartitions) +
+				(blockMode.isDualPlane ? 2 : 0);
+			const int	numBitsForColorEndpoints = 128 - numWeightDataBits - numConfigDataBits;
+			const int	extraCemBitsStart = 127 - numWeightDataBits - (isSingleUniqueCem ? -1
+				: numPartitions == 4 ? 7
+				: numPartitions == 3 ? 4
+				: numPartitions == 2 ? 1
+				: 0);
+			// Decode color endpoint modes.
+			deUint32 colorEndpointModes[4];
+			decodeColorEndpointModes(&colorEndpointModes[0], blockData, numPartitions, extraCemBitsStart);
+			const int numColorEndpointValues = computeNumColorEndpointValues(colorEndpointModes, numPartitions);
+			// Check for errors in color endpoint value count.
+			if (numColorEndpointValues > 18 || numBitsForColorEndpoints < (int)deDivRoundUp32(13 * numColorEndpointValues, 5))
+			{
+				setASTCErrorColorBlock(dst, blockWidth, blockHeight, isSRGB);
+				return DECOMPRESS_RESULT_ERROR;
+			}
+			// Compute color endpoints.
+			ColorEndpointPair colorEndpoints[4];
+			computeColorEndpoints(&colorEndpoints[0], blockData, &colorEndpointModes[0], numPartitions, numColorEndpointValues,
+				computeMaximumRangeISEParams(numBitsForColorEndpoints, numColorEndpointValues), numBitsForColorEndpoints);
+			// Compute texel weights.
+			TexelWeightPair texelWeights[MAX_BLOCK_WIDTH * MAX_BLOCK_HEIGHT];
+			computeTexelWeights(&texelWeights[0], blockData, blockWidth, blockHeight, blockMode);
+			// Set texel colors.
+			const int		ccs = blockMode.isDualPlane ? (int)blockData.getBits(extraCemBitsStart - 2, extraCemBitsStart - 1) : -1;
+			const deUint32	partitionIndexSeed = numPartitions > 1 ? blockData.getBits(13, 22) : (deUint32)-1;
+			return setTexelColors(dst, &colorEndpoints[0], &texelWeights[0], ccs, partitionIndexSeed, numPartitions, blockWidth, blockHeight, isSRGB, isLDR, &colorEndpointModes[0]);
+		}
+
+
+		bool decompress(uint8_t* pDst, const uint8_t* data, bool isSRGB, int blockWidth, int blockHeight)
+		{
+			// rg - We only support LDR here, although adding back in HDR would be easy.
+			const bool isLDR = true;
+			DE_ASSERT(isLDR || !isSRGB);
+
+			float linear[MAX_BLOCK_WIDTH * MAX_BLOCK_HEIGHT * 4];
+
+			const Block128 blockData(data);
+			if (decompressBlock(isSRGB ? (void*)pDst : (void*)&linear[0],
+				blockData, blockWidth, blockHeight, isSRGB, isLDR) != DECOMPRESS_RESULT_VALID_BLOCK)
+				return false;
+
+			if (!isSRGB)
+			{
+				int pix = 0;
+				for (int i = 0; i < blockHeight; i++)
+				{
+					for (int j = 0; j < blockWidth; j++, pix++)
+					{
+						pDst[4 * pix + 0] = (uint8_t)(FMath::Clamp<int>((int)(linear[pix * 4 + 0] * 65536.0f + .5f), 0, 65535) >> 8);
+						pDst[4 * pix + 1] = (uint8_t)(FMath::Clamp<int>((int)(linear[pix * 4 + 1] * 65536.0f + .5f), 0, 65535) >> 8);
+						pDst[4 * pix + 2] = (uint8_t)(FMath::Clamp<int>((int)(linear[pix * 4 + 2] * 65536.0f + .5f), 0, 65535) >> 8);
+						pDst[4 * pix + 3] = (uint8_t)(FMath::Clamp<int>((int)(linear[pix * 4 + 3] * 65536.0f + .5f), 0, 65535) >> 8);
+					}
+				}
 			}
 
-			imageblock_initialize_orig_from_work(blk, 16);
+			return true;
 		}
 
-	} // namespace arm
+	}
 
 
 #endif // MIRO_INCLUDE_ASTC
@@ -8811,11 +6224,11 @@ namespace miro
 
 		if (!s_initialised_astc_decompress)
 		{
-			arm::build_quantization_mode_table();
+			//arm::build_quantization_mode_table();
 
-			// pregenerate astc decompression tables for 4x4x1 blocks
-			arm::get_block_size_descriptor();
-			arm::generate_partition_tables();
+			//// pregenerate astc decompression tables for 4x4x1 blocks
+			//arm::get_block_size_descriptor();
+			//arm::generate_partition_tables();
 
 			s_initialised_astc_decompress = 1;
 		}
@@ -9819,14 +7232,15 @@ namespace miro
 	//---------------------------------------------------------------------------------------------
 #if MIRO_INCLUDE_ASTC
 
-	void RGB_to_ASTC4x4RGBL(uint32 sx, uint32 sy, const uint8* from, uint8* to, int q)
+	template<uint32 BLOCK_SIZE>
+	void Generic_RGB_to_ASTCRGBL(uint32 sx, uint32 sy, const uint8* from, uint8* to, int q)
 	{
 		(void)q;
 		astcrt::PhysicalBlock physical_block_zero = { {0} };
 		astcrt::PhysicalBlock* dst_re = reinterpret_cast<astcrt::PhysicalBlock*>(to);
 
-		uint32 bx = FMath::DivideAndRoundUp(sx, 4u);
-		uint32 by = FMath::DivideAndRoundUp(sy, 4u);
+		uint32 bx = FMath::DivideAndRoundUp(sx, BLOCK_SIZE);
+		uint32 by = FMath::DivideAndRoundUp(sy, BLOCK_SIZE);
 
 		//for ( uint32 y=0; y<by; ++y )
 		ParallelFor(by,
@@ -9838,17 +7252,17 @@ namespace miro
 				for (uint32 x = 0; x < bx; ++x)
 				{
 					// Fetch the block
-					uint8 block[4 * 4 * 4];
+					uint8 block[BLOCK_SIZE * BLOCK_SIZE * 4];
 					int b = 0;
-					for (int j = 0; j < 4; ++j)
+					for (uint32 j = 0; j < BLOCK_SIZE; ++j)
 					{
 						// Clamp y in case block goes beyond the image. Biases the result.
-						uint32 rj = FMath::Min(sy - 1, y * 4 + j);
+						uint32 rj = FMath::Min(sy - 1, y * BLOCK_SIZE + j);
 
-						for (int i = 0; i < 4; ++i)
+						for (uint32 i = 0; i < BLOCK_SIZE; ++i)
 						{
 							// Clamp x in case block goes beyond the image. Biases the result.
-							uint32 ri = FMath::Min(sx - 1, x * 4 + i);
+							uint32 ri = FMath::Min(sx - 1, x * BLOCK_SIZE + i);
 
 							const uint8* pPixelSource = from + 3 * (sx * rj + ri);
 
@@ -9869,14 +7283,15 @@ namespace miro
 
 
 	//---------------------------------------------------------------------------------------------
-	void RGBA_to_ASTC4x4RGBL(uint32 sx, uint32 sy, const uint8* from, uint8* to, int q)
+	template<uint32 BLOCK_SIZE>
+	void Generic_RGBA_to_ASTCRGBL(uint32 sx, uint32 sy, const uint8* from, uint8* to, int q)
 	{
 		(void)q;
 		astcrt::PhysicalBlock physical_block_zero = { {0} };
 		astcrt::PhysicalBlock* dst_re = reinterpret_cast<astcrt::PhysicalBlock*>(to);
 
-		uint32 bx = FMath::DivideAndRoundUp(sx, 4u);
-		uint32 by = FMath::DivideAndRoundUp(sy, 4u);
+		uint32 bx = FMath::DivideAndRoundUp(sx, BLOCK_SIZE);
+		uint32 by = FMath::DivideAndRoundUp(sy, BLOCK_SIZE);
 
 		//for ( uint32 y=0; y<by; ++y )
 		ParallelFor(by,
@@ -9889,17 +7304,17 @@ namespace miro
 				for (uint32 x = 0; x < bx; ++x)
 				{
 					// Fetch the block
-					uint8 block[4 * 4 * 4];
+					uint8 block[BLOCK_SIZE * BLOCK_SIZE * 4];
 					int b = 0;
-					for (int j = 0; j < 4; ++j)
+					for (uint32 j = 0; j < BLOCK_SIZE; ++j)
 					{
 						// Clamp y in case block goes beyond the image. Biases the result.
-						uint32 rj = FMath::Min(sy - 1, y * 4 + j);
+						uint32 rj = FMath::Min(sy - 1, y * BLOCK_SIZE + j);
 
-						for (int i = 0; i < 4; ++i)
+						for (uint32 i = 0; i < BLOCK_SIZE; ++i)
 						{
 							// Clamp x in case block goes beyond the image. Biases the result.
-							uint32 ri = FMath::Min(sx - 1, x * 4 + i);
+							uint32 ri = FMath::Min(sx - 1, x * BLOCK_SIZE + i);
 
 							const uint8* pPixelSource = from + 4 * (sx * rj + ri);
 
@@ -9920,31 +7335,23 @@ namespace miro
 
 
 	//---------------------------------------------------------------------------------------------
-	void ASTC4x4RGBL_to_RGB(uint32 sx, uint32 sy, const uint8* from, uint8* to)
+	template<uint32 BLOCK_SIZE>
+	void Generic_ASTCRGBL_to_RGB(uint32 sx, uint32 sy, const uint8* from, uint8* to)
 	{
 		init_astc_decompress();
 
-		// actually use arm's decompression
-		arm::symbolic_compressed_block symbolic;
-
-		for (uint32 y = 0; y < sy; y += 4)
+		for (uint32 y = 0; y < sy; y += BLOCK_SIZE)
 		{
-			for (uint32 x = 0; x < sx; x += 4)
+			for (uint32 x = 0; x < sx; x += BLOCK_SIZE)
 			{
-				arm::physical_compressed_block physical;
-				FMemory::Memcpy(physical.data, from, 16);
+				bool bIsSRGB = false;
+				uint8 Block[BLOCK_SIZE * BLOCK_SIZE * 4];
+				bool bSuccess = astcdec::decompress(Block, from, bIsSRGB, BLOCK_SIZE, BLOCK_SIZE);
+				check(bSuccess);
 
-				arm::physical_to_symbolic(physical, &symbolic);
-
-				arm::imageblock block;
-				arm::decompress_symbolic_block(arm::DECODE_LDR, &symbolic, &block);
-
-				// Copy data from block
-				const float* fptr = block.orig_data;
-				float data[3];
-				for (uint32 py = 0; py < 4; py++)
+				for (uint32 py = 0; py < BLOCK_SIZE; py++)
 				{
-					for (uint32 px = 0; px < 4; px++)
+					for (uint32 px = 0; px < BLOCK_SIZE; px++)
 					{
 						uint32 xi = x + px;
 						uint32 yi = y + py;
@@ -9952,56 +7359,38 @@ namespace miro
 						if (xi < sx && yi < sy)
 						{
 							uint8* toPixel = to + (yi * sx + xi) * 3;
-
-							data[0] = FMath::Min(1.0f, fptr[0]);
-							data[1] = FMath::Min(1.0f, fptr[1]);
-							data[2] = FMath::Min(1.0f, fptr[2]);
-
-							// pack the data
-							int ri = static_cast <int>(floor(data[0] * 255.0f + 0.5f));
-							int gi = static_cast <int>(floor(data[1] * 255.0f + 0.5f));
-							int bi = static_cast <int>(floor(data[2] * 255.0f + 0.5f));
-
-							toPixel[0] = (uint8)ri;
-							toPixel[1] = (uint8)gi;
-							toPixel[2] = (uint8)bi;
+							toPixel[0] = Block[py * BLOCK_SIZE * 4 + px * 4 + 0];
+							toPixel[1] = Block[py * BLOCK_SIZE * 4 + px * 4 + 1];
+							toPixel[2] = Block[py * BLOCK_SIZE * 4 + px * 4 + 2];
 						}
-						fptr += 4;
 					}
 				}
 
 				from += 16;
 			}
 		}
+
 	}
 
 
 	//---------------------------------------------------------------------------------------------
-	void ASTC4x4RGBL_to_RGBA(uint32 sx, uint32 sy, const uint8* from, uint8* to)
+	template<uint32 BLOCK_SIZE>
+	void Generic_ASTCRGBL_to_RGBA(uint32 sx, uint32 sy, const uint8* from, uint8* to)
 	{
 		init_astc_decompress();
 
-		// actually use arm's decompression
-		arm::symbolic_compressed_block symbolic;
-
-		for (uint32 y = 0; y < sy; y += 4)
+		for (uint32 y = 0; y < sy; y += BLOCK_SIZE)
 		{
-			for (uint32 x = 0; x < sx; x += 4)
+			for (uint32 x = 0; x < sx; x += BLOCK_SIZE)
 			{
-				arm::physical_compressed_block physical;
-				FMemory::Memcpy(physical.data, from, 16);
+				bool bIsSRGB = false;
+				uint8 Block[BLOCK_SIZE * BLOCK_SIZE * 4];
+				bool bSuccess = astcdec::decompress(Block, from, bIsSRGB, BLOCK_SIZE, BLOCK_SIZE);
+				check(bSuccess);
 
-				arm::physical_to_symbolic(physical, &symbolic);
-
-				arm::imageblock block;
-				arm::decompress_symbolic_block(arm::DECODE_LDR, &symbolic, &block);
-
-				// Copy data from block
-				const float* fptr = block.orig_data;
-				float data[3];
-				for (uint32 py = 0; py < 4; py++)
+				for (uint32 py = 0; py < BLOCK_SIZE; py++)
 				{
-					for (uint32 px = 0; px < 4; px++)
+					for (uint32 px = 0; px < BLOCK_SIZE; px++)
 					{
 						uint32 xi = x + px;
 						uint32 yi = y + py;
@@ -10009,33 +7398,24 @@ namespace miro
 						if (xi < sx && yi < sy)
 						{
 							uint8* toPixel = to + (yi * sx + xi) * 4;
-
-							data[0] = FMath::Min(1.0f, fptr[0]);
-							data[1] = FMath::Min(1.0f, fptr[1]);
-							data[2] = FMath::Min(1.0f, fptr[2]);
-
-							// pack the data
-							int ri = static_cast <int>(floor(data[0] * 255.0f + 0.5f));
-							int gi = static_cast <int>(floor(data[1] * 255.0f + 0.5f));
-							int bi = static_cast <int>(floor(data[2] * 255.0f + 0.5f));
-
-							toPixel[0] = (uint8)ri;
-							toPixel[1] = (uint8)gi;
-							toPixel[2] = (uint8)bi;
+							toPixel[0] = Block[py * BLOCK_SIZE * 4 + px * 4 + 0];
+							toPixel[1] = Block[py * BLOCK_SIZE * 4 + px * 4 + 1];
+							toPixel[2] = Block[py * BLOCK_SIZE * 4 + px * 4 + 2];
 							toPixel[3] = 0;
 						}
-						fptr += 4;
 					}
 				}
 
 				from += 16;
 			}
 		}
+
 	}
 
 
 	//---------------------------------------------------------------------------------------------
-	void RGBA_to_ASTC4x4RGBAL(uint32 sx, uint32 sy, const uint8* from, uint8* to, int q)
+	template<uint32 BLOCK_SIZE>
+	void Generic_RGBA_to_ASTCRGBAL(uint32 sx, uint32 sy, const uint8* from, uint8* to, int q)
 	{
 		(void)q;
 		astcrt::PhysicalBlock physical_block_zero = { {0} };
@@ -10055,17 +7435,17 @@ namespace miro
 				for (uint32 x = 0; x < bx; ++x)
 				{
 					// Fetch the block
-					uint8 block[4 * 4 * 4];
+					uint8 block[BLOCK_SIZE * BLOCK_SIZE * 4];
 					int b = 0;
-					for (int j = 0; j < 4; ++j)
+					for (int j = 0; j < BLOCK_SIZE; ++j)
 					{
 						// Clamp y in case block goes beyond the image. Biases the result.
-						uint32 rj = FMath::Min(sy - 1, y * 4 + j);
+						uint32 rj = FMath::Min(sy - 1, y * BLOCK_SIZE + j);
 
-						for (int i = 0; i < 4; ++i)
+						for (int i = 0; i < BLOCK_SIZE; ++i)
 						{
 							// Clamp x in case block goes beyond the image. Biases the result.
-							uint32 ri = FMath::Min(sx - 1, x * 4 + i);
+							uint32 ri = FMath::Min(sx - 1, x * BLOCK_SIZE + i);
 
 							const uint8* pPixelSource = from + 4 * (sx * rj + ri);
 
@@ -10086,14 +7466,15 @@ namespace miro
 
 
 	//---------------------------------------------------------------------------------------------
-	void RGB_to_ASTC4x4RGBAL(uint32 sx, uint32 sy, const uint8* from, uint8* to, int q)
+	template<uint32 BLOCK_SIZE>
+	void Generic_RGB_to_ASTCRGBAL(uint32 sx, uint32 sy, const uint8* from, uint8* to, int q)
 	{
 		(void)q;
 		astcrt::PhysicalBlock physical_block_zero = { {0} };
 		astcrt::PhysicalBlock* dst_re = reinterpret_cast<astcrt::PhysicalBlock*>(to);
 
-		uint32 bx = FMath::DivideAndRoundUp(sx, 4u);
-		uint32 by = FMath::DivideAndRoundUp(sy, 4u);
+		uint32 bx = FMath::DivideAndRoundUp(sx, BLOCK_SIZE);
+		uint32 by = FMath::DivideAndRoundUp(sy, BLOCK_SIZE);
 
 		//for ( uint32 y=0; y<by; ++y )
 		ParallelFor(by,
@@ -10106,17 +7487,17 @@ namespace miro
 				for (uint32 x = 0; x < bx; ++x)
 				{
 					// Fetch the block
-					uint8 block[4 * 4 * 4];
+					uint8 block[BLOCK_SIZE * BLOCK_SIZE * 4];
 					int b = 0;
-					for (int j = 0; j < 4; ++j)
+					for (int j = 0; j < BLOCK_SIZE; ++j)
 					{
 						// Clamp y in case block goes beyond the image. Biases the result.
-						uint32 rj = FMath::Min(sy - 1, y * 4 + j);
+						uint32 rj = FMath::Min(sy - 1, y * BLOCK_SIZE + j);
 
-						for (int i = 0; i < 4; ++i)
+						for (int i = 0; i < BLOCK_SIZE; ++i)
 						{
 							// Clamp x in case block goes beyond the image. Biases the result.
-							uint32 ri = FMath::Min(sx - 1, x * 4 + i);
+							uint32 ri = FMath::Min(sx - 1, x * BLOCK_SIZE + i);
 
 							const uint8* pPixelSource = from + 3 * (sx * rj + ri);
 
@@ -10137,36 +7518,23 @@ namespace miro
 
 
 	//---------------------------------------------------------------------------------------------
-	void ASTC4x4RGBAL_to_RGBA(uint32 sx, uint32 sy, const uint8* from, uint8* to)
+	template<uint32 BLOCK_SIZE>
+	void Generic_ASTCRGBAL_to_RGBA(uint32 sx, uint32 sy, const uint8* from, uint8* to)
 	{
 		init_astc_decompress();
 
-		// actually use arm's decompression
-		arm::symbolic_compressed_block symbolic;
-
-		for (uint32 y = 0; y < sy; y += 4)
+		for (uint32 y = 0; y < sy; y += BLOCK_SIZE)
 		{
-			for (uint32 x = 0; x < sx; x += 4)
+			for (uint32 x = 0; x < sx; x += BLOCK_SIZE)
 			{
-				arm::physical_compressed_block physical;
-				FMemory::Memcpy(physical.data, from, 16);
-
-				//                if (x==8 && y==8)
-				//                    s_debuglog = true;
-
-				physical_to_symbolic(physical, &symbolic);
+				bool bIsSRGB = false;
+				uint8 Block[BLOCK_SIZE * BLOCK_SIZE *4 ];
+				bool bSuccess = astcdec::decompress( Block, from, bIsSRGB, BLOCK_SIZE, BLOCK_SIZE);
+				check(bSuccess);
 
-				arm::imageblock block;
-				decompress_symbolic_block(arm::DECODE_LDR, &symbolic, &block);
-
-				//                s_debuglog = false;
-
-				// Copy data from block
-				const float* fptr = block.orig_data;
-				float data[4];
-				for (uint32 py = 0; py < 4; py++)
+				for (uint32 py = 0; py < BLOCK_SIZE; py++)
 				{
-					for (uint32 px = 0; px < 4; px++)
+					for (uint32 px = 0; px < BLOCK_SIZE; px++)
 					{
 						uint32 xi = x + px;
 						uint32 yi = y + py;
@@ -10174,24 +7542,11 @@ namespace miro
 						if (xi < sx && yi < sy)
 						{
 							uint8* toPixel = to + (yi * sx + xi) * 4;
-
-							data[0] = FMath::Min(1.0f, fptr[0]);
-							data[1] = FMath::Min(1.0f, fptr[1]);
-							data[2] = FMath::Min(1.0f, fptr[2]);
-							data[3] = FMath::Min(1.0f, fptr[3]);
-
-							// pack the data
-							int ri = static_cast <int>(floor(data[0] * 255.0f + 0.5f));
-							int gi = static_cast <int>(floor(data[1] * 255.0f + 0.5f));
-							int bi = static_cast <int>(floor(data[2] * 255.0f + 0.5f));
-							int ai = static_cast <int>(floor(data[3] * 255.0f + 0.5f));
-
-							toPixel[0] = (uint8)ri;
-							toPixel[1] = (uint8)gi;
-							toPixel[2] = (uint8)bi;
-							toPixel[3] = (uint8)ai;
+							toPixel[0] = Block[py * BLOCK_SIZE * 4 + px * 4 + 0];
+							toPixel[1] = Block[py * BLOCK_SIZE * 4 + px * 4 + 1];
+							toPixel[2] = Block[py * BLOCK_SIZE * 4 + px * 4 + 2];
+							toPixel[3] = Block[py * BLOCK_SIZE * 4 + px * 4 + 3];
 						}
-						fptr += 4;
 					}
 				}
 
@@ -10201,33 +7556,24 @@ namespace miro
 	}
 
 
-
 	//---------------------------------------------------------------------------------------------
-	void ASTC4x4RGBAL_to_RGB(uint32 sx, uint32 sy, const uint8* from, uint8* to)
+	template<uint32 BLOCK_SIZE>
+	void Generic_ASTCRGBAL_to_RGB(uint32 sx, uint32 sy, const uint8* from, uint8* to)
 	{
 		init_astc_decompress();
 
-		// actually use arm's decompression
-		arm::symbolic_compressed_block symbolic;
-
-		for (uint32 y = 0; y < sy; y += 4)
+		for (uint32 y = 0; y < sy; y += BLOCK_SIZE)
 		{
-			for (uint32 x = 0; x < sx; x += 4)
+			for (uint32 x = 0; x < sx; x += BLOCK_SIZE)
 			{
-				arm::physical_compressed_block physical;
-				FMemory::Memcpy(physical.data, from, 16);
-
-				physical_to_symbolic(physical, &symbolic);
-
-				arm::imageblock block;
-				decompress_symbolic_block(arm::DECODE_LDR, &symbolic, &block);
+				bool bIsSRGB = false;
+				uint8 Block[BLOCK_SIZE * BLOCK_SIZE * 4];
+				bool bSuccess = astcdec::decompress(Block, from, bIsSRGB, BLOCK_SIZE, BLOCK_SIZE);
+				check(bSuccess);
 
-				// Copy data from block
-				const float* fptr = block.orig_data;
-				float data[4];
-				for (uint32 py = 0; py < 4; py++)
+				for (uint32 py = 0; py < BLOCK_SIZE; py++)
 				{
-					for (uint32 px = 0; px < 4; px++)
+					for (uint32 px = 0; px < BLOCK_SIZE; px++)
 					{
 						uint32 xi = x + px;
 						uint32 yi = y + py;
@@ -10235,21 +7581,10 @@ namespace miro
 						if (xi < sx && yi < sy)
 						{
 							uint8* toPixel = to + (yi * sx + xi) * 3;
-
-							data[0] = FMath::Min(1.0f, fptr[0]);
-							data[1] = FMath::Min(1.0f, fptr[1]);
-							data[2] = FMath::Min(1.0f, fptr[2]);
-
-							// pack the data
-							int ri = static_cast <int>(floor(data[0] * 255.0f + 0.5f));
-							int gi = static_cast <int>(floor(data[1] * 255.0f + 0.5f));
-							int bi = static_cast <int>(floor(data[2] * 255.0f + 0.5f));
-
-							toPixel[0] = (uint8)ri;
-							toPixel[1] = (uint8)gi;
-							toPixel[2] = (uint8)bi;
+							toPixel[0] = Block[py * BLOCK_SIZE * 4 + px * 4 + 0];
+							toPixel[1] = Block[py * BLOCK_SIZE * 4 + px * 4 + 1];
+							toPixel[2] = Block[py * BLOCK_SIZE * 4 + px * 4 + 2];
 						}
-						fptr += 4;
 					}
 				}
 
@@ -10261,15 +7596,16 @@ namespace miro
 
 
 	//---------------------------------------------------------------------------------------------
-	void L_to_ASTC4x4RGBL(uint32 sx, uint32 sy, const uint8_t* from, uint8_t* to, int q)
+	template<uint32 BLOCK_SIZE>
+	void Generic_L_to_ASTCRGBL(uint32 sx, uint32 sy, const uint8_t* from, uint8_t* to, int q)
 	{
 		// Almost the same as RGB_to_ASTC4x4RGBL
 		(void)q;
 		astcrt::PhysicalBlock physical_block_zero = { {0} };
 		astcrt::PhysicalBlock* dst_re = reinterpret_cast<astcrt::PhysicalBlock*>(to);
 
-		uint32 bx = FMath::DivideAndRoundUp(sx, 4u);
-		uint32 by = FMath::DivideAndRoundUp(sy, 4u);
+		uint32 bx = FMath::DivideAndRoundUp(sx, BLOCK_SIZE);
+		uint32 by = FMath::DivideAndRoundUp(sy, BLOCK_SIZE);
 
 		//for ( uint32 y=0; y<by; ++y )
 		ParallelFor(by,
@@ -10281,17 +7617,17 @@ namespace miro
 				for (uint32 x = 0; x < bx; ++x)
 				{
 					// Fetch the block
-					uint8 block[4 * 4 * 4];
+					uint8 block[BLOCK_SIZE * BLOCK_SIZE * 4];
 					int b = 0;
-					for (int j = 0; j < 4; ++j)
+					for (int j = 0; j < BLOCK_SIZE; ++j)
 					{
 						// Clamp y in case block goes beyond the image. Biases the result.
-						uint32 rj = FMath::Min(sy - 1, y * 4 + j);
+						uint32 rj = FMath::Min(sy - 1, y * BLOCK_SIZE + j);
 
-						for (int i = 0; i < 4; ++i)
+						for (int i = 0; i < BLOCK_SIZE; ++i)
 						{
 							// Clamp x in case block goes beyond the image. Biases the result.
-							uint32 ri = FMath::Min(sx - 1, x * 4 + i);
+							uint32 ri = FMath::Min(sx - 1, x * BLOCK_SIZE + i);
 
 							const uint8* pPixelSource = from + (sx * rj + ri);
 
@@ -10311,16 +7647,15 @@ namespace miro
 	}
 
 	//---------------------------------------------------------------------------------------------
-	//---------------------------------------------------------------------------------------------
-	//---------------------------------------------------------------------------------------------
-	void RGB_to_ASTC4x4RGL(uint32 sx, uint32 sy, const uint8* from, uint8* to, int q)
+	template<uint32 BLOCK_SIZE>
+	void Generic_RGB_to_ASTCRGL(uint32 sx, uint32 sy, const uint8* from, uint8* to, int q)
 	{
 		(void)q;
 		astcrt::PhysicalBlock physical_block_zero = { {0} };
 		astcrt::PhysicalBlock* dst_re = reinterpret_cast<astcrt::PhysicalBlock*>(to);
 
-		uint32 bx = FMath::DivideAndRoundUp(sx, 4u);
-		uint32 by = FMath::DivideAndRoundUp(sy, 4u);
+		uint32 bx = FMath::DivideAndRoundUp(sx, BLOCK_SIZE);
+		uint32 by = FMath::DivideAndRoundUp(sy, BLOCK_SIZE);
 
 		//for ( uint32 y=0; y<by; ++y )
 		ParallelFor(by,
@@ -10333,17 +7668,17 @@ namespace miro
 				for (uint32 x = 0; x < bx; ++x)
 				{
 					// Fetch the block
-					uint8 block[4 * 4 * 4];
+					uint8 block[BLOCK_SIZE * BLOCK_SIZE * 4];
 					int b = 0;
-					for (int j = 0; j < 4; ++j)
+					for (int j = 0; j < BLOCK_SIZE; ++j)
 					{
 						// Clamp y in case block goes beyond the image. Biases the result.
-						uint32 rj = FMath::Min(sy - 1, y * 4 + j);
+						uint32 rj = FMath::Min(sy - 1, y * BLOCK_SIZE + j);
 
-						for (int i = 0; i < 4; ++i)
+						for (int i = 0; i < BLOCK_SIZE; ++i)
 						{
 							// Clamp x in case block goes beyond the image. Biases the result.
-							uint32 ri = FMath::Min(sx - 1, x * 4 + i);
+							uint32 ri = FMath::Min(sx - 1, x * BLOCK_SIZE + i);
 
 							const uint8* pPixelSource = from + 3 * (sx * rj + ri);
 
@@ -10363,30 +7698,15 @@ namespace miro
 
 
 	//---------------------------------------------------------------------------------------------
-	void ASTC4x4RGBL_to_ASTC4x4RGBAL(uint32 sx, uint32 sy, const uint8* from, uint8* to)
-	{
-		int blockCount = FMath::DivideAndRoundUp(sy, 4u) * FMath::DivideAndRoundUp(sx, 4u);
-		FMemory::Memcpy(to, from, blockCount * 16);
-	}
-
-
-	//---------------------------------------------------------------------------------------------
-	void ASTC4x4RGBAL_to_ASTC4x4RGBL(uint32 sx, uint32 sy, const uint8* from, uint8* to)
-	{
-		int blockCount = FMath::DivideAndRoundUp(sy, 4u) * FMath::DivideAndRoundUp(sx, 4u);
-		FMemory::Memcpy(to, from, blockCount * 16);
-	}
-
-
-	//---------------------------------------------------------------------------------------------
-	void RGBA_to_ASTC4x4RGL(uint32 sx, uint32 sy, const uint8* from, uint8* to, int q)
+	template<uint32 BLOCK_SIZE>
+	void Generic_RGBA_to_ASTCRGL(uint32 sx, uint32 sy, const uint8* from, uint8* to, int q)
 	{
 		(void)q;
 		astcrt::PhysicalBlock physical_block_zero = { {0} };
 		astcrt::PhysicalBlock* dst_re = reinterpret_cast<astcrt::PhysicalBlock*>(to);
 
-		uint32 bx = FMath::DivideAndRoundUp(sx, 4u);
-		uint32 by = FMath::DivideAndRoundUp(sy, 4u);
+		uint32 bx = FMath::DivideAndRoundUp(sx, BLOCK_SIZE);
+		uint32 by = FMath::DivideAndRoundUp(sy, BLOCK_SIZE);
 
 		//for ( uint32 y=0; y<by; ++y )
 		ParallelFor(by,
@@ -10399,17 +7719,17 @@ namespace miro
 				for (uint32 x = 0; x < bx; ++x)
 				{
 					// Fetch the block
-					uint8 block[4 * 4 * 4];
+					uint8 block[BLOCK_SIZE * BLOCK_SIZE * 4];
 					int b = 0;
-					for (int j = 0; j < 4; ++j)
+					for (int j = 0; j < BLOCK_SIZE; ++j)
 					{
 						// Clamp y in case block goes beyond the image. Biases the result.
-						uint32 rj = FMath::Min(sy - 1, y * 4 + j);
+						uint32 rj = FMath::Min(sy - 1, y * BLOCK_SIZE + j);
 
-						for (int i = 0; i < 4; ++i)
+						for (int i = 0; i < BLOCK_SIZE; ++i)
 						{
 							// Clamp x in case block goes beyond the image. Biases the result.
-							uint32 ri = FMath::Min(sx - 1, x * 4 + i);
+							uint32 ri = FMath::Min(sx - 1, x * BLOCK_SIZE + i);
 
 							const uint8* pPixelSource = from + 4 * (sx * rj + ri);
 
@@ -10430,38 +7750,23 @@ namespace miro
 
 
 	//---------------------------------------------------------------------------------------------
-	void ASTC4x4RGL_to_RGB(uint32 sx, uint32 sy, const uint8* from, uint8* to)
+	template<uint32 BLOCK_SIZE>
+	void Generic_ASTCRGL_to_RGB(uint32 sx, uint32 sy, const uint8* from, uint8* to)
 	{
 		init_astc_decompress();
 
-		// actually use arm's decompression
-		arm::symbolic_compressed_block symbolic;
-
-		for (uint32 y = 0; y < sy; y += 4)
+		for (uint32 y = 0; y < sy; y += BLOCK_SIZE)
 		{
-			for (uint32 x = 0; x < sx; x += 4)
+			for (uint32 x = 0; x < sx; x += BLOCK_SIZE)
 			{
-				arm::physical_compressed_block physical;
-				FMemory::Memcpy(physical.data, from, 16);
-
-				//                if (x==476 && y==256)
-				//                    s_debuglog = true;
-
-				arm::physical_to_symbolic(physical, &symbolic);
+				bool bIsSRGB = false;
+				uint8 Block[BLOCK_SIZE * BLOCK_SIZE * 4];
+				bool bSuccess = astcdec::decompress(Block, from, bIsSRGB, BLOCK_SIZE, BLOCK_SIZE);
+				check(bSuccess);
 
-				arm::imageblock block;
-				arm::decompress_symbolic_block(arm::DECODE_LDR,
-					&symbolic,
-					&block);
-
-				//                s_debuglog = false;
-
-								// Copy data from block
-				const float* fptr = block.orig_data;
-				float data[4];
-				for (uint32 py = 0; py < 4; py++)
+				for (uint32 py = 0; py < BLOCK_SIZE; py++)
 				{
-					for (uint32 px = 0; px < 4; px++)
+					for (uint32 px = 0; px < BLOCK_SIZE; px++)
 					{
 						uint32 xi = x + px;
 						uint32 yi = y + py;
@@ -10469,19 +7774,10 @@ namespace miro
 						if (xi < sx && yi < sy)
 						{
 							uint8* toPixel = to + (yi * sx + xi) * 3;
-
-							data[0] = FMath::Min(1.0f, fptr[0]);
-							data[3] = FMath::Min(1.0f, fptr[3]);
-
-							// pack the data
-							int ri = static_cast <int>(floor(data[0] * 255.0f + 0.5f));
-							int gi = static_cast <int>(floor(data[3] * 255.0f + 0.5f));
-
-							toPixel[0] = (uint8)ri;
-							toPixel[1] = (uint8)gi;
+							toPixel[0] = Block[py * BLOCK_SIZE * 4 + px * 4 + 0];
+							toPixel[1] = Block[py * BLOCK_SIZE * 4 + px * 4 + 3];
 							toPixel[2] = 0;
 						}
-						fptr += 4;
 					}
 				}
 
@@ -10492,31 +7788,23 @@ namespace miro
 
 
 	//---------------------------------------------------------------------------------------------
-	void ASTC4x4RGL_to_RGBA(uint32 sx, uint32 sy, const uint8* from, uint8* to)
+	template<uint32 BLOCK_SIZE>
+	void Generic_ASTCRGL_to_RGBA(uint32 sx, uint32 sy, const uint8* from, uint8* to)
 	{
 		init_astc_decompress();
 
-		// actually use arm's decompression
-		arm::symbolic_compressed_block symbolic;
-
-		for (uint32 y = 0; y < sy; y += 4)
+		for (uint32 y = 0; y < sy; y += BLOCK_SIZE)
 		{
-			for (uint32 x = 0; x < sx; x += 4)
+			for (uint32 x = 0; x < sx; x += BLOCK_SIZE)
 			{
-				arm::physical_compressed_block physical;
-				FMemory::Memcpy(physical.data, from, 16);
-
-				arm::physical_to_symbolic(physical, &symbolic);
+				bool bIsSRGB = false;
+				uint8 Block[BLOCK_SIZE * BLOCK_SIZE * 4];
+				bool bSuccess = astcdec::decompress(Block, from, bIsSRGB, BLOCK_SIZE, BLOCK_SIZE);
+				check(bSuccess);
 
-				arm::imageblock block;
-				arm::decompress_symbolic_block(arm::DECODE_LDR, &symbolic, &block);
-
-				// Copy data from block
-				const float* fptr = block.orig_data;
-				float data[4];
-				for (uint32 py = 0; py < 4; py++)
+				for (uint32 py = 0; py < BLOCK_SIZE; py++)
 				{
-					for (uint32 px = 0; px < 4; px++)
+					for (uint32 px = 0; px < BLOCK_SIZE; px++)
 					{
 						uint32 xi = x + px;
 						uint32 yi = y + py;
@@ -10524,20 +7812,11 @@ namespace miro
 						if (xi < sx && yi < sy)
 						{
 							uint8* toPixel = to + (yi * sx + xi) * 4;
-
-							data[0] = FMath::Min(1.0f, fptr[0]);
-							data[3] = FMath::Min(1.0f, fptr[3]);
-
-							// pack the data
-							int ri = static_cast <int>(floor(data[0] * 255.0f + 0.5f));
-							int gi = static_cast <int>(floor(data[3] * 255.0f + 0.5f));
-
-							toPixel[0] = (uint8)ri;
-							toPixel[1] = (uint8)gi;
+							toPixel[0] = Block[py * BLOCK_SIZE * 4 + px * 4 + 0];
+							toPixel[1] = Block[py * BLOCK_SIZE * 4 + px * 4 + 3];
 							toPixel[2] = 0;
 							toPixel[3] = 0;
 						}
-						fptr += 4;
 					}
 				}
 
@@ -10546,6 +7825,155 @@ namespace miro
 		}
 	}
 
+
+	//---------------------------------------------------------------------------------------------
+	// 4x4
+	//---------------------------------------------------------------------------------------------
+	void ASTC4x4RGBL_to_ASTC4x4RGBAL(uint32 sx, uint32 sy, const uint8* from, uint8* to)
+	{
+		int blockCount = FMath::DivideAndRoundUp(sy, 4u) * FMath::DivideAndRoundUp(sx, 4u);
+		FMemory::Memcpy(to, from, blockCount * 16);
+	}
+
+	void ASTC4x4RGBAL_to_ASTC4x4RGBL(uint32 sx, uint32 sy, const uint8* from, uint8* to)
+	{
+		int blockCount = FMath::DivideAndRoundUp(sy, 4u) * FMath::DivideAndRoundUp(sx, 4u);
+		FMemory::Memcpy(to, from, blockCount * 16);
+	}
+
+	void RGBA_to_ASTC4x4RGBAL(uint32 sx, uint32 sy, const uint8* from, uint8* to, int q)
+	{
+		Generic_RGBA_to_ASTCRGBAL<4>(sx, sy, from, to, q);
+	}
+
+	void RGB_to_ASTC4x4RGBAL(uint32 sx, uint32 sy, const uint8* from, uint8* to, int q)
+	{
+		Generic_RGB_to_ASTCRGBAL<4>(sx, sy, from, to, q);
+	}
+
+	void ASTC4x4RGBAL_to_RGBA(uint32 sx, uint32 sy, const uint8* from, uint8* to)
+	{
+		Generic_ASTCRGBAL_to_RGBA<4>(sx, sy, from, to);
+	}
+
+	void ASTC4x4RGBAL_to_RGB(uint32 sx, uint32 sy, const uint8* from, uint8* to)
+	{
+		Generic_ASTCRGBAL_to_RGB<4>(sx, sy, from, to);
+	}
+
+	void RGB_to_ASTC4x4RGBL(uint32 sx, uint32 sy, const uint8* from, uint8* to, int q)
+	{
+		Generic_RGB_to_ASTCRGBL<4>(sx, sy, from, to, q);
+	}
+
+	void RGBA_to_ASTC4x4RGBL(uint32 sx, uint32 sy, const uint8* from, uint8* to, int q)
+	{
+		Generic_RGBA_to_ASTCRGBL<4>(sx, sy, from, to, q);
+	}
+
+	void ASTC4x4RGBL_to_RGB(uint32 sx, uint32 sy, const uint8* from, uint8* to)
+	{
+		Generic_ASTCRGBL_to_RGB<4>(sx, sy, from, to);
+	}
+
+	void ASTC4x4RGBL_to_RGBA(uint32 sx, uint32 sy, const uint8* from, uint8* to)
+	{
+		Generic_ASTCRGBL_to_RGBA<4>(sx, sy, from, to);
+	}
+	
+	void L_to_ASTC4x4RGBL(uint32 sx, uint32 sy, const uint8* from, uint8* to, int q)
+	{
+		Generic_L_to_ASTCRGBL<4>(sx, sy, from, to,q);
+	}
+
+	void RGB_to_ASTC4x4RGL(uint32 sx, uint32 sy, const uint8* from, uint8* to, int q)
+	{
+		Generic_RGB_to_ASTCRGL<4>(sx, sy, from, to,q);
+	}
+	
+	void RGBA_to_ASTC4x4RGL(uint32 sx, uint32 sy, const uint8* from, uint8* to, int q)
+	{
+		Generic_RGBA_to_ASTCRGL<4>(sx, sy, from, to,q);
+	}
+	
+	void ASTC4x4RGL_to_RGB(uint32 sx, uint32 sy, const uint8* from, uint8* to)
+	{
+		Generic_ASTCRGL_to_RGB<4>(sx, sy, from, to);
+	}
+	
+	void ASTC4x4RGL_to_RGBA(uint32 sx, uint32 sy, const uint8* from, uint8* to)
+	{
+		Generic_ASTCRGL_to_RGBA<4>(sx, sy, from, to);
+	}
+
+
+	//---------------------------------------------------------------------------------------------
+	// 8x8
+	//---------------------------------------------------------------------------------------------
+	void ASTC8x8RGBAL_to_RGBA(uint32 sx, uint32 sy, const uint8* from, uint8* to)
+	{
+		Generic_ASTCRGBAL_to_RGBA<8>(sx, sy, from, to);
+	}
+
+	void ASTC8x8RGBAL_to_RGB(uint32 sx, uint32 sy, const uint8* from, uint8* to)
+	{
+		Generic_ASTCRGBAL_to_RGB<8>(sx, sy, from, to);
+	}
+
+	void ASTC8x8RGBL_to_RGB(uint32 sx, uint32 sy, const uint8* from, uint8* to)
+	{
+		Generic_ASTCRGBL_to_RGB<8>(sx, sy, from, to);
+	}
+
+	void ASTC8x8RGBL_to_RGBA(uint32 sx, uint32 sy, const uint8* from, uint8* to)
+	{
+		Generic_ASTCRGBL_to_RGBA<8>(sx, sy, from, to);
+	}
+
+	void ASTC8x8RGL_to_RGB(uint32 sx, uint32 sy, const uint8* from, uint8* to)
+	{
+		Generic_ASTCRGL_to_RGB<8>(sx, sy, from, to);
+	}
+
+	void ASTC8x8RGL_to_RGBA(uint32 sx, uint32 sy, const uint8* from, uint8* to)
+	{
+		Generic_ASTCRGL_to_RGBA<8>(sx, sy, from, to);
+	}
+
+
+	//---------------------------------------------------------------------------------------------
+	// 12x12
+	//---------------------------------------------------------------------------------------------
+	void ASTC12x12RGBAL_to_RGBA(uint32 sx, uint32 sy, const uint8* from, uint8* to)
+	{
+		Generic_ASTCRGBAL_to_RGBA<12>(sx, sy, from, to);
+	}
+
+	void ASTC12x12RGBAL_to_RGB(uint32 sx, uint32 sy, const uint8* from, uint8* to)
+	{
+		Generic_ASTCRGBAL_to_RGB<12>(sx, sy, from, to);
+	}
+
+	void ASTC12x12RGBL_to_RGB(uint32 sx, uint32 sy, const uint8* from, uint8* to)
+	{
+		Generic_ASTCRGBL_to_RGB<12>(sx, sy, from, to);
+	}
+
+	void ASTC12x12RGBL_to_RGBA(uint32 sx, uint32 sy, const uint8* from, uint8* to)
+	{
+		Generic_ASTCRGBL_to_RGBA<12>(sx, sy, from, to);
+	}
+
+	void ASTC12x12RGL_to_RGB(uint32 sx, uint32 sy, const uint8* from, uint8* to)
+	{
+		Generic_ASTCRGL_to_RGB<12>(sx, sy, from, to);
+	}
+
+	void ASTC12x12RGL_to_RGBA(uint32 sx, uint32 sy, const uint8* from, uint8* to)
+	{
+		Generic_ASTCRGL_to_RGBA<12>(sx, sy, from, to);
+	}
+
 #endif // MIRO_INCLUDE_ASTC
 
 
diff --git a/Engine/Plugins/Experimental/Mutable/Source/MutableRuntime/Private/MuR/BlockCompression/Miro/Miro.h b/Engine/Plugins/Experimental/Mutable/Source/MutableRuntime/Private/MuR/BlockCompression/Miro/Miro.h
--- a/Engine/Plugins/Experimental/Mutable/Source/MutableRuntime/Private/MuR/BlockCompression/Miro/Miro.h
+++ b/Engine/Plugins/Experimental/Mutable/Source/MutableRuntime/Private/MuR/BlockCompression/Miro/Miro.h
@@ -99,6 +99,30 @@ namespace miro
     extern void ASTC4x4RGL_to_RGB( uint32 sx, uint32 sy, const uint8* from, uint8* to );
     extern void ASTC4x4RGL_to_RGBA( uint32 sx, uint32 sy, const uint8* from, uint8* to );
 
+	//! ASTC 8x8 RGBA low support
+	extern void ASTC8x8RGBAL_to_RGBA(uint32 sx, uint32 sy, const uint8* from, uint8* to);
+	extern void ASTC8x8RGBAL_to_RGB(uint32 sx, uint32 sy, const uint8* from, uint8* to);
+
+	//! ASTC 8x8 RGB low support
+	extern void ASTC8x8RGBL_to_RGB(uint32 sx, uint32 sy, const uint8* from, uint8* to);
+	extern void ASTC8x8RGBL_to_RGBA(uint32 sx, uint32 sy, const uint8* from, uint8* to);
+
+	//! ASTC 8x8 RG low support
+	extern void ASTC8x8RGL_to_RGB(uint32 sx, uint32 sy, const uint8* from, uint8* to);
+	extern void ASTC8x8RGL_to_RGBA(uint32 sx, uint32 sy, const uint8* from, uint8* to);
+
+	//! ASTC 12x12 RGBA low support
+	extern void ASTC12x12RGBAL_to_RGBA(uint32 sx, uint32 sy, const uint8* from, uint8* to);
+	extern void ASTC12x12RGBAL_to_RGB(uint32 sx, uint32 sy, const uint8* from, uint8* to);
+
+	//! ASTC 12x12 RGB low support
+	extern void ASTC12x12RGBL_to_RGB(uint32 sx, uint32 sy, const uint8* from, uint8* to);
+	extern void ASTC12x12RGBL_to_RGBA(uint32 sx, uint32 sy, const uint8* from, uint8* to);
+
+	//! ASTC 12x12 RG low support
+	extern void ASTC12x12RGL_to_RGB(uint32 sx, uint32 sy, const uint8* from, uint8* to);
+	extern void ASTC12x12RGL_to_RGBA(uint32 sx, uint32 sy, const uint8* from, uint8* to);
+
 #endif
 
 }
diff --git a/Engine/Plugins/Experimental/Mutable/Source/MutableRuntime/Private/MuR/Image.cpp b/Engine/Plugins/Experimental/Mutable/Source/MutableRuntime/Private/MuR/Image.cpp
--- a/Engine/Plugins/Experimental/Mutable/Source/MutableRuntime/Private/MuR/Image.cpp
+++ b/Engine/Plugins/Experimental/Mutable/Source/MutableRuntime/Private/MuR/Image.cpp
@@ -39,11 +39,19 @@ namespace mu
         FImageFormatData( 4, 4, 16, 3 ),	// IF_BC6
         FImageFormatData( 4, 4, 16, 4 ),	// IF_BC7
 
-        FImageFormatData( 1, 1, 4, 4 ),	// IF_BGRA_UBYTE
+        FImageFormatData( 1, 1, 4, 4 ),		// IF_BGRA_UBYTE
 
         FImageFormatData( 4, 4, 16, 3 ),	// IF_ASTC_4x4_RGB_LDR
         FImageFormatData( 4, 4, 16, 4 ),	// IF_ASTC_4x4_RGBA_LDR
         FImageFormatData( 4, 4, 16, 2 ),	// IF_ASTC_4x4_RG_LDR
+		
+		FImageFormatData(8, 8, 16, 3),		// IF_ASTC_8x8_RGB_LDR,
+		FImageFormatData(8, 8, 16, 4),		// IF_ASTC_8x8_RGBA_LDR,
+		FImageFormatData(8, 8, 16, 2),		// IF_ASTC_8x8_RG_LDR,
+		FImageFormatData(12, 12, 16, 3),	// IF_ASTC_12x12_RGB_LDR
+		FImageFormatData(12, 12, 16, 4),	// IF_ASTC_12x12_RGBA_LDR
+		FImageFormatData(12, 12, 16, 2),	// IF_ASTC_12x12_RG_LDR
+
     };
 
 
diff --git a/Engine/Plugins/Experimental/Mutable/Source/MutableRuntime/Private/MuR/OpImagePixelFormat.cpp b/Engine/Plugins/Experimental/Mutable/Source/MutableRuntime/Private/MuR/OpImagePixelFormat.cpp
--- a/Engine/Plugins/Experimental/Mutable/Source/MutableRuntime/Private/MuR/OpImagePixelFormat.cpp
+++ b/Engine/Plugins/Experimental/Mutable/Source/MutableRuntime/Private/MuR/OpImagePixelFormat.cpp
@@ -552,6 +552,78 @@ bool ImagePixelFormatInPlace( int imageCompressionQuality, Image* pResult, const
                 break;
             }
 
+			case EImageFormat::IF_ASTC_8x8_RGB_LDR:
+			{
+				for (int m = 0; m < resultLODCount; ++m)
+				{
+					FIntVector2 mipSize = pResult->CalculateMipSize(m);
+					miro::ASTC8x8RGBL_to_RGB(mipSize[0], mipSize[1],
+						pBase->GetMipData(baseLOD + m),
+						pResult->GetMipData(m));
+				}
+				break;
+			}
+
+			case EImageFormat::IF_ASTC_8x8_RGBA_LDR:
+			{
+				for (int m = 0; m < resultLODCount; ++m)
+				{
+					FIntVector2 mipSize = pResult->CalculateMipSize(m);
+					miro::ASTC8x8RGBAL_to_RGB(mipSize[0], mipSize[1],
+						pBase->GetMipData(baseLOD + m),
+						pResult->GetMipData(m));
+				}
+				break;
+			}
+
+			case EImageFormat::IF_ASTC_8x8_RG_LDR:
+			{
+				for (int m = 0; m < resultLODCount; ++m)
+				{
+					FIntVector2 mipSize = pResult->CalculateMipSize(m);
+					miro::ASTC8x8RGL_to_RGB(mipSize[0], mipSize[1],
+						pBase->GetMipData(baseLOD + m),
+						pResult->GetMipData(m));
+				}
+				break;
+			}
+
+			case EImageFormat::IF_ASTC_12x12_RGB_LDR:
+			{
+				for (int m = 0; m < resultLODCount; ++m)
+				{
+					FIntVector2 mipSize = pResult->CalculateMipSize(m);
+					miro::ASTC12x12RGBL_to_RGB(mipSize[0], mipSize[1],
+						pBase->GetMipData(baseLOD + m),
+						pResult->GetMipData(m));
+				}
+				break;
+			}
+
+			case EImageFormat::IF_ASTC_12x12_RGBA_LDR:
+			{
+				for (int m = 0; m < resultLODCount; ++m)
+				{
+					FIntVector2 mipSize = pResult->CalculateMipSize(m);
+					miro::ASTC12x12RGBAL_to_RGB(mipSize[0], mipSize[1],
+						pBase->GetMipData(baseLOD + m),
+						pResult->GetMipData(m));
+				}
+				break;
+			}
+
+			case EImageFormat::IF_ASTC_12x12_RG_LDR:
+			{
+				for (int m = 0; m < resultLODCount; ++m)
+				{
+					FIntVector2 mipSize = pResult->CalculateMipSize(m);
+					miro::ASTC12x12RGL_to_RGB(mipSize[0], mipSize[1],
+						pBase->GetMipData(baseLOD + m),
+						pResult->GetMipData(m));
+				}
+				break;
+			}
+
             default:
                 // Case not implemented
                 check( false );
@@ -731,6 +803,78 @@ bool ImagePixelFormatInPlace( int imageCompressionQuality, Image* pResult, const
                 break;
             }
 
+			case EImageFormat::IF_ASTC_8x8_RGB_LDR:
+			{
+				for (int m = 0; m < resultLODCount; ++m)
+				{
+					FIntVector2 mipSize = pResult->CalculateMipSize(m);
+					miro::ASTC8x8RGBL_to_RGBA(mipSize[0], mipSize[1],
+						pBase->GetMipData(baseLOD + m),
+						pResult->GetMipData(m));
+				}
+				break;
+			}
+
+			case EImageFormat::IF_ASTC_8x8_RGBA_LDR:
+			{
+				for (int m = 0; m < resultLODCount; ++m)
+				{
+					FIntVector2 mipSize = pResult->CalculateMipSize(m);
+					miro::ASTC8x8RGBAL_to_RGBA(mipSize[0], mipSize[1],
+						pBase->GetMipData(baseLOD + m),
+						pResult->GetMipData(m));
+				}
+				break;
+			}
+
+			case EImageFormat::IF_ASTC_8x8_RG_LDR:
+			{
+				for (int m = 0; m < resultLODCount; ++m)
+				{
+					FIntVector2 mipSize = pResult->CalculateMipSize(m);
+					miro::ASTC8x8RGL_to_RGBA(mipSize[0], mipSize[1],
+						pBase->GetMipData(baseLOD + m),
+						pResult->GetMipData(m));
+				}
+				break;
+			}
+
+			case EImageFormat::IF_ASTC_12x12_RGB_LDR:
+			{
+				for (int m = 0; m < resultLODCount; ++m)
+				{
+					FIntVector2 mipSize = pResult->CalculateMipSize(m);
+					miro::ASTC12x12RGBL_to_RGBA(mipSize[0], mipSize[1],
+						pBase->GetMipData(baseLOD + m),
+						pResult->GetMipData(m));
+				}
+				break;
+			}
+
+			case EImageFormat::IF_ASTC_12x12_RGBA_LDR:
+			{
+				for (int m = 0; m < resultLODCount; ++m)
+				{
+					FIntVector2 mipSize = pResult->CalculateMipSize(m);
+					miro::ASTC12x12RGBAL_to_RGBA(mipSize[0], mipSize[1],
+						pBase->GetMipData(baseLOD + m),
+						pResult->GetMipData(m));
+				}
+				break;
+			}
+
+			case EImageFormat::IF_ASTC_12x12_RG_LDR:
+			{
+				for (int m = 0; m < resultLODCount; ++m)
+				{
+					FIntVector2 mipSize = pResult->CalculateMipSize(m);
+					miro::ASTC12x12RGL_to_RGBA(mipSize[0], mipSize[1],
+						pBase->GetMipData(baseLOD + m),
+						pResult->GetMipData(m));
+				}
+				break;
+			}
+
 			case EImageFormat::IF_RGB_UBYTE_RLE:
 			{
 				check(onlyLOD == -1);
diff --git a/Engine/Plugins/Experimental/Mutable/Source/MutableRuntime/Private/MuR/Operations.cpp b/Engine/Plugins/Experimental/Mutable/Source/MutableRuntime/Private/MuR/Operations.cpp
--- a/Engine/Plugins/Experimental/Mutable/Source/MutableRuntime/Private/MuR/Operations.cpp
+++ b/Engine/Plugins/Experimental/Mutable/Source/MutableRuntime/Private/MuR/Operations.cpp
@@ -17,129 +17,129 @@ namespace mu
     // clang-format off
     static const OP_DESC s_opDescs[OP_DESC_COUNT] =
 	{ 
-		// type			cached	GPU-izable	supported base image formats
-        { DT_NONE,		false,	false,		{ 0, 0, 0, 0, 0, 0, 0, 0, 0, 0, 0, 0, 0, 0, 0, 0, 0, 0, 0, 0, 0, 0 }		},	// NONE
-
-        { DT_BOOL,		false,	false,		{ 0, 0, 0, 0, 0, 0, 0, 0, 0, 0, 0, 0, 0, 0, 0, 0, 0, 0, 0, 0, 0, 0 }		},	// BO_CONSTANT
-        { DT_INT,		false,	false,		{ 0, 0, 0, 0, 0, 0, 0, 0, 0, 0, 0, 0, 0, 0, 0, 0, 0, 0, 0, 0, 0, 0 }		},	// NU_CONSTANT
-        { DT_SCALAR,	false,	false,		{ 0, 0, 0, 0, 0, 0, 0, 0, 0, 0, 0, 0, 0, 0, 0, 0, 0, 0, 0, 0, 0, 0 }		},	// SC_CONSTANT
-        { DT_COLOUR,	false,	false,		{ 0, 0, 0, 0, 0, 0, 0, 0, 0, 0, 0, 0, 0, 0, 0, 0, 0, 0, 0, 0, 0, 0 }		},	// CO_CONSTANT
-        { DT_IMAGE,		true,   false,		{ 1, 1, 1, 1, 1, 1, 1, 1, 1, 1, 1, 1, 0, 0, 0, 0, 0, 0, 0, 0, 0, 0 }		},	// IM_CONSTANT
-        { DT_MESH,		true,   false,		{ 0, 0, 0, 0, 0, 0, 0, 0, 0, 0, 0, 0, 0, 0, 0, 0, 0, 0, 0, 0, 0, 0 }		},	// ME_CONSTANT
-        { DT_LAYOUT,	false,	false,		{ 0, 0, 0, 0, 0, 0, 0, 0, 0, 0, 0, 0, 0, 0, 0, 0, 0, 0, 0, 0, 0, 0 }		},	// LA_CONSTANT
-        { DT_PROJECTOR,	false,	false,		{ 0, 0, 0, 0, 0, 0, 0, 0, 0, 0, 0, 0, 0, 0, 0, 0, 0, 0, 0, 0, 0, 0 }		},	// PR_CONSTANT
-        { DT_STRING,	false,	false,		{ 0, 0, 0, 0, 0, 0, 0, 0, 0, 0, 0, 0, 0, 0, 0, 0, 0, 0, 0, 0, 0, 0 }		},	// ST_CONSTANT
-
-        { DT_BOOL,		false,	false,		{ 0, 0, 0, 0, 0, 0, 0, 0, 0, 0, 0, 0, 0, 0, 0, 0, 0, 0, 0, 0, 0, 0 }		},	// BO_PARAMETER
-		{ DT_INT,		false,	false,		{ 0, 0, 0, 0, 0, 0, 0, 0, 0, 0, 0, 0, 0, 0, 0, 0, 0, 0, 0, 0, 0, 0 }		},	// NU_PARAMETER
-		{ DT_SCALAR,	false,	false,		{ 0, 0, 0, 0, 0, 0, 0, 0, 0, 0, 0, 0, 0, 0, 0, 0, 0, 0, 0, 0, 0, 0 }		},	// SC_PARAMETER
-		{ DT_COLOUR,	false,	false,		{ 0, 0, 0, 0, 0, 0, 0, 0, 0, 0, 0, 0, 0, 0, 0, 0, 0, 0, 0, 0, 0, 0 }		},	// CO_PARAMETER
-        { DT_PROJECTOR,	false,	false,		{ 0, 0, 0, 0, 0, 0, 0, 0, 0, 0, 0, 0, 0, 0, 0, 0, 0, 0, 0, 0, 0, 0 }		},	// PR_PARAMETER
-        { DT_IMAGE,		false,	false,		{ 0, 0, 0, 0, 0, 0, 0, 0, 0, 0, 0, 0, 0, 0, 0, 0, 0, 0, 0, 0, 0, 0 }		},	// IM_PARAMETER
-        { DT_STRING,	false,	false,		{ 0, 0, 0, 0, 0, 0, 0, 0, 0, 0, 0, 0, 0, 0, 0, 0, 0, 0, 0, 0, 0, 0 }		},	// ST_PARAMETER
-
-        { DT_INT,		false,	false,		{ 0, 0, 0, 0, 0, 0, 0, 0, 0, 0, 0, 0, 0, 0, 0, 0, 0, 0, 0, 0, 0, 0 }		},	// NU_CONDITIONAL
-		{ DT_SCALAR,	false,	false,		{ 0, 0, 0, 0, 0, 0, 0, 0, 0, 0, 0, 0, 0, 0, 0, 0, 0, 0, 0, 0, 0, 0 }		},	// SC_CONDITIONAL
-		{ DT_COLOUR,	false,	false,		{ 0, 0, 0, 0, 0, 0, 0, 0, 0, 0, 0, 0, 0, 0, 0, 0, 0, 0, 0, 0, 0, 0 }		},	// CO_CONDITIONAL
-        { DT_IMAGE,		true,	false,		{ 1, 1, 1, 1, 1, 1, 1, 1, 1, 1, 1, 1, 1, 1, 1, 1, 1, 1, 1, 1, 1, 1 }		},	// IM_CONDITIONAL
-        { DT_MESH,		true,	false,		{ 0, 0, 0, 0, 0, 0, 0, 0, 0, 0, 0, 0, 0, 0, 0, 0, 0, 0, 0, 0, 0, 0 }		},	// ME_CONDITIONAL
-        { DT_LAYOUT,	false,  false,		{ 0, 0, 0, 0, 0, 0, 0, 0, 0, 0, 0, 0, 0, 0, 0, 0, 0, 0, 0, 0, 0, 0 }		},	// LA_CONDITIONAL
-        { DT_INSTANCE,	false,	false,		{ 0, 0, 0, 0, 0, 0, 0, 0, 0, 0, 0, 0, 0, 0, 0, 0, 0, 0, 0, 0, 0, 0 }		},	// IN_CONDITIONAL
-
-        { DT_INT,		false,	false,		{ 0, 0, 0, 0, 0, 0, 0, 0, 0, 0, 0, 0, 0, 0, 0, 0, 0, 0, 0, 0, 0, 0 }		},	// NU_SWITCH
-		{ DT_SCALAR,	false,	false,		{ 0, 0, 0, 0, 0, 0, 0, 0, 0, 0, 0, 0, 0, 0, 0, 0, 0, 0, 0, 0, 0, 0 }		},	// SC_SWITCH
-		{ DT_COLOUR,	false,	false,		{ 0, 0, 0, 0, 0, 0, 0, 0, 0, 0, 0, 0, 0, 0, 0, 0, 0, 0, 0, 0, 0, 0 }		},	// CO_SWITCH
-		{ DT_IMAGE,		true,	false,		{ 1, 1, 1, 1, 1, 1, 1, 1, 1, 1, 1, 1, 1, 1, 1, 1, 1, 1, 1, 1, 1, 1 }		},	// IM_SWITCH
-		{ DT_MESH,		true,	false,		{ 0, 0, 0, 0, 0, 0, 0, 0, 0, 0, 0, 0, 0, 0, 0, 0, 0, 0, 0, 0, 0, 0 }		},	// ME_SWITCH
-		{ DT_LAYOUT,	false,	false,		{ 0, 0, 0, 0, 0, 0, 0, 0, 0, 0, 0, 0, 0, 0, 0, 0, 0, 0, 0, 0, 0, 0 }		},	// LA_SWITCH
-		{ DT_INSTANCE,	false,	false,		{ 0, 0, 0, 0, 0, 0, 0, 0, 0, 0, 0, 0, 0, 0, 0, 0, 0, 0, 0, 0, 0, 0 }		},	// IN_SWITCH
-
-        { DT_BOOL,		false,  false,		{ 0, 0, 0, 0, 0, 0, 0, 0, 0, 0, 0, 0, 0, 0, 0, 0, 0, 0, 0, 0, 0, 0 }		},	// BO_LESS
-        { DT_BOOL,		false,  false,		{ 0, 0, 0, 0, 0, 0, 0, 0, 0, 0, 0, 0, 0, 0, 0, 0, 0, 0, 0, 0, 0, 0 }		},	// BO_EQUAL_SC_CONST
-        { DT_BOOL,		false,  false,		{ 0, 0, 0, 0, 0, 0, 0, 0, 0, 0, 0, 0, 0, 0, 0, 0, 0, 0, 0, 0, 0, 0 }		},	// BO_AND
-        { DT_BOOL,		false,  false,		{ 0, 0, 0, 0, 0, 0, 0, 0, 0, 0, 0, 0, 0, 0, 0, 0, 0, 0, 0, 0, 0, 0 }		},	// BO_OR
-        { DT_BOOL,		false,  false,		{ 0, 0, 0, 0, 0, 0, 0, 0, 0, 0, 0, 0, 0, 0, 0, 0, 0, 0, 0, 0, 0, 0 }		},	// BO_NOT
-
-        { DT_SCALAR,	true,   false,		{ 0, 0, 0, 0, 0, 0, 0, 0, 0, 0, 0, 0, 0, 0, 0, 0, 0, 0, 0, 0, 0, 0 }		},	// SC_MULTIPLYADD
-        { DT_SCALAR,	true,   false,		{ 0, 0, 0, 0, 0, 0, 0, 0, 0, 0, 0, 0, 0, 0, 0, 0, 0, 0, 0, 0, 0, 0 }		},	// SC_ARITHMETIC
-        { DT_SCALAR,	true,   false,		{ 0, 0, 0, 0, 0, 0, 0, 0, 0, 0, 0, 0, 0, 0, 0, 0, 0, 0, 0, 0, 0, 0 }		},	// SC_CURVE
-
-        { DT_COLOUR,	false,	false,		{ 0, 0, 0, 0, 0, 0, 0, 0, 0, 0, 0, 0, 0, 0, 0, 0, 0, 0, 0, 0, 0, 0 }		},	// CO_SAMPLEIMAGE
-        { DT_COLOUR,	false,	false,		{ 0, 0, 0, 0, 0, 0, 0, 0, 0, 0, 0, 0, 0, 0, 0, 0, 0, 0, 0, 0, 0, 0 }		},	// CO_SWIZZLE
-        { DT_COLOUR,	false,	false,		{ 0, 0, 0, 0, 0, 0, 0, 0, 0, 0, 0, 0, 0, 0, 0, 0, 0, 0, 0, 0, 0, 0 }		},	// CO_IMAGESIZE
-        { DT_COLOUR,	false,	false,		{ 0, 0, 0, 0, 0, 0, 0, 0, 0, 0, 0, 0, 0, 0, 0, 0, 0, 0, 0, 0, 0, 0 }		},	// CO_LAYOUTBLOCKTRANSFORM
-        { DT_COLOUR,	false,	false,		{ 0, 0, 0, 0, 0, 0, 0, 0, 0, 0, 0, 0, 0, 0, 0, 0, 0, 0, 0, 0, 0, 0 }		},	// CO_FROMSCALARS
-        { DT_COLOUR,	false,	false,		{ 0, 0, 0, 0, 0, 0, 0, 0, 0, 0, 0, 0, 0, 0, 0, 0, 0, 0, 0, 0, 0, 0 }		},	// CO_ARITHMETIC
-
-        { DT_IMAGE,		true,	true,		{ 0, 1, 1, 1, 0, 0, 0, 0, 0, 0, 0, 0, 0, 0, 0, 0, 0, 0, 0, 0, 0, 0 }		},	// IM_LAYER
-        { DT_IMAGE,		true,	true,		{ 0, 1, 1, 1, 0, 0, 0, 0, 0, 0, 0, 0, 0, 0, 0, 0, 0, 0, 0, 0, 0, 0 }		},	// IM_LAYERCOLOUR
-        { DT_IMAGE,		true,	false,		{ 1, 1, 1, 1, 1, 1, 1, 1, 1, 1, 1, 1, 0, 0, 0, 0, 0, 0, 0, 0, 0, 0 }		},	// IM_PIXELFORMAT
-        { DT_IMAGE,		true,	false,		{ 1, 1, 1, 1, 0, 0, 0, 0, 0, 0, 0, 0, 0, 0, 0, 0, 0, 0, 0, 0, 0, 0 }		},	// IM_MIPMAP
-        { DT_IMAGE,		true,	false,		{ 0, 1, 1, 1, 0, 0, 0, 0, 0, 0, 0, 0, 0, 0, 0, 0, 0, 0, 0, 0, 0, 0 }		},	// IM_RESIZE
-        { DT_IMAGE,		true,	false,		{ 0, 1, 1, 1, 0, 0, 0, 0, 0, 0, 0, 0, 0, 0, 0, 0, 0, 0, 0, 0, 0, 0 }		},	// IM_RESIZELIKE
-		{ DT_IMAGE,		true,	false,		{ 0, 1, 1, 1, 0, 0, 0, 0, 0, 0, 0, 0, 0, 0, 0, 0, 0, 0, 0, 0, 0, 0 }		},	// IM_RESIZEREL
-		{ DT_IMAGE,		true,   false,		{ 1, 1, 1, 1, 1, 1, 1, 1, 1, 1, 1, 1, 0, 0, 0, 0, 0, 0, 0, 0, 0, 0 }		},	// IM_BLANKLAYOUT
-		{ DT_IMAGE,		true,	true,		{ 1, 1, 1, 1, 1, 1, 1, 1, 1, 1, 1, 0, 0, 0, 1, 1, 1, 1, 1, 1, 1, 1 }		},	// IM_COMPOSE
-        { DT_IMAGE,		true,	false,		{ 0, 1, 1, 0, 0, 0, 0, 0, 0, 0, 0, 0, 0, 0, 0, 0, 0, 0, 0, 0, 0, 0 }		},	// IM_DIFFERENCE
-		{ DT_IMAGE,		true,	false,		{ 0, 1, 1, 1, 0, 0, 0, 0, 0, 0, 0, 0, 0, 0, 0, 0, 0, 0, 0, 0, 0, 0 }		},	// IM_INTERPOLATE
-		{ DT_IMAGE,		true,	true,		{ 0, 1, 1, 1, 0, 0, 0, 0, 0, 0, 0, 0, 0, 0, 0, 0, 0, 0, 0, 0, 0, 0 }		},	// IM_INTERPOLATE3
-		{ DT_IMAGE,		true,	false,		{ 0, 1, 1, 0, 0, 0, 0, 0, 0, 0, 0, 0, 0, 0, 0, 0, 0, 0, 0, 0, 0, 0 }		},	// IM_SATURATE
-		{ DT_IMAGE,		true,	false,		{ 0, 1, 1, 0, 0, 0, 0, 0, 0, 0, 0, 0, 0, 0, 0, 0, 0, 0, 0, 0, 0, 0 }		},	// IM_LUMINANCE
-        { DT_IMAGE,		true,	false,		{ 0, 1, 1, 1, 0, 0, 0, 0, 0, 0, 0, 0, 0, 0, 0, 0, 0, 0, 0, 0, 0, 0 }		},	// IM_SWIZZLE
-		{ DT_IMAGE,		true,	false,		{ 0, 1, 1, 1, 0, 0, 0, 0, 0, 0, 0, 0, 0, 0, 0, 0, 0, 0, 0, 0, 0, 0 }		},	// IM_SELECTCOLOUR
-		{ DT_IMAGE,		true,	false,		{ 0, 1, 1, 1, 0, 0, 0, 0, 0, 0, 0, 0, 0, 0, 0, 0, 0, 0, 0, 0, 0, 0 }		},	// IM_COLOURMAP
-		{ DT_IMAGE,		true,	false,		{ 0, 0, 0, 0, 0, 0, 0, 0, 0, 0, 0, 0, 0, 0, 0, 0, 0, 0, 0, 0, 0, 0 }		},	// IM_GRADIENT
-        { DT_IMAGE,		true,	false,		{ 0, 1, 1, 1, 0, 0, 0, 0, 0, 0, 0, 0, 0, 0, 0, 0, 0, 0, 0, 0, 0, 0 }		},	// IM_BINARISE
-		{ DT_IMAGE,		true,	false,		{ 0, 0, 0, 0, 0, 0, 0, 0, 0, 0, 0, 0, 0, 0, 0, 0, 0, 0, 0, 0, 0, 0 }		},	// IM_PLAINCOLOUR
-		{ DT_IMAGE,		true,	false,		{ 0, 0, 0, 0, 0, 0, 0, 0, 0, 0, 0, 0, 0, 0, 0, 0, 0, 0, 0, 0, 0, 0 }		},	// IM_GPU
-		{ DT_IMAGE,		true,	false,		{ 0, 0, 0, 0, 0, 0, 0, 0, 0, 0, 0, 0, 0, 0, 0, 0, 0, 0, 0, 0, 0, 0 }		},	// IM_CROP
-		{ DT_IMAGE,		true,	false,		{ 0, 0, 0, 0, 0, 0, 0, 0, 0, 0, 0, 0, 0, 0, 0, 0, 0, 0, 0, 0, 0, 0 }		},	// IM_PATCH
-        { DT_IMAGE,		true,	false,		{ 0, 0, 0, 0, 0, 0, 0, 0, 0, 0, 0, 0, 0, 0, 0, 0, 0, 0, 0, 0, 0, 0 }		},	// IM_RASTERMESH
-		{ DT_IMAGE,		true,	false,		{ 0, 0, 0, 0, 0, 0, 0, 0, 0, 0, 0, 0, 0, 0, 0, 0, 0, 0, 0, 0, 0, 0 }		},	// IM_MAKEGROWMAP
-		{ DT_IMAGE,		true,	false,		{ 0, 0, 0, 0, 0, 0, 0, 0, 0, 0, 0, 0, 0, 0, 0, 0, 0, 0, 0, 0, 0, 0 }		},	// IM_DISPLACE
-		{ DT_IMAGE,		true,	false,		{ 0, 0, 0, 0, 0, 0, 0, 0, 0, 0, 0, 0, 0, 0, 0, 0, 0, 0, 0, 0, 0, 0 }		},	// IM_MULTILAYER
-		{ DT_IMAGE,		true,	false,		{ 0, 1, 1, 1, 0, 0, 0, 0, 0, 0, 0, 0, 0, 0, 0, 0, 0, 0, 0, 0, 0, 0 }		},	// IM_INVERT
-		{ DT_IMAGE,		true,	false,		{ 0, 1, 1, 1, 0, 0, 0, 0, 0, 0, 0, 0, 0, 0, 0, 0, 0, 0, 0, 0, 0, 0 }		},	// IM_NORMALCOMPOSITE
-		{ DT_IMAGE,		true,	false,		{ 0, 1, 1, 1, 0, 0, 0, 0, 0, 0, 0, 0, 0, 0, 0, 0, 0, 0, 0, 0, 0, 0 }		},	// IM_TRANSFORM
-
-        { DT_MESH,		true,	false,		{ 0, 0, 0, 0, 0, 0, 0, 0, 0, 0, 0, 0, 0, 0, 0, 0, 0, 0, 0, 0, 0, 0 }		},	// ME_APPLYLAYOUT
-        { DT_MESH,		true,	false,		{ 0, 0, 0, 0, 0, 0, 0, 0, 0, 0, 0, 0, 0, 0, 0, 0, 0, 0, 0, 0, 0, 0 }		},	// ME_DIFFERENCE
-        { DT_MESH,		true,	false,		{ 0, 0, 0, 0, 0, 0, 0, 0, 0, 0, 0, 0, 0, 0, 0, 0, 0, 0, 0, 0, 0, 0 }		},	// ME_MORPH2
-        { DT_MESH,		true,	false,		{ 0, 0, 0, 0, 0, 0, 0, 0, 0, 0, 0, 0, 0, 0, 0, 0, 0, 0, 0, 0, 0, 0 }		},	// ME_MERGE
-        { DT_MESH,		true,	false,		{ 0, 0, 0, 0, 0, 0, 0, 0, 0, 0, 0, 0, 0, 0, 0, 0, 0, 0, 0, 0, 0, 0 }		},	// ME_INTERPOLATE
-        { DT_MESH,		true,	false,		{ 0, 0, 0, 0, 0, 0, 0, 0, 0, 0, 0, 0, 0, 0, 0, 0, 0, 0, 0, 0, 0, 0 }		},	// ME_MASKCLIPMESH
-        { DT_MESH,		true,	false,		{ 0, 0, 0, 0, 0, 0, 0, 0, 0, 0, 0, 0, 0, 0, 0, 0, 0, 0, 0, 0, 0, 0 }		},	// ME_MASKDIFF
-        { DT_MESH,		true,	false,		{ 0, 0, 0, 0, 0, 0, 0, 0, 0, 0, 0, 0, 0, 0, 0, 0, 0, 0, 0, 0, 0, 0 }		},	// ME_SUBTRACT
-        { DT_MESH,		true,	false,		{ 0, 0, 0, 0, 0, 0, 0, 0, 0, 0, 0, 0, 0, 0, 0, 0, 0, 0, 0, 0, 0, 0 }		},	// ME_REMOVEMASK
-        { DT_MESH,		true,	false,		{ 0, 0, 0, 0, 0, 0, 0, 0, 0, 0, 0, 0, 0, 0, 0, 0, 0, 0, 0, 0, 0, 0 }		},	// ME_FORMAT
-        { DT_MESH,		true,	false,		{ 0, 0, 0, 0, 0, 0, 0, 0, 0, 0, 0, 0, 0, 0, 0, 0, 0, 0, 0, 0, 0, 0 }		},	// ME_EXTRACTLAYOUTBLOCK
-        { DT_MESH,		true,	false,		{ 0, 0, 0, 0, 0, 0, 0, 0, 0, 0, 0, 0, 0, 0, 0, 0, 0, 0, 0, 0, 0, 0 }		},	// ME_EXTRACTFACEGROUP
-        { DT_MESH,		true,	false,		{ 0, 0, 0, 0, 0, 0, 0, 0, 0, 0, 0, 0, 0, 0, 0, 0, 0, 0, 0, 0, 0, 0 }		},	// ME_TRANSFORM
-        { DT_MESH,		true,	false,		{ 0, 0, 0, 0, 0, 0, 0, 0, 0, 0, 0, 0, 0, 0, 0, 0, 0, 0, 0, 0, 0, 0 }		},	// ME_CLIPMORPHPLANE
-        { DT_MESH,		true,	false,		{ 0, 0, 0, 0, 0, 0, 0, 0, 0, 0, 0, 0, 0, 0, 0, 0, 0, 0, 0, 0, 0, 0 }		},	// ME_CLIPWITHMESH
-        { DT_MESH,		true,	false,		{ 0, 0, 0, 0, 0, 0, 0, 0, 0, 0, 0, 0, 0, 0, 0, 0, 0, 0, 0, 0, 0, 0 }		},	// ME_SETSKELETON
-        { DT_MESH,		true,	false,		{ 0, 0, 0, 0, 0, 0, 0, 0, 0, 0, 0, 0, 0, 0, 0, 0, 0, 0, 0, 0, 0, 0 }		},	// ME_PROJECT
-        { DT_MESH,		true,	false,		{ 0, 0, 0, 0, 0, 0, 0, 0, 0, 0, 0, 0, 0, 0, 0, 0, 0, 0, 0, 0, 0, 0 }		},	// ME_APPLYPOSE
-		{ DT_MESH,		true,	false,		{ 0, 0, 0, 0, 0, 0, 0, 0, 0, 0, 0, 0, 0, 0, 0, 0, 0, 0, 0, 0, 0, 0 }		},	// ME_REMAPINDICES
-		{ DT_MESH,		true,	false,		{ 0, 0, 0, 0, 0, 0, 0, 0, 0, 0, 0, 0, 0, 0, 0, 0, 0, 0, 0, 0, 0, 0 }		},	// ME_GEOMETRYOPERATION
-		{ DT_MESH,		true,	false,		{ 0, 0, 0, 0, 0, 0, 0, 0, 0, 0, 0, 0, 0, 0, 0, 0, 0, 0, 0, 0, 0, 0 }		},	// ME_BINDSHAPE
-		{ DT_MESH,		true,	false,		{ 0, 0, 0, 0, 0, 0, 0, 0, 0, 0, 0, 0, 0, 0, 0, 0, 0, 0, 0, 0, 0, 0 }		},	// ME_APPLYSHAPE
-		{ DT_MESH,		true,	false,		{ 0, 0, 0, 0, 0, 0, 0, 0, 0, 0, 0, 0, 0, 0, 0, 0, 0, 0, 0, 0, 0, 0 }		},	// ME_CLIPDEFORM
-		{ DT_MESH,		true,	false,		{ 0, 0, 0, 0, 0, 0, 0, 0, 0, 0, 0, 0, 0, 0, 0, 0, 0, 0, 0, 0, 0, 0 }		},	// ME_MORPHRESHAPE
-
-        { DT_INSTANCE,	false,	false,		{ 0, 0, 0, 0, 0, 0, 0, 0, 0, 0, 0, 0, 0, 0, 0, 0, 0, 0, 0, 0, 0, 0 }		},	// IN_ADDMESH
-        { DT_INSTANCE,	false,	false,		{ 0, 0, 0, 0, 0, 0, 0, 0, 0, 0, 0, 0, 0, 0, 0, 0, 0, 0, 0, 0, 0, 0 }		},	// IN_ADDIMAGE
-        { DT_INSTANCE,	false,	false,		{ 0, 0, 0, 0, 0, 0, 0, 0, 0, 0, 0, 0, 0, 0, 0, 0, 0, 0, 0, 0, 0, 0 }		},	// IN_ADDVECTOR
-        { DT_INSTANCE,	false,	false,		{ 0, 0, 0, 0, 0, 0, 0, 0, 0, 0, 0, 0, 0, 0, 0, 0, 0, 0, 0, 0, 0, 0 }		},	// IN_ADDSCALAR
-        { DT_INSTANCE,	false,	false,		{ 0, 0, 0, 0, 0, 0, 0, 0, 0, 0, 0, 0, 0, 0, 0, 0, 0, 0, 0, 0, 0, 0 }		},	// IN_ADDSTRING
-        { DT_INSTANCE,	false,	false,		{ 0, 0, 0, 0, 0, 0, 0, 0, 0, 0, 0, 0, 0, 0, 0, 0, 0, 0, 0, 0, 0, 0 }		},	// IN_ADDSURFACE
-        { DT_INSTANCE,	false,	false,		{ 0, 0, 0, 0, 0, 0, 0, 0, 0, 0, 0, 0, 0, 0, 0, 0, 0, 0, 0, 0, 0, 0 }		},	// IN_ADDCOMPONENT
-        { DT_INSTANCE,	false,	false,		{ 0, 0, 0, 0, 0, 0, 0, 0, 0, 0, 0, 0, 0, 0, 0, 0, 0, 0, 0, 0, 0, 0 }		},	// IN_ADDLOD
-
-        { DT_LAYOUT,	true,   false,		{ 0, 0, 0, 0, 0, 0, 0, 0, 0, 0, 0, 0, 0, 0, 0, 0, 0, 0, 0, 0, 0, 0 }		},	// LA_PACK
-        { DT_LAYOUT,	true,   false,		{ 0, 0, 0, 0, 0, 0, 0, 0, 0, 0, 0, 0, 0, 0, 0, 0, 0, 0, 0, 0, 0, 0 }		},	// LA_MERGE
-		{ DT_LAYOUT,	true,   false,		{ 0, 0, 0, 0, 0, 0, 0, 0, 0, 0, 0, 0, 0, 0, 0, 0, 0, 0, 0, 0, 0, 0 }		},	// LA_REMOVEBLOCKS
-		{ DT_LAYOUT,	true,   false,		{ 0, 0, 0, 0, 0, 0, 0, 0, 0, 0, 0, 0, 0, 0, 0, 0, 0, 0, 0, 0, 0, 0 }		},	// LA_FROMMESH
+		// type			cached	supported base image formats
+        { DT_NONE,		false,	{ 0, 0, 0, 0, 0, 0, 0, 0, 0, 0, 0, 0, 0, 0, 0, 0, 0, 0, 0, 0, 0, 0, 0, 0, 0, 0, 0, 0, 0 }		},	// NONE
+
+		{ DT_BOOL,		false,	{ 0, 0, 0, 0, 0, 0, 0, 0, 0, 0, 0, 0, 0, 0, 0, 0, 0, 0, 0, 0, 0, 0, 0, 0, 0, 0, 0, 0, 0 }		},	// BO_CONSTANT
+		{ DT_INT,		false,	{ 0, 0, 0, 0, 0, 0, 0, 0, 0, 0, 0, 0, 0, 0, 0, 0, 0, 0, 0, 0, 0, 0, 0, 0, 0, 0, 0, 0, 0 }		},	// NU_CONSTANT
+		{ DT_SCALAR,	false,	{ 0, 0, 0, 0, 0, 0, 0, 0, 0, 0, 0, 0, 0, 0, 0, 0, 0, 0, 0, 0, 0, 0, 0, 0, 0, 0, 0, 0, 0 }		},	// SC_CONSTANT
+		{ DT_COLOUR,	false,	{ 0, 0, 0, 0, 0, 0, 0, 0, 0, 0, 0, 0, 0, 0, 0, 0, 0, 0, 0, 0, 0, 0, 0, 0, 0, 0, 0, 0, 0 }		},	// CO_CONSTANT
+		{ DT_IMAGE,		true,   { 1, 1, 1, 1, 1, 1, 1, 1, 1, 1, 1, 1, 0, 0, 0, 0, 0, 0, 0, 0, 0, 0, 0, 0, 0, 0, 0, 0, 0 }		},	// IM_CONSTANT
+		{ DT_MESH,		true,   { 0, 0, 0, 0, 0, 0, 0, 0, 0, 0, 0, 0, 0, 0, 0, 0, 0, 0, 0, 0, 0, 0, 0, 0, 0, 0, 0, 0, 0 }		},	// ME_CONSTANT
+		{ DT_LAYOUT,	false,	{ 0, 0, 0, 0, 0, 0, 0, 0, 0, 0, 0, 0, 0, 0, 0, 0, 0, 0, 0, 0, 0, 0, 0, 0, 0, 0, 0, 0, 0 }		},	// LA_CONSTANT
+		{ DT_PROJECTOR,	false,	{ 0, 0, 0, 0, 0, 0, 0, 0, 0, 0, 0, 0, 0, 0, 0, 0, 0, 0, 0, 0, 0, 0, 0, 0, 0, 0, 0, 0, 0 }		},	// PR_CONSTANT
+		{ DT_STRING,	false,	{ 0, 0, 0, 0, 0, 0, 0, 0, 0, 0, 0, 0, 0, 0, 0, 0, 0, 0, 0, 0, 0, 0, 0, 0, 0, 0, 0, 0, 0 }		},	// ST_CONSTANT
+
+		{ DT_BOOL,		false,	{ 0, 0, 0, 0, 0, 0, 0, 0, 0, 0, 0, 0, 0, 0, 0, 0, 0, 0, 0, 0, 0, 0, 0, 0, 0, 0, 0, 0, 0 }		},	// BO_PARAMETER
+		{ DT_INT,		false,	{ 0, 0, 0, 0, 0, 0, 0, 0, 0, 0, 0, 0, 0, 0, 0, 0, 0, 0, 0, 0, 0, 0, 0, 0, 0, 0, 0, 0, 0 }		},	// NU_PARAMETER
+		{ DT_SCALAR,	false,	{ 0, 0, 0, 0, 0, 0, 0, 0, 0, 0, 0, 0, 0, 0, 0, 0, 0, 0, 0, 0, 0, 0, 0, 0, 0, 0, 0, 0, 0 }		},	// SC_PARAMETER
+		{ DT_COLOUR,	false,	{ 0, 0, 0, 0, 0, 0, 0, 0, 0, 0, 0, 0, 0, 0, 0, 0, 0, 0, 0, 0, 0, 0, 0, 0, 0, 0, 0, 0, 0 }		},	// CO_PARAMETER
+		{ DT_PROJECTOR,	false,	{ 0, 0, 0, 0, 0, 0, 0, 0, 0, 0, 0, 0, 0, 0, 0, 0, 0, 0, 0, 0, 0, 0, 0, 0, 0, 0, 0, 0, 0 }		},	// PR_PARAMETER
+		{ DT_IMAGE,		false,	{ 0, 0, 0, 0, 0, 0, 0, 0, 0, 0, 0, 0, 0, 0, 0, 0, 0, 0, 0, 0, 0, 0, 0, 0, 0, 0, 0, 0, 0 }		},	// IM_PARAMETER
+		{ DT_STRING,	false,	{ 0, 0, 0, 0, 0, 0, 0, 0, 0, 0, 0, 0, 0, 0, 0, 0, 0, 0, 0, 0, 0, 0, 0, 0, 0, 0, 0, 0, 0 }		},	// ST_PARAMETER
+
+		{ DT_INT,		false,	{ 0, 0, 0, 0, 0, 0, 0, 0, 0, 0, 0, 0, 0, 0, 0, 0, 0, 0, 0, 0, 0, 0, 0, 0, 0, 0, 0, 0, 0 }		},	// NU_CONDITIONAL
+		{ DT_SCALAR,	false,	{ 0, 0, 0, 0, 0, 0, 0, 0, 0, 0, 0, 0, 0, 0, 0, 0, 0, 0, 0, 0, 0, 0, 0, 0, 0, 0, 0, 0, 0 }		},	// SC_CONDITIONAL
+		{ DT_COLOUR,	false,	{ 0, 0, 0, 0, 0, 0, 0, 0, 0, 0, 0, 0, 0, 0, 0, 0, 0, 0, 0, 0, 0, 0, 0, 0, 0, 0, 0, 0, 0 }		},	// CO_CONDITIONAL
+		{ DT_IMAGE,		true,	{ 1, 1, 1, 1, 1, 1, 1, 1, 1, 1, 1, 1, 1, 1, 1, 1, 1, 1, 1, 1, 1, 1, 0, 0, 1, 1, 0, 0, 0 }		},	// IM_CONDITIONAL
+		{ DT_MESH,		true,	{ 0, 0, 0, 0, 0, 0, 0, 0, 0, 0, 0, 0, 0, 0, 0, 0, 0, 0, 0, 0, 0, 0, 0, 0, 0, 0, 0, 0, 0 }		},	// ME_CONDITIONAL
+		{ DT_LAYOUT,	false,  { 0, 0, 0, 0, 0, 0, 0, 0, 0, 0, 0, 0, 0, 0, 0, 0, 0, 0, 0, 0, 0, 0, 0, 0, 0, 0, 0, 0, 0 }		},	// LA_CONDITIONAL
+		{ DT_INSTANCE,	false,	{ 0, 0, 0, 0, 0, 0, 0, 0, 0, 0, 0, 0, 0, 0, 0, 0, 0, 0, 0, 0, 0, 0, 0, 0, 0, 0, 0, 0, 0 }		},	// IN_CONDITIONAL
+
+		{ DT_INT,		false,	{ 0, 0, 0, 0, 0, 0, 0, 0, 0, 0, 0, 0, 0, 0, 0, 0, 0, 0, 0, 0, 0, 0, 0, 0, 0, 0, 0, 0, 0 }		},	// NU_SWITCH
+		{ DT_SCALAR,	false,	{ 0, 0, 0, 0, 0, 0, 0, 0, 0, 0, 0, 0, 0, 0, 0, 0, 0, 0, 0, 0, 0, 0, 0, 0, 0, 0, 0, 0, 0 }		},	// SC_SWITCH
+		{ DT_COLOUR,	false,	{ 0, 0, 0, 0, 0, 0, 0, 0, 0, 0, 0, 0, 0, 0, 0, 0, 0, 0, 0, 0, 0, 0, 0, 0, 0, 0, 0, 0, 0 }		},	// CO_SWITCH
+		{ DT_IMAGE,		true,	{ 1, 1, 1, 1, 1, 1, 1, 1, 1, 1, 1, 1, 1, 1, 1, 1, 1, 1, 1, 1, 1, 1, 0, 0, 1, 1, 0, 0, 0 }		},	// IM_SWITCH
+		{ DT_MESH,		true,	{ 0, 0, 0, 0, 0, 0, 0, 0, 0, 0, 0, 0, 0, 0, 0, 0, 0, 0, 0, 0, 0, 0, 0, 0, 0, 0, 0, 0, 0 }		},	// ME_SWITCH
+		{ DT_LAYOUT,	false,	{ 0, 0, 0, 0, 0, 0, 0, 0, 0, 0, 0, 0, 0, 0, 0, 0, 0, 0, 0, 0, 0, 0, 0, 0, 0, 0, 0, 0, 0 }		},	// LA_SWITCH
+		{ DT_INSTANCE,	false,	{ 0, 0, 0, 0, 0, 0, 0, 0, 0, 0, 0, 0, 0, 0, 0, 0, 0, 0, 0, 0, 0, 0, 0, 0, 0, 0, 0, 0, 0 }		},	// IN_SWITCH
+
+		{ DT_BOOL,		false,  { 0, 0, 0, 0, 0, 0, 0, 0, 0, 0, 0, 0, 0, 0, 0, 0, 0, 0, 0, 0, 0, 0, 0, 0, 0, 0, 0, 0, 0 }		},	// BO_LESS
+		{ DT_BOOL,		false,  { 0, 0, 0, 0, 0, 0, 0, 0, 0, 0, 0, 0, 0, 0, 0, 0, 0, 0, 0, 0, 0, 0, 0, 0, 0, 0, 0, 0, 0 }		},	// BO_EQUAL_SC_CONST
+		{ DT_BOOL,		false,  { 0, 0, 0, 0, 0, 0, 0, 0, 0, 0, 0, 0, 0, 0, 0, 0, 0, 0, 0, 0, 0, 0, 0, 0, 0, 0, 0, 0, 0 }		},	// BO_AND
+		{ DT_BOOL,		false,  { 0, 0, 0, 0, 0, 0, 0, 0, 0, 0, 0, 0, 0, 0, 0, 0, 0, 0, 0, 0, 0, 0, 0, 0, 0, 0, 0, 0, 0 }		},	// BO_OR
+		{ DT_BOOL,		false,  { 0, 0, 0, 0, 0, 0, 0, 0, 0, 0, 0, 0, 0, 0, 0, 0, 0, 0, 0, 0, 0, 0, 0, 0, 0, 0, 0, 0, 0 }		},	// BO_NOT
+
+		{ DT_SCALAR,	true,   { 0, 0, 0, 0, 0, 0, 0, 0, 0, 0, 0, 0, 0, 0, 0, 0, 0, 0, 0, 0, 0, 0, 0, 0, 0, 0, 0, 0, 0 }		},	// SC_MULTIPLYADD
+		{ DT_SCALAR,	true,   { 0, 0, 0, 0, 0, 0, 0, 0, 0, 0, 0, 0, 0, 0, 0, 0, 0, 0, 0, 0, 0, 0, 0, 0, 0, 0, 0, 0, 0 }		},	// SC_ARITHMETIC
+		{ DT_SCALAR,	true,   { 0, 0, 0, 0, 0, 0, 0, 0, 0, 0, 0, 0, 0, 0, 0, 0, 0, 0, 0, 0, 0, 0, 0, 0, 0, 0, 0, 0, 0 }		},	// SC_CURVE
+
+		{ DT_COLOUR,	false,	{ 0, 0, 0, 0, 0, 0, 0, 0, 0, 0, 0, 0, 0, 0, 0, 0, 0, 0, 0, 0, 0, 0, 0, 0, 0, 0, 0, 0, 0 }		},	// CO_SAMPLEIMAGE
+		{ DT_COLOUR,	false,	{ 0, 0, 0, 0, 0, 0, 0, 0, 0, 0, 0, 0, 0, 0, 0, 0, 0, 0, 0, 0, 0, 0, 0, 0, 0, 0, 0, 0, 0 }		},	// CO_SWIZZLE
+		{ DT_COLOUR,	false,	{ 0, 0, 0, 0, 0, 0, 0, 0, 0, 0, 0, 0, 0, 0, 0, 0, 0, 0, 0, 0, 0, 0, 0, 0, 0, 0, 0, 0, 0 }		},	// CO_IMAGESIZE
+		{ DT_COLOUR,	false,	{ 0, 0, 0, 0, 0, 0, 0, 0, 0, 0, 0, 0, 0, 0, 0, 0, 0, 0, 0, 0, 0, 0, 0, 0, 0, 0, 0, 0, 0 }		},	// CO_LAYOUTBLOCKTRANSFORM
+		{ DT_COLOUR,	false,	{ 0, 0, 0, 0, 0, 0, 0, 0, 0, 0, 0, 0, 0, 0, 0, 0, 0, 0, 0, 0, 0, 0, 0, 0, 0, 0, 0, 0, 0 }		},	// CO_FROMSCALARS
+		{ DT_COLOUR,	false,	{ 0, 0, 0, 0, 0, 0, 0, 0, 0, 0, 0, 0, 0, 0, 0, 0, 0, 0, 0, 0, 0, 0, 0, 0, 0, 0, 0, 0, 0 }		},	// CO_ARITHMETIC
+
+		{ DT_IMAGE,		true,	{ 0, 1, 1, 1, 0, 0, 0, 0, 0, 0, 0, 0, 0, 0, 0, 0, 0, 0, 0, 0, 0, 0, 0, 0, 0, 0, 0, 0, 0 }		},	// IM_LAYER
+		{ DT_IMAGE,		true,	{ 0, 1, 1, 1, 0, 0, 0, 0, 0, 0, 0, 0, 0, 0, 0, 0, 0, 0, 0, 0, 0, 0, 0, 0, 0, 0, 0, 0, 0 }		},	// IM_LAYERCOLOUR
+		{ DT_IMAGE,		true,	{ 1, 1, 1, 1, 1, 1, 1, 1, 1, 1, 1, 1, 0, 0, 0, 0, 0, 0, 0, 0, 0, 0, 0, 0, 0, 0, 0, 0, 0 }		},	// IM_PIXELFORMAT
+		{ DT_IMAGE,		true,	{ 1, 1, 1, 1, 0, 0, 0, 0, 0, 0, 0, 0, 0, 0, 0, 0, 0, 0, 0, 0, 0, 0, 0, 0, 0, 0, 0, 0, 0 }		},	// IM_MIPMAP
+		{ DT_IMAGE,		true,	{ 0, 1, 1, 1, 0, 0, 0, 0, 0, 0, 0, 0, 0, 0, 0, 0, 0, 0, 0, 0, 0, 0, 0, 0, 0, 0, 0, 0, 0 }		},	// IM_RESIZE
+		{ DT_IMAGE,		true,	{ 0, 1, 1, 1, 0, 0, 0, 0, 0, 0, 0, 0, 0, 0, 0, 0, 0, 0, 0, 0, 0, 0, 0, 0, 0, 0, 0, 0, 0 }		},	// IM_RESIZELIKE
+		{ DT_IMAGE,		true,	{ 0, 1, 1, 1, 0, 0, 0, 0, 0, 0, 0, 0, 0, 0, 0, 0, 0, 0, 0, 0, 0, 0, 0, 0, 0, 0, 0, 0, 0 }		},	// IM_RESIZEREL
+		{ DT_IMAGE,		true,   { 1, 1, 1, 1, 1, 1, 1, 1, 1, 1, 1, 1, 0, 0, 0, 0, 0, 0, 0, 0, 0, 0, 0, 0, 0, 0, 0, 0, 0 }		},	// IM_BLANKLAYOUT
+		{ DT_IMAGE,		true,	{ 1, 1, 1, 1, 1, 1, 1, 1, 1, 1, 1, 0, 0, 0, 1, 1, 1, 1, 1, 1, 1, 1, 0, 0, 1, 1, 0, 0, 0 }		},	// IM_COMPOSE
+		{ DT_IMAGE,		true,	{ 0, 1, 1, 0, 0, 0, 0, 0, 0, 0, 0, 0, 0, 0, 0, 0, 0, 0, 0, 0, 0, 0, 0, 0, 0, 0, 0, 0, 0 }		},	// IM_DIFFERENCE
+		{ DT_IMAGE,		true,	{ 0, 1, 1, 1, 0, 0, 0, 0, 0, 0, 0, 0, 0, 0, 0, 0, 0, 0, 0, 0, 0, 0, 0, 0, 0, 0, 0, 0, 0 }		},	// IM_INTERPOLATE
+		{ DT_IMAGE,		true,	{ 0, 1, 1, 1, 0, 0, 0, 0, 0, 0, 0, 0, 0, 0, 0, 0, 0, 0, 0, 0, 0, 0, 0, 0, 0, 0, 0, 0, 0 }		},	// IM_INTERPOLATE3
+		{ DT_IMAGE,		true,	{ 0, 1, 1, 0, 0, 0, 0, 0, 0, 0, 0, 0, 0, 0, 0, 0, 0, 0, 0, 0, 0, 0, 0, 0, 0, 0, 0, 0, 0 }		},	// IM_SATURATE
+		{ DT_IMAGE,		true,	{ 0, 1, 1, 0, 0, 0, 0, 0, 0, 0, 0, 0, 0, 0, 0, 0, 0, 0, 0, 0, 0, 0, 0, 0, 0, 0, 0, 0, 0 }		},	// IM_LUMINANCE
+		{ DT_IMAGE,		true,	{ 0, 1, 1, 1, 0, 0, 0, 0, 0, 0, 0, 0, 0, 0, 0, 0, 0, 0, 0, 0, 0, 0, 0, 0, 0, 0, 0, 0, 0 }		},	// IM_SWIZZLE
+		{ DT_IMAGE,		true,	{ 0, 1, 1, 1, 0, 0, 0, 0, 0, 0, 0, 0, 0, 0, 0, 0, 0, 0, 0, 0, 0, 0, 0, 0, 0, 0, 0, 0, 0 }		},	// IM_SELECTCOLOUR
+		{ DT_IMAGE,		true,	{ 0, 1, 1, 1, 0, 0, 0, 0, 0, 0, 0, 0, 0, 0, 0, 0, 0, 0, 0, 0, 0, 0, 0, 0, 0, 0, 0, 0, 0 }		},	// IM_COLOURMAP
+		{ DT_IMAGE,		true,	{ 0, 0, 0, 0, 0, 0, 0, 0, 0, 0, 0, 0, 0, 0, 0, 0, 0, 0, 0, 0, 0, 0, 0, 0, 0, 0, 0, 0, 0 }		},	// IM_GRADIENT
+		{ DT_IMAGE,		true,	{ 0, 1, 1, 1, 0, 0, 0, 0, 0, 0, 0, 0, 0, 0, 0, 0, 0, 0, 0, 0, 0, 0, 0, 0, 0, 0, 0, 0, 0 }		},	// IM_BINARISE
+		{ DT_IMAGE,		true,	{ 0, 0, 0, 0, 0, 0, 0, 0, 0, 0, 0, 0, 0, 0, 0, 0, 0, 0, 0, 0, 0, 0, 0, 0, 0, 0, 0, 0, 0 }		},	// IM_PLAINCOLOUR
+		{ DT_IMAGE,		true,	{ 0, 0, 0, 0, 0, 0, 0, 0, 0, 0, 0, 0, 0, 0, 0, 0, 0, 0, 0, 0, 0, 0, 0, 0, 0, 0, 0, 0, 0 }		},	// IM_GPU
+		{ DT_IMAGE,		true,	{ 0, 0, 0, 0, 0, 0, 0, 0, 0, 0, 0, 0, 0, 0, 0, 0, 0, 0, 0, 0, 0, 0, 0, 0, 0, 0, 0, 0, 0 }		},	// IM_CROP
+		{ DT_IMAGE,		true,	{ 0, 0, 0, 0, 0, 0, 0, 0, 0, 0, 0, 0, 0, 0, 0, 0, 0, 0, 0, 0, 0, 0, 0, 0, 0, 0, 0, 0, 0 }		},	// IM_PATCH
+		{ DT_IMAGE,		true,	{ 0, 0, 0, 0, 0, 0, 0, 0, 0, 0, 0, 0, 0, 0, 0, 0, 0, 0, 0, 0, 0, 0, 0, 0, 0, 0, 0, 0, 0 }		},	// IM_RASTERMESH
+		{ DT_IMAGE,		true,	{ 0, 0, 0, 0, 0, 0, 0, 0, 0, 0, 0, 0, 0, 0, 0, 0, 0, 0, 0, 0, 0, 0, 0, 0, 0, 0, 0, 0, 0 }		},	// IM_MAKEGROWMAP
+		{ DT_IMAGE,		true,	{ 0, 0, 0, 0, 0, 0, 0, 0, 0, 0, 0, 0, 0, 0, 0, 0, 0, 0, 0, 0, 0, 0, 0, 0, 0, 0, 0, 0, 0 }		},	// IM_DISPLACE
+		{ DT_IMAGE,		true,	{ 0, 0, 0, 0, 0, 0, 0, 0, 0, 0, 0, 0, 0, 0, 0, 0, 0, 0, 0, 0, 0, 0, 0, 0, 0, 0, 0, 0, 0 }		},	// IM_MULTILAYER
+		{ DT_IMAGE,		true,	{ 0, 1, 1, 1, 0, 0, 0, 0, 0, 0, 0, 0, 0, 0, 0, 0, 0, 0, 0, 0, 0, 0, 0, 0, 0, 0, 0, 0, 0 }		},	// IM_INVERT
+		{ DT_IMAGE,		true,	{ 0, 1, 1, 1, 0, 0, 0, 0, 0, 0, 0, 0, 0, 0, 0, 0, 0, 0, 0, 0, 0, 0, 0, 0, 0, 0, 0, 0, 0 }		},	// IM_NORMALCOMPOSITE
+		{ DT_IMAGE,		true,	{ 0, 1, 1, 1, 0, 0, 0, 0, 0, 0, 0, 0, 0, 0, 0, 0, 0, 0, 0, 0, 0, 0, 0, 0, 0, 0, 0, 0, 0 }		},	// IM_TRANSFORM
+
+		{ DT_MESH,		true,	{ 0, 0, 0, 0, 0, 0, 0, 0, 0, 0, 0, 0, 0, 0, 0, 0, 0, 0, 0, 0, 0, 0, 0, 0, 0, 0, 0, 0, 0 }		},	// ME_APPLYLAYOUT
+		{ DT_MESH,		true,	{ 0, 0, 0, 0, 0, 0, 0, 0, 0, 0, 0, 0, 0, 0, 0, 0, 0, 0, 0, 0, 0, 0, 0, 0, 0, 0, 0, 0, 0 }		},	// ME_DIFFERENCE
+		{ DT_MESH,		true,	{ 0, 0, 0, 0, 0, 0, 0, 0, 0, 0, 0, 0, 0, 0, 0, 0, 0, 0, 0, 0, 0, 0, 0, 0, 0, 0, 0, 0, 0 }		},	// ME_MORPH2
+		{ DT_MESH,		true,	{ 0, 0, 0, 0, 0, 0, 0, 0, 0, 0, 0, 0, 0, 0, 0, 0, 0, 0, 0, 0, 0, 0, 0, 0, 0, 0, 0, 0, 0 }		},	// ME_MERGE
+		{ DT_MESH,		true,	{ 0, 0, 0, 0, 0, 0, 0, 0, 0, 0, 0, 0, 0, 0, 0, 0, 0, 0, 0, 0, 0, 0, 0, 0, 0, 0, 0, 0, 0 }		},	// ME_INTERPOLATE
+		{ DT_MESH,		true,	{ 0, 0, 0, 0, 0, 0, 0, 0, 0, 0, 0, 0, 0, 0, 0, 0, 0, 0, 0, 0, 0, 0, 0, 0, 0, 0, 0, 0, 0 }		},	// ME_MASKCLIPMESH
+		{ DT_MESH,		true,	{ 0, 0, 0, 0, 0, 0, 0, 0, 0, 0, 0, 0, 0, 0, 0, 0, 0, 0, 0, 0, 0, 0, 0, 0, 0, 0, 0, 0, 0 }		},	// ME_MASKDIFF
+		{ DT_MESH,		true,	{ 0, 0, 0, 0, 0, 0, 0, 0, 0, 0, 0, 0, 0, 0, 0, 0, 0, 0, 0, 0, 0, 0, 0, 0, 0, 0, 0, 0, 0 }		},	// ME_SUBTRACT
+		{ DT_MESH,		true,	{ 0, 0, 0, 0, 0, 0, 0, 0, 0, 0, 0, 0, 0, 0, 0, 0, 0, 0, 0, 0, 0, 0, 0, 0, 0, 0, 0, 0, 0 }		},	// ME_REMOVEMASK
+		{ DT_MESH,		true,	{ 0, 0, 0, 0, 0, 0, 0, 0, 0, 0, 0, 0, 0, 0, 0, 0, 0, 0, 0, 0, 0, 0, 0, 0, 0, 0, 0, 0, 0 }		},	// ME_FORMAT
+		{ DT_MESH,		true,	{ 0, 0, 0, 0, 0, 0, 0, 0, 0, 0, 0, 0, 0, 0, 0, 0, 0, 0, 0, 0, 0, 0, 0, 0, 0, 0, 0, 0, 0 }		},	// ME_EXTRACTLAYOUTBLOCK
+		{ DT_MESH,		true,	{ 0, 0, 0, 0, 0, 0, 0, 0, 0, 0, 0, 0, 0, 0, 0, 0, 0, 0, 0, 0, 0, 0, 0, 0, 0, 0, 0, 0, 0 }		},	// ME_EXTRACTFACEGROUP
+		{ DT_MESH,		true,	{ 0, 0, 0, 0, 0, 0, 0, 0, 0, 0, 0, 0, 0, 0, 0, 0, 0, 0, 0, 0, 0, 0, 0, 0, 0, 0, 0, 0, 0 }		},	// ME_TRANSFORM
+		{ DT_MESH,		true,	{ 0, 0, 0, 0, 0, 0, 0, 0, 0, 0, 0, 0, 0, 0, 0, 0, 0, 0, 0, 0, 0, 0, 0, 0, 0, 0, 0, 0, 0 }		},	// ME_CLIPMORPHPLANE
+		{ DT_MESH,		true,	{ 0, 0, 0, 0, 0, 0, 0, 0, 0, 0, 0, 0, 0, 0, 0, 0, 0, 0, 0, 0, 0, 0, 0, 0, 0, 0, 0, 0, 0 }		},	// ME_CLIPWITHMESH
+		{ DT_MESH,		true,	{ 0, 0, 0, 0, 0, 0, 0, 0, 0, 0, 0, 0, 0, 0, 0, 0, 0, 0, 0, 0, 0, 0, 0, 0, 0, 0, 0, 0, 0 }		},	// ME_SETSKELETON
+		{ DT_MESH,		true,	{ 0, 0, 0, 0, 0, 0, 0, 0, 0, 0, 0, 0, 0, 0, 0, 0, 0, 0, 0, 0, 0, 0, 0, 0, 0, 0, 0, 0, 0 }		},	// ME_PROJECT
+		{ DT_MESH,		true,	{ 0, 0, 0, 0, 0, 0, 0, 0, 0, 0, 0, 0, 0, 0, 0, 0, 0, 0, 0, 0, 0, 0, 0, 0, 0, 0, 0, 0, 0 }		},	// ME_APPLYPOSE
+		{ DT_MESH,		true,	{ 0, 0, 0, 0, 0, 0, 0, 0, 0, 0, 0, 0, 0, 0, 0, 0, 0, 0, 0, 0, 0, 0, 0, 0, 0, 0, 0, 0, 0 }		},	// ME_REMAPINDICES
+		{ DT_MESH,		true,	{ 0, 0, 0, 0, 0, 0, 0, 0, 0, 0, 0, 0, 0, 0, 0, 0, 0, 0, 0, 0, 0, 0, 0, 0, 0, 0, 0, 0, 0 }		},	// ME_GEOMETRYOPERATION
+		{ DT_MESH,		true,	{ 0, 0, 0, 0, 0, 0, 0, 0, 0, 0, 0, 0, 0, 0, 0, 0, 0, 0, 0, 0, 0, 0, 0, 0, 0, 0, 0, 0, 0 }		},	// ME_BINDSHAPE
+		{ DT_MESH,		true,	{ 0, 0, 0, 0, 0, 0, 0, 0, 0, 0, 0, 0, 0, 0, 0, 0, 0, 0, 0, 0, 0, 0, 0, 0, 0, 0, 0, 0, 0 }		},	// ME_APPLYSHAPE
+		{ DT_MESH,		true,	{ 0, 0, 0, 0, 0, 0, 0, 0, 0, 0, 0, 0, 0, 0, 0, 0, 0, 0, 0, 0, 0, 0, 0, 0, 0, 0, 0, 0, 0 }		},	// ME_CLIPDEFORM
+		{ DT_MESH,		true,	{ 0, 0, 0, 0, 0, 0, 0, 0, 0, 0, 0, 0, 0, 0, 0, 0, 0, 0, 0, 0, 0, 0, 0, 0, 0, 0, 0, 0, 0 }		},	// ME_MORPHRESHAPE
+
+		{ DT_INSTANCE,	false,	{ 0, 0, 0, 0, 0, 0, 0, 0, 0, 0, 0, 0, 0, 0, 0, 0, 0, 0, 0, 0, 0, 0, 0, 0, 0, 0, 0, 0, 0 }		},	// IN_ADDMESH
+		{ DT_INSTANCE,	false,	{ 0, 0, 0, 0, 0, 0, 0, 0, 0, 0, 0, 0, 0, 0, 0, 0, 0, 0, 0, 0, 0, 0, 0, 0, 0, 0, 0, 0, 0 }		},	// IN_ADDIMAGE
+		{ DT_INSTANCE,	false,	{ 0, 0, 0, 0, 0, 0, 0, 0, 0, 0, 0, 0, 0, 0, 0, 0, 0, 0, 0, 0, 0, 0, 0, 0, 0, 0, 0, 0, 0 }		},	// IN_ADDVECTOR
+		{ DT_INSTANCE,	false,	{ 0, 0, 0, 0, 0, 0, 0, 0, 0, 0, 0, 0, 0, 0, 0, 0, 0, 0, 0, 0, 0, 0, 0, 0, 0, 0, 0, 0, 0 }		},	// IN_ADDSCALAR
+		{ DT_INSTANCE,	false,	{ 0, 0, 0, 0, 0, 0, 0, 0, 0, 0, 0, 0, 0, 0, 0, 0, 0, 0, 0, 0, 0, 0, 0, 0, 0, 0, 0, 0, 0 }		},	// IN_ADDSTRING
+		{ DT_INSTANCE,	false,	{ 0, 0, 0, 0, 0, 0, 0, 0, 0, 0, 0, 0, 0, 0, 0, 0, 0, 0, 0, 0, 0, 0, 0, 0, 0, 0, 0, 0, 0 }		},	// IN_ADDSURFACE
+		{ DT_INSTANCE,	false,	{ 0, 0, 0, 0, 0, 0, 0, 0, 0, 0, 0, 0, 0, 0, 0, 0, 0, 0, 0, 0, 0, 0, 0, 0, 0, 0, 0, 0, 0 }		},	// IN_ADDCOMPONENT
+		{ DT_INSTANCE,	false,	{ 0, 0, 0, 0, 0, 0, 0, 0, 0, 0, 0, 0, 0, 0, 0, 0, 0, 0, 0, 0, 0, 0, 0, 0, 0, 0, 0, 0, 0 }		},	// IN_ADDLOD
+
+		{ DT_LAYOUT,	true,   { 0, 0, 0, 0, 0, 0, 0, 0, 0, 0, 0, 0, 0, 0, 0, 0, 0, 0, 0, 0, 0, 0, 0, 0, 0, 0, 0, 0, 0 }		},	// LA_PACK
+		{ DT_LAYOUT,	true,   { 0, 0, 0, 0, 0, 0, 0, 0, 0, 0, 0, 0, 0, 0, 0, 0, 0, 0, 0, 0, 0, 0, 0, 0, 0, 0, 0, 0, 0 }		},	// LA_MERGE
+		{ DT_LAYOUT,	true,   { 0, 0, 0, 0, 0, 0, 0, 0, 0, 0, 0, 0, 0, 0, 0, 0, 0, 0, 0, 0, 0, 0, 0, 0, 0, 0, 0, 0, 0 }		},	// LA_REMOVEBLOCKS
+		{ DT_LAYOUT,	true,   { 0, 0, 0, 0, 0, 0, 0, 0, 0, 0, 0, 0, 0, 0, 0, 0, 0, 0, 0, 0, 0, 0, 0, 0, 0, 0, 0, 0, 0 }		},	// LA_FROMMESH
 };
 
     // clang-format on
@@ -149,7 +149,7 @@ namespace mu
 	//---------------------------------------------------------------------------------------------
 	const OP_DESC& GetOpDesc( OP_TYPE type )
 	{
-        return s_opDescs[ (int)type ];
+        return s_opDescs[ (int32)type ];
 	}
 
 
diff --git a/Engine/Plugins/Experimental/Mutable/Source/MutableRuntime/Private/MuR/Operations.h b/Engine/Plugins/Experimental/Mutable/Source/MutableRuntime/Private/MuR/Operations.h
--- a/Engine/Plugins/Experimental/Mutable/Source/MutableRuntime/Private/MuR/Operations.h
+++ b/Engine/Plugins/Experimental/Mutable/Source/MutableRuntime/Private/MuR/Operations.h
@@ -1163,9 +1163,6 @@ namespace mu
         //! True if the instruction is worth caching when generating models
         bool cached;
 
-        //! True if the instruction is worth executing in the GPU
-        bool gpuizable;
-
         //! For image instructions, for every image format, true if it is supported as the base
         //! format of the operation.
         //! TODO: Move to tools library?
diff --git a/Engine/Plugins/Experimental/Mutable/Source/MutableRuntime/Public/MuR/Image.h b/Engine/Plugins/Experimental/Mutable/Source/MutableRuntime/Public/MuR/Image.h
--- a/Engine/Plugins/Experimental/Mutable/Source/MutableRuntime/Public/MuR/Image.h
+++ b/Engine/Plugins/Experimental/Mutable/Source/MutableRuntime/Public/MuR/Image.h
@@ -65,6 +65,13 @@ namespace mu
         IF_ASTC_4x4_RGBA_LDR,
         IF_ASTC_4x4_RG_LDR,
 
+		IF_ASTC_8x8_RGB_LDR,
+		IF_ASTC_8x8_RGBA_LDR,
+		IF_ASTC_8x8_RG_LDR,
+		IF_ASTC_12x12_RGB_LDR,
+		IF_ASTC_12x12_RGBA_LDR,
+		IF_ASTC_12x12_RG_LDR,
+
         IF_COUNT
 	};
 
diff --git a/Engine/Plugins/Experimental/Mutable/Source/MutableTools/Private/MuT/TypeInfo.cpp b/Engine/Plugins/Experimental/Mutable/Source/MutableTools/Private/MuT/TypeInfo.cpp
--- a/Engine/Plugins/Experimental/Mutable/Source/MutableTools/Private/MuT/TypeInfo.cpp
+++ b/Engine/Plugins/Experimental/Mutable/Source/MutableTools/Private/MuT/TypeInfo.cpp
@@ -39,6 +39,12 @@ namespace mu
 		"ASTC_4x4_RGBA_LDR",
 		"ASTC_4x4_RG_LDR",
 
+		"ASTC_8x8_RGB_LDR",
+		"ASTC_8x8_RGBA_LDR",
+		"ASTC_8x8_RG_LDR",
+		"ASTC_12x12_RGB_LDR",
+		"ASTC_12x12_RGBA_LDR",
+		"ASTC_12x12_RG_LDR",
 	};
 
 
