[mutable] fix possible crash with missing masks.
#rnx

diff --git a/Engine/Plugins/Experimental/Mutable/Source/CustomizableObject/Private/MuCO/CustomizableObjectInstance.cpp b/Engine/Plugins/Experimental/Mutable/Source/CustomizableObject/Private/MuCO/CustomizableObjectInstance.cpp
--- a/Engine/Plugins/Experimental/Mutable/Source/CustomizableObject/Private/MuCO/CustomizableObjectInstance.cpp
+++ b/Engine/Plugins/Experimental/Mutable/Source/CustomizableObject/Private/MuCO/CustomizableObjectInstance.cpp
@@ -2716,7 +2716,7 @@ FTexturePlatformData* UCustomizableInstancePrivateData::MutableCreateImagePlatfo
 	//check(SizeX == MaxPossibleSize || SizeY == MaxPossibleSize || FullLODCount == 1);
 	if (!(SizeX == MaxPossibleSize || SizeY == MaxPossibleSize || FullLODCount == 1))
 	{
-		UE_LOG(LogMutable, Warning, TEXT("Building instance: unsuported texture size %d x %d."), SizeX, SizeY);
+		UE_LOG(LogMutable, Warning, TEXT("Building instance: unsupported texture size %d x %d."), SizeX, SizeY);
 		//return nullptr;
 	}
 
diff --git a/Engine/Plugins/Experimental/Mutable/Source/MutableRuntime/Private/MuR/CodeRunnerTasked.cpp b/Engine/Plugins/Experimental/Mutable/Source/MutableRuntime/Private/MuR/CodeRunnerTasked.cpp
--- a/Engine/Plugins/Experimental/Mutable/Source/MutableRuntime/Private/MuR/CodeRunnerTasked.cpp
+++ b/Engine/Plugins/Experimental/Mutable/Source/MutableRuntime/Private/MuR/CodeRunnerTasked.cpp
@@ -664,7 +664,7 @@ namespace mu
 		EImageFormat InitialFormat = m_base->GetFormat();
 		check(InitialFormat < EImageFormat::IF_COUNT);
 
-		if (Args.mask)
+		if (Args.mask && m_mask)
 		{
 			if (m_base->GetSize() != m_mask->GetSize())
 			{
@@ -698,7 +698,7 @@ namespace mu
 			pNew = mu::CloneOrTakeOver<mu::Image>(m_base.get());
 			//pNew = new Image(m_base->GetSizeX(), m_base->GetSizeY(), m_base->GetLODCount(), m_base->GetFormat());
 
-			if (Args.mask)
+			if (Args.mask && m_mask)
 			{			
 				// \todo: precalculated tables for softlight
 				switch (EBlendType(Args.blendType))
@@ -750,7 +750,7 @@ namespace mu
 				AlphaOffset = 1;
 			}
 
-			if (Args.mask)
+			if (Args.mask && m_mask)
 			{
 				// \todo: precalculated tables for softlight
 				switch (EBlendType(Args.blendTypeAlpha))
