[mutable] Avoid unnecessary clones when not converting the pixel format.
#rnx

diff --git a/Engine/Plugins/Experimental/Mutable/Source/MutableRuntime/Private/MuR/CodeRunnerTasked.cpp b/Engine/Plugins/Experimental/Mutable/Source/MutableRuntime/Private/MuR/CodeRunnerTasked.cpp
--- a/Engine/Plugins/Experimental/Mutable/Source/MutableRuntime/Private/MuR/CodeRunnerTasked.cpp
+++ b/Engine/Plugins/Experimental/Mutable/Source/MutableRuntime/Private/MuR/CodeRunnerTasked.cpp
@@ -860,17 +860,16 @@ namespace mu
 
 		if (targetFormat != EImageFormat::IF_NONE)
 		{
-			vec2<int> resultSize;
-			int resultLODCount = 0;
-
-			resultSize[0] = m_base->GetSizeX();
-			resultSize[1] = m_base->GetSizeY();
-			resultLODCount = m_base->GetLODCount();
-
-			Ptr<Image> TempResult = ImagePixelFormat(m_imageCompressionQuality, m_base.get(), targetFormat , -1);
-			TempResult->m_flags = m_base->m_flags;
-
-			Result = TempResult;
+			if (targetFormat == m_base->GetFormat())
+			{
+				Result = CloneOrTakeOver<>(m_base.get());
+			}
+			else
+			{
+				Ptr<Image> TempResult = ImagePixelFormat(m_imageCompressionQuality, m_base.get(), targetFormat, -1);
+				TempResult->m_flags = m_base->m_flags;
+				Result = TempResult;
+			}
 		}
 		else
 		{
diff --git a/Engine/Plugins/Experimental/Mutable/Source/MutableRuntime/Private/MuR/OpImagePixelFormat.cpp b/Engine/Plugins/Experimental/Mutable/Source/MutableRuntime/Private/MuR/OpImagePixelFormat.cpp
--- a/Engine/Plugins/Experimental/Mutable/Source/MutableRuntime/Private/MuR/OpImagePixelFormat.cpp
+++ b/Engine/Plugins/Experimental/Mutable/Source/MutableRuntime/Private/MuR/OpImagePixelFormat.cpp
@@ -24,6 +24,9 @@ ImagePtr ImagePixelFormat( int imageCompressionQuality, const Image* pBase,
 {
 	MUTABLE_CPUPROFILER_SCOPE(ImagePixelFormat);
 
+	TRACE_CPUPROFILER_EVENT_SCOPE_TEXT(*FString::Printf(TEXT("From %d to %d"), int32(pBase->GetFormat()), int32(targetFormat)));
+
+
 	FIntVector2 resultSize;
     int resultLODCount = 0;
 
