While Investigating a client hang, found some code that is either unsafe, not clearly safe, or violates UE coding standards.

Code cleanup:
* Added a lot of checks() to make sure we crash explicitly with useful information rather than hoping a crash will have enough info to work from.
** NOTE: I'm aware that _some_ of the "GetPrivate()" calls have a check internal to them, but they're not named to indicate that, so it may change, and other GetPrivate() calls do NOT have internal checks for nullptr.  So to reduce the need for specific implementation knowledge and the possible confusion between them, I just re-assert in those cases if the pointer is nullptr.
* Fixed a number of coding standard issues, including:
** Variable names must start with initial capitals
** Spacing, indentation, and bracing issues that didn't match UE coding standards.
* Cleaned up some other spacing and indentation issues to be consistent within the file and for general readability.

Also, removed a line of code that would cause a crash due to dereferencing a pointer that's guaranteed to be nullptr or invalid.

Added a few temporary comments to investigate some other oddities in the code.

#UE5 #RNX #NoReleaseNotes

diff --git a/Engine/Plugins/Experimental/Mutable/Source/CustomizableObject/Private/MuCO/CustomizableObjectSystem.cpp b/Engine/Plugins/Experimental/Mutable/Source/CustomizableObject/Private/MuCO/CustomizableObjectSystem.cpp
--- a/Engine/Plugins/Experimental/Mutable/Source/CustomizableObject/Private/MuCO/CustomizableObjectSystem.cpp
+++ b/Engine/Plugins/Experimental/Mutable/Source/CustomizableObject/Private/MuCO/CustomizableObjectSystem.cpp
@@ -170,6 +170,8 @@ void FMutableQueue::Enqueue(const FMutableQueueElem& TaskToEnqueue)
 {
 	for (FMutableQueueElem& QueueElem : Array)
 	{
+		check(QueueElem.Operation);
+		check(TaskToEnqueue.Operation);
 		if (QueueElem.Operation->CustomizableObjectInstance == TaskToEnqueue.Operation->CustomizableObjectInstance)
 		{
 			QueueElem.Operation = TaskToEnqueue.Operation;
@@ -194,9 +196,13 @@ void FMutableQueue::ChangePriorities()
 {
 	for (FMutableQueueElem& Elem : Array)
 	{
-		if (Elem.Operation->CustomizableObjectInstance.IsValid() && Elem.Operation->CustomizableObjectInstance->GetPrivate())
+		check(Elem.Operation);
+		if (Elem.Operation->CustomizableObjectInstance.IsValid())
 		{
-			Elem.Priority = Elem.Operation->CustomizableObjectInstance->GetPrivate()->MinSquareDistFromComponentToPlayer;
+			if (UCustomizableInstancePrivateData* CustomizableObjectInstancePrivateData = Elem.Operation->CustomizableObjectInstance->GetPrivate())
+			{
+				Elem.Priority = CustomizableObjectInstancePrivateData->MinSquareDistFromComponentToPlayer;
+			}
 		}
 	}
 }
@@ -206,8 +212,10 @@ void FMutableQueue::UpdatePriority(const UCustomizableObjectInstance* Instance)
 {
 	for (FMutableQueueElem& Elem : Array)
 	{
+		check(Elem.Operation);
 		if (Elem.Operation->CustomizableObjectInstance == Instance)
 		{
+			check(Instance->GetPrivate() != nullptr);
 			Elem.Priority = Instance->GetPrivate()->MinSquareDistFromComponentToPlayer;
 			break;
 		}
@@ -219,6 +227,7 @@ const FMutableQueueElem* FMutableQueue::Get(const UCustomizableObjectInstance* I
 {
 	for (const FMutableQueueElem& Elem : Array)
 	{
+		check(Elem.Operation);
 		if (Elem.Operation->CustomizableObjectInstance == Instance)
 		{
 			return &Elem;
@@ -244,6 +253,7 @@ UCustomizableObjectSystem* UCustomizableObjectSystem::GetInstance()
 		check(IsInGameThread());
 
 		FCustomizableObjectSystemPrivate::SSystem = NewObject<UCustomizableObjectSystem>(UCustomizableObjectSystem::StaticClass());
+		check(FCustomizableObjectSystemPrivate::SSystem != nullptr);
 		checkf(!GUObjectArray.IsDisregardForGC(FCustomizableObjectSystemPrivate::SSystem), TEXT("Mutable was initialized too early in the UE4 init process, for instance, in the constructor of a default UObject."));
 		FCustomizableObjectSystemPrivate::SSystem->AddToRoot();
 		checkf(!GUObjectArray.IsDisregardForGC(FCustomizableObjectSystemPrivate::SSystem), TEXT("Mutable was initialized too early in the UE4 init process, for instance, in the constructor of a default UObject."));
@@ -288,6 +298,8 @@ void UCustomizableObjectSystem::LogShowData(bool bFullInfo, bool ShowMaterialInf
 
 	ArrayData.Sort([](UCustomizableObjectInstance& A, UCustomizableObjectInstance& B)
 	{
+		check(A.GetPrivate() != nullptr);
+		check(B.GetPrivate() != nullptr);
 		return (A.GetPrivate()->LastMinSquareDistFromComponentToPlayer < B.GetPrivate()->LastMinSquareDistFromComponentToPlayer);
 	});
 
@@ -334,6 +346,7 @@ void UCustomizableObjectSystem::InitSystem()
 {
 	// Everything initialized in Init() instead of constructor to prevent the default UCustomizableObjectSystem from registering a tick function
 	Private = MakeShareable(new FCustomizableObjectSystemPrivate());
+	check(Private != nullptr);
 	Private->NewCompilerFunc = nullptr;
 
 	Private->bReplaceDiscardedWithReferenceMesh = false;
@@ -362,16 +375,20 @@ void UCustomizableObjectSystem::InitSystem()
 #endif
 
 	mu::SettingsPtr pSettings = new mu::Settings;
+	check(pSettings);
 	pSettings->SetProfile(false);
 	pSettings->SetStreamingCache(MUTABLE_STREAMING_CACHE);
 	Private->MutableSystem = new mu::System(pSettings);
+	check(Private->MutableSystem);
 
 	Private->LastStreamingMemorySize = MUTABLE_STREAMING_CACHE;
 	Private->Streamer = new FUnrealMutableModelBulkStreamer();
+	check(Private->Streamer != nullptr);
 	Private->MutableSystem->SetStreamingInterface(Private->Streamer);
 
 	// Set up the external image provider, for image parameters.
 	FUnrealMutableImageProvider* Provider = new FUnrealMutableImageProvider();
+	check(Provider != nullptr);
 	Private->ImageProvider = Provider;
 	Private->MutableSystem->SetImageParameterGenerator(Provider);
 
@@ -388,6 +405,7 @@ void UCustomizableObjectSystem::InitSystem()
 	}
 
 	DefaultInstanceLODManagement = NewObject<UCustomizableInstanceLODManagement>();
+	check(DefaultInstanceLODManagement != nullptr);
 	CurrentInstanceLODManagement = DefaultInstanceLODManagement;
 }
 
@@ -426,6 +444,7 @@ void UCustomizableObjectSystem::BeginDestroy()
 		Private->CurrentMutableOperation = nullptr;
 
 		// Deallocate streaming
+		check(Private->Streamer != nullptr);
 		Private->Streamer->EndStreaming();
 
 		Private->CurrentInstanceBeingUpdated = nullptr;
@@ -456,7 +475,8 @@ FCustomizableObjectCompilerBase* (*FCustomizableObjectSystemPrivate::NewCompiler
 
 FCustomizableObjectCompilerBase* UCustomizableObjectSystem::GetNewCompiler()
 {
-	if (Private->NewCompilerFunc!=nullptr)
+	check(Private != nullptr);
+	if (Private->NewCompilerFunc != nullptr)
 	{
 		return Private->NewCompilerFunc();
 	}
@@ -469,6 +489,7 @@ FCustomizableObjectCompilerBase* UCustomizableObjectSystem::GetNewCompiler()
 
 void UCustomizableObjectSystem::SetNewCompilerFunc(FCustomizableObjectCompilerBase* (*InNewCompilerFunc)())
 {
+	check(Private != nullptr);
 	Private->NewCompilerFunc = InNewCompilerFunc;
 }
 
@@ -578,22 +599,23 @@ void FCustomizableObjectSystemPrivate::UpdateStats()
 
 				if (bHasResourceForRendering)
 				{
-					for (const FParameterDecorations& parameterDecoration : CustomizableObjectInstance->GetPrivate()->ParameterDecorations)
+					for (const FParameterDecorations& ParameterDecoration : CustomizableObjectInstance->GetPrivate()->ParameterDecorations)
 					{
-						for (const UTexture2D* parameterDecorationImage : parameterDecoration.Images)
+						for (const UTexture2D* ParameterDecorationImage : ParameterDecoration.Images)
 						{
-							if (parameterDecorationImage && parameterDecorationImage->IsValidLowLevel())
+							if (ParameterDecorationImage && ParameterDecorationImage->IsValidLowLevel())
 							{
-								SizeDecoration += parameterDecorationImage->CalcTextureMemorySizeEnum(TMC_AllMips);
+								SizeDecoration += ParameterDecorationImage->CalcTextureMemorySizeEnum(TMC_AllMips);
 							}
 						}
 					}
 
-					for (const FGeneratedTexture& generatedTextures : CustomizableObjectInstance->GetPrivate()->GeneratedTextures)
+					for (const FGeneratedTexture& GeneratedTextures : CustomizableObjectInstance->GetPrivate()->GeneratedTextures)
 					{
-						if (generatedTextures.Texture)
+						if (GeneratedTextures.Texture)
 						{
-							SizeGenerated += generatedTextures.Texture->CalcTextureMemorySizeEnum(TMC_AllMips);
+							check(GeneratedTextures.Texture != nullptr);
+							SizeGenerated += GeneratedTextures.Texture->CalcTextureMemorySizeEnum(TMC_AllMips);
 						}
 					}
 				}
@@ -635,7 +657,7 @@ void FCustomizableObjectSystemPrivate::UpdateStats()
 					UCustomizableObjectInstance* Instance = CustomizableSkeletalComponent->CustomizableObjectInstance;
 					APawn* PlayerPawn = UGameplayStatics::GetPlayerPawn(CustomizableSkeletalComponent->GetWorld(), 0);
 					
-					if (ParentActor && ParentActor == PlayerPawn && Instance)
+					if (ParentActor && (ParentActor == PlayerPawn) && Instance)
 					{
 						bFoundPlayer = true;
 
@@ -650,6 +672,7 @@ void FCustomizableObjectSystemPrivate::UpdateStats()
 												
 						GEngine->AddOnScreenDebugMessage(MsgIndex++, .0f, FColor::Green, TEXT("Animation tags: ") + TagString);
 
+						check(Instance->GetPrivate() != nullptr);
 						FCustomizableInstanceComponentData* ComponentData = Instance->GetPrivate()->GetComponentData(CustomizableSkeletalComponent->ComponentIndex);
 
 						if (ComponentData)
@@ -722,6 +745,7 @@ void UpdateSkeletalMesh(UCustomizableObjectInstance* CustomizableObjectInstance)
 	// This used to be CustomizableObjectInstance::UpdateSkeletalMesh_PostBeginUpdate3
 	{
 		MUTABLE_CPUPROFILER_SCOPE(UpdateSkeletalMesh_PostBeginUpdate3);
+		check(CustomizableObjectInstance != nullptr);
 
 #if !WITH_EDITOR
 		for (int32 ComponentIndex = 0; ComponentIndex < CustomizableObjectInstance->SkeletalMeshes.Num(); ++ComponentIndex)
@@ -733,19 +757,23 @@ void UpdateSkeletalMesh(UCustomizableObjectInstance* CustomizableObjectInstance)
 		}
 #endif
 
+		UCustomizableInstancePrivateData* CustomizableObjectInstancePrivateData = CustomizableObjectInstance->GetPrivate();
+		check(CustomizableObjectInstancePrivateData != nullptr);
 		for (TObjectIterator<UCustomizableSkeletalComponent> It; It; ++It)
 		{
 			UCustomizableSkeletalComponent* CustomizableSkeletalComponent = *It;
 
-			if (CustomizableSkeletalComponent && CustomizableSkeletalComponent->CustomizableObjectInstance == CustomizableObjectInstance
-				&& CustomizableObjectInstance->SkeletalMeshes.IsValidIndex(CustomizableSkeletalComponent->ComponentIndex))
+			if (CustomizableSkeletalComponent &&
+				(CustomizableSkeletalComponent->CustomizableObjectInstance == CustomizableObjectInstance) &&
+				CustomizableObjectInstance->SkeletalMeshes.IsValidIndex(CustomizableSkeletalComponent->ComponentIndex)
+			   )
 			{
 				MUTABLE_CPUPROFILER_SCOPE(UpdateSkeletalMesh_SetSkeletalMesh);
 
-				const bool bIsCreatingSkeletalMesh = CustomizableObjectInstance->GetPrivate()->HasCOInstanceFlags(CreatingSkeletalMesh);
+				const bool bIsCreatingSkeletalMesh = CustomizableObjectInstancePrivateData->HasCOInstanceFlags(CreatingSkeletalMesh);
 				CustomizableSkeletalComponent->SetSkeletalMesh(CustomizableObjectInstance->SkeletalMeshes[CustomizableSkeletalComponent->ComponentIndex], false, bIsCreatingSkeletalMesh);
 
-				if (CustomizableObjectInstance->GetPrivate()->HasCOInstanceFlags(ReplacePhysicsAssets))
+				if (CustomizableObjectInstancePrivateData->HasCOInstanceFlags(ReplacePhysicsAssets))
 				{
 					CustomizableSkeletalComponent->SetPhysicsAsset(
 						CustomizableObjectInstance->SkeletalMeshes[CustomizableSkeletalComponent->ComponentIndex] ? 
@@ -754,8 +782,8 @@ void UpdateSkeletalMesh(UCustomizableObjectInstance* CustomizableObjectInstance)
 			}
 		}
 
-		CustomizableObjectInstance->GetPrivate()->SetCOInstanceFlags(Generated);
-		CustomizableObjectInstance->GetPrivate()->ClearCOInstanceFlags(CreatingSkeletalMesh);
+		CustomizableObjectInstancePrivateData->SetCOInstanceFlags(Generated);
+		CustomizableObjectInstancePrivateData->ClearCOInstanceFlags(CreatingSkeletalMesh);
 
 		CustomizableObjectInstance->bEditorPropertyChanged = false;
 
@@ -778,17 +806,19 @@ void UpdateSkeletalMesh(UCustomizableObjectInstance* CustomizableObjectInstance)
 #endif //WITH_EDITOR
 		}
 
+		check(UCustomizableObjectSystem::GetInstance() != nullptr);
 		FCustomizableObjectSystemPrivate* CustomizableObjectSystem = UCustomizableObjectSystem::GetInstance()->GetPrivate();
+		check(CustomizableObjectSystem != nullptr);
 		if (CustomizableObjectSystem->bReleaseTexturesImmediately)
 		{
 			FMutableResourceCache& Cache = CustomizableObjectSystem->GetObjectCache(CustomizableObjectInstance->GetCustomizableObject());
 
-			for (TPair<uint32, FGeneratedTexture>& Item : CustomizableObjectInstance->GetPrivate()->TexturesToRelease)
+			for (TPair<uint32, FGeneratedTexture>& Item : CustomizableObjectInstancePrivateData->TexturesToRelease)
 			{
 				UCustomizableInstancePrivateData::ReleaseMutableTexture(Item.Value.Id, Item.Value.Texture, Cache);
 			}
 
-			CustomizableObjectInstance->GetPrivate()->TexturesToRelease.Empty();
+			CustomizableObjectInstancePrivateData->TexturesToRelease.Empty();
 		}
 	}
 }
@@ -861,6 +891,7 @@ void FCustomizableObjectSystemPrivate::InitDiscardResourcesSkeletalMesh(UCustomi
 	if (InCustomizableObjectInstance && InCustomizableObjectInstance->IsValidLowLevel())
 	{
 		const TSharedRef<FMutableOperation> Operation = MakeShared<FMutableOperation>(FMutableOperation::CreateInstanceDiscard(InCustomizableObjectInstance));
+		check(InCustomizableObjectInstance->GetPrivate() != nullptr);
 		MutableOperationQueue.Enqueue(FMutableQueueElem::Create(Operation, FMutableQueueElem::EQueuePriorityType::High, InCustomizableObjectInstance->GetPrivate()->MinSquareDistFromComponentToPlayer));
 	}
 }
@@ -891,6 +922,7 @@ void UCustomizableObjectSystem::ClearResourceCacheProtected()
 	check(IsInGameThread());
 
 	ProtectedCachedTextures.Reset(0);
+	check(GetPrivate() != nullptr);
 	GetPrivate()->ProtectedObjectCachedImages.Reset(0);
 }
 
@@ -898,7 +930,8 @@ void UCustomizableObjectSystem::ClearResourceCacheProtected()
 #if WITH_EDITOR
 bool UCustomizableObjectSystem::LockObject(const class UCustomizableObject* InObject)
 {
-	check(InObject);
+	check(InObject != nullptr);
+	check(InObject->GetPrivate() != nullptr);
 	check(!InObject->GetPrivate()->bLocked);
 	check(IsInGameThread() && !IsInParallelGameThread());
 
@@ -941,6 +974,7 @@ bool UCustomizableObjectSystem::LockObject(const class UCustomizableObject* InOb
 	}
 
 	// Ensure that we don't try to handle any further streaming operations for this object
+	check(GetPrivate() != nullptr);
 	if (GetPrivate()->Streamer)
 	{
 		GetPrivate()->Streamer->CancelStreamingForObject(InObject);
@@ -958,6 +992,8 @@ bool UCustomizableObjectSystem::LockObject(const class UCustomizableObject* InOb
 
 void UCustomizableObjectSystem::UnlockObject(const class UCustomizableObject* Obj)
 {
+	check(Obj != nullptr);
+	check(Obj->GetPrivate() != nullptr);
 	check(Obj->GetPrivate()->bLocked);
 	check(IsInGameThread() && !IsInParallelGameThread());
 	Obj->GetPrivate()->bLocked = false;
@@ -1027,48 +1063,49 @@ void UCustomizableObjectSystem::PurgePendingReleaseSkeletalMesh()
 	double	CurTime = FPlatformTime::Seconds();
 	static double TimeToDelete = 1.0;
 
-	for( int32 i = PendingReleaseSkeletalMesh.Num() - 1; i >= 0; --i )
+	for (int32 InfoIndex = PendingReleaseSkeletalMesh.Num() - 1; InfoIndex >= 0; --InfoIndex)
 	{
-		FPendingReleaseSkeletalMeshInfo& info = PendingReleaseSkeletalMesh[ i ];
+		FPendingReleaseSkeletalMeshInfo& Info = PendingReleaseSkeletalMesh[InfoIndex];
 		
-		if (info.SkeletalMesh != nullptr)
+		if (Info.SkeletalMesh != nullptr)
 		{
-			if (( CurTime - info.TimeStamp ) >= TimeToDelete)
+			if ((CurTime - Info.TimeStamp) >= TimeToDelete)
 			{
-				if (info.SkeletalMesh->GetSkeleton())
+				if (Info.SkeletalMesh->GetSkeleton())
 				{
-					info.SkeletalMesh->GetSkeleton()->ClearCacheData();
+					Info.SkeletalMesh->GetSkeleton()->ClearCacheData();
 				}
-				info.SkeletalMesh->GetRefSkeleton().Empty();
-				info.SkeletalMesh->GetMaterials().Empty();
-				info.SkeletalMesh->GetRefBasesInvMatrix().Empty();
-				info.SkeletalMesh->ReleaseResources();
-				info.SkeletalMesh->ReleaseResourcesFence.Wait();
+				Info.SkeletalMesh->GetRefSkeleton().Empty();
+				Info.SkeletalMesh->GetMaterials().Empty();
+				Info.SkeletalMesh->GetRefBasesInvMatrix().Empty();
+				Info.SkeletalMesh->ReleaseResources();
+				Info.SkeletalMesh->ReleaseResourcesFence.Wait();
 
 				// TODO: This was done for version 4.18
 				//info.SkeletalMesh->GetImportedResource()->LODModels.Empty();
 				//info.SkeletalMesh->LODInfo.Empty();
 
-				PendingReleaseSkeletalMesh.RemoveAt( i );
+				PendingReleaseSkeletalMesh.RemoveAt(InfoIndex);
 			}
 		}
 	}
 }
 
-void UCustomizableObjectSystem::AddPendingReleaseSkeletalMesh( USkeletalMesh* SkeletalMesh )
+void UCustomizableObjectSystem::AddPendingReleaseSkeletalMesh(USkeletalMesh* SkeletalMesh)
 {
-	check( SkeletalMesh != nullptr );
+	check(SkeletalMesh != nullptr);
 
-	FPendingReleaseSkeletalMeshInfo	info;
-	info.SkeletalMesh = SkeletalMesh;
-	info.TimeStamp = FPlatformTime::Seconds();
+	FPendingReleaseSkeletalMeshInfo	Info;
+	Info.SkeletalMesh = SkeletalMesh;
+	Info.TimeStamp = FPlatformTime::Seconds();
 
-	PendingReleaseSkeletalMesh.Add( info );
+	PendingReleaseSkeletalMesh.Add(Info);
 }
 
 
 void UCustomizableObjectSystem::ClearCurrentMutableOperation()
 {
+	check(Private != nullptr);
 	Private->CurrentInstanceBeingUpdated = nullptr;
 	Private->CurrentMutableOperation = nullptr;
 	ClearResourceCacheProtected();
@@ -1085,6 +1122,7 @@ void FCustomizableObjectSystemPrivate::UpdateStreamingLimit()
 	if (MemoryBytes != LastStreamingMemorySize)
 	{
 		LastStreamingMemorySize = MemoryBytes;
+		check(MutableSystem);
 		MutableSystem->SetStreamingCache(MemoryBytes);
 	}
 }
@@ -1102,31 +1140,34 @@ namespace impl
 	{
 		MUTABLE_CPUPROFILER_SCOPE(UCustomizableInstancePrivateData::UpdateParameterDecorations)
 
+		check(OperationData);
 		OperationData->ParametersUpdateData.Clear();
 		OperationData->RelevantParametersInProgress.Empty();
 
 		check(MutableParameters);
 
 		// This must run in the mutable thread.
-		auto MutableSystem = UCustomizableObjectSystem::GetInstance()->GetPrivate()->MutableSystem;
+		check(UCustomizableObjectSystem::GetInstance() != nullptr);
+		check(UCustomizableObjectSystem::GetInstance()->GetPrivate() != nullptr);
+		mu::SystemPtr MutableSystem = UCustomizableObjectSystem::GetInstance()->GetPrivate()->MutableSystem;
 
 		// Decorations
-		int32 NumParameter = MutableParameters->GetCount();
+		int32 NumParameters = MutableParameters->GetCount();
 
 		{
 			MUTABLE_CPUPROFILER_SCOPE(ParameterDecorations)
 
 				// Generate all the images that decorate the parameters UI, only in the editor.
-				for (int32 p = 0; p < NumParameter; ++p)
+				for (int32 ParamIndex = 0; ParamIndex < NumParameters; ++ParamIndex)
 				{
 					FParameterDecorationsUpdateData::FParameterDesc ParamData;
 
-					for (int32 i = 0; i < MutableParameters->GetAdditionalImageCount(p); ++i)
+					for (int32 ImageIndex = 0; ImageIndex < MutableParameters->GetAdditionalImageCount(ParamIndex); ++ImageIndex)
 					{
 						mu::ImagePtrConst Image = MutableSystem->BuildParameterAdditionalImage(
 							MutableModel,
 							MutableParameters,
-							p, i);
+							ParamIndex, ImageIndex);
 
 						ParamData.Images.Add(Image);
 					};
@@ -1139,19 +1180,18 @@ namespace impl
 		{
 			MUTABLE_CPUPROFILER_SCOPE(ParameterRelevancy)
 
-			TArray<bool> relevant;
-			relevant.SetNumZeroed(NumParameter);
-			MutableSystem->GetParameterRelevancy(MutableModel, MutableParameters, relevant.GetData());
+			TArray<bool> Relevant;
+			Relevant.SetNumZeroed(NumParameters);
+			MutableSystem->GetParameterRelevancy(MutableModel, MutableParameters, Relevant.GetData());
 
-			for (int32 p = 0; p < NumParameter; ++p)
+			for (int32 ParamIndex = 0; ParamIndex < NumParameters; ++ParamIndex)
 			{
-				if (relevant[p])
+				if (Relevant[ParamIndex])
 				{
-					OperationData->RelevantParametersInProgress.Add(p);
+					OperationData->RelevantParametersInProgress.Add(ParamIndex);
 				}
 			}
 		}
-
 	}
 
 
@@ -1161,10 +1201,13 @@ namespace impl
 
 		// This runs in the mutable thread.
 		check(MutableParameters);
-
+		check(OperationData);
 		OperationData->InstanceUpdateData.Clear();
 
+		check(UCustomizableObjectSystem::GetInstance() != nullptr);
+		check(UCustomizableObjectSystem::GetInstance()->GetPrivate() != nullptr);
 		mu::System* System = UCustomizableObjectSystem::GetInstance()->GetPrivate()->MutableSystem.get();
+		check(System != nullptr);
 
 		// For now, we are forcing the recreation of mutable-side instances with every update.
 		OperationData->InstanceID = System->NewInstance(Model);
@@ -1316,13 +1359,18 @@ namespace impl
 
 		// This runs in the mutable thread.
 
+		check(OperationData);
 		check(MutableParameters);
 
-		mu::System* System = UCustomizableObjectSystem::GetInstance()->GetPrivate()->MutableSystem.get();
+		check(UCustomizableObjectSystem::GetInstance() != nullptr);
+		FCustomizableObjectSystemPrivate* CustomizableObjectSystemPrivateData = UCustomizableObjectSystem::GetInstance()->GetPrivate();
+		check(CustomizableObjectSystemPrivateData != nullptr);
+		mu::System* System = CustomizableObjectSystemPrivateData->MutableSystem.get();
+		check(System != nullptr);
 
 		// Generate all the required resources, that are not cached
 		TArray<mu::RESOURCE_ID> ImagesInThisInstance;
-		for ( FInstanceUpdateData::FImage& Image: OperationData->InstanceUpdateData.Images )
+		for (FInstanceUpdateData::FImage& Image : OperationData->InstanceUpdateData.Images)
 		{
 			MUTABLE_CPUPROFILER_SCOPE(GetImage);
 
@@ -1342,7 +1390,7 @@ namespace impl
 			// See if it is cached from another instance
 			if (!bCached)
 			{
-				bCached = UCustomizableObjectSystem::GetInstance()->GetPrivate()->ProtectedObjectCachedImages.Contains(Image.ImageID);
+				bCached = CustomizableObjectSystemPrivateData->ProtectedObjectCachedImages.Contains(Image.ImageID);
 			}
 
 			if (bCached)
@@ -1362,16 +1410,18 @@ namespace impl
 				Image.Image = System->GetImage(OperationData->InstanceID, Image.ImageID, MipsToSkip);
 
 				// We need one mip or the complete chain. Otherwise there was a bug.
+				check(Image.Image);
 				int32 FullMipCount = Image.Image->GetMipmapCount(Image.Image->GetSizeX(), Image.Image->GetSizeY());
 				int32 RealMipCount = Image.Image->GetLODCount();
-				if ( RealMipCount!=1 && RealMipCount!=FullMipCount )
+				if ((RealMipCount != 1) && (RealMipCount != FullMipCount))
 				{
 					MUTABLE_CPUPROFILER_SCOPE(GetImage_MipFix);
 
 					UE_LOG(LogMutable, Warning, TEXT("Mutable generated an incomplete mip chain for image %d."), Image.ImageID);
 
 					// Force the right number of mips. The missing data will be black.
-					mu::ImagePtr NewImage = new mu::Image( Image.Image->GetSizeX(), Image.Image->GetSizeY(), FullMipCount, Image.Image->GetFormat() );
+					mu::ImagePtr NewImage = new mu::Image(Image.Image->GetSizeX(), Image.Image->GetSizeY(), FullMipCount, Image.Image->GetFormat());
+					check(NewImage);
 					if (NewImage->GetDataSize() >= Image.Image->GetDataSize())
 					{
 						FMemory::Memcpy(NewImage->GetData(), Image.Image->GetData(), Image.Image->GetDataSize());
@@ -1390,6 +1440,7 @@ namespace impl
 	{
 		// This runs in a worker thread
 
+		check(OperationData);
 		for (const FInstanceUpdateData::FSurface& Surface : OperationData->InstanceUpdateData.Surfaces)
 		{
 			for (int32 ImageIndex = 0; ImageIndex<Surface.ImageCount; ++ImageIndex)
@@ -1417,6 +1468,7 @@ namespace impl
 
 		MUTABLE_CPUPROFILER_SCOPE(PrepareSkeletonData);
 
+		check(OperationData);
 		const int32 LODCount = OperationData->InstanceUpdateData.LODs.Num();
 		const FInstanceUpdateData::FLOD& MinLOD = OperationData->InstanceUpdateData.LODs[OperationData->CurrentMinLOD];
 		const int32 ComponentCount = MinLOD.ComponentCount;
@@ -1475,6 +1527,7 @@ namespace impl
 				}
 
 				const mu::Skeleton* Skeleton = Mesh->GetSkeleton().get();
+				// We don't need to check(Skeleton != nullptr) because we would've continued if (!Mesh->GetSkeleton()) above.
 				const int32 BoneCount = Skeleton->GetBoneCount();
 
 				{
@@ -1504,7 +1557,7 @@ namespace impl
 					UsedBones.Empty(BoneCount);
 					UsedBones.AddDefaulted(BoneCount);
 
-					for (int32 Surface = 0; Surface < SurfaceCount; Surface++)
+					for (int32 Surface = 0; Surface < SurfaceCount; ++Surface)
 					{
 						int32 IndexStart, IndexCount, VertexStart, VertexCount;
 						Mesh->GetSurface(Surface, &VertexStart, &VertexCount, &IndexStart, &IndexCount);
@@ -1518,11 +1571,11 @@ namespace impl
 
 						if (BoneIndexFormat == mu::MBF_UINT8)
 						{
-							for (int32 v = 0; v < VertexCount; ++v)
+							for (int32 VertexIndex = 0; VertexIndex < VertexCount; ++VertexIndex)
 							{
-								for (int32 i = 0; i < BoneIndexComponents; ++i)
+								for (int32 Index = 0; Index < BoneIndexComponents; ++Index)
 								{
-									size_t SectionBoneIndex = VertexBoneIndexPtr[i];
+									size_t SectionBoneIndex = VertexBoneIndexPtr[Index];
 									UsedBones[SectionBoneIndex] = true;
 								}
 								VertexBoneIndexPtr += ElemSize;
@@ -1530,11 +1583,11 @@ namespace impl
 						}
 						else if (BoneIndexFormat == mu::MBF_UINT16)
 						{
-							for (int32 v = 0; v < VertexCount; ++v)
+							for (int32 VertexIndex = 0; VertexIndex < VertexCount; ++VertexIndex)
 							{
-								for (int32 i = 0; i < BoneIndexComponents; ++i)
+								for (int32 Index = 0; Index < BoneIndexComponents; ++Index)
 								{
-									size_t SectionBoneIndex = ((uint16*)VertexBoneIndexPtr)[i];
+									size_t SectionBoneIndex = ((uint16*)VertexBoneIndexPtr)[Index];
 									UsedBones[SectionBoneIndex] = true;
 								}
 								VertexBoneIndexPtr += ElemSize;
@@ -1542,11 +1595,11 @@ namespace impl
 						}
 						else if (BoneIndexFormat == mu::MBF_UINT32)
 						{
-							for (int32 v = 0; v < VertexCount; ++v)
+							for (int32 VertexIndex = 0; VertexIndex < VertexCount; ++VertexIndex)
 							{
-								for (int32 i = 0; i < BoneIndexComponents; ++i)
+								for (int32 Index = 0; Index < BoneIndexComponents; ++Index)
 								{
-									size_t SectionBoneIndex = ((uint32_t*)VertexBoneIndexPtr)[i];
+									size_t SectionBoneIndex = ((uint32_t*)VertexBoneIndexPtr)[Index];
 									UsedBones[SectionBoneIndex] = true;
 								}
 								VertexBoneIndexPtr += ElemSize;
@@ -1763,6 +1816,7 @@ namespace impl
 		if (OperationData->InstanceID > 0)
 		{
 			MUTABLE_CPUPROFILER_SCOPE(EndUpdate);
+			check(MutableSystem);
 			MutableSystem->EndUpdate(OperationData->InstanceID);
 			OperationData->InstanceUpdateData.Clear();
 
@@ -1777,6 +1831,7 @@ namespace impl
 	{
 		MUTABLE_CPUPROFILER_SCOPE(Task_Game_ReleasePlatformData)
 
+		check(OperationData.IsValid());
 		TMap<uint32, FTexturePlatformData*>& ImageToPlatformDataMap = OperationData->ImageToPlatformDataMap;
 		for (TPair<uint32, FTexturePlatformData*> Pair : ImageToPlatformDataMap)
 		{
@@ -1813,20 +1868,23 @@ namespace impl
 		// TODO: T2927
 		if (LogBenchmarkUtil::isLoggingActive())
 		{
-			// Log the time stat		
-			double deltaSeconds = FPlatformTime::Seconds() - System->GetPrivate()->CurrentMutableOperation->StartUpdateTime;
+			// Log the time stat
+			FCustomizableObjectSystemPrivate* CustomizableObjectSystemPrivateData = System->GetPrivate();
+			check(CustomizableObjectSystemPrivateData != nullptr);
+
+			double deltaSeconds = FPlatformTime::Seconds() - CustomizableObjectSystemPrivateData->CurrentMutableOperation->StartUpdateTime;
 			int32 deltaMs = int32(deltaSeconds * 1000);
-			int64 streamingCache = System->GetPrivate()->LastStreamingMemorySize / 1024;
+			int64 streamingCache = CustomizableObjectSystemPrivateData->LastStreamingMemorySize / 1024;
 
 			LogBenchmarkUtil::updateStat("customizable_instance_build_time", deltaMs);
 			LogBenchmarkUtil::updateStat("mutable_streaming_cache_memory", (long double)streamingCache / 1024.0);
-			System->GetPrivate()->TotalBuildMs += deltaMs;
-			System->GetPrivate()->TotalBuiltInstances++;
+			CustomizableObjectSystemPrivateData->TotalBuildMs += deltaMs;
+			CustomizableObjectSystemPrivateData->TotalBuiltInstances++;
 			SET_DWORD_STAT(STAT_MutableInstanceBuildTime, deltaMs);
-			SET_DWORD_STAT(STAT_MutableInstanceBuildTimeAvrg, System->GetPrivate()->TotalBuildMs / System->GetPrivate()->TotalBuiltInstances);
+			SET_DWORD_STAT(STAT_MutableInstanceBuildTimeAvrg, CustomizableObjectSystemPrivateData->TotalBuildMs / CustomizableObjectSystemPrivateData->TotalBuiltInstances);
 			SET_DWORD_STAT(STAT_MutableStreamingCache, streamingCache);
 
-			mu::System* MutableSystem = System->GetPrivate()->MutableSystem.get();
+			mu::System* MutableSystem = CustomizableObjectSystemPrivateData->MutableSystem.get();
 			if (MutableSystem)
 			{
 				SET_DWORD_STAT(STAT_MutableProfile_LiveInstanceCount, MutableSystem->GetProfileMetric(mu::System::ProfileMetric::LiveInstanceCount));
@@ -1858,6 +1916,8 @@ namespace impl
 			return;
 		}
 
+		check(OperationData.IsValid());
+
 		UCustomizableObjectInstance* CustomizableObjectInstance = CustomizableObjectInstancePtr.Get();
 
 		// TODO: Review checks.
@@ -1870,79 +1930,85 @@ namespace impl
 		// Actual work
 
 
-		// Process the pending texture coverage queries
 		if (!bCancel)
 		{
-			MUTABLE_CPUPROFILER_SCOPE(GameTextureQueries);
-			for (const FPendingTextureCoverageQuery& Query : OperationData->PendingTextureCoverageQueries)
+			UCustomizableInstancePrivateData* CustomizableInstancePrivateData = CustomizableObjectInstance->GetPrivate();
+			check(CustomizableInstancePrivateData != nullptr);
+
+			// Process the pending texture coverage queries
 			{
-				UMaterialInterface* Material = nullptr;
-				const uint32* InstanceIndex = CustomizableObjectInstance->GetPrivate()->ObjectToInstanceIndexMap.Find(Query.MaterialIndex);
-				if (InstanceIndex && CustomizableObjectInstance->GetPrivate()->ReferencedMaterials.IsValidIndex(*InstanceIndex))
+				MUTABLE_CPUPROFILER_SCOPE(GameTextureQueries);
+				for (const FPendingTextureCoverageQuery& Query : OperationData->PendingTextureCoverageQueries)
 				{
-					Material = CustomizableObjectInstance->GetPrivate()->ReferencedMaterials[*InstanceIndex];
-				}
+					UMaterialInterface* Material = nullptr;
+					const uint32* InstanceIndex = CustomizableInstancePrivateData->ObjectToInstanceIndexMap.Find(Query.MaterialIndex);
+					if (InstanceIndex && CustomizableInstancePrivateData->ReferencedMaterials.IsValidIndex(*InstanceIndex))
+					{
+						Material = CustomizableInstancePrivateData->ReferencedMaterials[*InstanceIndex];
+					}
 
-				UCustomizableInstancePrivateData::ProcessTextureCoverageQueries(OperationData, CustomizableObjectInstance->GetCustomizableObject(), Query.KeyName, Query.PlatformData, Material);
+					UCustomizableInstancePrivateData::ProcessTextureCoverageQueries(OperationData, CustomizableObjectInstance->GetCustomizableObject(), Query.KeyName, Query.PlatformData, Material);
+				}
+				OperationData->PendingTextureCoverageQueries.Empty();
 			}
-			OperationData->PendingTextureCoverageQueries.Empty();
-		}
 
-		// Process texture coverage queries because it's safe to do now that the Mutable thread is stopped
-		if (!bCancel)
-		{
-			UCustomizableInstancePrivateData* PrivateInstance = CustomizableObjectInstance->GetPrivate();
-			if (OperationData->TextureCoverageQueries_MutableThreadResults.Num() > 0)
+			// Process texture coverage queries because it's safe to do now that the Mutable thread is stopped
 			{
-				for (auto& Result : OperationData->TextureCoverageQueries_MutableThreadResults)
+				if (OperationData->TextureCoverageQueries_MutableThreadResults.Num() > 0)
 				{
-					FTextureCoverageQueryData* FinalResultData = PrivateInstance->TextureCoverageQueries.Find(Result.Key);
-					*FinalResultData = Result.Value;
-				}
+					for (auto& Result : OperationData->TextureCoverageQueries_MutableThreadResults)
+					{
+						FTextureCoverageQueryData* FinalResultData = CustomizableInstancePrivateData->TextureCoverageQueries.Find(Result.Key);
+						*FinalResultData = Result.Value;
+					}
 
-				OperationData->TextureCoverageQueries_MutableThreadResults.Empty();
-			}
+					OperationData->TextureCoverageQueries_MutableThreadResults.Empty();
+				}
 
 #if WITH_EDITOR
-			CustomizableObjectInstance->LastUpdateMutableRuntimeCycles = OperationData->MutableRuntimeCycles;
+				CustomizableObjectInstance->LastUpdateMutableRuntimeCycles = OperationData->MutableRuntimeCycles;
 #endif
 
-			// Convert Step
-			//-------------------------------------------------------------
+				// Convert Step
+				//-------------------------------------------------------------
 
-			// \TODO: Bring that code here instead of keeping it in the UCustomizableObjectInstance
-			if (CustomizableObjectInstance->GetPrivate()->UpdateSkeletalMesh_PostBeginUpdate0(CustomizableObjectInstance, OperationData))
-			{
-				// This used to be CustomizableObjectInstance::UpdateSkeletalMesh_PostBeginUpdate1
+				// \TODO: Bring that code here instead of keeping it in the UCustomizableObjectInstance
+				if (CustomizableInstancePrivateData->UpdateSkeletalMesh_PostBeginUpdate0(CustomizableObjectInstance, OperationData))
 				{
-					MUTABLE_CPUPROFILER_SCOPE(UpdateSkeletalMesh_PostBeginUpdate1);
-
-					// \TODO: Bring here
-					CustomizableObjectInstance->GetPrivate()->BuildMaterials(OperationData, CustomizableObjectInstance);
-				}
+					// This used to be CustomizableObjectInstance::UpdateSkeletalMesh_PostBeginUpdate1
+					{
+						MUTABLE_CPUPROFILER_SCOPE(UpdateSkeletalMesh_PostBeginUpdate1);
 
-				// This used to be CustomizableObjectInstance::UpdateSkeletalMesh_PostBeginUpdate2
-				{
-					MUTABLE_CPUPROFILER_SCOPE(UpdateSkeletalMesh_PostBeginUpdate2);
+						// \TODO: Bring here
+						CustomizableInstancePrivateData->BuildMaterials(OperationData, CustomizableObjectInstance);
+					}
 
-					for (int32 Component = 0; Component < CustomizableObjectInstance->SkeletalMeshes.Num(); ++Component)
+					// This used to be CustomizableObjectInstance::UpdateSkeletalMesh_PostBeginUpdate2
 					{
-						if (CustomizableObjectInstance->SkeletalMeshes[Component] && CustomizableObjectInstance->SkeletalMeshes[Component]->GetLODInfoArray().Num())
+						MUTABLE_CPUPROFILER_SCOPE(UpdateSkeletalMesh_PostBeginUpdate2);
+
+						for (int32 Component = 0; Component < CustomizableObjectInstance->SkeletalMeshes.Num(); ++Component)
 						{
-							MUTABLE_CPUPROFILER_SCOPE(UpdateSkeletalMesh_PostEditChangeProperty);
+							if (CustomizableObjectInstance->SkeletalMeshes[Component] && CustomizableObjectInstance->SkeletalMeshes[Component]->GetLODInfoArray().Num())
+							{
+								MUTABLE_CPUPROFILER_SCOPE(UpdateSkeletalMesh_PostEditChangeProperty);
 
-							CustomizableObjectInstance->GetPrivate()->PostEditChangePropertyWithoutEditor(CustomizableObjectInstance->SkeletalMeshes[Component]);
+								CustomizableInstancePrivateData->PostEditChangePropertyWithoutEditor(CustomizableObjectInstance->SkeletalMeshes[Component]);
+							}
 						}
 					}
 				}
-			}
-		}
+			} // END - Process texture coverage queries
+		} // if (!bCancel)
+
+		FCustomizableObjectSystemPrivate* CustomizableObjectSystemPrivateData = System->GetPrivate();
+		check(CustomizableObjectSystemPrivateData != nullptr);
 
 		// Next Task: Release Mutable. We need this regardless if we cancel or not
 		//-------------------------------------------------------------		
 		{
-			mu::SystemPtr MutableSystem = System->GetPrivate()->MutableSystem;
-			System->GetPrivate()->AddMutableThreadTask(
+			mu::SystemPtr MutableSystem = CustomizableObjectSystemPrivateData->MutableSystem;
+			CustomizableObjectSystemPrivateData->AddMutableThreadTask(
 				TEXT("Task_Mutable_ReleaseInstance"),
 				[OperationData, MutableSystem]() {Task_Mutable_ReleaseInstance(OperationData, MutableSystem); },
 				UE::Tasks::ETaskPriority::BackgroundNormal);
@@ -1954,8 +2020,9 @@ namespace impl
 		if (!bCancel)
 		{
 			TSharedPtr<FMutableReleasePlatformOperationData> ReleaseOperationData = MakeShared<FMutableReleasePlatformOperationData>();
+			check(ReleaseOperationData);
 			ReleaseOperationData->ImageToPlatformDataMap = MoveTemp(OperationData->ImageToPlatformDataMap);
-			System->GetPrivate()->AddAnyThreadTask(
+			CustomizableObjectSystemPrivateData->AddAnyThreadTask(
 				TEXT("Mutable_ReleasePlatformData"),
 				[ReleaseOperationData]() { Task_Game_ReleasePlatformData(ReleaseOperationData); },
 				UE::Tasks::ETaskPriority::BackgroundNormal
@@ -1973,7 +2040,7 @@ namespace impl
 
 			// Next Task: Callbacks
 			//-------------------------------------------------------------
-			System->GetPrivate()->AddGameThreadTask(
+			CustomizableObjectSystemPrivateData->AddGameThreadTask(
 				{
 				FMutableTaskDelegate::CreateLambda(
 					[OperationData,CustomizableObjectInstancePtr]()
@@ -2011,32 +2078,40 @@ namespace impl
 			System->ClearCurrentMutableOperation();
 			return;
 		}
+
+		check(OperationData.IsValid());
+
+		UCustomizableInstancePrivateData* ObjectInstancePrivateData = ObjectInstance->GetPrivate();
+		check(ObjectInstancePrivateData != nullptr);
 		
 		// Process the parameter decorations if requested
 		if (bBuildParameterDecorations)
 		{
-			ObjectInstance->GetPrivate()->UpdateParameterDecorationsEngineResources(OperationData);
+			ObjectInstancePrivateData->UpdateParameterDecorationsEngineResources(OperationData);
 		}
 
 
-		// Selectivelly lock the resource cache for the object used by this instance to avoid the destruction of resources that we may want to reuse.
+		// Selectively lock the resource cache for the object used by this instance to avoid the destruction of resources that we may want to reuse.
 		// When protecting textures there mustn't be any left from a previous update
 		check(System->ProtectedCachedTextures.Num() == 0);
 
-		FMutableResourceCache& Cache = System->GetPrivate()->GetObjectCache(CustomizableObject);
+		FCustomizableObjectSystemPrivate* SystemPrivateData = System->GetPrivate();
+		check(SystemPrivateData != nullptr);
+
+		FMutableResourceCache& Cache = SystemPrivateData->GetObjectCache(CustomizableObject);
 
 		System->ProtectedCachedTextures.Reset(Cache.Images.Num());
-		System->GetPrivate()->ProtectedObjectCachedImages.Reset(Cache.Images.Num());
+		SystemPrivateData->ProtectedObjectCachedImages.Reset(Cache.Images.Num());
 
 		for (const FInstanceUpdateData::FImage& Image : OperationData->InstanceUpdateData.Images)
 		{
 			FMutableImageCacheKey Key(Image.ImageID, OperationData->MipsToSkip);
 			TWeakObjectPtr<UTexture2D>* TexturePtr = Cache.Images.Find(Key);
 
-			if (TexturePtr && TexturePtr->Get() && System->GetPrivate()->TextureHasReferences(Image.ImageID))
+			if (TexturePtr && TexturePtr->Get() && SystemPrivateData->TextureHasReferences(Image.ImageID))
 			{
 				System->ProtectedCachedTextures.Add(TexturePtr->Get());
-				System->GetPrivate()->ProtectedObjectCachedImages.Add(Image.ImageID);
+				SystemPrivateData->ProtectedObjectCachedImages.Add(Image.ImageID);
 			}
 		}
 
@@ -2044,17 +2119,18 @@ namespace impl
 		// which will safely access the GlobalExternalImages map, and then just get the cached image or issue a disk read
 
 		// Copy data generated in the mutable thread over to the instance
-		ObjectInstance->GetPrivate()->PrepareForUpdate(OperationData);
+		ObjectInstancePrivateData->PrepareForUpdate(OperationData);
 
 		// Task: Mutable GetImages
 		//-------------------------------------------------------------
 		FGraphEventRef Mutable_GetImagesTask;
 		{
 			// Task inputs
+			check(CustomizableObject->GetPrivate() != nullptr);
 			TSharedPtr<mu::Model, ESPMode::ThreadSafe> Model = CustomizableObject->GetPrivate()->GetModel();
-			int32 State = CustomizableObjectInstancePtr->GetState();
+			int32 State = ObjectInstance->GetState();
 
-			Mutable_GetImagesTask = System->GetPrivate()->AddMutableThreadTask(
+			Mutable_GetImagesTask = SystemPrivateData->AddMutableThreadTask(
 					TEXT("Task_Mutable_GetImages"),
 					[OperationData, Parameters, Model, State]()
 					{
@@ -2066,12 +2142,15 @@ namespace impl
 
 		// Next Task: Load Unreal Assets
 		//-------------------------------------------------------------
-		FGraphEventRef Game_LoadUnrealAssets = ObjectInstance->GetPrivate()->LoadAdditionalAssetsAsync(OperationData, ObjectInstance, UCustomizableObjectSystem::GetInstance()->GetStreamableManager());
-		if (Game_LoadUnrealAssets) Game_LoadUnrealAssets->SetDebugName(TEXT("LoadAdditionalAssetsAsync"));
+		FGraphEventRef Game_LoadUnrealAssets = ObjectInstancePrivateData->LoadAdditionalAssetsAsync(OperationData, ObjectInstance, UCustomizableObjectSystem::GetInstance()->GetStreamableManager());
+		if (Game_LoadUnrealAssets)
+		{
+			Game_LoadUnrealAssets->SetDebugName(TEXT("LoadAdditionalAssetsAsync"));
+		}
 
 		// Next-next Task: Convert Resources
 		//-------------------------------------------------------------
-		System->GetPrivate()->AddGameThreadTask(
+		SystemPrivateData->AddGameThreadTask(
 			{
 			FMutableTaskDelegate::CreateLambda(
 				[OperationData,CustomizableObjectInstancePtr]()
@@ -2090,9 +2169,11 @@ namespace impl
 	/** "Start Update" */
 	void Task_Game_StartUpdate(TSharedPtr<FMutableOperation> Operation)
 	{
+		check(Operation);
 		check(Operation->Type == FMutableOperation::EOperationType::Update);
 
 		UCustomizableObjectSystem* System = UCustomizableObjectSystem::GetInstance();
+		check(System != nullptr);
 
 		if (!Operation->CustomizableObjectInstance.IsValid() || !Operation->CustomizableObjectInstance->IsValidLowLevel()) // Only start if it hasn't been already destroyed (i.e. GC after finish PIE)
 		{
@@ -2100,15 +2181,16 @@ namespace impl
 			return;
 		}
 
-
-		if (Operation->CustomizableObjectInstance->GetPrivate()->HasCOInstanceFlags(PendingLODsUpdate))
+		UCustomizableInstancePrivateData* ObjectInstancePrivateData = Operation->CustomizableObjectInstance->GetPrivate();
+		check(ObjectInstancePrivateData != nullptr);
+		if (ObjectInstancePrivateData->HasCOInstanceFlags(PendingLODsUpdate))
 		{
-			Operation->CustomizableObjectInstance->GetPrivate()->ClearCOInstanceFlags(PendingLODsUpdate);
+			ObjectInstancePrivateData->ClearCOInstanceFlags(PendingLODsUpdate);
 			// TODO: Is anything needed for this now?
 			//Operation->CustomizableObjectInstance->ReleaseMutableInstanceId(); // To make mutable regenerate the LODs even if the instance parameters have not changed
 		}
 
-		float LastMinSquareDistFromComponentToPlayer = Operation->CustomizableObjectInstance->GetPrivate()->LastMinSquareDistFromComponentToPlayer;
+		float LastMinSquareDistFromComponentToPlayer = ObjectInstancePrivateData->LastMinSquareDistFromComponentToPlayer;
 
 		UCustomizableObjectInstance* CandidateInstance = Operation->CustomizableObjectInstance.Get();
 
@@ -2117,24 +2199,32 @@ namespace impl
 		mu::Ptr<const mu::Parameters> Parameters;
 
 		// If the object is locked (for instance, compiling) we skip any instance update.
-		if (!CandidateInstance)
+		if (!CandidateInstance) // TODO DRB: Shouldn't this be impossible?  This is Operation->CustomizableObjectInstance, and we returned above if it wasn't valid.
 		{
 			bCancel = true;
 		}
 		else
 		{
 			CustomizableObject = CandidateInstance->GetCustomizableObject();
-			if (!CustomizableObject || CustomizableObject->GetPrivate()->bLocked)
+			if (!CustomizableObject)
 			{
 				bCancel = true;
 			}
+			else
+			{
+				check(CustomizableObject->GetPrivate() != nullptr);
+				if (CustomizableObject->GetPrivate()->bLocked)
+				{
+					bCancel = true;
+				}
+			}
 
 			// Only update resources if the instance is in range (it could have got far from the player since the task was queued)
-			if (
-				System->CurrentInstanceLODManagement->IsOnlyUpdateCloseCustomizableObjectsEnabled()
+			check(System->CurrentInstanceLODManagement != nullptr);
+			if (System->CurrentInstanceLODManagement->IsOnlyUpdateCloseCustomizableObjectsEnabled()
 				&& LastMinSquareDistFromComponentToPlayer > FMath::Square(System->CurrentInstanceLODManagement->GetOnlyUpdateCloseCustomizableObjectsDist())
 				&& LastMinSquareDistFromComponentToPlayer != FLT_MAX // This means it is the first frame so it has to be updated
-				)
+			   )
 			{
 				bCancel = true;
 			}
@@ -2148,9 +2238,9 @@ namespace impl
 
 		if (bCancel)
 		{
-			if (Operation->CustomizableObjectInstance.IsValid())
+			if (Operation->CustomizableObjectInstance.IsValid()) // TODO DRB: If this was true at the top of this functoin, we returned.  Can it have become invalid in between?  Any reason we're not using CandidateInstance here?
 			{
-				Operation->CustomizableObjectInstance->GetPrivate()->ClearCOInstanceFlags(Updating);
+				ObjectInstancePrivateData->ClearCOInstanceFlags(Updating);
 			}
 			System->ClearCurrentMutableOperation();
 			return;
@@ -2161,37 +2251,47 @@ namespace impl
 			Operation->StartUpdateTime = FPlatformTime::Seconds();
 		}
 
-		System->GetPrivate()->CurrentInstanceBeingUpdated = CandidateInstance;
+		FCustomizableObjectSystemPrivate* SystemPrivateData = System->GetPrivate();
+		check(SystemPrivateData != nullptr);
+
+		SystemPrivateData->CurrentInstanceBeingUpdated = CandidateInstance;
 
-		UCustomizableInstancePrivateData* PrivateInstance = CandidateInstance->GetPrivate();
+		UCustomizableInstancePrivateData* CandidateInstancePrivateData = CandidateInstance->GetPrivate();
+		check(CandidateInstancePrivateData != nullptr);
 
 		// Prepare streaming for the current customizable object
-		System->GetPrivate()->Streamer->PrepareStreamingForObject(CustomizableObject);
+		check(SystemPrivateData->Streamer != nullptr);
+		SystemPrivateData->Streamer->PrepareStreamingForObject(CustomizableObject);
 
 		CandidateInstance->CommitMinMaxLOD();
 		
 		// Task: Mutable Update and GetMesh
 		//-------------------------------------------------------------
 		TSharedPtr<FMutableOperationData> CurrentOperationData = MakeShared<FMutableOperationData>();
-		CurrentOperationData->TextureCoverageQueries_MutableThreadParams = PrivateInstance->TextureCoverageQueries;
+		check(CurrentOperationData);
+		CurrentOperationData->TextureCoverageQueries_MutableThreadParams = CandidateInstancePrivateData->TextureCoverageQueries;
 		CurrentOperationData->TextureCoverageQueries_MutableThreadResults.Empty();
 		CurrentOperationData->CurrentMinLOD = Operation->InstanceDescriptorRuntimeHash.GetMinLOD();
 		CurrentOperationData->CurrentMaxLOD = Operation->InstanceDescriptorRuntimeHash.GetMaxLOD();
 		CurrentOperationData->bNeverStream = Operation->bNeverStream;
 		CurrentOperationData->MipsToSkip = Operation->MipsToSkip;
 		CurrentOperationData->MutableParameters = Parameters;
+		check(Operation->CustomizableObjectInstance.IsValid());
 		CurrentOperationData->State = Operation->CustomizableObjectInstance->GetState();
 
 		FGraphEventRef Mutable_GetMeshTask;
 		{
 			// Task inputs
 			TWeakObjectPtr<UCustomizableObjectInstance> CustomizableObjectInstancePtr = CandidateInstance;
-			TSharedPtr<FMutableOperation> CurrentMutableOperation = System->GetPrivate()->CurrentMutableOperation;
+			check(CustomizableObjectInstancePtr.IsValid());
+			TSharedPtr<FMutableOperation> CurrentMutableOperation = SystemPrivateData->CurrentMutableOperation;
+			check(CurrentMutableOperation);
 			bool bBuildParameterDecorations = CurrentMutableOperation->IsBuildParameterDecorations();
+			check(CustomizableObject->GetPrivate() != nullptr);
 			TSharedPtr<mu::Model, ESPMode::ThreadSafe> Model = CustomizableObject->GetPrivate()->GetModel();
 			int32 State = CustomizableObjectInstancePtr->GetState();
 
-			Mutable_GetMeshTask = System->GetPrivate()->AddMutableThreadTask(
+			Mutable_GetMeshTask = SystemPrivateData->AddMutableThreadTask(
 				TEXT("Task_Mutable_Update_GetMesh"),
 				[CurrentOperationData, bBuildParameterDecorations, Parameters, Model, State]()
 				{
@@ -2206,9 +2306,12 @@ namespace impl
 		{
 			// Task inputs
 			TWeakObjectPtr<UCustomizableObjectInstance> CustomizableObjectInstancePtr = CandidateInstance;
-			bool bBuildParameterDecorations = System->GetPrivate()->CurrentMutableOperation->IsBuildParameterDecorations();
+			check(CustomizableObjectInstancePtr.IsValid());
+			TSharedPtr<FMutableOperation> CurrentMutableOperation = SystemPrivateData->CurrentMutableOperation;
+			check(CurrentMutableOperation);
+			bool bBuildParameterDecorations = CurrentMutableOperation->IsBuildParameterDecorations();
 
-			System->GetPrivate()->AddGameThreadTask(
+			SystemPrivateData->AddGameThreadTask(
 				{
 				FMutableTaskDelegate::CreateLambda(
 					[CurrentOperationData, CustomizableObjectInstancePtr, bBuildParameterDecorations, Parameters]()
@@ -2227,6 +2330,8 @@ void UCustomizableObjectSystem::AdvanceCurrentOperation()
 {
 	MUTABLE_CPUPROFILER_SCOPE(AdvanceCurrentOperation);
 
+	check(Private != nullptr);
+
 	// See if we can clear the last reference to the mutable-thread task
 	Private->ClearMutableTaskIfDone();
 
@@ -2258,59 +2363,64 @@ void UCustomizableObjectSystem::AdvanceCurrentOperation()
 	// - we have an ongoing operations
 	// - we have no pending work for the ongoing operation
 	// - so we are starting it.
-
 	switch (Private->CurrentMutableOperation->Type)
 	{
-	case FMutableOperation::EOperationType::Discard:
-	{
-		MUTABLE_CPUPROFILER_SCOPE(OperationDiscard);
+		case FMutableOperation::EOperationType::Discard:
+		{
+			MUTABLE_CPUPROFILER_SCOPE(OperationDiscard);
 
-		// \TODO: Discards could be done in any case, concurrently with update operations. Should they be
-		// in their own "queue"?
+			// \TODO: Discards could be done in any case, concurrently with update operations. Should they be
+			// in their own "queue"?
 
-		UCustomizableObjectInstance* COI = Private->CurrentMutableOperation->CustomizableObjectInstance.Get();
+			UCustomizableObjectInstance* COI = Private->CurrentMutableOperation->CustomizableObjectInstance.Get();
+			UCustomizableInstancePrivateData* COIPrivateData = COI ? COI->GetPrivate() : nullptr;
 
-		// Only discard resources if the instance is still out range (it could have got closer to the player since the task was queued)
-		if (!CurrentInstanceLODManagement->IsOnlyUpdateCloseCustomizableObjectsEnabled()
-			|| !COI
-			|| COI->GetPrivate()->LastMinSquareDistFromComponentToPlayer > FMath::Square(CurrentInstanceLODManagement->GetOnlyUpdateCloseCustomizableObjectsDist()))
-		{
-			if (COI && COI->IsValidLowLevel())
+			// Only discard resources if the instance is still out range (it could have got closer to the player since the task was queued)
+			if (!CurrentInstanceLODManagement->IsOnlyUpdateCloseCustomizableObjectsEnabled() ||
+				!COI ||
+				((COIPrivateData != nullptr) &&
+				 (COIPrivateData->LastMinSquareDistFromComponentToPlayer > FMath::Square(CurrentInstanceLODManagement->GetOnlyUpdateCloseCustomizableObjectsDist()))
+				)
+			   )
 			{
-				COI->GetPrivate()->DiscardResourcesAndSetReferenceSkeletalMesh(COI);
-				COI->GetPrivate()->ClearCOInstanceFlags(Updating);
-				COI->SkeletalMeshStatus = ESkeletalMeshState::Correct;
+				if (COI && COI->IsValidLowLevel())
+				{
+					check(COIPrivateData != nullptr);
+					COIPrivateData->DiscardResourcesAndSetReferenceSkeletalMesh(COI);
+					COIPrivateData->ClearCOInstanceFlags(Updating);
+					COI->SkeletalMeshStatus = ESkeletalMeshState::Correct;
+				}
+			}
+			else
+			{
+				check(COIPrivateData != nullptr);
+				COIPrivateData->ClearCOInstanceFlags(Updating);
 			}
-		}
-		else
-		{
-			COI->GetPrivate()->ClearCOInstanceFlags(Updating);
-		}
-
-		if (COI && !COI->HasAnySkeletalMesh())
-		{
-			// To solve the problem in the Mutable demo where PIE just after editor start made all instances appear as reference mesh until editor restart
-			COI->GetPrivate()->ClearCOInstanceFlags(Generated);
-		}
 
-		ClearCurrentMutableOperation();
+			if (COI && !COI->HasAnySkeletalMesh())
+			{
+				// To solve the problem in the Mutable demo where PIE just after editor start made all instances appear as reference mesh until editor restart
+				check(COIPrivateData != nullptr);
+				COIPrivateData->ClearCOInstanceFlags(Generated);
+			}
 
-		break;
-	}
+			ClearCurrentMutableOperation();
 
-	case FMutableOperation::EOperationType::Update:
-	{
-		MUTABLE_CPUPROFILER_SCOPE(OperationUpdate);
+			break;
+		}
 
-		// Start the first task of the update process. See namespace impl comments above.
-		impl::Task_Game_StartUpdate(Private->CurrentMutableOperation);
-		break;
-	}
+		case FMutableOperation::EOperationType::Update:
+		{
+			MUTABLE_CPUPROFILER_SCOPE(OperationUpdate);
 
-	default:
-		check(false);
-		break;
+			// Start the first task of the update process. See namespace impl comments above.
+			impl::Task_Game_StartUpdate(Private->CurrentMutableOperation);
+			break;
+		}
 
+		default:
+			check(false);
+			break;
 	}
 
 }
@@ -2388,16 +2498,19 @@ bool UCustomizableObjectSystem::Tick(float DeltaTime)
 	{
 		if (IsValidChecked(*CustomizableObjectInstance) && CustomizableObjectInstance->GetPrivate())
 		{
-			if (CustomizableObjectInstance->GetPrivate()->HasCOInstanceFlags(UsedByComponentInPlay))
+			UCustomizableInstancePrivateData* ObjectInstancePrivateData = CustomizableObjectInstance->GetPrivate();
+			check(ObjectInstancePrivateData != nullptr);
+
+			if (ObjectInstancePrivateData->HasCOInstanceFlags(UsedByComponentInPlay))
 			{
-				CustomizableObjectInstance->GetPrivate()->TickUpdateCloseCustomizableObjects(**CustomizableObjectInstance);
+				ObjectInstancePrivateData->TickUpdateCloseCustomizableObjects(**CustomizableObjectInstance);
 			}
-			else if (CustomizableObjectInstance->GetPrivate()->HasCOInstanceFlags(UsedByComponent))
+			else if (ObjectInstancePrivateData->HasCOInstanceFlags(UsedByComponent))
 			{
-				CustomizableObjectInstance->GetPrivate()->UpdateInstanceIfNotGenerated(**CustomizableObjectInstance);
+				ObjectInstancePrivateData->UpdateInstanceIfNotGenerated(**CustomizableObjectInstance);
 			}
 
-			CustomizableObjectInstance->GetPrivate()->ClearCOInstanceFlags((ECOInstanceFlags)(UsedByComponent | UsedByComponentInPlay));
+			ObjectInstancePrivateData->ClearCOInstanceFlags((ECOInstanceFlags)(UsedByComponent | UsedByComponentInPlay));
 		}
 	}
 
@@ -2421,11 +2534,11 @@ bool UCustomizableObjectSystem::Tick(float DeltaTime)
 	if (LogBenchmarkUtil::isLoggingActive())
 	{
 		uint64 SizeCache = 0;
-		for (const UTexture2D* cachedTextures : ProtectedCachedTextures)
+		for (const UTexture2D* CachedTextures : ProtectedCachedTextures)
 		{
-			if (cachedTextures)
+			if (CachedTextures)
 			{
-				SizeCache += cachedTextures->CalcTextureMemorySizeEnum(TMC_AllMips);
+				SizeCache += CachedTextures->CalcTextureMemorySizeEnum(TMC_AllMips);
 			}
 		}
 		SET_DWORD_STAT(STAT_MutableTextureCacheMemory, SizeCache / 1024.f);
@@ -2445,6 +2558,8 @@ TArray<FCustomizableObjectExternalTexture> UCustomizableObjectSystem::GetTexture
 {
 	TArray<FCustomizableObjectExternalTexture> Result;
 
+	check(Private != nullptr);
+	check(Private->ImageProvider != nullptr);
 	for (const TWeakObjectPtr<UCustomizableSystemImageProvider> Provider : Private->ImageProvider->ImageProviders)
 	{
 		if (Provider.IsValid())
@@ -2459,12 +2574,18 @@ TArray<FCustomizableObjectExternalTexture> UCustomizableObjectSystem::GetTexture
 
 void UCustomizableObjectSystem::RegisterImageProvider(UCustomizableSystemImageProvider* Provider)
 {
+	check(Private != nullptr);
+	check(Private->ImageProvider != nullptr);
+
 	Private->ImageProvider->ImageProviders.Add(Provider);
 }
 
 
 void UCustomizableObjectSystem::UnregisterImageProvider(UCustomizableSystemImageProvider* Provider)
 {
+	check(Private != nullptr);
+	check(Private->ImageProvider != nullptr);
+
 	Private->ImageProvider->ImageProviders.Remove(Provider);
 }
 
@@ -2483,6 +2604,10 @@ UDefaultImageProvider& UCustomizableObjectSystem::GetDefaultImageProvider()
 
 FMutableOperation FMutableOperation::CreateInstanceUpdate(UCustomizableObjectInstance* InCustomizableObjectInstance, bool bInNeverStream, int32 InMipsToSkip)
 {
+	check(InCustomizableObjectInstance != nullptr);
+	check(InCustomizableObjectInstance->GetPrivate() != nullptr);
+	check(InCustomizableObjectInstance->GetCustomizableObject() != nullptr);
+
 	FMutableOperation Op;
 	Op.Type = EOperationType::Update;
 	Op.bNeverStream = bInNeverStream;
@@ -2507,7 +2632,8 @@ FMutableOperation FMutableOperation::CreateInstanceUpdate(UCustomizableObjectIns
 
 FMutableOperation FMutableOperation::CreateInstanceDiscard(UCustomizableObjectInstance* InCustomizableObjectInstance)
 {
-	check(InCustomizableObjectInstance);
+	check(InCustomizableObjectInstance != nullptr);
+	check(InCustomizableObjectInstance->GetPrivate() != nullptr);
 
 	FMutableOperation Op;
 	Op.Type = EOperationType::Discard;
@@ -2521,6 +2647,10 @@ FMutableOperation FMutableOperation::CreateInstanceDiscard(UCustomizableObjectIn
 
 void FMutableOperation::MutableIsDisabledCase()
 {
+	check(CustomizableObjectInstance != nullptr);
+	check(CustomizableObjectInstance->GetPrivate() != nullptr);
+	check(CustomizableObjectInstance->GetCustomizableObject() != nullptr);
+
 	for (TObjectIterator<UCustomizableSkeletalComponent> It; It; ++It)
 	{
 		UCustomizableSkeletalComponent* CustomizableSkeletalComponent = *It;
@@ -2549,50 +2679,61 @@ void FMutableOperation::MutableIsDisabledCase()
 
 int32 UCustomizableObjectSystem::GetNumInstances() const
 {
+	check(Private != nullptr);
 	return Private->NumInstances;
 }
 
 int32 UCustomizableObjectSystem::GetNumPendingInstances() const
 {
+	check(Private != nullptr);
 	return Private->NumPendingInstances;
 }
 
 int32 UCustomizableObjectSystem::GetTotalInstances() const
 {
+	check(Private != nullptr);
 	return Private->TotalInstances;
 }
 
 int32 UCustomizableObjectSystem::GetTextureMemoryUsed() const
 {
+	check(Private != nullptr);
 	return int32(Private->TextureMemoryUsed);
 }
 
 int32 UCustomizableObjectSystem::GetAverageBuildTime() const
 {
+	check(Private != nullptr);
 	return Private->TotalBuiltInstances == 0 ? 0 : Private->TotalBuildMs / Private->TotalBuiltInstances;
 }
 
 
 bool UCustomizableObjectSystem::IsCompactSerializationEnabled() const
 {
+	check(Private != nullptr);
 	return Private->bCompactSerialization;
 }
 
 
 bool UCustomizableObjectSystem::IsSupport16BitBoneIndexEnabled() const
 {
+	check(Private != nullptr);
 	return Private->bSupport16BitBoneIndex;
 }
 
 bool UCustomizableObjectSystem::IsProgressiveMipStreamingEnabled() const
 {
+	check(Private != nullptr);
 	return Private->EnableMutableProgressiveMipStreaming != 0;
 }
 
 
 void UCustomizableObjectSystem::AddUncompiledCOWarning(UCustomizableObject* InObject)
 {
-	if (!InObject) return;
+	if (!InObject)
+	{
+		return;
+	}
 
 	FString Msg;
 	Msg += FString::Printf(TEXT("Warning: Customizable Object [%s] not compiled. Please compile and save the object."), *InObject->GetName());
@@ -2633,6 +2774,7 @@ void UCustomizableObjectSystem::EndBenchmark()
 
 void UCustomizableObjectSystem::SetReleaseMutableTexturesImmediately(bool bReleaseTextures)
 {
+	check(Private != nullptr);
 	Private->bReleaseTexturesImmediately = bReleaseTextures;
 }
 
@@ -2641,7 +2783,7 @@ void UCustomizableObjectSystem::SetReleaseMutableTexturesImmediately(bool bRelea
 
 void UCustomizableObjectSystem::OnPreBeginPIE(const bool bIsSimulatingInEditor)
 {
-	if(!EditorSettings.bCompileRootObjectsOnStartPIE || IsRunningGame() || IsCompilationDisabled())
+	if (!EditorSettings.bCompileRootObjectsOnStartPIE || IsRunningGame() || IsCompilationDisabled())
 	{
 		return;
 	}
@@ -2665,7 +2807,7 @@ void UCustomizableObjectSystem::OnPreBeginPIE(const bool bIsSimulatingInEditor)
 		}
 		
 		const UCustomizableObject* Object = Cast<UCustomizableObject>(Asset.GetAsset());
-		if(!Object || Object->IsCompiled() || Object->IsLocked() || Object->bIsChildObject)
+		if (!Object || Object->IsCompiled() || Object->IsLocked() || Object->bIsChildObject)
 		{
 			continue;
 		}
@@ -2687,7 +2829,7 @@ void UCustomizableObjectSystem::OnPreBeginPIE(const bool bIsSimulatingInEditor)
 
 void UCustomizableObjectSystem::StartNextRecompile()
 {
-	if (ObjectsToRecompile.Num() % 10 == 0)
+	if (ObjectsToRecompile.Num() % 10 == 0) // TODO DRB: What?  If we can add more than one between calls, this might never be hit.  Why % 10?
 	{
 		CollectGarbage(GARBAGE_COLLECTION_KEEPFLAGS);
 	}
@@ -2703,6 +2845,7 @@ void UCustomizableObjectSystem::StartNextRecompile()
 		// Use default options
 		FCompilationOptions Options = CustomizableObject->CompileOptions;
 		Options.bSilentCompilation = true;
+		check(RecompileCustomizableObjectsCompiler != nullptr);
 		RecompileCustomizableObjectsCompiler->Compile(*CustomizableObject, Options, true);
 	}
 }
@@ -2710,17 +2853,17 @@ void UCustomizableObjectSystem::StartNextRecompile()
 void UCustomizableObjectSystem::RecompileCustomizableObjectAsync(const FAssetData& InAssetData,
 	const UCustomizableObject* InObject)
 {
-	if(IsRunningGame() || IsCompilationDisabled())
+	if (IsRunningGame() || IsCompilationDisabled())
 	{
 		return;
 	}
 	
-	if((InObject && InObject->IsLocked()) || ObjectsToRecompile.Find((InAssetData)) != INDEX_NONE)
+	if ((InObject && InObject->IsLocked()) || ObjectsToRecompile.Find((InAssetData)) != INDEX_NONE)
 	{
 		return;
 	}
 	
-	if(!ObjectsToRecompile.IsEmpty())
+	if (!ObjectsToRecompile.IsEmpty())
 	{
 		ObjectsToRecompile.Add(InAssetData);
 	}
@@ -2754,7 +2897,7 @@ void UCustomizableObjectSystem::RecompileCustomizableObjects(const TArray<FAsset
 		TotalNumObjectsToRecompile = ObjectsToRecompile.Num();
 		NumObjectsCompiled = 0;
 
-		if(RecompileNotificationHandle.IsValid())
+		if (RecompileNotificationHandle.IsValid())
 		{
 			++TotalNumObjectsToRecompile;
 			FSlateNotificationManager::Get().UpdateProgressNotification(RecompileNotificationHandle, NumObjectsCompiled, TotalNumObjectsToRecompile);
@@ -2806,7 +2949,7 @@ uint64 UCustomizableObjectSystem::GetMaxChunkSizeForPlatform(const ITargetPlatfo
 {
 	const FString& PlatformName = TargetPlatform ? TargetPlatform->IniPlatformName() : FPlatformProperties::IniPlatformName();
 
-	if(const int64* CachedMaxChunkSize = PlatformMaxChunkSize.Find(PlatformName))
+	if (const int64* CachedMaxChunkSize = PlatformMaxChunkSize.Find(PlatformName))
 	{
 		return *CachedMaxChunkSize;
 	}
@@ -2825,7 +2968,7 @@ uint64 UCustomizableObjectSystem::GetMaxChunkSizeForPlatform(const ITargetPlatfo
 	}
 
 	// If no limit is specified default it to MUTABLE_STREAMED_DATA_MAXCHUNKSIZE 
-	if(MaxChunkSize <= 0)
+	if (MaxChunkSize <= 0)
 	{
 		MaxChunkSize = MUTABLE_STREAMED_DATA_MAXCHUNKSIZE;
 	}
@@ -2840,28 +2983,32 @@ uint64 UCustomizableObjectSystem::GetMaxChunkSizeForPlatform(const ITargetPlatfo
 
 void UCustomizableObjectSystem::CacheImage(uint64 ImageId)
 {
-	check(GetPrivate()->ImageProvider);
+	check(GetPrivate() != nullptr);
+	check(GetPrivate()->ImageProvider != nullptr);
 	GetPrivate()->ImageProvider->CacheImage(ImageId);
 }
 
 
 void UCustomizableObjectSystem::UnCacheImage(uint64 ImageId)
 {
-	check(GetPrivate()->ImageProvider);
+	check(GetPrivate() != nullptr); 
+	check(GetPrivate()->ImageProvider != nullptr);
 	GetPrivate()->ImageProvider->UnCacheImage(ImageId);
 }
 
 
 void UCustomizableObjectSystem::CacheAllImagesInAllProviders(bool bClearPreviousCacheImages)
 {
-	check(GetPrivate()->ImageProvider);
+	check(GetPrivate() != nullptr);
+	check(GetPrivate()->ImageProvider != nullptr);
 	GetPrivate()->ImageProvider->CacheAllImagesInAllProviders(bClearPreviousCacheImages);
 }
 
 
 void UCustomizableObjectSystem::ClearImageCache()
 {
-	check(GetPrivate()->ImageProvider);
+	check(GetPrivate() != nullptr);
+	check(GetPrivate()->ImageProvider != nullptr);
 	GetPrivate()->ImageProvider->ClearCache();
 }
 
diff --git a/Engine/Plugins/Experimental/Mutable/Source/CustomizableObject/Private/MuCO/CustomizableObjectSystemPrivate.h b/Engine/Plugins/Experimental/Mutable/Source/CustomizableObject/Private/MuCO/CustomizableObjectSystemPrivate.h
--- a/Engine/Plugins/Experimental/Mutable/Source/CustomizableObject/Private/MuCO/CustomizableObjectSystemPrivate.h
+++ b/Engine/Plugins/Experimental/Mutable/Source/CustomizableObject/Private/MuCO/CustomizableObjectSystemPrivate.h
@@ -110,6 +110,7 @@ struct FMutableQueueElem
 
 		MutableQueueElem.PriorityType = NewPriorityType;
 		MutableQueueElem.Priority = NewPriority;
+		check(InOperation);
 		MutableQueueElem.Operation = InOperation;
 		MutableQueueElem.bIsDiscardResources = (InOperation->Type==FMutableOperation::EOperationType::Discard);
 
