[mutable] excessive checks to keep static analysis happy.
#rnx

diff --git a/Engine/Plugins/Experimental/Mutable/Source/MutableRuntime/Private/MuR/BlockCompression/Miro/Miro.cpp b/Engine/Plugins/Experimental/Mutable/Source/MutableRuntime/Private/MuR/BlockCompression/Miro/Miro.cpp
--- a/Engine/Plugins/Experimental/Mutable/Source/MutableRuntime/Private/MuR/BlockCompression/Miro/Miro.cpp
+++ b/Engine/Plugins/Experimental/Mutable/Source/MutableRuntime/Private/MuR/BlockCompression/Miro/Miro.cpp
@@ -6043,8 +6043,9 @@ namespace impl
 				for (int texelX = 0; texelX < blockWidth; texelX++)
 				{
 					const int				texelNdx = texelY * blockWidth + texelX;
-					const int				colorEndpointNdx = numPartitions == 1 ? 0 : computeTexelPartition(partitionIndexSeed, texelX, texelY, 0, numPartitions, smallBlock);
+					int						colorEndpointNdx = numPartitions == 1 ? 0 : computeTexelPartition(partitionIndexSeed, texelX, texelY, 0, numPartitions, smallBlock);
 					DE_ASSERT(colorEndpointNdx < numPartitions);
+					colorEndpointNdx = FMath::Clamp<int>(colorEndpointNdx,0, FMath::Min(3,numPartitions));
 					const UVec4& e0 = colorEndpoints[colorEndpointNdx].e0;
 					const UVec4& e1 = colorEndpoints[colorEndpointNdx].e1;
 					const TexelWeightPair& weight = texelWeights[texelNdx];
@@ -6066,7 +6067,7 @@ namespace impl
 						}
 						result = DECOMPRESS_RESULT_ERROR;
 					}
-					else
+					else if (colorEndpointNdx >= 0 && colorEndpointNdx < 4)
 					{
 						for (int channelNdx = 0; channelNdx < 4; channelNdx++)
 						{
