[mutable] tex decompression fix for blank blocks.
#rnx

diff --git a/Engine/Plugins/Experimental/Mutable/Source/MutableRuntime/Private/MuR/BlockCompression/Miro/Miro.cpp b/Engine/Plugins/Experimental/Mutable/Source/MutableRuntime/Private/MuR/BlockCompression/Miro/Miro.cpp
--- a/Engine/Plugins/Experimental/Mutable/Source/MutableRuntime/Private/MuR/BlockCompression/Miro/Miro.cpp
+++ b/Engine/Plugins/Experimental/Mutable/Source/MutableRuntime/Private/MuR/BlockCompression/Miro/Miro.cpp
@@ -5033,6 +5033,15 @@ namespace impl
 				DE_ASSERT(inBounds(ndx, 0, 128));
 				return getBit(ndx) != 0;
 			}
+
+			bool isZero() const
+			{
+				for (int wordNdx = 0; wordNdx < NUM_WORDS; wordNdx++)
+				{
+					if (m_words[wordNdx]) return false;
+				}
+				return true;
+			}
 		private:
 			Word m_words[NUM_WORDS];
 		};
@@ -6108,6 +6117,7 @@ namespace impl
 		DecompressResult decompressBlock(void* dst, const Block128& blockData, int blockWidth, int blockHeight, bool isSRGB, bool isLDR)
 		{
 			DE_ASSERT(isLDR || !isSRGB);
+
 			// Decode block mode.
 			const ASTCBlockMode blockMode = getASTCBlockMode(blockData.getBits(0, 10));
 			// Check for block mode errors.
@@ -6177,6 +6187,21 @@ namespace impl
 			float linear[MAX_BLOCK_WIDTH * MAX_BLOCK_HEIGHT * 4];
 
 			const Block128 blockData(data);
+
+			// (anticto) shortcut for blank blocks
+			if (blockData.isZero())
+			{
+				int32 pix = 0;
+				for (int32 i = 0; i < blockHeight* blockWidth; i++)
+				{
+					pDst[4 * i + 0] = 0;
+					pDst[4 * i + 1] = 0;
+					pDst[4 * i + 2] = 0;
+					pDst[4 * i + 3] = 0;
+				}
+				return true;
+			}
+
 			if (decompressBlock(isSRGB ? (void*)pDst : (void*)&linear[0],
 				blockData, blockWidth, blockHeight, isSRGB, isLDR) != DECOMPRESS_RESULT_VALID_BLOCK)
 				return false;
