[mutable] Fix crash when getting an external image

#rnx

diff --git a/Engine/Plugins/Experimental/Mutable/Source/CustomizableObject/Private/MuCO/UnrealMutableImageProvider.cpp b/Engine/Plugins/Experimental/Mutable/Source/CustomizableObject/Private/MuCO/UnrealMutableImageProvider.cpp
--- a/Engine/Plugins/Experimental/Mutable/Source/CustomizableObject/Private/MuCO/UnrealMutableImageProvider.cpp
+++ b/Engine/Plugins/Experimental/Mutable/Source/CustomizableObject/Private/MuCO/UnrealMutableImageProvider.cpp
@@ -91,19 +91,45 @@ namespace
 }
 
 
+mu::EImageFormat GetMutablePixelFormat(EPixelFormat InTextureFormat)
+{
+	switch (InTextureFormat)
+	{
+	case PF_B8G8R8A8: return mu::EImageFormat::IF_BGRA_UBYTE;
+	case PF_R8G8B8A8: return mu::EImageFormat::IF_RGBA_UBYTE;
+	case PF_DXT1: return mu::EImageFormat::IF_BC1;
+	case PF_DXT3: return mu::EImageFormat::IF_BC2;
+	case PF_DXT5: return mu::EImageFormat::IF_BC3;
+	case PF_BC4: return mu::EImageFormat::IF_BC4;
+	case PF_BC5: return mu::EImageFormat::IF_BC5;
+	case PF_G8: return mu::EImageFormat::IF_L_UBYTE;
+	case PF_ASTC_4x4: return mu::EImageFormat::IF_ASTC_4x4_RGBA_LDR;
+	case PF_ASTC_8x8: return mu::EImageFormat::IF_ASTC_8x8_RGBA_LDR;
+	case PF_ASTC_12x12: return mu::EImageFormat::IF_ASTC_12x12_RGBA_LDR;
+	default: return mu::EImageFormat::IF_NONE;
+	}
+}
+
+
 //-------------------------------------------------------------------------------------------------
 
 mu::ImagePtr FUnrealMutableImageProvider::GetImage(mu::EXTERNAL_IMAGE_ID id)
 {
 	// Thread: Mutable worker
+	MUTABLE_CPUPROFILER_SCOPE(FUnrealMutableImageProvider::GetImage);
+
+	// Out Texture
+	mu::ImagePtr Image;
 
 	// Some data that may have to be copied from the GlobalExternalImages while it's locked
-	TUniquePtr<IAsyncReadFileHandle> FileHandle;
-	int32 LODs = 1;
-	int32 SizeX = 0;
-	int32 SizeY = 0;
+	TUniquePtr<IBulkDataIORequest> IORequest;
+	const int32 LODs = 1;
+
 	EPixelFormat Format = EPixelFormat::PF_Unknown;
 	int32 BulkDataSize = 0;
+
+	mu::EImageFormat MutImageFormat = mu::EImageFormat::IF_NONE;
+	int32 MutImageDataSize = 0;
 	
 	{
 		FScopeLock Lock(&ExternalImagesLock);
@@ -132,17 +158,58 @@ mu::ImagePtr FUnrealMutableImageProvider::GetImage(mu::EXTERNAL_IMAGE_ID id)
 			// In the editor the src data can be directly accessed
 			return ConvertTextureUnrealToMutable(TextureToLoad);
 #else
+			// Texture format and the equivalent mutable format
+			Format = TextureToLoad->GetPlatformData()->PixelFormat;
+			MutImageFormat = GetMutablePixelFormat(Format);
+
+			// Check if it's a format we support
+			if (MutImageFormat == mu::EImageFormat::IF_NONE)
+			{
+				UE_LOG(LogMutable, Warning, TEXT("Failed to get external image. Unexpected image format. EImageFormat [%s]."), GetPixelFormatString(Format));
+				return CreateDummy();
+			}
+
+			Image = new mu::Image(TextureToLoad->GetSizeX(), TextureToLoad->GetSizeY(), LODs, MutImageFormat);
+			MutImageDataSize = Image->GetDataSize();
+
 			// In a packaged game the bulk data has to be loaded
 			// Get the actual file to read the mip 0 data, do not keep any reference to TextureToLoad because once outside of the lock
 			// it may be GCed or changed. Just keep the actual file handle and some sizes instead of the texture
 			FByteBulkData& BulkData = TextureToLoad->GetPlatformData()->Mips[0].BulkData;
-			check(BulkData.GetBulkDataSize() > 0);
-			FileHandle.Reset(BulkData.OpenAsyncReadHandle());
-
-			SizeX = TextureToLoad->GetSizeX();
-			SizeY = TextureToLoad->GetSizeY();
 			BulkDataSize = BulkData.GetBulkDataSize();
-			Format = TextureToLoad->GetPlatformData()->PixelFormat;
+			check(BulkDataSize > 0);
+
+			if (BulkDataSize != MutImageDataSize)
+			{
+				UE_LOG(LogMutable, Warning, TEXT("Failed to get external image. Bulk data size is different than the expected size. BulkData size [%d]. Mutable image data size [%d]."),
+					BulkDataSize, MutImageDataSize);
+
+				return CreateDummy();
+			}
+
+			// Create a streaming request if the data is not loaded or copy the mip data
+			if (!BulkData.IsBulkDataLoaded())
+			{
+				IORequest.Reset(BulkData.CreateStreamingRequest(EAsyncIOPriorityAndFlags::AIOP_High, nullptr, Image->GetData()));
+			}
+			else
+			{
+				// Bulk data already loaded
+				const void* Data = BulkData.IsUnlocked() ? BulkData.LockReadOnly() : nullptr; // TODO: Retry if it fails?
+				
+				if (Data)
+				{
+					FMemory::Memcpy(Image->GetData(), Data, BulkDataSize);
+
+					BulkData.Unlock();
+					return Image;
+				}
+				else
+				{
+					UE_LOG(LogMutable, Warning, TEXT("Failed to get external image. Bulk data already locked or null."));
+					return CreateDummy();
+				}
+			}
 #endif
 		}
 		else
@@ -153,50 +220,47 @@ mu::ImagePtr FUnrealMutableImageProvider::GetImage(mu::EXTERNAL_IMAGE_ID id)
 		}
 	}
 
-	// Case where the bulk data has to be loaded from disk. The ExternalImagesLock has been unlocked because the file access
-	// must not be done while the lock is on as it could stall the game thread for too long
-	if (FileHandle)
+
+	if (IORequest)
 	{
-		mu::ImagePtr Image;
-		IAsyncReadRequest* AsyncRequest = FileHandle->ReadRequest(0, BulkDataSize,
-			AIOP_FLAG_DONTCACHE | AIOP_BelowNormal, nullptr, nullptr);
+		// It's OK to wait since this must run on the Mutable thread
+		bool bComplete = IORequest->WaitCompletion();
+		uint8* Results = IORequest->GetReadResults(); // required?
 
-		if (AsyncRequest)
+		if (bComplete && Results && Image->GetDataSize() == (int32)IORequest->GetSize())
 		{
-			// It's OK to wait since this must run on the Mutable thread
-			bool bComplete = AsyncRequest->WaitCompletion();
-			uint8* Results = AsyncRequest->GetReadResults();
-
-			if (bComplete && Results)
-			{
-				// TODO: For the moment only this uncompressed format is supported in packaged games
-				check(Format == EPixelFormat::PF_B8G8R8A8);
-				mu::EImageFormat MutableFormat = mu::EImageFormat::IF_BGRA_UBYTE;
-
-				check(BulkDataSize == AsyncRequest->GetSizeResults());
-				check(Image->GetDataSize() == BulkDataSize);
-				Image = new mu::Image(SizeX, SizeY, LODs, MutableFormat);
-				FMemory::Memcpy(Image->GetData(), Results, Image->GetDataSize());
-
-				FMemory::Free(Results);
-				Results = nullptr;
-			}
+			check(BulkDataSize == (int32)IORequest->GetSize());
+			check(Results == Image->GetData());
+			return Image;
 		}
 
-		delete AsyncRequest;
-
-		if (!Image)
+		if (!bComplete || !Results)
+		{
+			UE_LOG(LogMutable, Warning, TEXT("Failed to get external image. IO Request failed. Request completed: [%hhd]. Request results [%hhd]. Format: [%s]. MutableFormat: [%d]."),
+				bComplete,
+				(Results != nullptr),
+				GetPixelFormatString(Format),
+				(int32)MutImageFormat);
+		}
+		else if (MutImageDataSize != (int32)IORequest->GetSize())
 		{
-			// Something failed in the loading of the bulk data, just return a dummy
-			UE_LOG(LogMutable, Warning, TEXT("UTexture2D BulkData failed loading for an application-specific image parameter."));
-			CreateDummy();
+			UE_LOG(LogMutable, Warning, TEXT("Failed to get external image. Requested size is different than the expected size. RequestSize: [%lld]. ExpectedSize: [%d]. Format: [%s]. MutableFormat: [%d]."),
+				IORequest->GetSize(),
+				Image->GetDataSize(),
+				GetPixelFormatString(Format),
+				(int32)MutImageFormat);
+		}
+		else
+		{
+			UE_LOG(LogMutable, Warning, TEXT("Failed to get external image."));
 		}
 
-		return Image;
+		// Something failed when loading the bulk data, just return a dummy
+		return CreateDummy();
 	}
 	else
 	{
-		UE_LOG(LogMutable, Warning, TEXT("Failed to get a FileHandle for a UTexture2D BulkData for an application-specific image parameter."));
+		UE_LOG(LogMutable, Warning, TEXT("Failed to create an IORequest for a UTexture2D BulkData for an application-specific image parameter."));
 		return CreateDummy();
 	}
 }
