[mutable] More static analysis fixes.
#rnx

diff --git a/Engine/Plugins/Experimental/Mutable/Source/MutableRuntime/Private/MuR/BlockCompression/Miro/Miro.cpp b/Engine/Plugins/Experimental/Mutable/Source/MutableRuntime/Private/MuR/BlockCompression/Miro/Miro.cpp
--- a/Engine/Plugins/Experimental/Mutable/Source/MutableRuntime/Private/MuR/BlockCompression/Miro/Miro.cpp
+++ b/Engine/Plugins/Experimental/Mutable/Source/MutableRuntime/Private/MuR/BlockCompression/Miro/Miro.cpp
@@ -6038,7 +6038,7 @@ namespace impl
 		{
 			const bool			smallBlock = blockWidth * blockHeight < 31;
 			DecompressResult	result = DECOMPRESS_RESULT_VALID_BLOCK;
-			bool				isHDREndpoint[4];
+			bool				isHDREndpoint[4] = {false,false,false,false};
 			for (int i = 0; i < numPartitions; i++)
 			{
 				isHDREndpoint[i] = isColorEndpointModeHDR(colorEndpointModes[i]);
