[mutable] Add the mips-to-skip information to the external image request. Support still needs to be implemented at Unreal level.
#rnx

diff --git a/Engine/Plugins/Experimental/Mutable/Source/CustomizableObject/Classes/MuCO/CustomizableObject.h b/Engine/Plugins/Experimental/Mutable/Source/CustomizableObject/Classes/MuCO/CustomizableObject.h
--- a/Engine/Plugins/Experimental/Mutable/Source/CustomizableObject/Classes/MuCO/CustomizableObject.h
+++ b/Engine/Plugins/Experimental/Mutable/Source/CustomizableObject/Classes/MuCO/CustomizableObject.h
@@ -1223,7 +1223,7 @@ private:
 	// This is a manual version number for the binary blobs in this asset.
 	// Increasing it invalidates all the previously compiled models.
 	// Warning: If while merging code both versions have changed, take the highest+1.
-	static const int32 CurrentSupportedVersion = 360;
+	static const int32 CurrentSupportedVersion = 361;
 
 public:
 
diff --git a/Engine/Plugins/Experimental/Mutable/Source/CustomizableObject/Private/MuCO/UnrealMutableImageProvider.cpp b/Engine/Plugins/Experimental/Mutable/Source/CustomizableObject/Private/MuCO/UnrealMutableImageProvider.cpp
--- a/Engine/Plugins/Experimental/Mutable/Source/CustomizableObject/Private/MuCO/UnrealMutableImageProvider.cpp
+++ b/Engine/Plugins/Experimental/Mutable/Source/CustomizableObject/Private/MuCO/UnrealMutableImageProvider.cpp
@@ -113,9 +113,9 @@ mu::EImageFormat GetMutablePixelFormat(EPixelFormat InTextureFormat)
 
 //-------------------------------------------------------------------------------------------------
 
-mu::ImagePtr FUnrealMutableImageProvider::GetImage(mu::EXTERNAL_IMAGE_ID id)
+mu::ImagePtr FUnrealMutableImageProvider::GetImage(mu::EXTERNAL_IMAGE_ID id, uint8 MipmapsToSkip )
 {
-	// Thread: Mutable worker
+	// Thread: worker
 	MUTABLE_CPUPROFILER_SCOPE(FUnrealMutableImageProvider::GetImage);
 
 	// Out Texture
diff --git a/Engine/Plugins/Experimental/Mutable/Source/CustomizableObject/Private/MuCO/UnrealMutableImageProvider.h b/Engine/Plugins/Experimental/Mutable/Source/CustomizableObject/Private/MuCO/UnrealMutableImageProvider.h
--- a/Engine/Plugins/Experimental/Mutable/Source/CustomizableObject/Private/MuCO/UnrealMutableImageProvider.h
+++ b/Engine/Plugins/Experimental/Mutable/Source/CustomizableObject/Private/MuCO/UnrealMutableImageProvider.h
@@ -18,8 +18,8 @@ class FUnrealMutableImageProvider : public mu::ImageParameterGenerator, public F
 public:
 
 	// mu::ImageParameterGenerator interface
-	// Thread: Mutable
-	mu::ImagePtr GetImage(mu::EXTERNAL_IMAGE_ID id) override;
+	// Thread: worker
+	mu::ImagePtr GetImage(mu::EXTERNAL_IMAGE_ID id, uint8 MipmapsToSkip) override;
 
 	// Own interface
 	// Thread: Game
@@ -42,7 +42,7 @@ private:
 	struct FUnrealMutableImageInfo
 	{
 		FUnrealMutableImageInfo() {}
-		FUnrealMutableImageInfo(mu::ImagePtr InImage, class UTexture2D* InTextureToLoad) : Image(InImage), TextureToLoad(InTextureToLoad) { }
+		FUnrealMutableImageInfo(const mu::ImagePtr& InImage, class UTexture2D* InTextureToLoad) : Image(InImage), TextureToLoad(InTextureToLoad) { }
 
 		mu::ImagePtr Image;
 
diff --git a/Engine/Plugins/Experimental/Mutable/Source/CustomizableObjectEditor/Private/MuCOE/SMutableCodeViewer.cpp b/Engine/Plugins/Experimental/Mutable/Source/CustomizableObjectEditor/Private/MuCOE/SMutableCodeViewer.cpp
--- a/Engine/Plugins/Experimental/Mutable/Source/CustomizableObjectEditor/Private/MuCOE/SMutableCodeViewer.cpp
+++ b/Engine/Plugins/Experimental/Mutable/Source/CustomizableObjectEditor/Private/MuCOE/SMutableCodeViewer.cpp
@@ -218,6 +218,7 @@ public:
 				OpName += mu::TypeInfo::s_blendModeName[int32(Args.blendType)];
 				OpName += TEXT(" a: ");
 				OpName += mu::TypeInfo::s_blendModeName[int32(Args.blendTypeAlpha)];
+				//OpName += FString::Printf(TEXT(" flags %d"), Args.flags);
 				break;
 			}
 
@@ -2191,11 +2192,13 @@ namespace
 	class TestImageProvider : public mu::ImageParameterGenerator
 	{
 	public:
-		mu::Ptr<mu::Image> GetImage(mu::EXTERNAL_IMAGE_ID id) override
+		mu::Ptr<mu::Image> GetImage(mu::EXTERNAL_IMAGE_ID id, uint8 MipmapsToSkip) override
 		{
 			MUTABLE_CPUPROFILER_SCOPE(TestImageProvider_GetImage);
 
-			const int32 Size = 1024;
+			int32 Size = 1024;
+			Size = FMath::Max(4, Size / (1 << MipmapsToSkip));
+
 			const int32 LODCount = 1;
 			const mu::EImageFormat Format = mu::EImageFormat::IF_RGBA_UBYTE;
 			mu::Ptr<mu::Image> Image = new mu::Image(Size, Size, LODCount, Format);
diff --git a/Engine/Plugins/Experimental/Mutable/Source/MutableRuntime/Private/MuR/CodeRunner.cpp b/Engine/Plugins/Experimental/Mutable/Source/MutableRuntime/Private/MuR/CodeRunner.cpp
--- a/Engine/Plugins/Experimental/Mutable/Source/MutableRuntime/Private/MuR/CodeRunner.cpp
+++ b/Engine/Plugins/Experimental/Mutable/Source/MutableRuntime/Private/MuR/CodeRunner.cpp
@@ -110,7 +110,7 @@ namespace mu
 
 
     //---------------------------------------------------------------------------------------------
-    ImagePtr CodeRunner::LoadExternalImage( EXTERNAL_IMAGE_ID id )
+    ImagePtr CodeRunner::LoadExternalImage( EXTERNAL_IMAGE_ID Id, uint8 MipmapsToSkip )
     {
 		MUTABLE_CPUPROFILER_SCOPE(LoadExternalImage);
 
@@ -120,13 +120,16 @@ namespace mu
 
 		if (m_pSystem->m_pImageParameterGenerator)
 		{
-			pResult = m_pSystem->m_pImageParameterGenerator->GetImage(id);
+			pResult = m_pSystem->m_pImageParameterGenerator->GetImage(Id, MipmapsToSkip);
+
 			// Don't cache for now. Need to figure out how to invalidate them.
+			// \TODO: Like constants? attached to a cache level?
 			//m_pSystem->m_externalImages.push_back( pair<EXTERNAL_IMAGE_ID,ImagePtr>(id,pResult) );
 		}
 		else
 		{
 			// Not found and there is no generator!
+			check(false);
 		}
 
 		return pResult;
@@ -2389,8 +2392,11 @@ namespace mu
         {
 			OP::ParameterArgs args = pModel->GetPrivate()->m_program.GetOpArgs<OP::ParameterArgs>(item.At);
 			Ptr<RangeIndex> Index = BuildCurrentOpRangeIndex(item, pParams, pModel, args.variable);
+
 			EXTERNAL_IMAGE_ID id = pParams->GetImageValue(args.variable, Index);
-			ImagePtr pResult = LoadExternalImage(id);
+
+			uint8 MipmapsToSkip = item.ExecutionOptions;
+			ImagePtr pResult = LoadExternalImage(id, MipmapsToSkip);
 			GetMemory().SetImage(item, pResult);
             break;
         }
@@ -5262,7 +5268,8 @@ namespace mu
 			check(item.Stage == 0);
 			OP::ParameterArgs args = program.GetOpArgs<OP::ParameterArgs>(item.At);
 			EXTERNAL_IMAGE_ID id = pParams->GetImageValue(args.variable);
-			ImagePtr pResult = LoadExternalImage(id);
+			uint8 MipsToSkip = item.ExecutionOptions;
+			ImagePtr pResult = LoadExternalImage(id, MipsToSkip);
 			m_heapImageDesc[item.CustomState].m_format = pResult->GetFormat();
 			m_heapImageDesc[item.CustomState].m_size = pResult->GetSize();
 			m_heapImageDesc[item.CustomState].m_lods = pResult->GetLODCount();
diff --git a/Engine/Plugins/Experimental/Mutable/Source/MutableRuntime/Private/MuR/CodeRunner.h b/Engine/Plugins/Experimental/Mutable/Source/MutableRuntime/Private/MuR/CodeRunner.h
--- a/Engine/Plugins/Experimental/Mutable/Source/MutableRuntime/Private/MuR/CodeRunner.h
+++ b/Engine/Plugins/Experimental/Mutable/Source/MutableRuntime/Private/MuR/CodeRunner.h
@@ -102,7 +102,7 @@ namespace  mu
 		void RunCodeImageDesc(FScheduledOp&, const Parameters*, const Model*, uint32 LodMask);
 
 
-        ImagePtr LoadExternalImage( EXTERNAL_IMAGE_ID id );
+        ImagePtr LoadExternalImage( EXTERNAL_IMAGE_ID, uint8 MipmapsToSkip );
 
         //! Heap of intermediate data pushed by some instructions and referred by others.
         //! It is not released until no operations are pending.
diff --git a/Engine/Plugins/Experimental/Mutable/Source/MutableRuntime/Public/MuR/System.h b/Engine/Plugins/Experimental/Mutable/Source/MutableRuntime/Public/MuR/System.h
--- a/Engine/Plugins/Experimental/Mutable/Source/MutableRuntime/Public/MuR/System.h
+++ b/Engine/Plugins/Experimental/Mutable/Source/MutableRuntime/Public/MuR/System.h
@@ -70,7 +70,7 @@ namespace mu
         virtual ~ImageParameterGenerator() = default;
 
         //!
-        virtual Ptr<Image> GetImage( EXTERNAL_IMAGE_ID id ) = 0;
+        virtual Ptr<Image> GetImage( EXTERNAL_IMAGE_ID id, uint8 MipmapsToSkip) = 0;
     };
 
 
