[mutable] Add operation to optimize the skinning of the generated meshes. If possible, bone index buffers will be formated and bonemaps shrinked.

#rnx

diff --git a/Engine/Plugins/Experimental/Mutable/Source/CustomizableObject/Private/MuCO/CustomizableObjectInstance.cpp b/Engine/Plugins/Experimental/Mutable/Source/CustomizableObject/Private/MuCO/CustomizableObjectInstance.cpp
--- a/Engine/Plugins/Experimental/Mutable/Source/CustomizableObject/Private/MuCO/CustomizableObjectInstance.cpp
+++ b/Engine/Plugins/Experimental/Mutable/Source/CustomizableObject/Private/MuCO/CustomizableObjectInstance.cpp
@@ -3070,13 +3070,10 @@ bool UCustomizableInstancePrivateData::BuildSkeletalMeshSkeletonData(const TShar
 				BoneIndex = BoneToFinalBoneIndexMap[MutSkeletonData.BoneNames[BoneIndex]];
 			}
 
-			TArray<uint16> UniqueActiveBones;
-			for (const uint16 BoneIndex : Component.ActiveBones)
+			for (uint16& BoneIndex : Component.ActiveBones)
 			{
-				UniqueActiveBones.AddUnique(BoneToFinalBoneIndexMap[MutSkeletonData.BoneNames[BoneIndex]]);
+				BoneIndex = BoneToFinalBoneIndexMap[MutSkeletonData.BoneNames[BoneIndex]];
 			}
-
-			Exchange(Component.ActiveBones, UniqueActiveBones);
 			Component.ActiveBones.Sort();
 		}
 	}
diff --git a/Engine/Plugins/Experimental/Mutable/Source/CustomizableObject/Private/MuCO/CustomizableObjectSystem.cpp b/Engine/Plugins/Experimental/Mutable/Source/CustomizableObject/Private/MuCO/CustomizableObjectSystem.cpp
--- a/Engine/Plugins/Experimental/Mutable/Source/CustomizableObject/Private/MuCO/CustomizableObjectSystem.cpp
+++ b/Engine/Plugins/Experimental/Mutable/Source/CustomizableObject/Private/MuCO/CustomizableObjectSystem.cpp
@@ -1476,14 +1476,9 @@ namespace impl
 		// Add SkeletonData for each component
 		OperationData->InstanceUpdateData.Skeletons.AddDefaulted(ComponentCount);
 
-		// Helpers
-		TMap<FName, int16> BoneNameToIndex;
-		TArray<const char*> ParentChain;
-		TArray<bool> UsedBones;
-
 		for (int32 ComponentIndex = 0; ComponentIndex < ComponentCount; ++ComponentIndex)
 		{
-			FInstanceUpdateData::FComponent& MinLODComponent = OperationData->InstanceUpdateData.Components[MinLOD.FirstComponent+ComponentIndex];
+			FInstanceUpdateData::FComponent& MinLODComponent = OperationData->InstanceUpdateData.Components[MinLOD.FirstComponent + ComponentIndex];
 
 			// Set the ComponentIndex
 			OperationData->InstanceUpdateData.Skeletons[ComponentIndex].ComponentIndex = MinLODComponent.Id;
@@ -1492,15 +1487,12 @@ namespace impl
 			TArray<uint32>& SkeletonIds = OperationData->InstanceUpdateData.Skeletons[ComponentIndex].SkeletonIds;
 			TArray<FName>& BoneNames = OperationData->InstanceUpdateData.Skeletons[ComponentIndex].BoneNames;
 			TMap<FName, FMatrix44f>& BoneMatricesWithScale = OperationData->InstanceUpdateData.Skeletons[ComponentIndex].BoneMatricesWithScale;
-			
+
 			// Use first valid LOD bone count as a potential total number of bones, used for pre-allocating data arrays
 			if (MinLODComponent.Mesh && MinLODComponent.Mesh->GetSkeleton())
 			{
 				const int32 TotalPossibleBones = MinLODComponent.Mesh->GetSkeleton()->GetBoneCount();
 
-				// Helpers
-				BoneNameToIndex.Empty(TotalPossibleBones);
-
 				// Out Data
 				BoneNames.Reserve(TotalPossibleBones);
 				BoneMatricesWithScale.Reserve(TotalPossibleBones);
@@ -1526,223 +1518,36 @@ namespace impl
 					SkeletonIds.AddUnique(Mesh->GetSkeletonID(SkeletonIndex));
 				}
 
-				const mu::Skeleton* Skeleton = Mesh->GetSkeleton().get();
-				// We don't need to check(Skeleton != nullptr) because we would've continued if (!Mesh->GetSkeleton()) above.
-				const int32 BoneCount = Skeleton->GetBoneCount();
-
-				{
-					MUTABLE_CPUPROFILER_SCOPE(PrepareSkeletonData_UsedBones);
-
-					const mu::FMeshBufferSet& MutableMeshVertexBuffers = Mesh->GetVertexBuffers();
-
-					const int32 NumVerticesLODModel = Mesh->GetVertexCount();
-					const int32 SurfaceCount = Mesh->GetSurfaceCount();
 
-					mu::MESH_BUFFER_FORMAT BoneIndexFormat = mu::MBF_NONE;
-					int32 BoneIndexComponents = 0;
-					int32 BoneIndexOffset = 0;
-					int32 BoneIndexBuffer = -1;
-					int32 BoneIndexChannel = -1;
-					MutableMeshVertexBuffers.FindChannel(mu::MBS_BONEINDICES, 0, &BoneIndexBuffer, &BoneIndexChannel);
-					if (BoneIndexBuffer >= 0 || BoneIndexChannel >= 0)
-					{
-						MutableMeshVertexBuffers.GetChannel(BoneIndexBuffer, BoneIndexChannel,
-							nullptr, nullptr, &BoneIndexFormat, &BoneIndexComponents, &BoneIndexOffset);
-					}
+				const int32 MaxBoneIndex = Mesh->GetBonePoseCount();
 
-					const int32 ElemSize = MutableMeshVertexBuffers.GetElementSize(BoneIndexBuffer);
+				// Add active bones and new names to the BoneNames array 
+				CurrentLODComponent.BoneMap.Reserve(MaxBoneIndex);
+				CurrentLODComponent.ActiveBones.Reserve(MaxBoneIndex);
 
-					const uint8_t* BufferStart = MutableMeshVertexBuffers.GetBufferData(BoneIndexBuffer) + BoneIndexOffset;
-
-					UsedBones.Empty(BoneCount);
-					UsedBones.AddDefaulted(BoneCount);
+				for (int32 BonePoseIndex = 0; BonePoseIndex < MaxBoneIndex; ++BonePoseIndex)
+				{
+					const FName BoneName = Mesh->GetBonePoseName(BonePoseIndex);
+					check(BoneName != NAME_None);
 
-					for (int32 Surface = 0; Surface < SurfaceCount; ++Surface)
+					int32 BoneIndex = BoneNames.Find(BoneName);
+					if (BoneIndex == INDEX_NONE)
 					{
-						int32 IndexStart, IndexCount, VertexStart, VertexCount;
-						Mesh->GetSurface(Surface, &VertexStart, &VertexCount, &IndexStart, &IndexCount);
-
-						if (VertexCount == 0 || IndexCount == 0)
-						{
-							continue;
-						}
+						BoneIndex = BoneNames.Add(BoneName);
 
-						const uint8_t* VertexBoneIndexPtr = BufferStart + VertexStart * ElemSize;
-
-						if (BoneIndexFormat == mu::MBF_UINT8)
-						{
-							for (int32 VertexIndex = 0; VertexIndex < VertexCount; ++VertexIndex)
-							{
-								for (int32 Index = 0; Index < BoneIndexComponents; ++Index)
-								{
-									size_t SectionBoneIndex = VertexBoneIndexPtr[Index];
-									UsedBones[SectionBoneIndex] = true;
-								}
-								VertexBoneIndexPtr += ElemSize;
-							}
-						}
-						else if (BoneIndexFormat == mu::MBF_UINT16)
-						{
-							for (int32 VertexIndex = 0; VertexIndex < VertexCount; ++VertexIndex)
-							{
-								for (int32 Index = 0; Index < BoneIndexComponents; ++Index)
-								{
-									size_t SectionBoneIndex = ((uint16*)VertexBoneIndexPtr)[Index];
-									UsedBones[SectionBoneIndex] = true;
-								}
-								VertexBoneIndexPtr += ElemSize;
-							}
-						}
-						else if (BoneIndexFormat == mu::MBF_UINT32)
-						{
-							for (int32 VertexIndex = 0; VertexIndex < VertexCount; ++VertexIndex)
-							{
-								for (int32 Index = 0; Index < BoneIndexComponents; ++Index)
-								{
-									size_t SectionBoneIndex = ((uint32_t*)VertexBoneIndexPtr)[Index];
-									UsedBones[SectionBoneIndex] = true;
-								}
-								VertexBoneIndexPtr += ElemSize;
-							}
-						}
-						else
-						{
-							// Unsupported bone index format in generated mutable mesh
-							check(false);
-						}
+						FTransform3f Transform;
+						Mesh->GetBoneTransform(BonePoseIndex, Transform);
+						BoneMatricesWithScale.Emplace(BoneName, Transform.Inverse().ToMatrixWithScale());
 					}
-				}
-
-				{
-					MUTABLE_CPUPROFILER_SCOPE(PrepareSkeletonData_ActiveBones);
-
-					TArray<uint16>& ComponentBoneMap = CurrentLODComponent.BoneMap;
-					TArray<uint16>& ComponentActiveBones = CurrentLODComponent.ActiveBones;
-
-					ComponentBoneMap.Reserve(BoneCount);
-					ComponentActiveBones.Reserve(BoneCount);
-
-					ParentChain.SetNumZeroed(BoneCount);
-
-					// Add new bones and their poses to the final hierarchy of active bones
-					auto AddBone = [&](const char* StrName) {
 
-						const FName BoneName = StrName;
-						const int16 FinalBoneIndex = BoneNames.Num();
-
-						// Add bone
-						BoneNames.Add(BoneName);
-
-						BoneNameToIndex.Add(BoneName, FinalBoneIndex);
-						ComponentActiveBones.Add(FinalBoneIndex);
-
-						// Add bone pose
-						const int32 BonePoseIndex = Mesh->FindBonePose(StrName);
-						if (BonePoseIndex != INDEX_NONE)
-						{
-							FTransform3f Transform;
-							Mesh->GetBoneTransform(BonePoseIndex, Transform);
-							BoneMatricesWithScale.Emplace(BoneName, Transform.Inverse().ToMatrixWithScale());
-						}
-
-						return FinalBoneIndex;
-					};
-
-					for (int32 BoneIndex = 0; BoneIndex < BoneCount; ++BoneIndex)
+					if (EnumHasAnyFlags(Mesh->GetBoneUsageFlags(BonePoseIndex), mu::EBoneUsageFlags::Skinning))
 					{
-						const char* BoneName = Skeleton->GetBoneName(BoneIndex);
-
-						int16* FinalBoneIndexPtr = BoneNameToIndex.Find(BoneName);
-						if (!UsedBones[BoneIndex])
-						{
-							if (!FinalBoneIndexPtr)
-							{
-								// Add bone even if it doesn't have weights
-								AddBone(BoneName);
-							}
-
-							// Add root as a placeholder
-							ComponentBoneMap.Add(0);
-							continue;
-						}
-
-						if (!FinalBoneIndexPtr) // New bone!
-						{
-							// Ensure parent chain is valid
-							int16 FinalParentIndex = INDEX_NONE;
-
-							const int16 MutParentIndex = Skeleton->GetBoneParent(BoneIndex);
-							if (MutParentIndex != INDEX_NONE) // Bone has a parent, ensure the parent exists
-							{
-								// Ensure parent chain until root exists
-								uint16 ParentChainCount = 0;
-
-								// Find parents to add
-								int16 NextMutParentIndex = MutParentIndex;
-								while (NextMutParentIndex != INDEX_NONE)
-								{
-									const char* ParentName = Skeleton->GetBoneName(NextMutParentIndex);
-									if (int16* Found = BoneNameToIndex.Find(ParentName))
-									{
-										FinalParentIndex = *Found;
-										break;
-									}
-
-									ParentChain[ParentChainCount] = ParentName;
-									++ParentChainCount;
-
-									NextMutParentIndex = Skeleton->GetBoneParent(NextMutParentIndex);
-								}
-
-								// Add parent bones to the list and to the active bones array
-								for (int16 ParentChainIndex = ParentChainCount - 1; ParentChainIndex >= 0; --ParentChainIndex)
-								{
-									// Add the parent
-									AddBone(ParentChain[ParentChainIndex]);
-								}
-							}
-
-							// Add the used bone to the hierarchy
-							const int16 NewBoneIndex = AddBone(BoneName);
-							ComponentBoneMap.Add(NewBoneIndex);
-						}
-						else
-						{
-							int16 FinalBoneIndex = *FinalBoneIndexPtr;
-							ComponentBoneMap.Add(FinalBoneIndex);
-
-							if (ComponentActiveBones.Find(FinalBoneIndex) != INDEX_NONE)
-							{
-								continue;
-							}
-
-							int16 MutParentIndex = Skeleton->GetBoneParent(BoneIndex);
-							while (MutParentIndex != INDEX_NONE)
-							{
-								FName ParentName = Skeleton->GetBoneName(MutParentIndex);
-								const int16* FinalParentIndex = BoneNameToIndex.Find(ParentName);
-								check(FinalParentIndex);
-
-								if (ComponentActiveBones.Find(*FinalParentIndex) != INDEX_NONE)
-								{
-									break;
-								}
-
-								ComponentActiveBones.Add(*FinalParentIndex);
-
-								MutParentIndex = Skeleton->GetBoneParent(MutParentIndex);
-							}
-
-							ComponentActiveBones.Add(FinalBoneIndex);
-						}
+						CurrentLODComponent.BoneMap.Add(BoneIndex);
 					}
 
-					ComponentBoneMap.Shrink();
-					ComponentActiveBones.Shrink();
-				}	
+					CurrentLODComponent.ActiveBones.Add(BoneIndex);
+				}
 			}
-
-			BoneNames.Shrink();
 		}
 	}
 
@@ -2487,7 +2292,6 @@ bool UCustomizableObjectSystem::Tick(float DeltaTime)
 			Private->CurrentMutableOperation = nullptr;
 		}
 	}
-	
 	// Advance the current operation
 	if (Private->CurrentMutableOperation)
 	{
diff --git a/Engine/Plugins/Experimental/Mutable/Source/MutableRuntime/Private/MuR/CodeRunner.cpp b/Engine/Plugins/Experimental/Mutable/Source/MutableRuntime/Private/MuR/CodeRunner.cpp
--- a/Engine/Plugins/Experimental/Mutable/Source/MutableRuntime/Private/MuR/CodeRunner.cpp
+++ b/Engine/Plugins/Experimental/Mutable/Source/MutableRuntime/Private/MuR/CodeRunner.cpp
@@ -61,6 +61,7 @@
 #include "MuR/OpMeshReshape.h"
 #include "MuR/OpMeshSubtract.h"
 #include "MuR/OpMeshTransform.h"
+#include "MuR/OpMeshOptimizeSkinning.h"
 #include "MuR/Operations.h"
 #include "MuR/Parameters.h"
 #include "MuR/ParametersPrivate.h"
@@ -2366,6 +2367,40 @@ namespace mu
             break;
         }
 
+		case OP_TYPE::ME_OPTIMIZESKINNING:
+		{
+			auto args = pModel->GetPrivate()->m_program.GetOpArgs<OP::MeshOptimizeSkinningArgs>(item.At);
+			switch (item.Stage)
+			{
+			case 0:
+				if (args.source)
+				{
+					AddOp(FScheduledOp(item.At, item, 1),
+						FScheduledOp(args.source, item));
+				}
+				else
+				{
+					GetMemory().SetMesh(item, nullptr);
+				}
+				break;
+
+			case 1:
+			{
+				Ptr<const Mesh> pSource = GetMemory().GetMesh(FCacheAddress(args.source, item));
+
+				MeshPtr pResult = MeshOptimizeSkinning(pSource.get());
+
+				GetMemory().SetMesh(item, pResult ? pResult : pSource->Clone());
+				break;
+			}
+
+			default:
+				check(false);
+			}
+
+			break;
+		}
+
         default:
             if (type!=OP_TYPE::NONE)
             {
diff --git a/Engine/Plugins/Experimental/Mutable/Source/MutableRuntime/Private/MuR/Mesh.cpp b/Engine/Plugins/Experimental/Mutable/Source/MutableRuntime/Private/MuR/Mesh.cpp
--- a/Engine/Plugins/Experimental/Mutable/Source/MutableRuntime/Private/MuR/Mesh.cpp
+++ b/Engine/Plugins/Experimental/Mutable/Source/MutableRuntime/Private/MuR/Mesh.cpp
@@ -522,6 +522,12 @@ void mu::Mesh::GetBoneTransform(int32 BoneIndex, FTransform3f& Transform) const
 	Transform = BoneIndex > INDEX_NONE ? BonePoses[BoneIndex].BoneTransform : FTransform3f::Identity;
 }
 
+EBoneUsageFlags Mesh::GetBoneUsageFlags(int32 BoneIndex) const
+{
+	check(BoneIndex >= 0 && BoneIndex < BonePoses.Num());
+	return BoneIndex > INDEX_NONE ? BonePoses[BoneIndex].BoneUsageFlags : EBoneUsageFlags::None;
+}
+
 
 //---------------------------------------------------------------------------------------------
 int32 Mesh::GetSkeletonIDsCount() const
diff --git a/Engine/Plugins/Experimental/Mutable/Source/MutableRuntime/Private/MuR/OpMeshOptimizeSkinning.h b/Engine/Plugins/Experimental/Mutable/Source/MutableRuntime/Private/MuR/OpMeshOptimizeSkinning.h
new file mode 100644
--- /dev/null
+++ b/Engine/Plugins/Experimental/Mutable/Source/MutableRuntime/Private/MuR/OpMeshOptimizeSkinning.h
@@ -0,0 +1,183 @@
+// Copyright Epic Games, Inc. All Rights Reserved.
+
+#pragma once
+
+#include "MuR/MeshPrivate.h"
+#include "MuR/ConvertData.h"
+#include "MuR/MutableTrace.h"
+#include "MuR/Platform.h"
+
+
+namespace mu
+{
+
+	//---------------------------------------------------------------------------------------------
+	//!
+	//---------------------------------------------------------------------------------------------
+	inline MeshPtr MeshOptimizeSkinning(const Mesh* InMesh)
+	{
+		MUTABLE_CPUPROFILER_SCOPE(MeshOptimizeSkinning);
+
+		if (!InMesh)
+		{
+			return nullptr;
+		}
+
+		// Prepare an array of BoneIndex to SkinnedBoneIndex to remap indices.
+		TArray<uint16> RemapedBoneIndices;
+
+		mu::SkeletonPtrConst Skeleton = InMesh->GetSkeleton();
+		RemapedBoneIndices.AddZeroed(Skeleton->GetBoneCount());
+
+		int32 SkinnedBonesCount = 0;
+
+		const int32 BonePoseCount = InMesh->GetBonePoseCount();
+		for (int32 BonePoseIndex = 0; BonePoseIndex < BonePoseCount; ++BonePoseIndex)
+		{
+			if (EnumHasAnyFlags(InMesh->BonePoses[BonePoseIndex].BoneUsageFlags, EBoneUsageFlags::Skinning))
+			{
+				const int32 SkeletonBoneIndex = Skeleton->FindBone(InMesh->BonePoses[BonePoseIndex].BoneName.c_str());
+				RemapedBoneIndices[SkeletonBoneIndex] = SkinnedBonesCount;
+				++SkinnedBonesCount;
+			}
+		}
+
+		// Find the desired bone index format
+		const MESH_BUFFER_FORMAT DesiredBoneIndexFormat = SkinnedBonesCount < MAX_uint8 ? mu::MBF_UINT8 : mu::MBF_UINT16;
+
+		const mu::FMeshBufferSet& SourceVertexBuffers = InMesh->GetVertexBuffers();
+
+		const int32 VertexBuffersCount = SourceVertexBuffers.GetBufferCount();
+		const int32 ElementCount = SourceVertexBuffers.GetElementCount();
+
+		// Clone mesh without VertexBuffers, they will be copied manually.
+		MeshPtr Result = InMesh->Clone(~(EMeshCloneFlags::WithVertexBuffers));
+
+		mu::FMeshBufferSet& VertexBuffers = Result->GetVertexBuffers();
+
+		VertexBuffers.m_buffers.Reserve(VertexBuffersCount);
+		VertexBuffers.SetElementCount(ElementCount);
+
+		for (int32 BufferIndex = 0; BufferIndex < VertexBuffersCount; ++BufferIndex)
+		{
+			const MESH_BUFFER& SourceBuffer = SourceVertexBuffers.m_buffers[BufferIndex];
+			const int32 ChannelsCount = SourceVertexBuffers.GetBufferChannelCount(BufferIndex);
+
+			int32 BoneIndexChannelIndex = INDEX_NONE;
+
+			int32 ElementSize = 0;
+			TArray<MESH_BUFFER_SEMANTIC> Semantics;
+			TArray<int32> SemanticIndices;
+			TArray<MESH_BUFFER_FORMAT> Formats;
+			TArray<int32> Components;
+			TArray<int32> Offsets;
+
+			MESH_BUFFER_FORMAT SourceBoneIndexFormat = mu::MBF_NONE;
+
+			// Offset accumulator
+			int32 AuxOffset = 0;
+
+			// Copy and fix channel details
+			for (int32 ChannelIndex = 0; ChannelIndex < ChannelsCount; ++ChannelIndex)
+			{
+				const MESH_BUFFER_CHANNEL& Channel = SourceBuffer.m_channels[ChannelIndex];
+				MESH_BUFFER_FORMAT Format = Channel.m_format;
+
+				if (Channel.m_semantic == MBS_BONEINDICES)
+				{
+					SourceBoneIndexFormat = Channel.m_format;
+					Format = DesiredBoneIndexFormat;
+					BoneIndexChannelIndex = ChannelIndex;
+				}
+
+				Semantics.Add(Channel.m_semantic);
+				SemanticIndices.Add(Channel.m_semanticIndex);
+				Formats.Add(Format);
+				Components.Add(Channel.m_componentCount);
+				Offsets.Add(AuxOffset);
+
+				const int32 FormatSize = GetMeshFormatData(Format).m_size;
+
+				ElementSize += FormatSize;
+				AuxOffset += FormatSize * Channel.m_componentCount;
+			}
+
+			// Copy buffers
+			if (BoneIndexChannelIndex != INDEX_NONE)
+			{
+				// This buffer has a BoneIndex channel that needs to be manually fixed. 
+				VertexBuffers.SetBufferCount(BufferIndex + 1);
+				VertexBuffers.SetBuffer(BufferIndex, ElementSize, ChannelsCount, Semantics.GetData(), SemanticIndices.GetData(), Formats.GetData(), Components.GetData(), Offsets.GetData());
+				uint8* Data = VertexBuffers.GetBufferData(BufferIndex);
+
+				const uint8* SourceData = SourceBuffer.m_data.GetData();
+				const int32 SourceBoneIndexSize = GetMeshFormatData(SourceBoneIndexFormat).m_size;
+
+				const int32 BoneIndexSize = GetMeshFormatData(DesiredBoneIndexFormat).m_size;
+				const int32 BoneIndexComponentCount = Components[BoneIndexChannelIndex];
+				const int32 BoneIndexChannelOffset = Offsets[BoneIndexChannelIndex];
+
+				const int32 TailSize = AuxOffset - (BoneIndexChannelOffset + BoneIndexComponentCount * BoneIndexSize);
+
+				for (int32 VertexIndex = 0; VertexIndex < ElementCount; ++VertexIndex)
+				{
+					FMemory::Memcpy(Data, SourceData, BoneIndexChannelOffset);
+					Data += BoneIndexChannelOffset;
+					SourceData += BoneIndexChannelOffset;
+
+					for (int32 ComponentIndex = 0; ComponentIndex < BoneIndexComponentCount; ++ComponentIndex)
+					{
+						switch (DesiredBoneIndexFormat)
+						{
+						case MESH_BUFFER_FORMAT::MBF_UINT8:
+						{
+							switch (SourceBoneIndexFormat)
+							{
+							case MESH_BUFFER_FORMAT::MBF_UINT8:
+							{
+								*Data = (uint8)RemapedBoneIndices[*SourceData];
+								break;
+							}
+							case MESH_BUFFER_FORMAT::MBF_UINT16:
+							{
+								*Data = (uint8)RemapedBoneIndices[*((uint16*)SourceData)];
+								break;
+							}
+							default:
+								unimplemented();
+								break;
+							}
+							break;
+						}
+						case MESH_BUFFER_FORMAT::MBF_UINT16:
+						{
+							check(SourceBoneIndexFormat == MESH_BUFFER_FORMAT::MBF_UINT16);
+							*((uint16*)Data) = RemapedBoneIndices[*((uint16*)SourceData)];
+							break;
+						}
+						default:
+							unimplemented();
+							break;
+						}
+
+						Data += BoneIndexSize;
+						SourceData += SourceBoneIndexSize;
+					}
+
+
+					FMemory::Memcpy(Data, SourceData, TailSize);
+					Data += TailSize;
+					SourceData += TailSize;
+				}
+			}
+			else
+			{
+				// Add buffers that don't require a fix up
+				VertexBuffers.AddBuffer(SourceVertexBuffers, BufferIndex);
+			}
+
+		}
+
+		return Result;
+	}
+}
diff --git a/Engine/Plugins/Experimental/Mutable/Source/MutableRuntime/Private/MuR/Operations.cpp b/Engine/Plugins/Experimental/Mutable/Source/MutableRuntime/Private/MuR/Operations.cpp
--- a/Engine/Plugins/Experimental/Mutable/Source/MutableRuntime/Private/MuR/Operations.cpp
+++ b/Engine/Plugins/Experimental/Mutable/Source/MutableRuntime/Private/MuR/Operations.cpp
@@ -126,6 +126,7 @@ namespace mu
 		{ DT_MESH,		true,	{ 0, 0, 0, 0, 0, 0, 0, 0, 0, 0, 0, 0, 0, 0, 0, 0, 0, 0, 0, 0, 0, 0, 0, 0, 0, 0, 0, 0, 0 }		},	// ME_APPLYSHAPE
 		{ DT_MESH,		true,	{ 0, 0, 0, 0, 0, 0, 0, 0, 0, 0, 0, 0, 0, 0, 0, 0, 0, 0, 0, 0, 0, 0, 0, 0, 0, 0, 0, 0, 0 }		},	// ME_CLIPDEFORM
 		{ DT_MESH,		true,	{ 0, 0, 0, 0, 0, 0, 0, 0, 0, 0, 0, 0, 0, 0, 0, 0, 0, 0, 0, 0, 0, 0, 0, 0, 0, 0, 0, 0, 0 }		},	// ME_MORPHRESHAPE
+		{ DT_MESH,		true,	{ 0, 0, 0, 0, 0, 0, 0, 0, 0, 0, 0, 0, 0, 0, 0, 0, 0, 0, 0, 0, 0, 0, 0, 0, 0, 0, 0, 0, 0 }		},	// ME_OPTIMIZESKINNING
 
 		{ DT_INSTANCE,	false,	{ 0, 0, 0, 0, 0, 0, 0, 0, 0, 0, 0, 0, 0, 0, 0, 0, 0, 0, 0, 0, 0, 0, 0, 0, 0, 0, 0, 0, 0 }		},	// IN_ADDMESH
 		{ DT_INSTANCE,	false,	{ 0, 0, 0, 0, 0, 0, 0, 0, 0, 0, 0, 0, 0, 0, 0, 0, 0, 0, 0, 0, 0, 0, 0, 0, 0, 0, 0, 0, 0 }		},	// IN_ADDIMAGE
@@ -1091,6 +1092,13 @@ namespace mu
 			break;
 		}
 
+		case OP_TYPE::ME_OPTIMIZESKINNING:
+		{
+			OP::MeshOptimizeSkinningArgs args = program.GetOpArgs<OP::MeshOptimizeSkinningArgs>(at);
+			f(args.source);
+			break;
+		}
+
         //-------------------------------------------------------------------------------------
         case OP_TYPE::IN_ADDMESH:
         case OP_TYPE::IN_ADDIMAGE:
diff --git a/Engine/Plugins/Experimental/Mutable/Source/MutableRuntime/Private/MuR/Operations.h b/Engine/Plugins/Experimental/Mutable/Source/MutableRuntime/Private/MuR/Operations.h
--- a/Engine/Plugins/Experimental/Mutable/Source/MutableRuntime/Private/MuR/Operations.h
+++ b/Engine/Plugins/Experimental/Mutable/Source/MutableRuntime/Private/MuR/Operations.h
@@ -313,6 +313,9 @@ namespace mu
         //! Mesh morph with Skeleton Reshape based on the morphed mesh.
         ME_MORPHRESHAPE,
 
+		//! Optimize skinning before adding a mesh to the component
+		ME_OPTIMIZESKINNING,
+
         //-----------------------------------------------------------------------------------------
         // Instance operations
         //-----------------------------------------------------------------------------------------
@@ -958,6 +961,11 @@ namespace mu
 			float clipWeightThreshold = 0.9f;
 		};
 
+		struct MeshOptimizeSkinningArgs
+		{
+			ADDRESS source;
+		};
+
 
         //-------------------------------------------------------------------------------------
         struct InstanceAddArgs
diff --git a/Engine/Plugins/Experimental/Mutable/Source/MutableRuntime/Public/MuR/Mesh.h b/Engine/Plugins/Experimental/Mutable/Source/MutableRuntime/Public/MuR/Mesh.h
--- a/Engine/Plugins/Experimental/Mutable/Source/MutableRuntime/Public/MuR/Mesh.h
+++ b/Engine/Plugins/Experimental/Mutable/Source/MutableRuntime/Public/MuR/Mesh.h
@@ -289,6 +289,9 @@ namespace mu
 		//! Return a matrix stored per bone. It is a set of 16-float values.
 		void GetBoneTransform(int32 BoneIndex, FTransform3f& Transform) const;
 
+		//! 
+		EBoneUsageFlags GetBoneUsageFlags(int32 BoneIndex) const;
+
 		//!
 		int32 GetSkeletonIDsCount() const;
 
diff --git a/Engine/Plugins/Experimental/Mutable/Source/MutableRuntime/Public/MuR/Model.h b/Engine/Plugins/Experimental/Mutable/Source/MutableRuntime/Public/MuR/Model.h
--- a/Engine/Plugins/Experimental/Mutable/Source/MutableRuntime/Public/MuR/Model.h
+++ b/Engine/Plugins/Experimental/Mutable/Source/MutableRuntime/Public/MuR/Model.h
@@ -13,7 +13,7 @@
 //! data structures. Compiled models are not necessarily compatible when the runtime is updated,
 //! so this version number can be used externally to verify this. It is not used internally, and
 //! serializing models from different versions than this runtime will probably result in a crash.
-#define MUTABLE_COMPILED_MODEL_CODE_VERSION		uint32( 52 )
+#define MUTABLE_COMPILED_MODEL_CODE_VERSION		uint32( 53 )
 #define MUTABLE_PARAMETERS_VERSION              uint32( 1 )
 
 
diff --git a/Engine/Plugins/Experimental/Mutable/Source/MutableTools/Private/MuT/ASTOpMeshOptimizeSkinning.cpp b/Engine/Plugins/Experimental/Mutable/Source/MutableTools/Private/MuT/ASTOpMeshOptimizeSkinning.cpp
new file mode 100644
--- /dev/null
+++ b/Engine/Plugins/Experimental/Mutable/Source/MutableTools/Private/MuT/ASTOpMeshOptimizeSkinning.cpp
@@ -0,0 +1,88 @@
+// Copyright Epic Games, Inc. All Rights Reserved.
+
+#include "MuT/ASTOpMeshOptimizeSkinning.h"
+
+#include "HAL/PlatformMath.h"
+#include "MuR/ModelPrivate.h"
+#include "MuR/RefCounted.h"
+#include "MuR/Types.h"
+#include "MuT/StreamsPrivate.h"
+
+#include <utility>
+
+
+namespace mu
+{
+
+
+	//-------------------------------------------------------------------------------------------------
+	//-------------------------------------------------------------------------------------------------
+	//-------------------------------------------------------------------------------------------------
+	ASTOpMeshOptimizeSkinning::ASTOpMeshOptimizeSkinning()
+		: source(this)
+	{
+	}
+
+
+	//-------------------------------------------------------------------------------------------------
+	ASTOpMeshOptimizeSkinning::~ASTOpMeshOptimizeSkinning()
+	{
+		// Explicit call needed to avoid recursive destruction
+		ASTOp::RemoveChildren();
+	}
+
+
+	//-------------------------------------------------------------------------------------------------
+	bool ASTOpMeshOptimizeSkinning::IsEqual(const ASTOp& otherUntyped) const
+	{
+		if (auto other = dynamic_cast<const ASTOpMeshOptimizeSkinning*>(&otherUntyped))
+		{
+			return source == other->source;
+		}
+		return false;
+	}
+
+
+	//-------------------------------------------------------------------------------------------------
+	uint64 ASTOpMeshOptimizeSkinning::Hash() const
+	{
+		uint64 res = std::hash<void*>()(source.child().get());
+		return res;
+	}
+
+
+	//-------------------------------------------------------------------------------------------------
+	mu::Ptr<ASTOp> ASTOpMeshOptimizeSkinning::Clone(MapChildFuncRef mapChild) const
+	{
+		Ptr<ASTOpMeshOptimizeSkinning> n = new ASTOpMeshOptimizeSkinning();
+		n->source = mapChild(source.child());
+		return n;
+	}
+
+
+	//-------------------------------------------------------------------------------------------------
+	void ASTOpMeshOptimizeSkinning::ForEachChild(const TFunctionRef<void(ASTChild&)> f)
+	{
+		f(source);
+	}
+
+
+	//-------------------------------------------------------------------------------------------------
+	void ASTOpMeshOptimizeSkinning::Link(FProgram& program, const FLinkerOptions*)
+	{
+		// Already linked?
+		if (!linkedAddress)
+		{
+			OP::MeshOptimizeSkinningArgs args;
+			FMemory::Memset(&args, 0, sizeof(args));
+
+			if (source) args.source = source->linkedAddress;
+
+			linkedAddress = (OP::ADDRESS)program.m_opAddress.Num();
+			program.m_opAddress.Add((uint32_t)program.m_byteCode.Num());
+			AppendCode(program.m_byteCode, OP_TYPE::ME_OPTIMIZESKINNING);
+			AppendCode(program.m_byteCode, args);
+		}
+	}
+
+}
diff --git a/Engine/Plugins/Experimental/Mutable/Source/MutableTools/Private/MuT/ASTOpMeshOptimizeSkinning.h b/Engine/Plugins/Experimental/Mutable/Source/MutableTools/Private/MuT/ASTOpMeshOptimizeSkinning.h
new file mode 100644
--- /dev/null
+++ b/Engine/Plugins/Experimental/Mutable/Source/MutableTools/Private/MuT/ASTOpMeshOptimizeSkinning.h
@@ -0,0 +1,40 @@
+// Copyright Epic Games, Inc. All Rights Reserved.
+
+#pragma once
+
+#include "HAL/Platform.h"
+#include "MuR/Operations.h"
+#include "MuR/Ptr.h"
+#include "MuT/AST.h"
+
+
+namespace mu
+{
+	struct FProgram;
+
+	//---------------------------------------------------------------------------------------------
+	class ASTOpMeshOptimizeSkinning : public ASTOp
+	{
+	public:
+
+		//! Mesh that will have the skinning optimized.
+		ASTChild source;
+
+	public:
+
+		ASTOpMeshOptimizeSkinning();
+		ASTOpMeshOptimizeSkinning(const ASTOpMeshOptimizeSkinning&) = delete;
+		~ASTOpMeshOptimizeSkinning() override;
+
+		// ASTOp interface
+		OP_TYPE GetOpType() const override { return OP_TYPE::ME_OPTIMIZESKINNING; }
+		uint64 Hash() const override;
+		void ForEachChild(const TFunctionRef<void(ASTChild&)>) override;
+		bool IsEqual(const ASTOp& otherUntyped) const override;
+		Ptr<ASTOp> Clone(MapChildFuncRef mapChild) const override;
+		void Link(FProgram& program, const FLinkerOptions*) override;
+	};
+
+
+}
+
diff --git a/Engine/Plugins/Experimental/Mutable/Source/MutableTools/Private/MuT/CodeGenerator.cpp b/Engine/Plugins/Experimental/Mutable/Source/MutableTools/Private/MuT/CodeGenerator.cpp
--- a/Engine/Plugins/Experimental/Mutable/Source/MutableTools/Private/MuT/CodeGenerator.cpp
+++ b/Engine/Plugins/Experimental/Mutable/Source/MutableTools/Private/MuT/CodeGenerator.cpp
@@ -30,6 +30,7 @@
 #include "MuT/ASTOpMeshRemoveMask.h"
 #include "MuT/ASTOpMeshDifference.h"
 #include "MuT/ASTOpMeshMorph.h"
+#include "MuT/ASTOpMeshOptimizeSkinning.h"
 #include "MuT/ASTOpParameter.h"
 #include "MuT/ASTOpLayoutRemoveBlocks.h"
 #include "MuT/ASTOpLayoutFromMesh.h"
@@ -1788,6 +1789,13 @@ namespace mu
 			}
 		}
 
+		// Add op to optimize the skinning of the resulting mesh
+		{
+			Ptr<ASTOpMeshOptimizeSkinning> mop = new ASTOpMeshOptimizeSkinning();
+			mop->source = lastMeshOp;
+			lastMeshOp = mop;
+		}
+
         // Add the component mesh
         {
             Ptr<ASTOpInstanceAdd> iop = new ASTOpInstanceAdd();
diff --git a/Engine/Plugins/Experimental/Mutable/Source/MutableTools/Private/MuT/ErrorLog.cpp b/Engine/Plugins/Experimental/Mutable/Source/MutableTools/Private/MuT/ErrorLog.cpp
--- a/Engine/Plugins/Experimental/Mutable/Source/MutableTools/Private/MuT/ErrorLog.cpp
+++ b/Engine/Plugins/Experimental/Mutable/Source/MutableTools/Private/MuT/ErrorLog.cpp
@@ -301,6 +301,7 @@ namespace mu
 		TEXT("ME_APPLYSHAPE	   "),
 		TEXT("ME_CLIPDEFORM	   "),
 		TEXT("ME_MORPHRESHAPE  "),
+		TEXT("ME_OPTIMIZESKIN  "),
 
 		TEXT("IN_ADDMESH       "),
 		TEXT("IN_ADDIMAGE      "),
@@ -605,6 +606,14 @@ namespace mu
 			}
 			break;
         }
+		case OP_TYPE::ME_OPTIMIZESKINNING:
+		{
+			auto args = program.GetOpArgs<OP::MeshOptimizeSkinningArgs>(at);
+			done += mutable_snprintf(temp, 1024,
+							"source : %d",
+							args.source);
+			break;
+		}
 
 		//-----------------------------------------------------------------------------------------
 		default:
