[mutable] static analysis fix.
#rnx

diff --git a/Engine/Plugins/Experimental/Mutable/Source/MutableRuntime/Private/MuR/BlockCompression/Miro/Miro.cpp b/Engine/Plugins/Experimental/Mutable/Source/MutableRuntime/Private/MuR/BlockCompression/Miro/Miro.cpp
--- a/Engine/Plugins/Experimental/Mutable/Source/MutableRuntime/Private/MuR/BlockCompression/Miro/Miro.cpp
+++ b/Engine/Plugins/Experimental/Mutable/Source/MutableRuntime/Private/MuR/BlockCompression/Miro/Miro.cpp
@@ -6076,7 +6076,7 @@ namespace impl
 						}
 						result = DECOMPRESS_RESULT_ERROR;
 					}
-					else if (colorEndpointNdx >= 0 && colorEndpointNdx < 4)
+					else 
 					{
 						for (int channelNdx = 0; channelNdx < 4; channelNdx++)
 						{
