[mutable] Add support to skip mips when requesting an external image.

#rnx

diff --git a/Engine/Plugins/Experimental/Mutable/Source/CustomizableObject/Private/MuCO/UnrealMutableImageProvider.cpp b/Engine/Plugins/Experimental/Mutable/Source/CustomizableObject/Private/MuCO/UnrealMutableImageProvider.cpp
--- a/Engine/Plugins/Experimental/Mutable/Source/CustomizableObject/Private/MuCO/UnrealMutableImageProvider.cpp
+++ b/Engine/Plugins/Experimental/Mutable/Source/CustomizableObject/Private/MuCO/UnrealMutableImageProvider.cpp
@@ -21,14 +21,16 @@
 namespace
 {
 
-	mu::ImagePtr ConvertTextureUnrealToMutable(UTexture2D* Texture)
+	mu::ImagePtr ConvertTextureUnrealToMutable(UTexture2D* Texture, uint8 MipmapsToSkip)
 	{
 		mu::ImagePtr pResult;
 
 #if WITH_EDITOR
 		int LODs = 1;
-		int SizeX = Texture->Source.GetSizeX();
-		int SizeY = Texture->Source.GetSizeY();
+		int SizeX = Texture->Source.GetSizeX() >> MipmapsToSkip;
+		int SizeY = Texture->Source.GetSizeY() >> MipmapsToSkip;
+		check(SizeX > 0 && SizeY > 0);
+
 		ETextureSourceFormat Format = Texture->Source.GetFormat();
 		mu::EImageFormat MutableFormat = mu::EImageFormat::IF_NONE;
 
@@ -46,19 +48,26 @@ namespace
 
 		// If not locked ReadOnly the Texture Source's FGuid can change, invalidating the texture's caching/shaders
 		// making shader compile and cook times increase
-		const uint8* pSource = Texture->Source.LockMipReadOnly(0);
+		const uint8* pSource = Texture->Source.LockMipReadOnly(MipmapsToSkip);
 		if (pSource)
 		{
 			FMemory::Memcpy(pResult->GetData(), pSource, pResult->GetDataSize());
+			Texture->Source.UnlockMip(MipmapsToSkip);
+		}
+		else
+		{
+			check(false);
 		}
 
-		Texture->Source.UnlockMip(0);
+
 #else
-		check(Texture->GetPlatformData()->Mips[0].BulkData.IsBulkDataLoaded());
+		check(Texture->GetPlatformData()->Mips[MipmapsToSkip].BulkData.IsBulkDataLoaded());
 
 		int32 LODs = 1;
-		int32 SizeX = Texture->GetSizeX();
-		int32 SizeY = Texture->GetSizeY();
+		int32 SizeX = Texture->GetSizeX() >> MipmapsToSkip;
+		int32 SizeY = Texture->GetSizeY() >> MipmapsToSkip;
+		check(SizeX > 0 && SizeY > 0);
+
 		EPixelFormat Format = Texture->GetPlatformData()->PixelFormat;
 		mu::EImageFormat MutableFormat = mu::EImageFormat::IF_NONE;
 
@@ -76,14 +85,18 @@ namespace
 
 		// If not locked ReadOnly the Texture Source's FGuid can change, invalidating the texture's caching/shaders
 		// making shader compile and cook times increase
-		const void* pSource = Texture->GetPlatformData()->Mips[0].BulkData.LockReadOnly();
+		const void* pSource = Texture->GetPlatformData()->Mips[MipmapsToSkip].BulkData.LockReadOnly();
 
 		if (pSource)
 		{
 			FMemory::Memcpy(pResult->GetData(), pSource, pResult->GetDataSize());
+			Texture->GetPlatformData()->Mips[MipmapsToSkip].BulkData.Unlock();
+		}
+		else
+		{
+			check(false);
 		}
 
-		Texture->GetPlatformData()->Mips[0].BulkData.Unlock();
 #endif
 
 		return pResult;
@@ -154,9 +167,20 @@ mu::ImagePtr FUnrealMutableImageProvider::GetImage(mu::EXTERNAL_IMAGE_ID id, uin
 			// not being modified by the game thread at the moment and the texture cannot be GCed because of the AddReferencedObjects
 			// in the FUnrealMutableImageProvider
 
+			uint8 MipIndex = MipmapsToSkip < TextureToLoad->GetNumMips() ? MipmapsToSkip : TextureToLoad->GetNumMips() - 1;
+
+			// Mips in the mip tail are inlined and can't be streamed, find the smallest mip available.
+			for (; MipIndex > 0; --MipIndex)
+			{
+				if (TextureToLoad->GetPlatformData()->Mips[MipIndex].BulkData.CanLoadFromDisk())
+				{
+					break;
+				}
+			}
+
 #if WITH_EDITOR
 			// In the editor the src data can be directly accessed
-			return ConvertTextureUnrealToMutable(TextureToLoad);
+			return ConvertTextureUnrealToMutable(TextureToLoad, MipIndex);
 #else
 			// Texture format and the equivalent mutable format
 			Format = TextureToLoad->GetPlatformData()->PixelFormat;
@@ -169,13 +193,16 @@ mu::ImagePtr FUnrealMutableImageProvider::GetImage(mu::EXTERNAL_IMAGE_ID id, uin
 				return CreateDummy();
 			}
 
-			Image = new mu::Image(TextureToLoad->GetSizeX(), TextureToLoad->GetSizeY(), LODs, MutImageFormat);
+			int SizeX = TextureToLoad->GetSizeX() >> MipIndex;
+			int SizeY = TextureToLoad->GetSizeY() >> MipIndex;
+
+			Image = new mu::Image(SizeX, SizeY, LODs, MutImageFormat);
 			MutImageDataSize = Image->GetDataSize();
 
 			// In a packaged game the bulk data has to be loaded
 			// Get the actual file to read the mip 0 data, do not keep any reference to TextureToLoad because once outside of the lock
 			// it may be GCed or changed. Just keep the actual file handle and some sizes instead of the texture
-			FByteBulkData& BulkData = TextureToLoad->GetPlatformData()->Mips[0].BulkData;
+			FByteBulkData& BulkData = TextureToLoad->GetPlatformData()->Mips[MipIndex].BulkData;
 			BulkDataSize = BulkData.GetBulkDataSize();
 			check(BulkDataSize > 0);
 
@@ -294,7 +321,7 @@ void FUnrealMutableImageProvider::CacheImage(mu::EXTERNAL_IMAGE_ID id)
 			case UCustomizableSystemImageProvider::ValueType::Unreal:
 			{
 				UTexture2D* UnrealTexture = Provider->GetTextureParameterValue(id);
-				pResult = ConvertTextureUnrealToMutable(UnrealTexture);
+				pResult = ConvertTextureUnrealToMutable(UnrealTexture, 0);
 				break;
 			}
 
