[mutable] Fix crash when compiling certain models that can generate empty meshes.
#rnx

diff --git a/Engine/Plugins/Experimental/Mutable/Source/MutableRuntime/Private/MuR/CodeRunner.cpp b/Engine/Plugins/Experimental/Mutable/Source/MutableRuntime/Private/MuR/CodeRunner.cpp
--- a/Engine/Plugins/Experimental/Mutable/Source/MutableRuntime/Private/MuR/CodeRunner.cpp
+++ b/Engine/Plugins/Experimental/Mutable/Source/MutableRuntime/Private/MuR/CodeRunner.cpp
@@ -2375,8 +2375,7 @@ namespace mu
 			case 0:
 				if (args.source)
 				{
-					AddOp(FScheduledOp(item.At, item, 1),
-						FScheduledOp(args.source, item));
+					AddOp(FScheduledOp(item.At, item, 1), FScheduledOp(args.source, item));
 				}
 				else
 				{
@@ -2388,9 +2387,9 @@ namespace mu
 			{
 				Ptr<const Mesh> pSource = GetMemory().GetMesh(FCacheAddress(args.source, item));
 
-				MeshPtr pResult = MeshOptimizeSkinning(pSource.get());
+				MeshPtrConst pResult = MeshOptimizeSkinning(pSource.get());
 
-				GetMemory().SetMesh(item, pResult ? pResult : pSource->Clone());
+				GetMemory().SetMesh(item, pResult ? pResult : pSource);
 				break;
 			}
 
