[mutable] Additional mutable code generation optimizations: improve removal of resizes and swizzles.
#rnx

diff --git a/Engine/Plugins/Experimental/Mutable/Source/CustomizableObject/Private/MuCO/CustomizableInstancePrivateData.h b/Engine/Plugins/Experimental/Mutable/Source/CustomizableObject/Private/MuCO/CustomizableInstancePrivateData.h
--- a/Engine/Plugins/Experimental/Mutable/Source/CustomizableObject/Private/MuCO/CustomizableInstancePrivateData.h
+++ b/Engine/Plugins/Experimental/Mutable/Source/CustomizableObject/Private/MuCO/CustomizableInstancePrivateData.h
@@ -246,7 +246,7 @@ public:
 	CUSTOMIZABLEOBJECT_API static void ConvertImage(class UTexture2D* Texture, mu::ImagePtrConst MutableImage, const FMutableModelImageProperties& Props, int32 OnlyLOD=-1, int32 ExtractChannel=-1);
 
 	/** Set OnlyLOD to -1 to generate all mips */
-	CUSTOMIZABLEOBJECT_API static FTexturePlatformData* MutableCreateImagePlatformData(const mu::Image* MutableImage, int32 OnlyLOD, uint16 FullSizeX, uint16 FullSizeY);
+	CUSTOMIZABLEOBJECT_API static FTexturePlatformData* MutableCreateImagePlatformData(mu::Ptr<const mu::Image> MutableImage, int32 OnlyLOD, uint16 FullSizeX, uint16 FullSizeY);
 
 private:
 
diff --git a/Engine/Plugins/Experimental/Mutable/Source/CustomizableObject/Private/MuCO/CustomizableObjectInstance.cpp b/Engine/Plugins/Experimental/Mutable/Source/CustomizableObject/Private/MuCO/CustomizableObjectInstance.cpp
--- a/Engine/Plugins/Experimental/Mutable/Source/CustomizableObject/Private/MuCO/CustomizableObjectInstance.cpp
+++ b/Engine/Plugins/Experimental/Mutable/Source/CustomizableObject/Private/MuCO/CustomizableObjectInstance.cpp
@@ -67,6 +67,7 @@
 #include "MuCO/MultilayerProjector.h"
 #include "MuCO/UnrealConversionUtils.h"
 #include "MuR/MeshBufferSet.h"
+#include "MuR/OpImagePixelFormat.h"
 #include "MuR/PhysicsBody.h"
 #include "MuR/Ptr.h"
 #include "MuR/RefCounted.h"
@@ -2666,7 +2667,7 @@ void CopyTextureProperties(UTexture2D* Texture, const UTexture2D* SourceTexture)
 
 // The memory allocated in the function and pointed by the returned pointer is owned by the caller and must be freed. 
 // If assigned to a UTexture2D, it will be freed by that UTexture2D
-FTexturePlatformData* UCustomizableInstancePrivateData::MutableCreateImagePlatformData(const mu::Image* MutableImage, int32 OnlyLOD, uint16 FullSizeX, uint16 FullSizeY)
+FTexturePlatformData* UCustomizableInstancePrivateData::MutableCreateImagePlatformData(mu::Ptr<const mu::Image> MutableImage, int32 OnlyLOD, uint16 FullSizeX, uint16 FullSizeY)
 {
 	int32 SizeX = FMath::Max(MutableImage->GetSize().x(), FullSizeX);
 	int32 SizeY = FMath::Max(MutableImage->GetSize().y(), FullSizeY);
@@ -2721,7 +2722,6 @@ FTexturePlatformData* UCustomizableInstancePrivateData::MutableCreateImagePlatfo
 		//return nullptr;
 	}
 
-
 	EPixelFormat PlatformFormat = PF_Unknown;
 	switch (MutableFormat)
 	{
@@ -2749,9 +2749,25 @@ FTexturePlatformData* UCustomizableInstancePrivateData::MutableCreateImagePlatfo
 	case mu::EImageFormat::IF_ASTC_4x4_RGBA_LDR:PlatformFormat = PF_ASTC_4x4;	break;
 	case mu::EImageFormat::IF_ASTC_4x4_RG_LDR:	PlatformFormat = PF_ASTC_4x4;	break;
 	default:
-		// Cannot prepare texture if it's not in the right format, this can happen if mutable is in debug mode
-		UE_LOG(LogMutable, Warning, TEXT("Building instance: a texture was generated in an unsupported format, which cannot be converted to Unreal."));
-		return nullptr; 
+		// Cannot prepare texture if it's not in the right format, this can happen if mutable is in debug mode or in case of bugs
+		UE_LOG(LogMutable, Warning, TEXT("Building instance: a texture was generated in an unsupported format, it will be converted to Unreal with a performance penalty."));
+
+		switch (mu::GetImageFormatData(MutableFormat).m_channels)
+		{
+		case 1:
+			PlatformFormat = PF_R8;
+			MutableImage = ImagePixelFormat(0, MutableImage.get(), mu::EImageFormat::IF_L_UBYTE);
+			break;
+		case 2:
+		case 3:
+		case 4:
+			PlatformFormat = PF_R8G8B8A8;
+			MutableImage = ImagePixelFormat(0, MutableImage.get(), mu::EImageFormat::IF_RGBA_UBYTE);
+			break;
+		default: 
+			// Absolutely worst case
+			return nullptr;
+		}		
 	}
 
 	FTexturePlatformData* PlatformData = new FTexturePlatformData();
@@ -2956,7 +2972,7 @@ void UCustomizableInstancePrivateData::ConvertImage(UTexture2D* Texture, mu::Ima
 		OnlyLOD = FMath::Min( OnlyLOD, MutableImage->GetLODCount()-1 );
 	}
 
-	Texture->SetPlatformData(MutableCreateImagePlatformData(MutableImage.get(),OnlyLOD,0,0) );
+	Texture->SetPlatformData(MutableCreateImagePlatformData(MutableImage,OnlyLOD,0,0) );
 
 }
 
diff --git a/Engine/Plugins/Experimental/Mutable/Source/CustomizableObject/Private/MuCO/CustomizableObjectSystem.cpp b/Engine/Plugins/Experimental/Mutable/Source/CustomizableObject/Private/MuCO/CustomizableObjectSystem.cpp
--- a/Engine/Plugins/Experimental/Mutable/Source/CustomizableObject/Private/MuCO/CustomizableObjectSystem.cpp
+++ b/Engine/Plugins/Experimental/Mutable/Source/CustomizableObject/Private/MuCO/CustomizableObjectSystem.cpp
@@ -1453,7 +1453,7 @@ namespace impl
 				// If the image is null, it must be in the cache (or repeated in this instance), and we don't need to do anything here.
 				if (MutableImage)
 				{
-					FTexturePlatformData* PlatformData = UCustomizableInstancePrivateData::MutableCreateImagePlatformData(MutableImage.get(), -1, Image.FullImageSizeX, Image.FullImageSizeY);
+					FTexturePlatformData* PlatformData = UCustomizableInstancePrivateData::MutableCreateImagePlatformData(MutableImage, -1, Image.FullImageSizeX, Image.FullImageSizeY);
 					OperationData->ImageToPlatformDataMap.Add(Image.ImageID, PlatformData);
 					OperationData->PendingTextureCoverageQueries.Add({ KeyName, Surface.MaterialIndex, PlatformData });
 				}
diff --git a/Engine/Plugins/Experimental/Mutable/Source/CustomizableObjectEditor/Private/MuCOE/SMutableCodeViewer.cpp b/Engine/Plugins/Experimental/Mutable/Source/CustomizableObjectEditor/Private/MuCOE/SMutableCodeViewer.cpp
--- a/Engine/Plugins/Experimental/Mutable/Source/CustomizableObjectEditor/Private/MuCOE/SMutableCodeViewer.cpp
+++ b/Engine/Plugins/Experimental/Mutable/Source/CustomizableObjectEditor/Private/MuCOE/SMutableCodeViewer.cpp
@@ -178,7 +178,7 @@ public:
 			case mu::OP_TYPE::IM_RESIZEREL:
 			{
 				mu::OP::ImageResizeRelArgs Args = Program.GetOpArgs<mu::OP::ImageResizeRelArgs>(RowItem->MutableOperation);
-				OpName += FString::Printf(TEXT(" %.3f x %.3f"), int32(Args.factor[0]), int32(Args.factor[1]));
+				OpName += FString::Printf(TEXT(" %.3f x %.3f"), Args.factor[0], Args.factor[1]);
 				break;
 			}
 
@@ -218,7 +218,7 @@ public:
 				OpName += mu::TypeInfo::s_blendModeName[int32(Args.blendType)];
 				OpName += TEXT(" a: ");
 				OpName += mu::TypeInfo::s_blendModeName[int32(Args.blendTypeAlpha)];
-				//OpName += FString::Printf(TEXT(" flags %d"), Args.flags);
+				OpName += FString::Printf(TEXT(" flags %d"), Args.flags);
 				break;
 			}
 
diff --git a/Engine/Plugins/Experimental/Mutable/Source/MutableRuntime/Private/MuR/CodeRunnerTasked.cpp b/Engine/Plugins/Experimental/Mutable/Source/MutableRuntime/Private/MuR/CodeRunnerTasked.cpp
--- a/Engine/Plugins/Experimental/Mutable/Source/MutableRuntime/Private/MuR/CodeRunnerTasked.cpp
+++ b/Engine/Plugins/Experimental/Mutable/Source/MutableRuntime/Private/MuR/CodeRunnerTasked.cpp
@@ -700,6 +700,9 @@ namespace mu
 
 			if (Args.mask && m_mask)
 			{			
+				// Not implemented yet
+				check(Args.flags==0);
+
 				// \todo: precalculated tables for softlight
 				switch (EBlendType(Args.blendType))
 				{
@@ -719,19 +722,40 @@ namespace mu
 			}
 			else
 			{
-				switch (EBlendType(Args.blendType))
+				// Not implemented yet
+				if (Args.flags & OP::ImageLayerArgs::FLAGS::F_BASE_RGB_FROM_ALPHA)
 				{
-				case EBlendType::BT_NORMAL_COMBINE: ImageNormalCombine(pNew.get(), m_base.get(), vec3f(m_col)); break;
-				case EBlendType::BT_SOFTLIGHT: BufferLayerColour<SoftLightChannel, false>(pNew.get(), m_base.get(), vec3f(m_col)); break;
-				case EBlendType::BT_HARDLIGHT: BufferLayerColour<HardLightChannel, false>(pNew.get(), m_base.get(), vec3f(m_col)); break;
-				case EBlendType::BT_BURN: BufferLayerColour<BurnChannel, false>(pNew.get(), m_base.get(), vec3f(m_col)); break;
-				case EBlendType::BT_DODGE: BufferLayerColour<DodgeChannel, false>(pNew.get(), m_base.get(), vec3f(m_col)); break;
-				case EBlendType::BT_SCREEN: BufferLayerColour<ScreenChannel, false>(pNew.get(), m_base.get(), vec3f(m_col)); break;
-				case EBlendType::BT_OVERLAY: BufferLayerColour<OverlayChannel, false>(pNew.get(), m_base.get(), vec3f(m_col)); break;
-				case EBlendType::BT_LIGHTEN: BufferLayerColour<LightenChannel, false>(pNew.get(), m_base.get(), vec3f(m_col)); break;
-				case EBlendType::BT_MULTIPLY: BufferLayerColour<MultiplyChannel, false>(pNew.get(), m_base.get(), vec3f(m_col)); break;
-				case EBlendType::BT_BLEND: FillPlainColourImage(pNew.get(), m_col); break;
-				default: check(false);
+					switch (EBlendType(Args.blendType))
+					{
+					case EBlendType::BT_NORMAL_COMBINE: check(false); break;
+					case EBlendType::BT_SOFTLIGHT: BufferLayerColourFromAlpha<SoftLightChannel, false>(pNew.get(), m_base.get(), vec3f(m_col)); break;
+					case EBlendType::BT_HARDLIGHT: BufferLayerColourFromAlpha<HardLightChannel, false>(pNew.get(), m_base.get(), vec3f(m_col)); break;
+					case EBlendType::BT_BURN: BufferLayerColourFromAlpha<BurnChannel, false>(pNew.get(), m_base.get(), vec3f(m_col)); break;
+					case EBlendType::BT_DODGE: BufferLayerColourFromAlpha<DodgeChannel, false>(pNew.get(), m_base.get(), vec3f(m_col)); break;
+					case EBlendType::BT_SCREEN: BufferLayerColourFromAlpha<ScreenChannel, false>(pNew.get(), m_base.get(), vec3f(m_col)); break;
+					case EBlendType::BT_OVERLAY: BufferLayerColourFromAlpha<OverlayChannel, false>(pNew.get(), m_base.get(), vec3f(m_col)); break;
+					case EBlendType::BT_LIGHTEN: BufferLayerColourFromAlpha<LightenChannel, false>(pNew.get(), m_base.get(), vec3f(m_col)); break;
+					case EBlendType::BT_MULTIPLY: BufferLayerColourFromAlpha<MultiplyChannel, false>(pNew.get(), m_base.get(), vec3f(m_col)); break;
+					case EBlendType::BT_BLEND: check(false); break;
+					default: check(false);
+					}
+				}
+				else
+				{
+					switch (EBlendType(Args.blendType))
+					{
+					case EBlendType::BT_NORMAL_COMBINE: ImageNormalCombine(pNew.get(), m_base.get(), vec3f(m_col)); break;
+					case EBlendType::BT_SOFTLIGHT: BufferLayerColour<SoftLightChannel, false>(pNew.get(), m_base.get(), vec3f(m_col)); break;
+					case EBlendType::BT_HARDLIGHT: BufferLayerColour<HardLightChannel, false>(pNew.get(), m_base.get(), vec3f(m_col)); break;
+					case EBlendType::BT_BURN: BufferLayerColour<BurnChannel, false>(pNew.get(), m_base.get(), vec3f(m_col)); break;
+					case EBlendType::BT_DODGE: BufferLayerColour<DodgeChannel, false>(pNew.get(), m_base.get(), vec3f(m_col)); break;
+					case EBlendType::BT_SCREEN: BufferLayerColour<ScreenChannel, false>(pNew.get(), m_base.get(), vec3f(m_col)); break;
+					case EBlendType::BT_OVERLAY: BufferLayerColour<OverlayChannel, false>(pNew.get(), m_base.get(), vec3f(m_col)); break;
+					case EBlendType::BT_LIGHTEN: BufferLayerColour<LightenChannel, false>(pNew.get(), m_base.get(), vec3f(m_col)); break;
+					case EBlendType::BT_MULTIPLY: BufferLayerColour<MultiplyChannel, false>(pNew.get(), m_base.get(), vec3f(m_col)); break;
+					case EBlendType::BT_BLEND: FillPlainColourImage(pNew.get(), m_col); break;
+					default: check(false);
+					}
 				}
 			}
 		}
diff --git a/Engine/Plugins/Experimental/Mutable/Source/MutableRuntime/Private/MuR/OpImageLayer.h b/Engine/Plugins/Experimental/Mutable/Source/MutableRuntime/Private/MuR/OpImageLayer.h
--- a/Engine/Plugins/Experimental/Mutable/Source/MutableRuntime/Private/MuR/OpImageLayer.h
+++ b/Engine/Plugins/Experimental/Mutable/Source/MutableRuntime/Private/MuR/OpImageLayer.h
@@ -14,8 +14,7 @@ namespace mu
 	//---------------------------------------------------------------------------------------------
 	//! Apply a blending function to an image with a colour source
 	//---------------------------------------------------------------------------------------------
-	template< unsigned (*BLEND_FUNC)(unsigned,unsigned),
-			  bool CLAMP >
+	template< unsigned (*BLEND_FUNC)(unsigned,unsigned), bool CLAMP >
     inline void BufferLayerColour( Image* pResult, const Image* pBase, vec3<float> col )
 	{
         check( pResult->GetFormat() == pBase->GetFormat() );
@@ -184,6 +183,170 @@ namespace mu
     }
 
 
+	//---------------------------------------------------------------------------------------------
+	template< unsigned (*BLEND_FUNC)(unsigned, unsigned), bool CLAMP >
+	inline void BufferLayerColourFromAlpha(Image* pResult, const Image* pBase, vec3<float> col)
+	{
+		check(pResult->GetFormat() == pBase->GetFormat());
+		check(pResult->GetSizeX() == pBase->GetSizeX());
+		check(pResult->GetSizeY() == pBase->GetSizeY());
+		check(pResult->GetLODCount() == pBase->GetLODCount());
+
+		EImageFormat baseFormat = pBase->GetFormat();
+
+		uint8* pDestBuf = pResult->GetData();
+		const uint8* pBaseBuf = pBase->GetData();
+
+		unsigned top[3];
+		top[0] = (unsigned)(255 * col[0]);
+		top[1] = (unsigned)(255 * col[1]);
+		top[2] = (unsigned)(255 * col[2]);
+
+		// Generic implementation
+		int32 pixelCount = pBase->CalculatePixelCount();
+
+		constexpr int PixelCountConcurrencyThreshold = 0xff;
+
+		switch (baseFormat)
+		{
+		case EImageFormat::IF_RGB_UBYTE:
+		{
+			// There is no alpha so we assume 255
+			const auto& ProcessPixel = [ pDestBuf, top ] (uint32 i)
+			{
+				for (int c = 0; c < 3; ++c)
+				{
+					unsigned base = 255;
+					unsigned result = BLEND_FUNC(base, top[c]);
+					if (CLAMP)
+					{
+						pDestBuf[3 * i + c] = (uint8)FMath::Min(255u, result);
+					}
+					else
+					{
+						pDestBuf[3 * i + c] = (uint8)result;
+					}
+				}
+			};
+
+			if (pixelCount > PixelCountConcurrencyThreshold)
+			{
+				ParallelFor(pixelCount, ProcessPixel);
+			}
+			else
+			{
+				for (int p = 0; p < pixelCount; ++p)
+				{
+					ProcessPixel(p);
+				}
+			}
+			break;
+		}
+		case EImageFormat::IF_RGBA_UBYTE:
+		{
+			const auto& ProcessPixel = [ pBaseBuf, pDestBuf, top ] (uint32 i)
+			{
+				for (int c = 0; c < 3; ++c)
+				{
+					unsigned base = pBaseBuf[4 * i + 3];
+					unsigned result = BLEND_FUNC(base, top[c]);
+					if (CLAMP)
+					{
+						pDestBuf[4 * i + c] = (uint8)FMath::Min(255u, result);
+					}
+					else
+					{
+						pDestBuf[4 * i + c] = (uint8)result;
+					}
+				}
+				pDestBuf[4 * i + 3] = pBaseBuf[4 * i + 3];
+			};
+
+			if (pixelCount > PixelCountConcurrencyThreshold)
+			{
+				ParallelFor(pixelCount, ProcessPixel);
+			}
+			else
+			{
+				for (int p = 0; p < pixelCount; ++p)
+				{
+					ProcessPixel(p);
+				}
+			}
+			break;
+		}
+		case EImageFormat::IF_BGRA_UBYTE:
+		{
+			const auto& ProcessPixel = [ pBaseBuf, pDestBuf, top ] (uint32 i)
+			{
+				for (int c = 0; c < 3; ++c)
+				{
+					unsigned base = pBaseBuf[4 * i + 3];
+					unsigned result = BLEND_FUNC(base, top[2 - c]);
+					if (CLAMP)
+					{
+						pDestBuf[4 * i + c] = (uint8)FMath::Min(255u, result);
+					}
+					else
+					{
+						pDestBuf[4 * i + c] = (uint8)result;
+					}
+				}
+				pDestBuf[4 * i + 3] = pBaseBuf[4 * i + 3];
+			};
+
+			if (pixelCount > PixelCountConcurrencyThreshold)
+			{
+				ParallelFor(pixelCount, ProcessPixel);
+			}
+			else
+			{
+				for (int p = 0; p < pixelCount; ++p)
+				{
+					ProcessPixel(p);
+				}
+			}
+			break;
+		}
+		case EImageFormat::IF_L_UBYTE:
+		{
+			const auto& ProcessPixel = [ pBaseBuf, pDestBuf, top ] (uint32 i)
+			{
+				// There is no alpha so we assume 255
+				unsigned base = 255;
+				unsigned result = BLEND_FUNC(base, top[0]);
+				if (CLAMP)
+				{
+					pDestBuf[i] = (uint8)FMath::Min(255u, result);
+				}
+				else
+				{
+					pDestBuf[i] = (uint8)result;
+				}
+			};
+
+
+				if (pixelCount > PixelCountConcurrencyThreshold)
+				{
+					ParallelFor(pixelCount, ProcessPixel);
+				}
+				else
+				{
+					for (int p = 0; p < pixelCount; ++p)
+					{
+						ProcessPixel(p);
+					}
+				}
+				break;
+		}
+
+		default:
+			checkf(false, TEXT("Unsupported format."));
+			break;
+		}
+	}
+
+
 	//---------------------------------------------------------------------------------------------
 	template< unsigned (*BLEND_FUNC_MASKED)(unsigned,unsigned,unsigned),
 		unsigned (*BLEND_FUNC)(unsigned,unsigned),
@@ -1016,12 +1179,6 @@ namespace mu
 	{
 		check(pBase->GetSizeX() == pBlend->GetSizeX() && pBase->GetSizeY() == pBlend->GetSizeY());
 		check(bOnlyFirstLOD || pBase->GetLODCount() <= pBlend->GetLODCount());
-		check(pBlend->GetFormat()==EImageFormat::IF_RGBA_UBYTE);
-
-		uint8* pDestBuf = pStartDestBuf;
-		const uint8* pBaseBuf = pBase->GetData() + BaseChannelOffset;
-		const uint8* pBlendedBuf = pBlend->GetData() + BaseChannelOffset;
-		const uint8* pMaskBuf = pBlend->GetData() + 3;
 
 		// The base determines the number of lods to process.
 		int32 LODCount = bOnlyFirstLOD ? 1 : pBase->GetLODCount();
@@ -1038,33 +1195,72 @@ namespace mu
 
 		int32 UnblendedChannels = BASE_CHANNEL_STRIDE - CHANNELS_TO_BLEND;
 
-		ParallelFor(PixelCount,
-			[
-				pBaseBuf, pBlendedBuf, pDestBuf, pMaskBuf, UnblendedChannels
-			] (uint32 i)
-			{
-				uint32 mask = pMaskBuf[BLENDED_CHANNEL_STRIDE * i];
-				for (int c = 0; c < CHANNELS_TO_BLEND; ++c)
+		uint8* pDestBuf = pStartDestBuf;
+		const uint8* pBaseBuf = pBase->GetData() + BaseChannelOffset;
+		const uint8* pBlendedBuf = pBlend->GetData() + BaseChannelOffset;
+
+		if (pBlend->GetFormat() == EImageFormat::IF_RGBA_UBYTE)
+		{
+			const uint8* pMaskBuf = pBlend->GetData() + 3;
+
+			ParallelFor(PixelCount,
+				[
+					pBaseBuf, pBlendedBuf, pDestBuf, pMaskBuf, UnblendedChannels
+				] (uint32 i)
 				{
-					uint32 base = pBaseBuf[BASE_CHANNEL_STRIDE * i + c];
-					uint32 blended = pBlendedBuf[BLENDED_CHANNEL_STRIDE * i + c];
-					uint32 result = BLEND_FUNC_MASKED(base, blended, mask);
-					if (CLAMP)
+					uint32 mask = pMaskBuf[BLENDED_CHANNEL_STRIDE * i];
+					for (int c = 0; c < CHANNELS_TO_BLEND; ++c)
 					{
-						pDestBuf[BASE_CHANNEL_STRIDE * i + c] = (uint8)FMath::Min(255u, result);
+						uint32 base = pBaseBuf[BASE_CHANNEL_STRIDE * i + c];
+						uint32 blended = pBlendedBuf[BLENDED_CHANNEL_STRIDE * i + c];
+						uint32 result = BLEND_FUNC_MASKED(base, blended, mask);
+						if (CLAMP)
+						{
+							pDestBuf[BASE_CHANNEL_STRIDE * i + c] = (uint8)FMath::Min(255u, result);
+						}
+						else
+						{
+							pDestBuf[BASE_CHANNEL_STRIDE * i + c] = (uint8)result;
+						}
 					}
-					else
+					// Copy the unblended channels
+					// \TODO: unnecessary when doing it in-place?
+					for (int c = 0; c < UnblendedChannels; ++c)
 					{
-						pDestBuf[BASE_CHANNEL_STRIDE * i + c] = (uint8)result;
+						pDestBuf[BASE_CHANNEL_STRIDE * i + CHANNELS_TO_BLEND + c] = pBaseBuf[BASE_CHANNEL_STRIDE * i + CHANNELS_TO_BLEND + c];
 					}
-				}
-				// Copy the unblended channels
-				// \TODO: unnecessary when doing it in-place?
-				for (int c = 0; c < UnblendedChannels; ++c)
+				});
+		}
+		else
+		{
+			ParallelFor(PixelCount,
+				[
+					pBaseBuf, pBlendedBuf, pDestBuf, UnblendedChannels
+				] (uint32 i)
 				{
-					pDestBuf[BASE_CHANNEL_STRIDE * i + CHANNELS_TO_BLEND + c] = pBaseBuf[BASE_CHANNEL_STRIDE * i + CHANNELS_TO_BLEND + c];
-				}
-			});
+					for (int c = 0; c < CHANNELS_TO_BLEND; ++c)
+					{
+						uint32 base = pBaseBuf[BASE_CHANNEL_STRIDE * i + c];
+						uint32 blended = pBlendedBuf[BLENDED_CHANNEL_STRIDE * i + c];
+						uint32 result = BLEND_FUNC(base, blended);
+						if (CLAMP)
+						{
+							pDestBuf[BASE_CHANNEL_STRIDE * i + c] = (uint8)FMath::Min(255u, result);
+						}
+						else
+						{
+							pDestBuf[BASE_CHANNEL_STRIDE * i + c] = (uint8)result;
+						}
+					}
+					// Copy the unblended channels
+					// \TODO: unnecessary when doing it in-place?
+					for (int c = 0; c < UnblendedChannels; ++c)
+					{
+						pDestBuf[BASE_CHANNEL_STRIDE * i + CHANNELS_TO_BLEND + c] = pBaseBuf[BASE_CHANNEL_STRIDE * i + CHANNELS_TO_BLEND + c];
+					}
+				});
+
+		}
 	}
 
 
diff --git a/Engine/Plugins/Experimental/Mutable/Source/MutableRuntime/Private/MuR/Operations.h b/Engine/Plugins/Experimental/Mutable/Source/MutableRuntime/Private/MuR/Operations.h
--- a/Engine/Plugins/Experimental/Mutable/Source/MutableRuntime/Private/MuR/Operations.h
+++ b/Engine/Plugins/Experimental/Mutable/Source/MutableRuntime/Private/MuR/Operations.h
@@ -541,7 +541,9 @@ namespace mu
 				/** If the image has 4 channels, apply to the fourth channel as well. */
 				F_APPLY_TO_ALPHA = 1 << 1,
 				/** Use the alpha channel of the blended image as mask. Mask should be null.*/
-				F_USE_MASK_FROM_BLENDED = 1 << 2
+				F_USE_MASK_FROM_BLENDED = 1 << 2,
+				/** Use the alpha channel of the base image as its RGB.*/
+				F_BASE_RGB_FROM_ALPHA = 1 << 3,
 			} FLAGS;
         };
 
@@ -564,6 +566,7 @@ namespace mu
             ADDRESS colour;
 			uint8 blendType;		// One of EBlendType
 			uint8 blendTypeAlpha;	// One of EBlendType
+			uint8 flags;			// Like in ImageLayerArgs
 		};
 
         struct ImagePixelFormatArgs
diff --git a/Engine/Plugins/Experimental/Mutable/Source/MutableTools/Private/MuT/ASTOpImageLayer.cpp b/Engine/Plugins/Experimental/Mutable/Source/MutableTools/Private/MuT/ASTOpImageLayer.cpp
--- a/Engine/Plugins/Experimental/Mutable/Source/MutableTools/Private/MuT/ASTOpImageLayer.cpp
+++ b/Engine/Plugins/Experimental/Mutable/Source/MutableTools/Private/MuT/ASTOpImageLayer.cpp
@@ -5,6 +5,7 @@
 #include "MuT/StreamsPrivate.h"
 #include "MuT/ASTOpImagePatch.h"
 #include "MuT/ASTOpImageLayerColor.h"
+#include "MuT/ASTOpImageSwizzle.h"
 #include "MuT/ASTOpSwitch.h"
 #include "MuR/ModelPrivate.h"
 #include "MuR/RefCounted.h"
@@ -186,7 +187,6 @@ namespace mu
 
 			const ASTOpFixed* BlendPlainColor = dynamic_cast<const ASTOpFixed*>(blendAt.get());
 
-			// \TODO: Add support for embedded alpha to image layer color.
 			if (!bUseMaskFromBlendAlpha)
 			{
 				Ptr<ASTOpImageLayerColor> NewLayerColor = new ASTOpImageLayerColor;
@@ -220,80 +220,264 @@ namespace mu
 			}
 		}
 
-		// Introduce crop if mask is constant and smaller than the base
-		FImageRect sourceMaskUsage;
-		FImageDesc maskDesc;
-		bool validUsageRect = false;
-		if (!at && maskAt)
+		// See if the mask is actually already in the alpha channel of the blended. In that case,
+		// remove the mask and enable the flag to use the alpha from blended.
+		// This sounds very specific but, experimentally, it seems to happen often.
+		if (!at && maskAt && blendAt 
+			&& 
+			Flags==0
+			)
 		{
-			//MUTABLE_CPUPROFILER_SCOPE(EvaluateAreasForCrop);
+			bool bMatchingAlphaExpression = true;
+			Ptr<ASTOp> CurrentMask = maskAt;
+			Ptr<ASTOp> CurrentBlend = blendAt;
+			while (bMatchingAlphaExpression)
+			{
+				// Skip blend ops that wouldn't change the alpha
+				bool bUpdated = true;
+				while (bUpdated)
+				{
+					bUpdated = false;
+					switch (CurrentBlend->GetOpType())
+					{
+					case OP_TYPE::IM_LAYERCOLOUR:
+					{
+						const ASTOpImageLayerColor* BlendLayer = dynamic_cast<const ASTOpImageLayerColor*>(CurrentBlend.get());
+						if (BlendLayer->blendTypeAlpha == EBlendType::BT_NONE)
+						{
+							CurrentBlend = BlendLayer->base.child();
+							bUpdated = true;
+						}
+						break;
+					}
 
-			validUsageRect = maskAt->GetNonBlackRect(sourceMaskUsage);
-			if (validUsageRect)
+					case OP_TYPE::IM_LAYER:
+					{
+						const ASTOpImageLayer* BlendLayer = dynamic_cast<const ASTOpImageLayer*>(CurrentBlend.get());
+						if (BlendLayer->blendTypeAlpha == EBlendType::BT_NONE)
+						{
+							CurrentBlend = BlendLayer->base.child();
+							bUpdated = true;
+						}
+						break;
+					}
+
+					default:
+						break;
+					}
+				}
+
+				// Matching ops?
+				if (CurrentMask->GetOpType() != CurrentBlend->GetOpType())
+				{
+					break;
+				}
+
+				switch (CurrentMask->GetOpType())
+				{
+
+				case OP_TYPE::IM_DISPLACE:
+				{
+					const ASTOpFixed* MaskDisplace = dynamic_cast<const ASTOpFixed*>(CurrentMask.get());
+					const ASTOpFixed* BlendDisplace = dynamic_cast<const ASTOpFixed*>(CurrentBlend.get());
+					if (MaskDisplace->children[MaskDisplace->op.args.ImageDisplace.displacementMap].child()
+						==
+						BlendDisplace->children[BlendDisplace->op.args.ImageDisplace.displacementMap].child())
+					{
+						CurrentMask = MaskDisplace->children[MaskDisplace->op.args.ImageDisplace.source].child();
+						CurrentBlend = BlendDisplace->children[BlendDisplace->op.args.ImageDisplace.source].child();
+					}
+					break;
+				}
+
+				case OP_TYPE::IM_RASTERMESH:
+				{
+					const ASTOpFixed* MaskRaster = dynamic_cast<const ASTOpFixed*>(CurrentMask.get());
+					const ASTOpFixed* BlendRaster = dynamic_cast<const ASTOpFixed*>(CurrentBlend.get());
+					if (MaskRaster->children[MaskRaster->op.args.ImageRasterMesh.mesh].child() == BlendRaster->children[BlendRaster->op.args.ImageRasterMesh.mesh].child()
+						&&
+						MaskRaster->children[MaskRaster->op.args.ImageRasterMesh.projector].child() == BlendRaster->children[BlendRaster->op.args.ImageRasterMesh.projector].child()
+						&&
+						MaskRaster->children[MaskRaster->op.args.ImageRasterMesh.mask].child() == BlendRaster->children[BlendRaster->op.args.ImageRasterMesh.mask].child()
+						&&
+						MaskRaster->children[MaskRaster->op.args.ImageRasterMesh.angleFadeProperties].child() == BlendRaster->children[BlendRaster->op.args.ImageRasterMesh.angleFadeProperties].child()
+						&&
+						MaskRaster->op.args.ImageRasterMesh.blockIndex == BlendRaster->op.args.ImageRasterMesh.blockIndex
+						)
+					{
+						CurrentMask = MaskRaster->children[MaskRaster->op.args.ImageRasterMesh.image].child();
+						CurrentBlend = BlendRaster->children[BlendRaster->op.args.ImageRasterMesh.image].child();
+					}
+					break;
+				}
+
+				case OP_TYPE::IM_RESIZE:
+				{
+					const ASTOpFixed* MaskResize = dynamic_cast<const ASTOpFixed*>(CurrentMask.get());
+					const ASTOpFixed* BlendResize = dynamic_cast<const ASTOpFixed*>(CurrentBlend.get());
+					if (MaskResize->op.args.ImageResize.size[0] == BlendResize->op.args.ImageResize.size[0]
+						&&
+						MaskResize->op.args.ImageResize.size[1] == BlendResize->op.args.ImageResize.size[1])
+					{
+						CurrentMask = MaskResize->children[MaskResize->op.args.ImageResize.source].child();
+						CurrentBlend = BlendResize->children[BlendResize->op.args.ImageResize.source].child();
+					}
+					break;
+				}
+
+				default:
+					// Case not supported, so don't optimize.
+					bMatchingAlphaExpression = false;
+					break;
+				}
+			}
+
+			if (bMatchingAlphaExpression && CurrentMask->GetOpType() == OP_TYPE::IM_SWIZZLE)
+			{
+				// End of the possible mask expression chain match should have a swizzle selecting the alpha.
+				const ASTOpImageSwizzle* MaskSwizzle = dynamic_cast<const ASTOpImageSwizzle*>(CurrentMask.get());
+				if (MaskSwizzle->SourceChannels[0] == 3
+					&&
+					!MaskSwizzle->SourceChannels[1]
+					&&
+					!MaskSwizzle->SourceChannels[2]
+					&&
+					!MaskSwizzle->SourceChannels[3]
+					&&
+					MaskSwizzle->Sources[0].child() == CurrentBlend
+					)
+				{
+					// we can do something good here
+					Ptr<ASTOpImageLayer> NewLayer = mu::Clone<ASTOpImageLayer>(this);
+					NewLayer->mask = nullptr;
+					NewLayer->Flags |= OP::ImageLayerArgs::F_USE_MASK_FROM_BLENDED;
+					at = NewLayer;
+				}
+			}
+
+		}
+
+		// Try to avoid child swizzle
+		if (!at)
+		{
+			// Is the base a swizzle expanding alpha from a texture?
+			if (baseAt->GetOpType() == OP_TYPE::IM_SWIZZLE)
 			{
-				check(sourceMaskUsage.size[0] > 0);
-				check(sourceMaskUsage.size[1] > 0);
+				const ASTOpImageSwizzle* TypedBase = dynamic_cast<const ASTOpImageSwizzle*>(baseAt.get());
+				bool bAreAllAlpha = true;
+				for (int32 c=0; c<MUTABLE_OP_MAX_SWIZZLE_CHANNELS; ++c)
+				{
+					if (TypedBase->Sources[c] && TypedBase->SourceChannels[c] != 3)
+					{
+						bAreAllAlpha = false;
+						break;
+					}
+				}
 
-				FGetImageDescContext context;
-				maskDesc = maskAt->GetImageDesc(false, &context);
+				if (bAreAllAlpha)
+				{
+					// TODO
+				}
+			}
+
+			// Is the mask a swizzle expanding alpha from a texture?
+			if (maskAt && maskAt->GetOpType() == OP_TYPE::IM_SWIZZLE)
+			{
+				const ASTOpImageSwizzle* TypedBase = dynamic_cast<const ASTOpImageSwizzle*>(maskAt.get());
+				bool bAreAllAlpha = true;
+				for (int32 c = 0; c < MUTABLE_OP_MAX_SWIZZLE_CHANNELS; ++c)
+				{
+					if (TypedBase->Sources[c] && TypedBase->SourceChannels[c] != 3)
+					{
+						bAreAllAlpha = false;
+						break;
+					}
+				}
+
+				if (bAreAllAlpha)
+				{
+					// TODO
+				}
 			}
 		}
 
-		if (!at && maskAt && validUsageRect)
+		// Introduce crop if mask is constant and smaller than the base
+		if (!at && maskAt)
 		{
-			// Adjust for compressed blocks (4), and some extra mips (2 more mips, which is 4)
-			constexpr int blockSize = 4 * 4;
-
-			FImageRect maskUsage;
-			maskUsage.min[0] = (sourceMaskUsage.min[0] / blockSize) * blockSize;
-			maskUsage.min[1] = (sourceMaskUsage.min[1] / blockSize) * blockSize;
-			vec2<uint16> minOffset = sourceMaskUsage.min - maskUsage.min;
-			maskUsage.size[0] = ((sourceMaskUsage.size[0] + minOffset[0] + blockSize - 1) / blockSize) * blockSize;
-			maskUsage.size[1] = ((sourceMaskUsage.size[1] + minOffset[1] + blockSize - 1) / blockSize) * blockSize;
-
-			// Is it worth?
-			float ratio = float(maskUsage.size[0] * maskUsage.size[1])
-				/ float(maskDesc.m_size[0] * maskDesc.m_size[1]);
-			float acceptableCropRatio = options.m_acceptableCropRatio;
-			if (ratio < acceptableCropRatio)
+			FImageRect sourceMaskUsage;
+			FImageDesc maskDesc;
+
+			bool validUsageRect = false;
+			{
+				//MUTABLE_CPUPROFILER_SCOPE(EvaluateAreasForCrop);
+
+				validUsageRect = maskAt->GetNonBlackRect(sourceMaskUsage);
+				if (validUsageRect)
+				{
+					check(sourceMaskUsage.size[0] > 0);
+					check(sourceMaskUsage.size[1] > 0);
+
+					FGetImageDescContext context;
+					maskDesc = maskAt->GetImageDesc(false, &context);
+				}
+			}
+
+			if (validUsageRect)
 			{
-				check(maskUsage.size[0] > 0);
-				check(maskUsage.size[1] > 0);
-
-				Ptr<ASTOpFixed> cropMask = new ASTOpFixed();
-				cropMask->op.type = OP_TYPE::IM_CROP;
-				cropMask->SetChild(cropMask->op.args.ImageCrop.source, mask.child());
-				cropMask->op.args.ImageCrop.minX = maskUsage.min[0];
-				cropMask->op.args.ImageCrop.minY = maskUsage.min[1];
-				cropMask->op.args.ImageCrop.sizeX = maskUsage.size[0];
-				cropMask->op.args.ImageCrop.sizeY = maskUsage.size[1];
-
-				Ptr<ASTOpFixed> cropBlended = new ASTOpFixed();
-				cropBlended->op.type = OP_TYPE::IM_CROP;
-				cropBlended->SetChild(cropBlended->op.args.ImageCrop.source, blend.child());
-				cropBlended->op.args.ImageCrop.minX = maskUsage.min[0];
-				cropBlended->op.args.ImageCrop.minY = maskUsage.min[1];
-				cropBlended->op.args.ImageCrop.sizeX = maskUsage.size[0];
-				cropBlended->op.args.ImageCrop.sizeY = maskUsage.size[1];
-
-				Ptr<ASTOpFixed> cropBase = new ASTOpFixed();
-				cropBase->op.type = OP_TYPE::IM_CROP;
-				cropBase->SetChild(cropBase->op.args.ImageCrop.source, base.child());
-				cropBase->op.args.ImageCrop.minX = maskUsage.min[0];
-				cropBase->op.args.ImageCrop.minY = maskUsage.min[1];
-				cropBase->op.args.ImageCrop.sizeX = maskUsage.size[0];
-				cropBase->op.args.ImageCrop.sizeY = maskUsage.size[1];
-
-				Ptr<ASTOpImageLayer> newLayer = mu::Clone<ASTOpImageLayer>(this);
-				newLayer->base = cropBase;
-				newLayer->blend = cropBlended;
-				newLayer->mask = cropMask;
-
-				Ptr<ASTOpImagePatch> patch = new ASTOpImagePatch();
-				patch->base = baseAt;
-				patch->patch = newLayer;
-				patch->location = maskUsage.min;
-				at = patch;
+				// Adjust for compressed blocks (4), and some extra mips (2 more mips, which is 4)
+				constexpr int blockSize = 4 * 4;
+
+				FImageRect maskUsage;
+				maskUsage.min[0] = (sourceMaskUsage.min[0] / blockSize) * blockSize;
+				maskUsage.min[1] = (sourceMaskUsage.min[1] / blockSize) * blockSize;
+				vec2<uint16> minOffset = sourceMaskUsage.min - maskUsage.min;
+				maskUsage.size[0] = ((sourceMaskUsage.size[0] + minOffset[0] + blockSize - 1) / blockSize) * blockSize;
+				maskUsage.size[1] = ((sourceMaskUsage.size[1] + minOffset[1] + blockSize - 1) / blockSize) * blockSize;
+
+				// Is it worth?
+				float ratio = float(maskUsage.size[0] * maskUsage.size[1])
+					/ float(maskDesc.m_size[0] * maskDesc.m_size[1]);
+				float acceptableCropRatio = options.m_acceptableCropRatio;
+				if (ratio < acceptableCropRatio)
+				{
+					check(maskUsage.size[0] > 0);
+					check(maskUsage.size[1] > 0);
+
+					Ptr<ASTOpFixed> cropMask = new ASTOpFixed();
+					cropMask->op.type = OP_TYPE::IM_CROP;
+					cropMask->SetChild(cropMask->op.args.ImageCrop.source, mask.child());
+					cropMask->op.args.ImageCrop.minX = maskUsage.min[0];
+					cropMask->op.args.ImageCrop.minY = maskUsage.min[1];
+					cropMask->op.args.ImageCrop.sizeX = maskUsage.size[0];
+					cropMask->op.args.ImageCrop.sizeY = maskUsage.size[1];
+
+					Ptr<ASTOpFixed> cropBlended = new ASTOpFixed();
+					cropBlended->op.type = OP_TYPE::IM_CROP;
+					cropBlended->SetChild(cropBlended->op.args.ImageCrop.source, blend.child());
+					cropBlended->op.args.ImageCrop.minX = maskUsage.min[0];
+					cropBlended->op.args.ImageCrop.minY = maskUsage.min[1];
+					cropBlended->op.args.ImageCrop.sizeX = maskUsage.size[0];
+					cropBlended->op.args.ImageCrop.sizeY = maskUsage.size[1];
+
+					Ptr<ASTOpFixed> cropBase = new ASTOpFixed();
+					cropBase->op.type = OP_TYPE::IM_CROP;
+					cropBase->SetChild(cropBase->op.args.ImageCrop.source, base.child());
+					cropBase->op.args.ImageCrop.minX = maskUsage.min[0];
+					cropBase->op.args.ImageCrop.minY = maskUsage.min[1];
+					cropBase->op.args.ImageCrop.sizeX = maskUsage.size[0];
+					cropBase->op.args.ImageCrop.sizeY = maskUsage.size[1];
+
+					Ptr<ASTOpImageLayer> newLayer = mu::Clone<ASTOpImageLayer>(this);
+					newLayer->base = cropBase;
+					newLayer->blend = cropBlended;
+					newLayer->mask = cropMask;
+
+					Ptr<ASTOpImagePatch> patch = new ASTOpImagePatch();
+					patch->base = baseAt;
+					patch->patch = newLayer;
+					patch->location = maskUsage.min;
+					at = patch;
+				}
 			}
 		}
 
@@ -421,7 +605,7 @@ namespace mu
 
 		}
 
-		// If we failed to optimize so far, see if it is worth optimizing the blended only.
+		// If we failed to optimize so far, see if it is worth optimizing the blended branch only.
 		if (!at)
 		{
 			OP_TYPE BlendType = blendAt->GetOpType();
@@ -483,6 +667,64 @@ namespace mu
 						}
 					}
 
+					at = NewSwitch;
+				}
+
+				break;
+			}
+
+			default:
+				break;
+
+			}
+		}
+
+		// If we failed to optimize so far, see if it is worth optimizing the mask branch only.
+		if (!at && maskAt)
+		{
+			OP_TYPE MaskType = maskAt->GetOpType();
+			switch (MaskType)
+			{
+
+			case OP_TYPE::IM_SWITCH:
+			{
+				const ASTOpSwitch* MaskSwitch = dynamic_cast<const ASTOpSwitch*>(maskAt.get());
+
+				// If at least a switch option is a plain colour, sink the layer into the switch
+				bool bWorthSinking = false;
+				for (int32 v = 0; v < MaskSwitch->cases.Num(); ++v)
+				{
+					if (MaskSwitch->cases[v].branch)
+					{
+						// \TODO: Use the smarter query function to detect plain images?
+						if (MaskSwitch->cases[v].branch->GetOpType() == OP_TYPE::IM_PLAINCOLOUR)
+						{
+							bWorthSinking = true;
+							break;
+						}
+					}
+				}
+
+				if (bWorthSinking)
+				{
+					Ptr<ASTOpSwitch> NewSwitch = mu::Clone<ASTOpSwitch>(MaskSwitch);
+
+					if (NewSwitch->def)
+					{
+						Ptr<ASTOpImageLayer> defOp = mu::Clone<ASTOpImageLayer>(this);
+						defOp->mask = MaskSwitch->def.child();
+						NewSwitch->def = defOp;
+					}
+
+					for (int32 v = 0; v < NewSwitch->cases.Num(); ++v)
+					{
+						if (NewSwitch->cases[v].branch)
+						{
+							Ptr<ASTOpImageLayer> BranchOp = mu::Clone<ASTOpImageLayer>(this);
+							BranchOp->mask = MaskSwitch->cases[v].branch.child();
+							NewSwitch->cases[v].branch = BranchOp;
+						}
+					}
 
 					at = NewSwitch;
 				}
diff --git a/Engine/Plugins/Experimental/Mutable/Source/MutableTools/Private/MuT/ASTOpImageLayerColor.cpp b/Engine/Plugins/Experimental/Mutable/Source/MutableTools/Private/MuT/ASTOpImageLayerColor.cpp
--- a/Engine/Plugins/Experimental/Mutable/Source/MutableTools/Private/MuT/ASTOpImageLayerColor.cpp
+++ b/Engine/Plugins/Experimental/Mutable/Source/MutableTools/Private/MuT/ASTOpImageLayerColor.cpp
@@ -3,6 +3,7 @@
 #include "MuT/ASTOpImageLayerColor.h"
 
 #include "MuT/ASTOpSwitch.h"
+#include "MuT/ASTOpImageSwizzle.h"
 #include "Containers/Map.h"
 #include "HAL/PlatformMath.h"
 #include "MuR/ModelPrivate.h"
@@ -40,7 +41,8 @@ namespace mu
 				color == Other->color &&
 				mask == Other->mask &&
 				blendType == Other->blendType &&
-				blendTypeAlpha == Other->blendTypeAlpha;
+				blendTypeAlpha == Other->blendTypeAlpha &&
+				Flags == Other->Flags;
 		}
 		return false;
 	}
@@ -66,6 +68,7 @@ namespace mu
 		n->mask = mapChild(mask.child());
 		n->blendType = blendType;
 		n->blendTypeAlpha = blendTypeAlpha;
+		n->Flags = Flags;
 		return n;
 	}
 
@@ -90,6 +93,7 @@ namespace mu
 
 			args.blendType = (uint8)blendType;
 			args.blendTypeAlpha = (uint8)blendTypeAlpha;
+			args.flags = Flags;
 
 			if (base) args.base = base->linkedAddress;
 			if (color) args.colour = color->linkedAddress;
@@ -182,9 +186,9 @@ namespace mu
 				else if (colour.Equals(FVector4f(1, 1, 1, 1), UE_SMALL_NUMBER))
 				{
 					// If the mask is white, we can remove it
-					Ptr<ASTOpImageLayerColor> nop = mu::Clone<ASTOpImageLayerColor>(this);
-					nop->mask = nullptr;
-					at = nop;
+					Ptr<ASTOpImageLayerColor> NewOp = mu::Clone<ASTOpImageLayerColor>(this);
+					NewOp->mask = nullptr;
+					at = NewOp;
 				}
 			}
 		}
@@ -205,16 +209,6 @@ namespace mu
 					ColorConst.Set(Value[0], Value[1], Value[2], Value[3]);
 				}
 
-				if (!bRGBUnchanged)
-				{
-					switch (blendType)
-					{ 
-					case EBlendType::BT_LIGHTEN: bRGBUnchanged = ColorConst.IsNearlyZero3(UE_SMALL_NUMBER); break;
-					case EBlendType::BT_MULTIPLY: bRGBUnchanged = ColorConst.Equals(FVector4f(1, 1, 1, 1)); break;
-					default: break;
-					}					
-				}
-
 				if (!bAlphaUnchanged)
 				{
 					// TODO: Update when alpha may come from alpha in the color.
@@ -225,6 +219,28 @@ namespace mu
 					default: break;
 					}
 				}
+
+				if (!bRGBUnchanged)
+				{
+					if (Flags & OP::ImageLayerArgs::FLAGS::F_BASE_RGB_FROM_ALPHA)
+					{
+						switch (blendType)
+						{
+						case EBlendType::BT_LIGHTEN: bRGBUnchanged = ColorConst[3]<UE_SMALL_NUMBER; break;
+						case EBlendType::BT_MULTIPLY: bRGBUnchanged = FMath::IsNearlyEqual(ColorConst[3],1.0f); break;
+						default: break;
+						}
+					}
+					else
+					{
+						switch (blendType)
+						{
+						case EBlendType::BT_LIGHTEN: bRGBUnchanged = ColorConst.IsNearlyZero3(UE_SMALL_NUMBER); break;
+						case EBlendType::BT_MULTIPLY: bRGBUnchanged = ColorConst.Equals(FVector4f(1, 1, 1, 1)); break;
+						default: break;
+						}
+					}
+				}
 			}
 
 			if (bRGBUnchanged && bAlphaUnchanged)
@@ -234,6 +250,47 @@ namespace mu
 			}
 		}
 
+
+		// Try to avoid child swizzle
+		if (!at)
+		{
+			// Is the base a swizzle getting expanding alpha from the same texture?
+			if (base.child()->GetOpType() == OP_TYPE::IM_SWIZZLE)
+			{
+				const ASTOpImageSwizzle* TypedBase = dynamic_cast<const ASTOpImageSwizzle*>(base.child().get());
+				bool bAreAllSameAlpha = true;
+				Ptr<ASTOp> Source = nullptr;
+				for (int32 c = 0; c < MUTABLE_OP_MAX_SWIZZLE_CHANNELS; ++c)
+				{
+					Ptr<ASTOp> ThisSource = TypedBase->Sources[c].child();
+					if (!Source)
+					{
+						Source = ThisSource;
+					}
+					else if (ThisSource && ThisSource!=Source)
+					{
+						bAreAllSameAlpha = false;
+						break;
+					}
+
+
+					if (ThisSource && TypedBase->SourceChannels[c] != 3)
+					{
+						bAreAllSameAlpha = false;
+						break;
+					}
+				}
+
+				if (bAreAllSameAlpha)
+				{
+					Ptr<ASTOpImageLayerColor> NewOp = mu::Clone<ASTOpImageLayerColor>(this);
+					NewOp->Flags |= OP::ImageLayerArgs::FLAGS::F_BASE_RGB_FROM_ALPHA;
+					NewOp->base = Source;
+					at = NewOp;
+				}
+			}
+		}
+
 		return at;
 	}
 
diff --git a/Engine/Plugins/Experimental/Mutable/Source/MutableTools/Private/MuT/ASTOpImageLayerColor.h b/Engine/Plugins/Experimental/Mutable/Source/MutableTools/Private/MuT/ASTOpImageLayerColor.h
--- a/Engine/Plugins/Experimental/Mutable/Source/MutableTools/Private/MuT/ASTOpImageLayerColor.h
+++ b/Engine/Plugins/Experimental/Mutable/Source/MutableTools/Private/MuT/ASTOpImageLayerColor.h
@@ -31,6 +31,9 @@ namespace mu
 		/** Blend type used for the alpha channel if any. */
 		EBlendType blendTypeAlpha = EBlendType::BT_NONE;
 
+		/** See ImageLayerArgs::Flags .*/
+		uint8 Flags = 0;
+
 	public:
 
 		ASTOpImageLayerColor();
diff --git a/Engine/Plugins/Experimental/Mutable/Source/MutableTools/Private/MuT/ASTOpImageSwizzle.cpp b/Engine/Plugins/Experimental/Mutable/Source/MutableTools/Private/MuT/ASTOpImageSwizzle.cpp
--- a/Engine/Plugins/Experimental/Mutable/Source/MutableTools/Private/MuT/ASTOpImageSwizzle.cpp
+++ b/Engine/Plugins/Experimental/Mutable/Source/MutableTools/Private/MuT/ASTOpImageSwizzle.cpp
@@ -477,11 +477,11 @@ namespace mu
 				const ASTOpImageLayer* AlphaLayer = dynamic_cast<const ASTOpImageLayer*>(Sources[3].child().get());
 				check(AlphaLayer);
 
-				bool bIsSpecialMultiLayer = !AlphaLayer->mask;
+				bool bIsSpecialMultiLayer = !AlphaLayer->mask && !ColorLayer->Flags && !AlphaLayer->Flags;
 
 				if (bIsSpecialMultiLayer)
 				{
-					// We can combine the 2 multilayers into the composite blend+lighten mode
+					// We can combine the 2 image_layers into the composite blend+lighten mode
 
 					Ptr<ASTOpImageSwizzle> NewBase = mu::Clone<ASTOpImageSwizzle>(this);
 					NewBase->Sources[0] = ColorLayer->base.child();
diff --git a/Engine/Plugins/Experimental/Mutable/Source/MutableTools/Private/MuT/CodeOptimiser_Semantic.cpp b/Engine/Plugins/Experimental/Mutable/Source/MutableTools/Private/MuT/CodeOptimiser_Semantic.cpp
--- a/Engine/Plugins/Experimental/Mutable/Source/MutableTools/Private/MuT/CodeOptimiser_Semantic.cpp
+++ b/Engine/Plugins/Experimental/Mutable/Source/MutableTools/Private/MuT/CodeOptimiser_Semantic.cpp
@@ -278,6 +278,39 @@ namespace mu
             break;
         }
 
+		//-------------------------------------------------------------------------------------
+		case OP_TYPE::IM_RESIZEREL:
+		{
+			auto sourceAt = children[op.args.ImageResizeRel.source].child();
+
+			// The instruction can be sunk
+			OP_TYPE sourceType = sourceAt->GetOpType();
+			switch (sourceType)
+			{
+
+			// This is done here instead of in the OptimizeSize step to prevent removing 
+			// resize_rel too early.
+            case OP_TYPE::IM_RESIZE:
+            {
+				Ptr<ASTOpFixed> NewOp = mu::Clone<ASTOpFixed>(sourceAt.get());
+				NewOp->op.args.ImageResize.size[0] =
+                        int16_t(NewOp->op.args.ImageResize.size[0] * op.args.ImageResizeRel.factor[0] );
+				NewOp->op.args.ImageResize.size[1] =
+                        int16_t(NewOp->op.args.ImageResize.size[1] * op.args.ImageResizeRel.factor[1] );
+
+				at = NewOp;
+                break;
+            }
+
+
+			default:
+				break;
+
+			}
+
+			break;
+		}
+
 
         //-------------------------------------------------------------------------------------
         case OP_TYPE::IM_PIXELFORMAT:
@@ -1614,78 +1647,88 @@ namespace mu
             switch ( sourceType )
             {
 
-                case OP_TYPE::IM_BLANKLAYOUT:
-                {
-                    auto newOp = mu::Clone<ASTOpFixed>(sourceAt);
-
-                    newOp->op.args.ImageBlankLayout.blockSize[0] =
-                            uint16( newOp->op.args.ImageBlankLayout.blockSize[0]
-                                    * op.args.ImageResizeRel.factor[0]
-                                    + 0.5f );
-                    newOp->op.args.ImageBlankLayout.blockSize[1] =
-                            uint16( newOp->op.args.ImageBlankLayout.blockSize[1]
-                                    * op.args.ImageResizeRel.factor[1]
-                                    + 0.5f );
-                    at = newOp;
-                    break;
-                }
+			case OP_TYPE::IM_BLANKLAYOUT:
+			{
+				auto newOp = mu::Clone<ASTOpFixed>(sourceAt);
+
+				newOp->op.args.ImageBlankLayout.blockSize[0] =
+					uint16(newOp->op.args.ImageBlankLayout.blockSize[0]
+						* op.args.ImageResizeRel.factor[0]
+						+ 0.5f);
+				newOp->op.args.ImageBlankLayout.blockSize[1] =
+					uint16(newOp->op.args.ImageBlankLayout.blockSize[1]
+						* op.args.ImageResizeRel.factor[1]
+						+ 0.5f);
+				at = newOp;
+				break;
+			}
+
+			case OP_TYPE::IM_PLAINCOLOUR:
+			{
+				auto newOp = mu::Clone<ASTOpFixed>(sourceAt);
+
+				newOp->op.args.ImagePlainColour.size[0] =
+					uint16(newOp->op.args.ImagePlainColour.size[0]
+						* op.args.ImageResizeRel.factor[0]
+						+ 0.5f);
+				newOp->op.args.ImagePlainColour.size[1] =
+					uint16(newOp->op.args.ImagePlainColour.size[1]
+						* op.args.ImageResizeRel.factor[1]
+						+ 0.5f);
+				at = newOp;
+				break;
+			}
 
 
 				// Don't combine. ResizeRel sometimes can resize more children than Resize can do. (see RasterMesh)
+				// It can be combined in an optimization step further in the process, when normal sizes may have been 
+				// optimized already (see OptimizeSemantic)
 //                case OP_TYPE::IM_RESIZE:
 //                {
-//                    OP op = program.m_code[sourceAt];
-//                    op.args.ImageResize.size[0] =
-//                            int16_t( op.args.ImageResize.size[0]
-//                                   * program.m_code[at].args.ImageResizeRel.factor[0] );
-//                    op.args.ImageResize.size[1] =
-//                            int16_t( op.args.ImageResize.size[1]
-//                                   * program.m_code[at].args.ImageResizeRel.factor[1] );
-//                    return op;
 //                    break;
 //                }
 
 
-                default:
-                {
-                    Sink_ImageResizeRelAST sinker;
-                    at = sinker.Apply(this);
-
-                    break;
-                }
-
-                }
+            default:
+            {
+                Sink_ImageResizeRelAST sinker;
+                at = sinker.Apply(this);
 
                 break;
             }
 
+            }
 
-            //-------------------------------------------------------------------------------------
-            case OP_TYPE::IM_RESIZELIKE:
-            {
-                Ptr<ImageSizeExpression> sourceSize =
-                        children[op.args.ImageResizeLike.source]->GetImageSizeExpression();
-                Ptr<ImageSizeExpression> sizeSourceSize =
-                        children[op.args.ImageResizeLike.sizeSource]->GetImageSizeExpression();
+            break;
+        }
 
-                if ( *sourceSize == *sizeSourceSize )
-                {
-                    at = children[op.args.ImageResizeLike.source].child();
-                }
 
-                else if ( sizeSourceSize->type == ImageSizeExpression::ISET_CONSTANT )
-                {
-                    Ptr<ASTOpFixed> newAt = new ASTOpFixed;
-                    newAt->op.type = OP_TYPE::IM_RESIZE;
-                    newAt->SetChild( newAt->op.args.ImageResize.source, children[op.args.ImageResizeLike.source] );
-                    newAt->op.args.ImageResize.size[0] = sizeSourceSize->size[0];
-                    newAt->op.args.ImageResize.size[1] = sizeSourceSize->size[1];
-                    at = newAt;
-                }
-                else if ( sizeSourceSize->type == ImageSizeExpression::ISET_LAYOUTFACTOR )
-                {
-                    // TODO
-                    // Skip intermediate ops until the layout
+        //-------------------------------------------------------------------------------------
+        case OP_TYPE::IM_RESIZELIKE:
+        {
+            Ptr<ImageSizeExpression> sourceSize =
+                    children[op.args.ImageResizeLike.source]->GetImageSizeExpression();
+            Ptr<ImageSizeExpression> sizeSourceSize =
+                    children[op.args.ImageResizeLike.sizeSource]->GetImageSizeExpression();
+
+            if ( *sourceSize == *sizeSourceSize )
+            {
+                at = children[op.args.ImageResizeLike.source].child();
+            }
+
+            else if ( sizeSourceSize->type == ImageSizeExpression::ISET_CONSTANT )
+            {
+                Ptr<ASTOpFixed> newAt = new ASTOpFixed;
+                newAt->op.type = OP_TYPE::IM_RESIZE;
+                newAt->SetChild( newAt->op.args.ImageResize.source, children[op.args.ImageResizeLike.source] );
+                newAt->op.args.ImageResize.size[0] = sizeSourceSize->size[0];
+                newAt->op.args.ImageResize.size[1] = sizeSourceSize->size[1];
+                at = newAt;
+            }
+            else if ( sizeSourceSize->type == ImageSizeExpression::ISET_LAYOUTFACTOR )
+            {
+                // TODO
+                // Skip intermediate ops until the layout
 //                    if ( program.m_code[sizeSourceAt].type != OP_TYPE::IM_BLANKLAYOUT )
 //                    {
 //                        m_modified = true;
@@ -1701,11 +1744,11 @@ namespace mu
 //                        opResize.args.ImageResizeLike.sizeSource = program.AddOp( blackLayoutOp );
 //                        at = program.AddOp( opResize );
 //                    }
-                }
-
-                break;
             }
 
+            break;
+        }
+
         default:
             break;
 
