[mutable] Removed "auto" uses outside the coding standards.
#rnx

diff --git a/Engine/Plugins/Experimental/Mutable/Source/MutableRuntime/Private/MuR/CodeRunner.cpp b/Engine/Plugins/Experimental/Mutable/Source/MutableRuntime/Private/MuR/CodeRunner.cpp
--- a/Engine/Plugins/Experimental/Mutable/Source/MutableRuntime/Private/MuR/CodeRunner.cpp
+++ b/Engine/Plugins/Experimental/Mutable/Source/MutableRuntime/Private/MuR/CodeRunner.cpp
@@ -145,7 +145,7 @@ namespace mu
         //MUTABLE_CPUPROFILER_SCOPE(RunCode_Conditional);
 
 		OP_TYPE type = pModel->GetPrivate()->m_program.GetOpType(item.At);
-        auto args = pModel->GetPrivate()->m_program.GetOpArgs<OP::ConditionalArgs>(item.At);
+		OP::ConditionalArgs args = pModel->GetPrivate()->m_program.GetOpArgs<OP::ConditionalArgs>(item.At);
 
         // Conditionals have the following execution stages:
         // 0: we need to run the condition
@@ -339,7 +339,7 @@ namespace mu
 
         case OP_TYPE::IN_ADDVECTOR:
         {
-            auto args = pModel->GetPrivate()->m_program.GetOpArgs<OP::InstanceAddArgs>(item.At);
+			OP::InstanceAddArgs args = pModel->GetPrivate()->m_program.GetOpArgs<OP::InstanceAddArgs>(item.At);
 
             switch (item.Stage)
             {
@@ -385,7 +385,7 @@ namespace mu
 
         case OP_TYPE::IN_ADDSCALAR:
         {
-            auto args = pModel->GetPrivate()->m_program.GetOpArgs<OP::InstanceAddArgs>(item.At);
+			OP::InstanceAddArgs args = pModel->GetPrivate()->m_program.GetOpArgs<OP::InstanceAddArgs>(item.At);
             switch (item.Stage)
             {
             case 0:
@@ -430,7 +430,7 @@ namespace mu
 
         case OP_TYPE::IN_ADDSTRING:
         {
-            auto args = pModel->GetPrivate()->m_program.GetOpArgs<OP::InstanceAddArgs>( item.At );
+			OP::InstanceAddArgs args = pModel->GetPrivate()->m_program.GetOpArgs<OP::InstanceAddArgs>( item.At );
             switch ( item.Stage )
             {
             case 0:
@@ -477,7 +477,7 @@ namespace mu
 
         case OP_TYPE::IN_ADDCOMPONENT:
         {
-            auto args = pModel->GetPrivate()->m_program.GetOpArgs<OP::InstanceAddArgs>(item.At);
+			OP::InstanceAddArgs args = pModel->GetPrivate()->m_program.GetOpArgs<OP::InstanceAddArgs>(item.At);
             switch (item.Stage)
             {
             case 0:
@@ -536,7 +536,7 @@ namespace mu
 
         case OP_TYPE::IN_ADDSURFACE:
         {
-            auto args = pModel->GetPrivate()->m_program.GetOpArgs<OP::InstanceAddArgs>(item.At);
+			OP::InstanceAddArgs args = pModel->GetPrivate()->m_program.GetOpArgs<OP::InstanceAddArgs>(item.At);
             switch (item.Stage)
             {
             case 0:
@@ -606,7 +606,7 @@ namespace mu
 
         case OP_TYPE::IN_ADDLOD:
         {
-            auto args = pModel->GetPrivate()->m_program.GetOpArgs<OP::InstanceAddLODArgs>(item.At);
+			OP::InstanceAddLODArgs args = pModel->GetPrivate()->m_program.GetOpArgs<OP::InstanceAddLODArgs>(item.At);
             switch (item.Stage)
             {
             case 0:
@@ -692,7 +692,7 @@ namespace mu
 
         case OP_TYPE::IN_ADDMESH:
         {
-            auto args = pModel->GetPrivate()->m_program.GetOpArgs<OP::InstanceAddArgs>(item.At);
+			OP::InstanceAddArgs args = pModel->GetPrivate()->m_program.GetOpArgs<OP::InstanceAddArgs>(item.At);
             switch (item.Stage)
             {
             case 0:
@@ -739,7 +739,7 @@ namespace mu
 
         case OP_TYPE::IN_ADDIMAGE:
         {
-            auto args = pModel->GetPrivate()->m_program.GetOpArgs<OP::InstanceAddArgs>(item.At);
+			OP::InstanceAddArgs args = pModel->GetPrivate()->m_program.GetOpArgs<OP::InstanceAddArgs>(item.At);
             switch (item.Stage)
             {
             case 0:
@@ -872,7 +872,7 @@ namespace mu
 
         case OP_TYPE::ME_APPLYLAYOUT:
         {
-            auto args = pModel->GetPrivate()->m_program.GetOpArgs<OP::MeshApplyLayoutArgs>(item.At);
+			OP::MeshApplyLayoutArgs args = pModel->GetPrivate()->m_program.GetOpArgs<OP::MeshApplyLayoutArgs>(item.At);
             switch (item.Stage)
             {
             case 0:
@@ -1151,7 +1151,7 @@ namespace mu
 
         case OP_TYPE::ME_MERGE:
         {
-            auto args = pModel->GetPrivate()->m_program.GetOpArgs<OP::MeshMergeArgs>(item.At);
+			OP::MeshMergeArgs args = pModel->GetPrivate()->m_program.GetOpArgs<OP::MeshMergeArgs>(item.At);
             switch (item.Stage)
             {
             case 0:
@@ -1214,7 +1214,7 @@ namespace mu
 
         case OP_TYPE::ME_INTERPOLATE:
         {
-            auto args = pModel->GetPrivate()->m_program.GetOpArgs<OP::MeshInterpolateArgs>(item.At);
+			OP::MeshInterpolateArgs args = pModel->GetPrivate()->m_program.GetOpArgs<OP::MeshInterpolateArgs>(item.At);
             switch (item.Stage)
             {
             case 0:
@@ -1406,7 +1406,7 @@ namespace mu
 
         case OP_TYPE::ME_MASKCLIPMESH:
         {
-            auto args = pModel->GetPrivate()->m_program.GetOpArgs<OP::MeshMaskClipMeshArgs>(item.At);
+			OP::MeshMaskClipMeshArgs args = pModel->GetPrivate()->m_program.GetOpArgs<OP::MeshMaskClipMeshArgs>(item.At);
             switch (item.Stage)
             {
             case 0:
@@ -1440,7 +1440,7 @@ namespace mu
 
         case OP_TYPE::ME_MASKDIFF:
         {
-            auto args = pModel->GetPrivate()->m_program.GetOpArgs<OP::MeshMaskDiffArgs>(item.At);
+			OP::MeshMaskDiffArgs args = pModel->GetPrivate()->m_program.GetOpArgs<OP::MeshMaskDiffArgs>(item.At);
             switch (item.Stage)
             {
             case 0:
@@ -1474,7 +1474,7 @@ namespace mu
 
         case OP_TYPE::ME_SUBTRACT:
         {
-            auto args = pModel->GetPrivate()->m_program.GetOpArgs<OP::MeshSubtractArgs>(item.At);
+			OP::MeshSubtractArgs args = pModel->GetPrivate()->m_program.GetOpArgs<OP::MeshSubtractArgs>(item.At);
             switch (item.Stage)
             {
             case 0:
@@ -1503,7 +1503,7 @@ namespace mu
 
         case OP_TYPE::ME_FORMAT:
         {
-            auto args = pModel->GetPrivate()->m_program.GetOpArgs<OP::MeshFormatArgs>(item.At);
+			OP::MeshFormatArgs args = pModel->GetPrivate()->m_program.GetOpArgs<OP::MeshFormatArgs>(item.At);
             switch (item.Stage)
             {
             case 0:
@@ -1555,16 +1555,16 @@ namespace mu
         {
             const uint8* data = pModel->GetPrivate()->m_program.GetOpArgsPointer(item.At);
 
-            mu::OP::ADDRESS source;
-            memcpy( &source, data, sizeof(OP::ADDRESS) );
+            OP::ADDRESS source;
+            FMemory::Memcpy( &source, data, sizeof(OP::ADDRESS) );
             data += sizeof(OP::ADDRESS);
 
             uint16 layout;
-            memcpy( &layout, data, sizeof(uint16) );
+			FMemory::Memcpy( &layout, data, sizeof(uint16) );
             data += sizeof(uint16);
 
             uint16 blockCount;
-            memcpy( &blockCount, data, sizeof(uint16) );
+			FMemory::Memcpy( &blockCount, data, sizeof(uint16) );
             data += sizeof(uint16);
 
             switch (item.Stage)
@@ -1587,7 +1587,7 @@ namespace mu
 
                 // Access with memcpy necessary for unaligned arm issues.
                 uint32 blocks[1024];
-                memcpy( blocks, data, sizeof(uint32)*FMath::Min(1024,int(blockCount)) );
+				FMemory::Memcpy( blocks, data, sizeof(uint32)*FMath::Min(1024,int(blockCount)) );
 
                 MeshPtr pResult;
                 pResult = MeshExtractLayoutBlock( pSource.get(),
@@ -1608,7 +1608,7 @@ namespace mu
 
         case OP_TYPE::ME_EXTRACTFACEGROUP:
         {
-            auto args = pModel->GetPrivate()->m_program.GetOpArgs<OP::MeshExtractFaceGroupArgs>(item.At);
+			OP::MeshExtractFaceGroupArgs args = pModel->GetPrivate()->m_program.GetOpArgs<OP::MeshExtractFaceGroupArgs>(item.At);
             switch (item.Stage)
             {
             case 0:
@@ -1644,7 +1644,7 @@ namespace mu
 
         case OP_TYPE::ME_TRANSFORM:
         {
-            auto args = pModel->GetPrivate()->m_program.GetOpArgs<OP::MeshTransformArgs>(item.At);
+			OP::MeshTransformArgs args = pModel->GetPrivate()->m_program.GetOpArgs<OP::MeshTransformArgs>(item.At);
             switch (item.Stage)
             {
             case 0:
@@ -1681,7 +1681,7 @@ namespace mu
 
         case OP_TYPE::ME_CLIPMORPHPLANE:
         {
-            auto args = pModel->GetPrivate()->m_program.GetOpArgs<OP::MeshClipMorphPlaneArgs>(item.At);
+			OP::MeshClipMorphPlaneArgs args = pModel->GetPrivate()->m_program.GetOpArgs<OP::MeshClipMorphPlaneArgs>(item.At);
             switch (item.Stage)
             {
             case 0:
@@ -1751,7 +1751,7 @@ namespace mu
 
         case OP_TYPE::ME_CLIPWITHMESH:
         {
-            auto args = pModel->GetPrivate()->m_program.GetOpArgs<OP::MeshClipWithMeshArgs>(item.At);
+			OP::MeshClipWithMeshArgs args = pModel->GetPrivate()->m_program.GetOpArgs<OP::MeshClipWithMeshArgs>(item.At);
             switch (item.Stage)
             {
             case 0:
@@ -1839,7 +1839,7 @@ namespace mu
 
         case OP_TYPE::ME_REMAPINDICES:
         {
-            auto args = pModel->GetPrivate()->m_program.GetOpArgs<OP::MeshRemapIndicesArgs>(item.At);
+			OP::MeshRemapIndicesArgs args = pModel->GetPrivate()->m_program.GetOpArgs<OP::MeshRemapIndicesArgs>(item.At);
             switch (item.Stage)
             {
             case 0:
@@ -1885,7 +1885,7 @@ namespace mu
 
         case OP_TYPE::ME_APPLYPOSE:
         {
-            auto args = pModel->GetPrivate()->m_program.GetOpArgs<OP::MeshApplyPoseArgs>(item.At);
+			OP::MeshApplyPoseArgs args = pModel->GetPrivate()->m_program.GetOpArgs<OP::MeshApplyPoseArgs>(item.At);
             switch (item.Stage)
             {
             case 0:
@@ -1931,7 +1931,7 @@ namespace mu
 
 		case OP_TYPE::ME_GEOMETRYOPERATION:
 		{
-			auto args = pModel->GetPrivate()->m_program.GetOpArgs<OP::MeshGeometryOperationArgs>(item.At);
+			OP::MeshGeometryOperationArgs args = pModel->GetPrivate()->m_program.GetOpArgs<OP::MeshGeometryOperationArgs>(item.At);
 			switch (item.Stage)
 			{
 			case 0:
@@ -2052,7 +2052,7 @@ namespace mu
 
 		case OP_TYPE::ME_APPLYSHAPE:
 		{
-			auto args = pModel->GetPrivate()->m_program.GetOpArgs<OP::MeshApplyShapeArgs>(item.At);
+			OP::MeshApplyShapeArgs args = pModel->GetPrivate()->m_program.GetOpArgs<OP::MeshApplyShapeArgs>(item.At);
 			switch (item.Stage)
 			{
 			case 0:
@@ -2144,7 +2144,7 @@ namespace mu
 
         case OP_TYPE::ME_SETSKELETON:
         {
-            auto args = pModel->GetPrivate()->m_program.GetOpArgs<OP::MeshSetSkeletonArgs>(item.At);
+			OP::MeshSetSkeletonArgs args = pModel->GetPrivate()->m_program.GetOpArgs<OP::MeshSetSkeletonArgs>(item.At);
             switch (item.Stage)
             {
             case 0:
@@ -2316,7 +2316,7 @@ namespace mu
 
         case OP_TYPE::ME_PROJECT:
         {
-            auto args = pModel->GetPrivate()->m_program.GetOpArgs<OP::MeshProjectArgs>(item.At);
+			OP::MeshProjectArgs args = pModel->GetPrivate()->m_program.GetOpArgs<OP::MeshProjectArgs>(item.At);
             switch (item.Stage)
             {
             case 0:
@@ -2369,7 +2369,7 @@ namespace mu
 
 		case OP_TYPE::ME_OPTIMIZESKINNING:
 		{
-			auto args = pModel->GetPrivate()->m_program.GetOpArgs<OP::MeshOptimizeSkinningArgs>(item.At);
+			OP::MeshOptimizeSkinningArgs args = pModel->GetPrivate()->m_program.GetOpArgs<OP::MeshOptimizeSkinningArgs>(item.At);
 			switch (item.Stage)
 			{
 			case 0:
@@ -2741,7 +2741,7 @@ namespace mu
 
         case OP_TYPE::IM_DIFFERENCE:
         {
-            auto args = pModel->GetPrivate()->m_program.GetOpArgs<OP::ImageDifferenceArgs>(item.At);
+			OP::ImageDifferenceArgs args = pModel->GetPrivate()->m_program.GetOpArgs<OP::ImageDifferenceArgs>(item.At);
             switch (item.Stage)
             {
             case 0:
@@ -3007,7 +3007,7 @@ namespace mu
 
         case OP_TYPE::IM_RESIZEREL:
         {
-            auto args = pModel->GetPrivate()->m_program.GetOpArgs<OP::ImageResizeRelArgs>(item.At);
+			OP::ImageResizeRelArgs args = pModel->GetPrivate()->m_program.GetOpArgs<OP::ImageResizeRelArgs>(item.At);
             switch (item.Stage)
             {
             case 0:
@@ -3063,7 +3063,7 @@ namespace mu
 
         case OP_TYPE::IM_BLANKLAYOUT:
         {
-            auto args = pModel->GetPrivate()->m_program.GetOpArgs<OP::ImageBlankLayoutArgs>(item.At);
+			OP::ImageBlankLayoutArgs args = pModel->GetPrivate()->m_program.GetOpArgs<OP::ImageBlankLayoutArgs>(item.At);
             switch (item.Stage)
             {
             case 0:
@@ -3140,7 +3140,7 @@ namespace mu
 
         case OP_TYPE::IM_COMPOSE:
         {
-            auto args = pModel->GetPrivate()->m_program.GetOpArgs<OP::ImageComposeArgs>( item.At );
+			OP::ImageComposeArgs args = pModel->GetPrivate()->m_program.GetOpArgs<OP::ImageComposeArgs>( item.At );
             switch ( item.Stage )
             {
             case 0:
@@ -3186,7 +3186,7 @@ namespace mu
 
         case OP_TYPE::IM_INTERPOLATE:
         {
-            auto args = pModel->GetPrivate()->m_program.GetOpArgs<OP::ImageInterpolateArgs>(item.At);
+			OP::ImageInterpolateArgs args = pModel->GetPrivate()->m_program.GetOpArgs<OP::ImageInterpolateArgs>(item.At);
             switch (item.Stage)
             {
             case 0:
@@ -3344,7 +3344,7 @@ namespace mu
 
         case OP_TYPE::IM_INTERPOLATE3:
         {
-            auto args = pModel->GetPrivate()->m_program.GetOpArgs<OP::ImageInterpolate3Args>(item.At);
+			OP::ImageInterpolate3Args args = pModel->GetPrivate()->m_program.GetOpArgs<OP::ImageInterpolate3Args>(item.At);
             switch (item.Stage)
             {
             case 0:
@@ -3383,7 +3383,7 @@ namespace mu
 
         case OP_TYPE::IM_SATURATE:
         {
-            auto args = pModel->GetPrivate()->m_program.GetOpArgs<OP::ImageSaturateArgs>(item.At);
+			OP::ImageSaturateArgs args = pModel->GetPrivate()->m_program.GetOpArgs<OP::ImageSaturateArgs>(item.At);
             switch (item.Stage)
             {
             case 0:
@@ -3412,7 +3412,7 @@ namespace mu
 
         case OP_TYPE::IM_LUMINANCE:
         {
-            auto args = pModel->GetPrivate()->m_program.GetOpArgs<OP::ImageLuminanceArgs>(item.At);
+			OP::ImageLuminanceArgs args = pModel->GetPrivate()->m_program.GetOpArgs<OP::ImageLuminanceArgs>(item.At);
             switch (item.Stage)
             {
             case 0:
@@ -3823,7 +3823,7 @@ namespace mu
 
         case OP_TYPE::IM_RASTERMESH:
         {
-            auto args = pModel->GetPrivate()->m_program.GetOpArgs<OP::ImageRasterMeshArgs>(item.At);
+			OP::ImageRasterMeshArgs args = pModel->GetPrivate()->m_program.GetOpArgs<OP::ImageRasterMeshArgs>(item.At);
             switch (item.Stage)
             {
             case 0:
@@ -3891,7 +3891,7 @@ namespace mu
                     float fadeEnd = 180.0f;
                     if ( args.angleFadeProperties )
                     {
-                        auto fadeProperties = GetMemory().GetColour(FScheduledOp::FromOpAndOptions(args.angleFadeProperties, item, 0));
+                        FVector4f fadeProperties = GetMemory().GetColour(FScheduledOp::FromOpAndOptions(args.angleFadeProperties, item, 0));
                         fadeStart = fadeProperties[0];
                         fadeEnd = fadeProperties[1];
                     }
@@ -3919,14 +3919,14 @@ namespace mu
                     float projectionAngle = 0;
                     if ( args.projector && pSource )
                     {
-                        auto pProjector = GetMemory().GetProjector(FScheduledOp::FromOpAndOptions(args.projector, item, 0));
+                        Ptr<const Projector> pProjector = GetMemory().GetProjector(FScheduledOp::FromOpAndOptions(args.projector, item, 0));
                         if (pProjector)
                         {
                             projectionAngle = pProjector->m_value.projectionAngle;
 
                             switch (pProjector->m_value.type)
                             {
-                            case mu::PROJECTOR_TYPE::PLANAR:
+                            case PROJECTOR_TYPE::PLANAR:
                                 ImageRasterProjectedPlanar( pMesh.get(), pNew.get(),
                                                             pSource.get(), pMask.get(),
                                                             fadeStart, fadeEnd,
@@ -3934,7 +3934,7 @@ namespace mu
                                                             &scratch );
                                 break;
 
-                            case mu::PROJECTOR_TYPE::WRAPPING:
+                            case PROJECTOR_TYPE::WRAPPING:
                                 ImageRasterProjectedWrapping( pMesh.get(), pNew.get(),
                                                               pSource.get(), pMask.get(),
                                                               fadeStart, fadeEnd,
@@ -3942,7 +3942,7 @@ namespace mu
                                                               &scratch );
                                 break;
 
-                            case mu::PROJECTOR_TYPE::CYLINDRICAL:
+                            case PROJECTOR_TYPE::CYLINDRICAL:
                                 ImageRasterProjectedCylindrical( pMesh.get(), pNew.get(),
                                                                  pSource.get(), pMask.get(),
                                                                  fadeStart, fadeEnd,
@@ -4147,13 +4147,13 @@ namespace mu
             return nullptr;
         }
 
-        const auto& program = pModel->GetPrivate()->m_program;
+        const FProgram& program = pModel->GetPrivate()->m_program;
         const FParameterDesc& paramDesc = program.m_parameters[ parameterIndex ];
         for( size_t rangeIndexInParam=0;
              rangeIndexInParam<paramDesc.m_ranges.Num();
              ++rangeIndexInParam )
         {
-            auto rangeIndexInModel = paramDesc.m_ranges[rangeIndexInParam];
+            uint32 rangeIndexInModel = paramDesc.m_ranges[rangeIndexInParam];
             const ExecutionIndex& currentIndex = GetMemory().GetRageIndex( item.ExecutionIndex );
             int position = currentIndex.GetFromModelRangeIndex(rangeIndexInModel);
             index->GetPrivate()->m_values[rangeIndexInParam] = position;
@@ -4171,14 +4171,14 @@ namespace mu
     {
         //MUTABLE_CPUPROFILER_SCOPE(RunCode_Bool);
 
-        const auto& program = pModel->GetPrivate()->m_program;
-        auto type = program.GetOpType(item.At);
+        const FProgram& program = pModel->GetPrivate()->m_program;
+        OP_TYPE type = program.GetOpType(item.At);
         switch (type)
         {
 
         case OP_TYPE::BO_CONSTANT:
         {
-            auto args = program.GetOpArgs<OP::BoolConstantArgs>(item.At);
+			OP::BoolConstantArgs args = program.GetOpArgs<OP::BoolConstantArgs>(item.At);
             bool result = args.value;
             GetMemory().SetBool( item, result );
             break;
@@ -4186,7 +4186,7 @@ namespace mu
 
         case OP_TYPE::BO_PARAMETER:
         {
-            auto args = program.GetOpArgs<OP::ParameterArgs>(item.At);
+			OP::ParameterArgs args = program.GetOpArgs<OP::ParameterArgs>(item.At);
             bool result = false;
 			Ptr<RangeIndex> index = BuildCurrentOpRangeIndex( item, pParams, pModel, args.variable );
             result = pParams->GetBoolValue( args.variable, index );
@@ -4196,7 +4196,7 @@ namespace mu
 
         case OP_TYPE::BO_LESS:
         {
-            auto args = program.GetOpArgs<OP::BoolLessArgs>(item.At);
+			OP::BoolLessArgs args = program.GetOpArgs<OP::BoolLessArgs>(item.At);
             switch (item.Stage)
             {
             case 0:
@@ -4223,7 +4223,7 @@ namespace mu
 
         case OP_TYPE::BO_AND:
         {
-            auto args = program.GetOpArgs<OP::BoolBinaryArgs>(item.At);
+			OP::BoolBinaryArgs args = program.GetOpArgs<OP::BoolBinaryArgs>(item.At);
             switch (item.Stage)
             {
             case 0:
@@ -4289,7 +4289,7 @@ namespace mu
 
         case OP_TYPE::BO_OR:
         {
-            auto args = program.GetOpArgs<OP::BoolBinaryArgs>(item.At);
+			OP::BoolBinaryArgs args = program.GetOpArgs<OP::BoolBinaryArgs>(item.At);
             switch (item.Stage)
             {
             case 0:
@@ -4355,7 +4355,7 @@ namespace mu
 
         case OP_TYPE::BO_NOT:
         {
-            auto args = program.GetOpArgs<OP::BoolNotArgs>(item.At);
+			OP::BoolNotArgs args = program.GetOpArgs<OP::BoolNotArgs>(item.At);
             switch (item.Stage)
             {
             case 0:
@@ -4378,7 +4378,7 @@ namespace mu
 
         case OP_TYPE::BO_EQUAL_INT_CONST:
         {
-            auto args = program.GetOpArgs<OP::BoolEqualScalarConstArgs>(item.At);
+			OP::BoolEqualScalarConstArgs args = program.GetOpArgs<OP::BoolEqualScalarConstArgs>(item.At);
             switch (item.Stage)
             {
             case 0:
@@ -4421,7 +4421,7 @@ namespace mu
 
         case OP_TYPE::NU_CONSTANT:
         {
-            auto args = pModel->GetPrivate()->m_program.GetOpArgs<OP::IntConstantArgs>(item.At);
+			OP::IntConstantArgs args = pModel->GetPrivate()->m_program.GetOpArgs<OP::IntConstantArgs>(item.At);
             int result = args.value;
             GetMemory().SetInt( item, result );
             break;
@@ -4429,7 +4429,7 @@ namespace mu
 
         case OP_TYPE::NU_PARAMETER:
         {
-            auto args = pModel->GetPrivate()->m_program.GetOpArgs<OP::ParameterArgs>(item.At);
+			OP::ParameterArgs args = pModel->GetPrivate()->m_program.GetOpArgs<OP::ParameterArgs>(item.At);
 			Ptr<RangeIndex> index = BuildCurrentOpRangeIndex( item, pParams, pModel, args.variable );
             int result = pParams->GetIntValue( args.variable, index );
 
@@ -4478,7 +4478,7 @@ namespace mu
 
         case OP_TYPE::SC_CONSTANT:
         {
-            auto args = pModel->GetPrivate()->m_program.GetOpArgs<OP::ScalarConstantArgs>(item.At);
+			OP::ScalarConstantArgs args = pModel->GetPrivate()->m_program.GetOpArgs<OP::ScalarConstantArgs>(item.At);
             float result = args.value;
             GetMemory().SetScalar( item, result );
             break;
@@ -4486,7 +4486,7 @@ namespace mu
 
         case OP_TYPE::SC_PARAMETER:
         {
-            auto args = pModel->GetPrivate()->m_program.GetOpArgs<OP::ParameterArgs>(item.At);
+			OP::ParameterArgs args = pModel->GetPrivate()->m_program.GetOpArgs<OP::ParameterArgs>(item.At);
 			Ptr<RangeIndex> index = BuildCurrentOpRangeIndex( item, pParams, pModel, args.variable );
             float result = pParams->GetFloatValue( args.variable, index );
             GetMemory().SetScalar( item, result );
@@ -4495,7 +4495,7 @@ namespace mu
 
         case OP_TYPE::SC_CURVE:
         {
-            auto args = pModel->GetPrivate()->m_program.GetOpArgs<OP::ScalarCurveArgs>(item.At);
+			OP::ScalarCurveArgs args = pModel->GetPrivate()->m_program.GetOpArgs<OP::ScalarCurveArgs>(item.At);
             switch (item.Stage)
             {
             case 0:
@@ -4527,7 +4527,7 @@ namespace mu
 
         case OP_TYPE::SC_ARITHMETIC:
         {
-            auto args = pModel->GetPrivate()->m_program.GetOpArgs<OP::ArithmeticArgs>(item.At);
+			OP::ArithmeticArgs args = pModel->GetPrivate()->m_program.GetOpArgs<OP::ArithmeticArgs>(item.At);
             switch (item.Stage)
             {
             case 0:
@@ -4583,8 +4583,7 @@ namespace mu
 
 
     //---------------------------------------------------------------------------------------------
-    void
-    CodeRunner::RunCode_String( FScheduledOp& item, const Parameters* pParams, const Model* pModel )
+    void CodeRunner::RunCode_String( FScheduledOp& item, const Parameters* pParams, const Model* pModel )
     {
         MUTABLE_CPUPROFILER_SCOPE(RunCode_String );
 
@@ -4597,7 +4596,7 @@ namespace mu
 			OP::ResourceConstantArgs args = pModel->GetPrivate()->m_program.GetOpArgs<OP::ResourceConstantArgs>( item.At );
             check( args.value < (uint32)pModel->GetPrivate()->m_program.m_constantStrings.Num() );
 
-            auto result = pModel->GetPrivate()->m_program.m_constantStrings[args.value];
+            const std::string& result = pModel->GetPrivate()->m_program.m_constantStrings[args.value];
             GetMemory().SetString( item, new String(result.c_str()) );
 
             break;
@@ -4629,7 +4628,7 @@ namespace mu
 
 		OP_TYPE type = pModel->GetPrivate()->m_program.GetOpType(item.At);
 
-        const auto& program = pModel->GetPrivate()->m_program;
+        const FProgram& program = pModel->GetPrivate()->m_program;
 
         switch ( type )
         {
@@ -4728,7 +4727,7 @@ namespace mu
                 {
                     if ( args.sources[t] )
                     {
-                        auto p = GetMemory().GetColour( FCacheAddress(args.sources[t],item) );
+                        FVector4f p = GetMemory().GetColour( FCacheAddress(args.sources[t],item) );
                         result[t] = p[ args.sourceChannels[t] ];
                     }
                 }
@@ -5278,7 +5277,7 @@ namespace mu
 		}
 
 
-		const mu::FProgram& program = pModel->GetPrivate()->m_program;
+		const FProgram& program = pModel->GetPrivate()->m_program;
 
 		OP_TYPE type = program.GetOpType(item.At);
 		switch (type)
@@ -5544,7 +5543,7 @@ namespace mu
 
 		case OP_TYPE::IM_RESIZE:
 		{
-			auto args = program.GetOpArgs<OP::ImageResizeArgs>(item.At);
+			OP::ImageResizeArgs args = program.GetOpArgs<OP::ImageResizeArgs>(item.At);
 			switch (item.Stage)
 			{
 			case 0:
@@ -5565,7 +5564,7 @@ namespace mu
 
 		case OP_TYPE::IM_RESIZELIKE:
 		{
-			auto args = program.GetOpArgs<OP::ImageResizeLikeArgs>(item.At);
+			OP::ImageResizeLikeArgs args = program.GetOpArgs<OP::ImageResizeLikeArgs>(item.At);
 			switch (item.Stage)
 			{
 			case 0:
@@ -5601,7 +5600,7 @@ namespace mu
 
 		case OP_TYPE::IM_RESIZEREL:
 		{
-			auto args = program.GetOpArgs<OP::ImageResizeRelArgs>(item.At);
+			OP::ImageResizeRelArgs args = program.GetOpArgs<OP::ImageResizeRelArgs>(item.At);
 			switch (item.Stage)
 			{
 			case 0:
diff --git a/Engine/Plugins/Experimental/Mutable/Source/MutableRuntime/Private/MuR/CodeRunnerTasked.cpp b/Engine/Plugins/Experimental/Mutable/Source/MutableRuntime/Private/MuR/CodeRunnerTasked.cpp
--- a/Engine/Plugins/Experimental/Mutable/Source/MutableRuntime/Private/MuR/CodeRunnerTasked.cpp
+++ b/Engine/Plugins/Experimental/Mutable/Source/MutableRuntime/Private/MuR/CodeRunnerTasked.cpp
@@ -229,7 +229,7 @@ namespace mu
 			}
 
 			// Look for completed streaming ops and complete the rom loading
-			for (auto& o : m_romLoadOps)
+			for (FRomLoadOp& o : m_romLoadOps)
 			{
 				if (o.m_romIndex>=0 && m_pSystem->m_pStreamInterface->IsReadCompleted(o.m_streamID))
 				{
@@ -308,7 +308,7 @@ namespace mu
 
 				// If we reached here it means we didn't find an op to wait for. Try to wait for a loading op.
 				// \todo: unify laoding ops with normal ones?
-				for (auto& o : m_romLoadOps)
+				for (FRomLoadOp& o : m_romLoadOps)
 				{
 					if (o.m_romIndex >= 0)
 					{
@@ -1248,7 +1248,7 @@ namespace mu
 			check(RomSize>0);
 
 			CodeRunner::FRomLoadOp* op = nullptr;
-			for (auto& o : runner->m_romLoadOps)
+			for (FRomLoadOp& o : runner->m_romLoadOps)
 			{
 				if (o.m_romIndex < 0)
 				{
@@ -1368,7 +1368,7 @@ namespace mu
 			check(RomSize > 0);
 
 			CodeRunner::FRomLoadOp* op = nullptr;
-			for (auto& o : runner->m_romLoadOps)
+			for (FRomLoadOp& o : runner->m_romLoadOps)
 			{
 				if (o.m_romIndex < 0)
 				{
@@ -1478,7 +1478,7 @@ namespace mu
 
 		FProgram& program = m_pModel->GetPrivate()->m_program;
 
-		auto type = program.GetOpType(item.At);
+		OP_TYPE type = program.GetOpType(item.At);
 
 		switch (type)
 		{
diff --git a/Engine/Plugins/Experimental/Mutable/Source/MutableRuntime/Private/MuR/MeshBufferSet.cpp b/Engine/Plugins/Experimental/Mutable/Source/MutableRuntime/Private/MuR/MeshBufferSet.cpp
--- a/Engine/Plugins/Experimental/Mutable/Source/MutableRuntime/Private/MuR/MeshBufferSet.cpp
+++ b/Engine/Plugins/Experimental/Mutable/Source/MutableRuntime/Private/MuR/MeshBufferSet.cpp
@@ -69,7 +69,7 @@ namespace mu
 
 		m_elementCount = count;
 
-		for (auto& buf : m_buffers)
+		for (MESH_BUFFER& buf : m_buffers)
 		{
 			buf.m_data.SetNumUninitialized(buf.m_elementSize * count, false);
 		}
@@ -379,7 +379,7 @@ namespace mu
 
 		if (fromIndex != toIndex)
 		{
-			for (auto& b : m_buffers)
+			for (MESH_BUFFER& b : m_buffers)
 			{
 				FMemory::Memcpy(&b.m_data[b.m_elementSize * toIndex],
 					&b.m_data[b.m_elementSize * fromIndex],
@@ -573,9 +573,9 @@ namespace mu
 		uint8_t currentIndices[MBS_COUNT];
 		memset(currentIndices, 0, sizeof(currentIndices));
 
-		for (auto& b : m_buffers)
+		for (MESH_BUFFER& b : m_buffers)
 		{
-			for (auto& c : b.m_channels)
+			for (MESH_BUFFER_CHANNEL& c : b.m_channels)
 			{
 				c.m_semanticIndex = currentIndices[c.m_semantic];
 				currentIndices[c.m_semantic]++;
@@ -588,7 +588,7 @@ namespace mu
 	void FMeshBufferSet::UpdateOffsets(int32 b)
 	{
 		uint8_t offset = 0;
-		for (auto& c : m_buffers[b].m_channels)
+		for (MESH_BUFFER_CHANNEL& c : m_buffers[b].m_channels)
 		{
 			if (c.m_offset < offset)
 			{
diff --git a/Engine/Plugins/Experimental/Mutable/Source/MutableRuntime/Private/MuR/OpImageProject.cpp b/Engine/Plugins/Experimental/Mutable/Source/MutableRuntime/Private/MuR/OpImageProject.cpp
--- a/Engine/Plugins/Experimental/Mutable/Source/MutableRuntime/Private/MuR/OpImageProject.cpp
+++ b/Engine/Plugins/Experimental/Mutable/Source/MutableRuntime/Private/MuR/OpImageProject.cpp
@@ -865,7 +865,7 @@ void ImageRasterProjected_Optimised( const Mesh* pMesh,
     check( (int)scratch->culledVertex.Num()==vertexCount );
 
     check( pMesh->GetVertexBuffers().GetElementSize(0)==sizeof(OPTIMISED_VERTEX) );
-    auto pVertices = reinterpret_cast<const OPTIMISED_VERTEX*>( pMesh->GetVertexBuffers().GetBufferData(0) );
+	const OPTIMISED_VERTEX* pVertices = reinterpret_cast<const OPTIMISED_VERTEX*>( pMesh->GetVertexBuffers().GetBufferData(0) );
 
     float fadeEndCos = cosf(fadeEnd);
     for ( int v=0; v<vertexCount; ++v )
@@ -883,7 +883,7 @@ void ImageRasterProjected_Optimised( const Mesh* pMesh,
 
     // Get the indices
     check( pMesh->GetIndexBuffers().GetElementSize(0)==4 );
-    auto pIndices = reinterpret_cast<const uint32*>( pMesh->GetIndexBuffers().GetBufferData(0) );
+	const uint32* pIndices = reinterpret_cast<const uint32*>( pMesh->GetIndexBuffers().GetBufferData(0) );
 
 
     // The mesh is supposed to contain only the faces in the selected block
@@ -944,7 +944,7 @@ void ImageRasterProjected_OptimisedWrapping( const Mesh* pMesh,
     check( (int)scratch->culledVertex.Num()==vertexCount );
 
     check( pMesh->GetVertexBuffers().GetElementSize(0)==sizeof(OPTIMISED_VERTEX_WRAPPING) );
-    auto pVertices = reinterpret_cast<const OPTIMISED_VERTEX_WRAPPING*>( pMesh->GetVertexBuffers().GetBufferData(0) );
+	const OPTIMISED_VERTEX_WRAPPING* pVertices = reinterpret_cast<const OPTIMISED_VERTEX_WRAPPING*>( pMesh->GetVertexBuffers().GetBufferData(0) );
 
     float fadeEndCos = cosf(fadeEnd);
     for ( int v=0; v<vertexCount; ++v )
@@ -968,7 +968,7 @@ void ImageRasterProjected_OptimisedWrapping( const Mesh* pMesh,
 
     // Get the indices
     check( pMesh->GetIndexBuffers().GetElementSize(0)==4 );
-    auto pIndices = reinterpret_cast<const uint32*>( pMesh->GetIndexBuffers().GetBufferData(0) );
+	const uint32* pIndices = reinterpret_cast<const uint32*>( pMesh->GetIndexBuffers().GetBufferData(0) );
 
 
     // The mesh is supposed to contain only the faces in the selected block
@@ -1303,7 +1303,7 @@ inline void MeshCreateCollapsedVertexMap( const Mesh* pMesh,
             MESH_BUFFER_SEMANTIC sem = MBSPriv.m_buffers[b].m_channels[c].m_semantic;
             int semIndex = MBSPriv.m_buffers[b].m_channels[c].m_semanticIndex;
 
-            auto it = UntypedMeshBufferIteratorConst(pMesh->GetVertexBuffers(), sem, semIndex);
+			UntypedMeshBufferIteratorConst it = UntypedMeshBufferIteratorConst(pMesh->GetVertexBuffers(), sem, semIndex);
 
 			box<FVector3f> meshBoundingBox;
 			Octree* pOctree = nullptr;
@@ -1782,11 +1782,11 @@ void MeshProject_Optimised_Wrapping( const Mesh* pMesh,
     // Get the vertices
     int vertexCount = pMesh->GetVertexCount();
     check( pMesh->GetVertexBuffers().GetElementSize(0)==sizeof(OPTIMISED_VERTEX_WRAPPING) );
-    auto pVertices = reinterpret_cast<const OPTIMISED_VERTEX_WRAPPING*>( pMesh->GetVertexBuffers().GetBufferData(0) );
+	const OPTIMISED_VERTEX_WRAPPING* pVertices = reinterpret_cast<const OPTIMISED_VERTEX_WRAPPING*>( pMesh->GetVertexBuffers().GetBufferData(0) );
 
     // Get the indices
     check( pMesh->GetIndexBuffers().GetElementSize(0)==4 );
-    auto pIndices = reinterpret_cast<const uint32*>( pMesh->GetIndexBuffers().GetBufferData(0) );
+	const uint32* pIndices = reinterpret_cast<const uint32*>( pMesh->GetIndexBuffers().GetBufferData(0) );
     int faceCount = pMesh->GetFaceCount();
 
 	TArray<PROJECTED_VERTEX> projectedPositions;
@@ -2581,16 +2581,16 @@ MeshPtr MeshProject_Optimised( const Mesh* pMesh,
         pResult = CreateMeshOptimisedForProjection( layout );
         pResult->GetVertexBuffers().SetElementCount(vertexCount);
         pResult->GetIndexBuffers().SetElementCount(indexCount);
-        auto pResultIndices = reinterpret_cast<uint32*>( pResult->GetIndexBuffers().GetBufferData(0) );
-        auto pResultVertices = reinterpret_cast<OPTIMISED_VERTEX*>( pResult->GetVertexBuffers().GetBufferData(0) );
+		uint32* pResultIndices = reinterpret_cast<uint32*>( pResult->GetIndexBuffers().GetBufferData(0) );
+		OPTIMISED_VERTEX* pResultVertices = reinterpret_cast<OPTIMISED_VERTEX*>( pResult->GetVertexBuffers().GetBufferData(0) );
 
         // Get the vertices
         check( pMesh->GetVertexBuffers().GetElementSize(0)==sizeof(OPTIMISED_VERTEX) );
-        auto pVertices = reinterpret_cast<const OPTIMISED_VERTEX*>( pMesh->GetVertexBuffers().GetBufferData(0) );
+		const OPTIMISED_VERTEX* pVertices = reinterpret_cast<const OPTIMISED_VERTEX*>( pMesh->GetVertexBuffers().GetBufferData(0) );
 
         // Get the indices
         check( pMesh->GetIndexBuffers().GetElementSize(0)==4 );
-        auto pIndices = reinterpret_cast<const uint32*>( pMesh->GetIndexBuffers().GetBufferData(0) );
+		const uint32* pIndices = reinterpret_cast<const uint32*>( pMesh->GetIndexBuffers().GetBufferData(0) );
         int faceCount = pMesh->GetFaceCount();
 
         if (projector.type==PROJECTOR_TYPE::PLANAR)
@@ -2623,8 +2623,8 @@ MeshPtr MeshProject_Optimised( const Mesh* pMesh,
         pResult = CreateMeshOptimisedForWrappingProjection( layout );
         pResult->GetVertexBuffers().SetElementCount(vertexCount);
         pResult->GetIndexBuffers().SetElementCount(indexCount);
-        auto pResultIndices = reinterpret_cast<uint32*>( pResult->GetIndexBuffers().GetBufferData(0) );
-        auto pResultVertices = reinterpret_cast<OPTIMISED_VERTEX_WRAPPING*>( pResult->GetVertexBuffers().GetBufferData(0) );
+		uint32* pResultIndices = reinterpret_cast<uint32*>( pResult->GetIndexBuffers().GetBufferData(0) );
+		OPTIMISED_VERTEX_WRAPPING* pResultVertices = reinterpret_cast<OPTIMISED_VERTEX_WRAPPING*>( pResult->GetVertexBuffers().GetBufferData(0) );
 
         MeshProject_Optimised_Wrapping( pMesh,
                                         projectorPosition, projectorDirection,
diff --git a/Engine/Plugins/Experimental/Mutable/Source/MutableRuntime/Private/MuR/OpMeshFormat.cpp b/Engine/Plugins/Experimental/Mutable/Source/MutableRuntime/Private/MuR/OpMeshFormat.cpp
--- a/Engine/Plugins/Experimental/Mutable/Source/MutableRuntime/Private/MuR/OpMeshFormat.cpp
+++ b/Engine/Plugins/Experimental/Mutable/Source/MutableRuntime/Private/MuR/OpMeshFormat.cpp
@@ -22,11 +22,6 @@
 namespace mu
 {
 
-
-	//-------------------------------------------------------------------------------------------------
-	// Eric Lengyel method's
-	// http://www.terathon.com/code/tangent.html
-	//-------------------------------------------------------------------------------------------------
 	namespace
 	{
 
@@ -213,7 +208,7 @@ namespace mu
 						{
 							for (int v = 0; v < vCount; ++v)
 							{
-								auto pTypedResultBuf = (float*)pResultBuf;
+								float* pTypedResultBuf = (float*)pResultBuf;
 								for (int i = 0; i < resultComponents; ++i)
 								{
 									pTypedResultBuf[i] = 1.0f;
@@ -227,7 +222,7 @@ namespace mu
 						{
 							for (int v = 0; v < vCount; ++v)
 							{
-								auto pTypedResultBuf = (uint8_t*)pResultBuf;
+								uint8_t* pTypedResultBuf = (uint8_t*)pResultBuf;
 								for (int i = 0; i < resultComponents; ++i)
 								{
 									pTypedResultBuf[i] = 255;
@@ -241,7 +236,7 @@ namespace mu
 						{
 							for (int v = 0; v < vCount; ++v)
 							{
-								auto pTypedResultBuf = (uint16*)pResultBuf;
+								uint16* pTypedResultBuf = (uint16*)pResultBuf;
 								for (int i = 0; i < resultComponents; ++i)
 								{
 									pTypedResultBuf[i] = 65535;
@@ -322,7 +317,7 @@ namespace mu
 
 
 							// Add the tangent sign
-							auto pData = (uint8_t*)pResultBuf;
+							uint8_t* pData = (uint8_t*)pResultBuf;
 
 							// Look for the full tangent space
 							int tanXBuf, tanXChan, tanYBuf, tanYChan, tanZBuf, tanZChan;
@@ -577,7 +572,7 @@ namespace mu
 					{
 						int resultBuf = 0;
 						int resultChan = 0;
-						auto& formBuffs = pResult->GetVertexBuffers();
+						FMeshBufferSet& formBuffs = pResult->GetVertexBuffers();
 						formBuffs.FindChannel(MBS_BONEINDICES, Channel.m_semanticIndex, &resultBuf, &resultChan);
 						if (resultBuf >= 0)
 						{
@@ -585,7 +580,7 @@ namespace mu
 							int32_t maxBoneIndex = 0;
 							for (int v = 0; v < VertexBuffers.GetElementCount(); ++v)
 							{
-								auto va = it.GetAsVec8i();
+								vec<int32_t, 8> va = it.GetAsVec8i();
 								for (int c = 0; c < it.GetComponents(); ++c)
 								{
 									maxBoneIndex = FMath::Max(maxBoneIndex, va[c]);
@@ -593,7 +588,7 @@ namespace mu
 								++it;
 							}
 
-							auto& format = formBuffs.m_buffers[resultBuf].m_channels[resultChan].m_format;
+							MESH_BUFFER_FORMAT& format = formBuffs.m_buffers[resultBuf].m_channels[resultChan].m_format;
 							if (maxBoneIndex > 0xffff && (format == MBF_UINT8 || format == MBF_UINT16))
 							{
 								format = MBF_UINT32;
@@ -663,7 +658,7 @@ namespace mu
 		pResult->SetPhysicsBody(pSource->GetPhysicsBody());
 
 		pResult->m_layouts.Empty();
-		for (const auto& Layout : pSource->m_layouts)
+		for (const Ptr<const Layout>& Layout : pSource->m_layouts)
 		{
 			pResult->m_layouts.Add(Layout->Clone());
 		}
diff --git a/Engine/Plugins/Experimental/Mutable/Source/MutableRuntime/Private/MuR/Operations.cpp b/Engine/Plugins/Experimental/Mutable/Source/MutableRuntime/Private/MuR/Operations.cpp
--- a/Engine/Plugins/Experimental/Mutable/Source/MutableRuntime/Private/MuR/Operations.cpp
+++ b/Engine/Plugins/Experimental/Mutable/Source/MutableRuntime/Private/MuR/Operations.cpp
@@ -431,7 +431,7 @@ namespace mu
 
         case OP_TYPE::SC_CURVE:
         {
-            auto args = program.GetOpArgs<OP::ScalarCurveArgs>(at);
+			OP::ScalarCurveArgs args = program.GetOpArgs<OP::ScalarCurveArgs>(at);
             f(args.time );
             break;
         }
@@ -444,7 +444,7 @@ namespace mu
         case OP_TYPE::LA_CONDITIONAL:
         case OP_TYPE::IN_CONDITIONAL:
         {
-            auto args = program.GetOpArgs<OP::ConditionalArgs>(at);
+			OP::ConditionalArgs args = program.GetOpArgs<OP::ConditionalArgs>(at);
             f(args.condition );
             f(args.yes );
             f(args.no );
@@ -498,7 +498,7 @@ namespace mu
         //-------------------------------------------------------------------------------------
         case OP_TYPE::BO_LESS:
         {
-            auto args = program.GetOpArgs<OP::BoolLessArgs>(at);
+			OP::BoolLessArgs args = program.GetOpArgs<OP::BoolLessArgs>(at);
             f(args.a );
             f(args.b );
             break;
@@ -507,7 +507,7 @@ namespace mu
         //-------------------------------------------------------------------------------------
         case OP_TYPE::BO_EQUAL_INT_CONST:
         {
-            auto args = program.GetOpArgs<OP::BoolEqualScalarConstArgs>(at);
+			OP::BoolEqualScalarConstArgs args = program.GetOpArgs<OP::BoolEqualScalarConstArgs>(at);
             f(args.value );
             break;
         }
@@ -516,7 +516,7 @@ namespace mu
         case OP_TYPE::BO_AND:
         case OP_TYPE::BO_OR:
         {
-            auto args = program.GetOpArgs<OP::BoolBinaryArgs>(at);
+			OP::BoolBinaryArgs args = program.GetOpArgs<OP::BoolBinaryArgs>(at);
             f(args.a );
             f(args.b );
             break;
@@ -525,7 +525,7 @@ namespace mu
         //-------------------------------------------------------------------------------------
         case OP_TYPE::BO_NOT:
         {
-            auto args = program.GetOpArgs<OP::BoolNotArgs>(at);
+			OP::BoolNotArgs args = program.GetOpArgs<OP::BoolNotArgs>(at);
             f(args.source );
             break;
         }
@@ -533,7 +533,7 @@ namespace mu
         //-------------------------------------------------------------------------------------
         case OP_TYPE::SC_MULTIPLYADD:
         {
-            auto args = program.GetOpArgs<OP::ScalarMultiplyAddArgs>(at);
+			OP::ScalarMultiplyAddArgs args = program.GetOpArgs<OP::ScalarMultiplyAddArgs>(at);
             f(args.factor0 );
             f(args.factor1 );
             f(args.add );
@@ -543,7 +543,7 @@ namespace mu
         //-------------------------------------------------------------------------------------
         case OP_TYPE::SC_ARITHMETIC:
         {
-            auto args = program.GetOpArgs<OP::ArithmeticArgs>(at);
+			OP::ArithmeticArgs args = program.GetOpArgs<OP::ArithmeticArgs>(at);
             f(args.a );
             f(args.b );
             break;
@@ -552,7 +552,7 @@ namespace mu
         //-------------------------------------------------------------------------------------
         case OP_TYPE::CO_SAMPLEIMAGE:
         {
-            auto args = program.GetOpArgs<OP::ColourSampleImageArgs>(at);
+			OP::ColourSampleImageArgs args = program.GetOpArgs<OP::ColourSampleImageArgs>(at);
             f(args.image );
             f(args.x );
             f(args.y );
@@ -561,7 +561,7 @@ namespace mu
 
         case OP_TYPE::CO_SWIZZLE:
         {
-            auto args = program.GetOpArgs<OP::ColourSwizzleArgs>(at);
+			OP::ColourSwizzleArgs args = program.GetOpArgs<OP::ColourSwizzleArgs>(at);
             for (int t=0;t<MUTABLE_OP_MAX_SWIZZLE_CHANNELS;++t)
             {
                 f(args.sources[t] );
@@ -571,21 +571,21 @@ namespace mu
 
         case OP_TYPE::CO_IMAGESIZE:
         {
-            auto args = program.GetOpArgs<OP::ColourSampleImageArgs>(at);
+			OP::ColourSampleImageArgs args = program.GetOpArgs<OP::ColourSampleImageArgs>(at);
             f(args.image );
             break;
         }
 
         case OP_TYPE::CO_LAYOUTBLOCKTRANSFORM:
         {
-            auto args = program.GetOpArgs<OP::ColourLayoutBlockTransformArgs>(at);
+			OP::ColourLayoutBlockTransformArgs args = program.GetOpArgs<OP::ColourLayoutBlockTransformArgs>(at);
             f(args.layout );
             break;
         }
 
         case OP_TYPE::CO_FROMSCALARS:
         {
-            auto args = program.GetOpArgs<OP::ColourFromScalarsArgs>(at);
+			OP::ColourFromScalarsArgs args = program.GetOpArgs<OP::ColourFromScalarsArgs>(at);
             f(args.x);
             f(args.y);
             f(args.z);
@@ -595,7 +595,7 @@ namespace mu
 
         case OP_TYPE::CO_ARITHMETIC:
         {
-            auto args = program.GetOpArgs<OP::ArithmeticArgs>(at);
+			OP::ArithmeticArgs args = program.GetOpArgs<OP::ArithmeticArgs>(at);
             f(args.a);
             f(args.b);
             break;
@@ -604,7 +604,7 @@ namespace mu
         //-------------------------------------------------------------------------------------
         case OP_TYPE::IM_LAYER:
         {
-            auto args = program.GetOpArgs<OP::ImageLayerArgs>(at);
+			OP::ImageLayerArgs args = program.GetOpArgs<OP::ImageLayerArgs>(at);
             f(args.base );
             if ( args.mask )
             {
@@ -616,7 +616,7 @@ namespace mu
 
         case OP_TYPE::IM_LAYERCOLOUR:
         {
-            auto args = program.GetOpArgs<OP::ImageLayerColourArgs>(at);
+			OP::ImageLayerColourArgs args = program.GetOpArgs<OP::ImageLayerColourArgs>(at);
             f(args.base );
             if ( args.mask )
             {
@@ -628,7 +628,7 @@ namespace mu
 
         case OP_TYPE::IM_MULTILAYER:
         {
-            auto args = program.GetOpArgs<OP::ImageMultiLayerArgs>(at);
+			OP::ImageMultiLayerArgs args = program.GetOpArgs<OP::ImageMultiLayerArgs>(at);
             f(args.rangeSize );
             f(args.base );
             if ( args.mask )
@@ -641,7 +641,7 @@ namespace mu
 
 		case OP_TYPE::IM_NORMALCOMPOSITE:
 		{
-			auto args = program.GetOpArgs<OP::ImageNormalCompositeArgs>(at);
+			OP::ImageNormalCompositeArgs args = program.GetOpArgs<OP::ImageNormalCompositeArgs>(at);
 			f(args.base);
 			f(args.normal);
 
@@ -650,28 +650,28 @@ namespace mu
 
         case OP_TYPE::IM_PIXELFORMAT:
         {
-            auto args = program.GetOpArgs<OP::ImagePixelFormatArgs>(at);
+			OP::ImagePixelFormatArgs args = program.GetOpArgs<OP::ImagePixelFormatArgs>(at);
             f(args.source );
             break;
         }
 
         case OP_TYPE::IM_MIPMAP:
         {
-            auto args = program.GetOpArgs<OP::ImageMipmapArgs>(at);
+			OP::ImageMipmapArgs args = program.GetOpArgs<OP::ImageMipmapArgs>(at);
             f(args.source );
             break;
         }
 
         case OP_TYPE::IM_RESIZE:
         {
-            auto args = program.GetOpArgs<OP::ImageResizeArgs>(at);
+			OP::ImageResizeArgs args = program.GetOpArgs<OP::ImageResizeArgs>(at);
             f(args.source );
             break;
         }
 
         case OP_TYPE::IM_RESIZELIKE:
         {
-            auto args = program.GetOpArgs<OP::ImageResizeLikeArgs>(at);
+			OP::ImageResizeLikeArgs args = program.GetOpArgs<OP::ImageResizeLikeArgs>(at);
             f(args.source );
             f(args.sizeSource );
             break;
@@ -679,21 +679,21 @@ namespace mu
 
         case OP_TYPE::IM_RESIZEREL:
         {
-            auto args = program.GetOpArgs<OP::ImageResizeRelArgs>(at);
+			OP::ImageResizeRelArgs args = program.GetOpArgs<OP::ImageResizeRelArgs>(at);
             f(args.source );
             break;
         }
 
         case OP_TYPE::IM_BLANKLAYOUT:
         {
-            auto args = program.GetOpArgs<OP::ImageBlankLayoutArgs>(at);
+			OP::ImageBlankLayoutArgs args = program.GetOpArgs<OP::ImageBlankLayoutArgs>(at);
             f(args.layout );
             break;
         }
 
         case OP_TYPE::IM_COMPOSE:
         {
-            auto args = program.GetOpArgs<OP::ImageComposeArgs>(at);
+			OP::ImageComposeArgs args = program.GetOpArgs<OP::ImageComposeArgs>(at);
             f(args.layout );
             f(args.base );
             f(args.blockImage );
@@ -703,7 +703,7 @@ namespace mu
 
         case OP_TYPE::IM_DIFFERENCE:
         {
-            auto args = program.GetOpArgs<OP::ImageDifferenceArgs>(at);
+			OP::ImageDifferenceArgs args = program.GetOpArgs<OP::ImageDifferenceArgs>(at);
             f(args.a );
             f(args.b );
             break;
@@ -711,7 +711,7 @@ namespace mu
 
         case OP_TYPE::IM_INTERPOLATE:
         {
-            auto args = program.GetOpArgs<OP::ImageInterpolateArgs>(at);
+			OP::ImageInterpolateArgs args = program.GetOpArgs<OP::ImageInterpolateArgs>(at);
             f(args.factor );
 
             for (int t=0;t<MUTABLE_OP_MAX_INTERPOLATE_COUNT;++t)
@@ -723,7 +723,7 @@ namespace mu
 
         case OP_TYPE::IM_INTERPOLATE3:
         {
-            auto args = program.GetOpArgs<OP::ImageInterpolate3Args>(at);
+			OP::ImageInterpolate3Args args = program.GetOpArgs<OP::ImageInterpolate3Args>(at);
             f(args.factor1 );
             f(args.factor2 );
             f(args.target0 );
@@ -734,7 +734,7 @@ namespace mu
 
         case OP_TYPE::IM_SWIZZLE:
         {
-            auto args = program.GetOpArgs<OP::ImageSwizzleArgs>(at);
+			OP::ImageSwizzleArgs args = program.GetOpArgs<OP::ImageSwizzleArgs>(at);
             for (int t=0;t<MUTABLE_OP_MAX_SWIZZLE_CHANNELS;++t)
             {
                 f(args.sources[t] );
@@ -744,7 +744,7 @@ namespace mu
 
         case OP_TYPE::IM_SATURATE:
         {
-            auto args = program.GetOpArgs<OP::ImageSaturateArgs>(at);
+			OP::ImageSaturateArgs args = program.GetOpArgs<OP::ImageSaturateArgs>(at);
             f(args.base );
             f(args.factor );
             break;
@@ -752,14 +752,14 @@ namespace mu
 
         case OP_TYPE::IM_LUMINANCE:
         {
-            auto args = program.GetOpArgs<OP::ImageLuminanceArgs>(at);
+			OP::ImageLuminanceArgs args = program.GetOpArgs<OP::ImageLuminanceArgs>(at);
             f(args.base );
             break;
         }
 
         case OP_TYPE::IM_SELECTCOLOUR:
         {
-            auto args = program.GetOpArgs<OP::ImageSelectColourArgs>(at);
+			OP::ImageSelectColourArgs args = program.GetOpArgs<OP::ImageSelectColourArgs>(at);
             f(args.base );
             f(args.colour );
             break;
@@ -767,7 +767,7 @@ namespace mu
 
         case OP_TYPE::IM_COLOURMAP:
         {
-            auto args = program.GetOpArgs<OP::ImageColourMapArgs>(at);
+			OP::ImageColourMapArgs args = program.GetOpArgs<OP::ImageColourMapArgs>(at);
             f(args.base );
             f(args.mask );
             f(args.map );
@@ -776,7 +776,7 @@ namespace mu
 
         case OP_TYPE::IM_GRADIENT:
         {
-            auto args = program.GetOpArgs<OP::ImageGradientArgs>(at);
+			OP::ImageGradientArgs args = program.GetOpArgs<OP::ImageGradientArgs>(at);
             f(args.colour0 );
             f(args.colour1 );
             break;
@@ -784,7 +784,7 @@ namespace mu
 
         case OP_TYPE::IM_BINARISE:
         {
-            auto args = program.GetOpArgs<OP::ImageBinariseArgs>(at);
+			OP::ImageBinariseArgs args = program.GetOpArgs<OP::ImageBinariseArgs>(at);
             f(args.base );
             f(args.threshold );
             break;
@@ -792,21 +792,21 @@ namespace mu
 
         case OP_TYPE::IM_PLAINCOLOUR:
         {
-            auto args = program.GetOpArgs<OP::ImagePlainColourArgs>(at);
+			OP::ImagePlainColourArgs args = program.GetOpArgs<OP::ImagePlainColourArgs>(at);
             f(args.colour );
             break;
         }
 
         case OP_TYPE::IM_CROP:
         {
-            auto args = program.GetOpArgs<OP::ImageCropArgs>(at);
+			OP::ImageCropArgs args = program.GetOpArgs<OP::ImageCropArgs>(at);
             f(args.source );
             break;
         }
 
         case OP_TYPE::IM_PATCH:
         {
-            auto args = program.GetOpArgs<OP::ImagePatchArgs>(at);
+			OP::ImagePatchArgs args = program.GetOpArgs<OP::ImagePatchArgs>(at);
             f(args.base );
             f(args.patch );
             break;
@@ -814,7 +814,7 @@ namespace mu
 
         case OP_TYPE::IM_RASTERMESH:
         {
-            auto args = program.GetOpArgs<OP::ImageRasterMeshArgs>(at);
+			OP::ImageRasterMeshArgs args = program.GetOpArgs<OP::ImageRasterMeshArgs>(at);
             f(args.mesh );
             f(args.image );
             f(args.mask );
@@ -825,14 +825,14 @@ namespace mu
 
         case OP_TYPE::IM_MAKEGROWMAP:
         {
-            auto args = program.GetOpArgs<OP::ImageMakeGrowMapArgs>(at);
+			OP::ImageMakeGrowMapArgs args = program.GetOpArgs<OP::ImageMakeGrowMapArgs>(at);
             f(args.mask );
             break;
         }
 
         case OP_TYPE::IM_DISPLACE:
         {
-            auto args = program.GetOpArgs<OP::ImageDisplaceArgs>(at);
+			OP::ImageDisplaceArgs args = program.GetOpArgs<OP::ImageDisplaceArgs>(at);
             f(args.source );
             f(args.displacementMap );
             break;
@@ -840,7 +840,7 @@ namespace mu
 
 		case OP_TYPE::IM_INVERT:
 		{
-			auto args = program.GetOpArgs<OP::ImageInvertArgs>(at);
+			OP::ImageInvertArgs args = program.GetOpArgs<OP::ImageInvertArgs>(at);
 			f(args.base);
 			break;
 		}
