#rnx

diff --git a/Engine/Plugins/Experimental/Mutable/Source/CustomizableObject/Classes/MuCO/CustomizableObject.h b/Engine/Plugins/Experimental/Mutable/Source/CustomizableObject/Classes/MuCO/CustomizableObject.h
--- a/Engine/Plugins/Experimental/Mutable/Source/CustomizableObject/Classes/MuCO/CustomizableObject.h
+++ b/Engine/Plugins/Experimental/Mutable/Source/CustomizableObject/Classes/MuCO/CustomizableObject.h
@@ -1223,7 +1223,7 @@ private:
 	// This is a manual version number for the binary blobs in this asset.
 	// Increasing it invalidates all the previously compiled models.
 	// Warning: If while merging code both versions have changed, take the highest+1.
-	static const int32 CurrentSupportedVersion = 361;
+	static const int32 CurrentSupportedVersion = 362;
 
 public:
 
diff --git a/Engine/Plugins/Experimental/Mutable/Source/CustomizableObject/Private/MuCO/CustomizableInstancePrivateData.h b/Engine/Plugins/Experimental/Mutable/Source/CustomizableObject/Private/MuCO/CustomizableInstancePrivateData.h
--- a/Engine/Plugins/Experimental/Mutable/Source/CustomizableObject/Private/MuCO/CustomizableInstancePrivateData.h
+++ b/Engine/Plugins/Experimental/Mutable/Source/CustomizableObject/Private/MuCO/CustomizableInstancePrivateData.h
@@ -180,6 +180,9 @@ public:
 	// Handle used to store a streaming request operation if one is ongoing.
 	TSharedPtr<FStreamableHandle> StreamingHandle;
 
+	// Only used in LiveUpdateMode to reuse core instances between updates and their temp data to speed up updates, but spend way more memory
+	mu::Instance::ID LiveUpdateModeInstanceID = 0;
+
 	FCustomizableInstanceComponentData* GetComponentData(int32 ComponentIndex);
 	const FCustomizableInstanceComponentData* GetComponentData(int32 ComponentIndex) const;
 
diff --git a/Engine/Plugins/Experimental/Mutable/Source/CustomizableObject/Private/MuCO/CustomizableObjectInstance.cpp b/Engine/Plugins/Experimental/Mutable/Source/CustomizableObject/Private/MuCO/CustomizableObjectInstance.cpp
--- a/Engine/Plugins/Experimental/Mutable/Source/CustomizableObject/Private/MuCO/CustomizableObjectInstance.cpp
+++ b/Engine/Plugins/Experimental/Mutable/Source/CustomizableObject/Private/MuCO/CustomizableObjectInstance.cpp
@@ -425,6 +425,19 @@ void UCustomizableObjectInstance::BeginDestroy()
 	BeginDestroyDelegate.Broadcast(this);
 	BeginDestroyNativeDelegate.Broadcast(this);
 
+	// Release the Live Instance ID if there it hadn't been released before
+	if (PrivateData && PrivateData->LiveUpdateModeInstanceID)
+	{
+		// If FCustomizableObjectSystemPrivate::SSystem is nullptr it means it has already been destroyed, no point in registering an instanceID release
+		// since the Mutable system has already been destroyed. Just checking UCustomizableObjectSystem::GetInstance() will try to recreate the system when
+		// everything is shutting down, so it's better to check FCustomizableObjectSystemPrivate::SSystem first here
+		if (FCustomizableObjectSystemPrivate::SSystem && UCustomizableObjectSystem::GetInstance() && UCustomizableObjectSystem::GetInstance()->GetPrivate())
+		{
+			UCustomizableObjectSystem::GetInstance()->GetPrivate()->InitInstanceIDRelease(PrivateData->LiveUpdateModeInstanceID);
+			PrivateData->LiveUpdateModeInstanceID = 0;
+		}
+	}
+
 	if (PrivateData && PrivateData->StreamingHandle)
 	{
 		PrivateData->StreamingHandle->CancelHandle();
diff --git a/Engine/Plugins/Experimental/Mutable/Source/CustomizableObject/Private/MuCO/CustomizableObjectSystem.cpp b/Engine/Plugins/Experimental/Mutable/Source/CustomizableObject/Private/MuCO/CustomizableObjectSystem.cpp
--- a/Engine/Plugins/Experimental/Mutable/Source/CustomizableObject/Private/MuCO/CustomizableObjectSystem.cpp
+++ b/Engine/Plugins/Experimental/Mutable/Source/CustomizableObject/Private/MuCO/CustomizableObjectSystem.cpp
@@ -897,6 +897,15 @@ void FCustomizableObjectSystemPrivate::InitDiscardResourcesSkeletalMesh(UCustomi
 }
 
 
+void FCustomizableObjectSystemPrivate::InitInstanceIDRelease(mu::Instance::ID IDToRelease)
+{
+	check(IsInGameThread());
+
+	const TSharedRef<FMutableOperation> Operation = MakeShared<FMutableOperation>(FMutableOperation::CreateInstanceIDRelease(IDToRelease));
+	MutableOperationQueue.Enqueue(FMutableQueueElem::Create(Operation, FMutableQueueElem::EQueuePriorityType::High, 0));
+}
+
+
 bool UCustomizableObjectSystem::IsReplaceDiscardedWithReferenceMeshEnabled() const
 {
 	if (Private.IsValid())
@@ -1209,9 +1218,30 @@ namespace impl
 		mu::System* System = UCustomizableObjectSystem::GetInstance()->GetPrivate()->MutableSystem.get();
 		check(System != nullptr);
 
-		// For now, we are forcing the recreation of mutable-side instances with every update.
-		OperationData->InstanceID = System->NewInstance(Model);
-		UE_LOG(LogMutable, Log, TEXT("Creating instance with id [%d] "), OperationData->InstanceID)
+		if (OperationData->bLiveUpdateMode)
+		{
+			if (OperationData->InstanceID == 0)
+			{
+				// It's the first update since the instance was put in LiveUpdate Mode, this ID will be reused from now on
+				OperationData->InstanceID = System->NewInstance(Model);
+				UE_LOG(LogMutable, Verbose, TEXT("Creating Mutable instance with id [%d] for reuse "), OperationData->InstanceID);
+			}
+			else
+			{
+				// The instance was already in LiveUpdate Mode, the ID is reused
+				check(OperationData->InstanceID);
+				UE_LOG(LogMutable, Verbose, TEXT("Reusing Mutable instance with id [%d] "), OperationData->InstanceID);
+			}
+		}
+		else
+		{
+			// In non-LiveUpdate mode, we are forcing the recreation of mutable-side instances with every update.
+			check(OperationData->InstanceID == 0);
+			OperationData->InstanceID = System->NewInstance(Model);
+			UE_LOG(LogMutable, Verbose, TEXT("Creating Mutable instance with id [%d] "), OperationData->InstanceID);
+		}
+
+		check(OperationData->InstanceID != 0);
 
 		const mu::Instance* Instance = nullptr;
 
@@ -1625,10 +1655,27 @@ namespace impl
 			MutableSystem->EndUpdate(OperationData->InstanceID);
 			OperationData->InstanceUpdateData.Clear();
 
+			if (!OperationData->bLiveUpdateMode)
+			{
+				MutableSystem->ReleaseInstance(OperationData->InstanceID);
+				OperationData->InstanceID = 0;
+			}
+		}
+	}
+
+
+	void Task_Mutable_ReleaseInstanceID(TSharedPtr<FMutableOperationData> OperationData, mu::SystemPtr MutableSystem)
+	{
+		MUTABLE_CPUPROFILER_SCOPE(Task_Mutable_ReleaseInstanceID)
+
+		// This runs in a worker thread.
+		check(OperationData.IsValid());
+
+		if (OperationData->InstanceID > 0)
+		{
 			MutableSystem->ReleaseInstance(OperationData->InstanceID);
 			OperationData->InstanceID = 0;
 		}
-
 	}
 
 
@@ -1877,6 +1924,20 @@ namespace impl
 			return;
 		}
 
+		UCustomizableInstancePrivateData* ObjectInstancePrivateData = ObjectInstance->GetPrivate();
+		check(ObjectInstancePrivateData != nullptr);
+
+		if (OperationData->bLiveUpdateMode)
+		{
+			check(OperationData->InstanceID != 0);
+
+			if (ObjectInstancePrivateData->LiveUpdateModeInstanceID == 0)
+			{
+				// From now this instance will reuse this InstanceID until it gets out of LiveUpdateMode
+				ObjectInstancePrivateData->LiveUpdateModeInstanceID = OperationData->InstanceID;
+			}
+		}
+
 		const UCustomizableObject* CustomizableObject = ObjectInstance->GetCustomizableObject(); 
 		if (!CustomizableObject)
 		{
@@ -1885,9 +1946,6 @@ namespace impl
 		}
 
 		check(OperationData.IsValid());
-
-		UCustomizableInstancePrivateData* ObjectInstancePrivateData = ObjectInstance->GetPrivate();
-		check(ObjectInstancePrivateData != nullptr);
 		
 		// Process the parameter decorations if requested
 		if (bBuildParameterDecorations)
@@ -1971,6 +2029,50 @@ namespace impl
 	}
 
 
+	/** Enqueue the release ID operation in the Mutable queue */
+	void Task_Game_ReleaseInstanceID(const mu::Instance::ID IDToRelease)
+	{
+		UCustomizableObjectSystem* System = UCustomizableObjectSystem::GetInstance();
+		check(System != nullptr);
+
+		FCustomizableObjectSystemPrivate* SystemPrivateData = System->GetPrivate();
+		check(SystemPrivateData != nullptr);
+
+		mu::SystemPtr MutableSystem = SystemPrivateData->MutableSystem;
+
+		// Task: Release Instance ID
+		//-------------------------------------------------------------
+		TSharedPtr<FMutableOperationData> CurrentOperationData = MakeShared<FMutableOperationData>();
+		check(CurrentOperationData);
+		CurrentOperationData->InstanceID = IDToRelease;
+
+		FGraphEventRef Mutable_GetMeshTask;
+		{
+			// Task inputs
+			TSharedPtr<FMutableOperation> CurrentMutableOperation = SystemPrivateData->CurrentMutableOperation;
+			check(CurrentMutableOperation);
+
+			Mutable_GetMeshTask = SystemPrivateData->AddMutableThreadTask(
+				TEXT("Task_Mutable_ReleaseInstanceID"),
+				[CurrentOperationData, MutableSystem]()
+				{
+					impl::Task_Mutable_ReleaseInstanceID(CurrentOperationData, MutableSystem);
+				},
+				UE::Tasks::ETaskPriority::BackgroundHigh);
+		}
+	}
+
+
+	/** Enqueue the release ID operation in the Mutable queue */
+	void Task_Game_ReleaseInstanceID(TSharedPtr<FMutableOperation> Operation)
+	{
+		check(Operation);
+		check(Operation->Type == FMutableOperation::EOperationType::IDRelease);
+
+		Task_Game_ReleaseInstanceID(Operation->IDToRelease);
+	}
+
+
 	/** "Start Update" */
 	void Task_Game_StartUpdate(TSharedPtr<FMutableOperation> Operation)
 	{
@@ -2069,6 +2171,27 @@ namespace impl
 		SystemPrivateData->Streamer->PrepareStreamingForObject(CustomizableObject);
 
 		CandidateInstance->CommitMinMaxLOD();
+
+		FString StateName = CandidateInstance->GetCustomizableObject()->GetStateName(CandidateInstance->GetState());
+		const FParameterUIData* StateData = CandidateInstance->GetCustomizableObject()->StateUIDataMap.Find(StateName);
+		bool bLiveUpdateMode = StateData ? StateData->bLiveUpdateMode : false;
+
+		if (bLiveUpdateMode && (!Operation->bNeverStream || Operation->MipsToSkip > 0))
+		{
+			UE_LOG(LogMutable, Warning, TEXT("Instance LiveUpdateMode does not yet support progressive streaming of Mutable textures. Disabling LiveUpdateMode for this update."));
+			bLiveUpdateMode = false;
+		}
+
+		FCustomizableObjectSystemPrivate* CustomizableObjectSystemPrivateData = System->GetPrivate();
+		check(CustomizableObjectSystemPrivateData != nullptr);
+
+		if (!bLiveUpdateMode && CandidateInstancePrivateData->LiveUpdateModeInstanceID != 0)
+		{
+			// The instance was in live update mode last update, but now it's not. So the Id and resources have to be released.
+			// Enqueue a new mutable task to release them
+			Task_Game_ReleaseInstanceID(CandidateInstancePrivateData->LiveUpdateModeInstanceID);
+			CandidateInstancePrivateData->LiveUpdateModeInstanceID = 0;
+		}
 		
 		// Task: Mutable Update and GetMesh
 		//-------------------------------------------------------------
@@ -2079,6 +2202,8 @@ namespace impl
 		CurrentOperationData->CurrentMinLOD = Operation->InstanceDescriptorRuntimeHash.GetMinLOD();
 		CurrentOperationData->CurrentMaxLOD = Operation->InstanceDescriptorRuntimeHash.GetMaxLOD();
 		CurrentOperationData->bNeverStream = Operation->bNeverStream;
+		CurrentOperationData->bLiveUpdateMode = bLiveUpdateMode;
+		CurrentOperationData->InstanceID = bLiveUpdateMode ? CandidateInstancePrivateData->LiveUpdateModeInstanceID : 0;
 		CurrentOperationData->MipsToSkip = Operation->MipsToSkip;
 		CurrentOperationData->MutableParameters = Parameters;
 		check(Operation->CustomizableObjectInstance.IsValid());
@@ -2223,6 +2348,15 @@ void UCustomizableObjectSystem::AdvanceCurrentOperation()
 			break;
 		}
 
+		case FMutableOperation::EOperationType::IDRelease:
+		{
+			impl::Task_Game_ReleaseInstanceID(Private->CurrentMutableOperation);
+
+			ClearCurrentMutableOperation();
+
+			break;
+		}
+
 		default:
 			check(false);
 			break;
@@ -2449,6 +2583,16 @@ FMutableOperation FMutableOperation::CreateInstanceDiscard(UCustomizableObjectIn
 }
 
 
+FMutableOperation FMutableOperation::CreateInstanceIDRelease(mu::Instance::ID IDToRelease)
+{
+	FMutableOperation Op;
+	Op.Type = EOperationType::IDRelease;
+	Op.IDToRelease = IDToRelease;
+
+	return Op;
+}
+
+
 void FMutableOperation::MutableIsDisabledCase()
 {
 	check(CustomizableObjectInstance != nullptr);
diff --git a/Engine/Plugins/Experimental/Mutable/Source/CustomizableObject/Private/MuCO/CustomizableObjectSystemPrivate.h b/Engine/Plugins/Experimental/Mutable/Source/CustomizableObject/Private/MuCO/CustomizableObjectSystemPrivate.h
--- a/Engine/Plugins/Experimental/Mutable/Source/CustomizableObject/Private/MuCO/CustomizableObjectSystemPrivate.h
+++ b/Engine/Plugins/Experimental/Mutable/Source/CustomizableObject/Private/MuCO/CustomizableObjectSystemPrivate.h
@@ -39,6 +39,7 @@ public:
 
 	static FMutableOperation CreateInstanceUpdate(UCustomizableObjectInstance* COInstance, bool bInNeverStream, int32 MipsToSkip);
 	static FMutableOperation CreateInstanceDiscard(UCustomizableObjectInstance* COInstance);
+	static FMutableOperation CreateInstanceIDRelease(mu::Instance::ID);
 
 
 	enum class EOperationType
@@ -47,7 +48,10 @@ public:
 		Update,
 
 		// Discard the resources of an instance.
-		Discard
+		Discard,
+
+		// Release the instance ID and all the temp data associated with it. Usually used with the LiveUpdateMode
+		IDRelease
 	};
 
 	// Type of the operation
@@ -74,7 +78,10 @@ public:
 
 	/** Instance optimization state. */
 	int32 State;
-	
+
+	/** Only used in the IDRelease operation type */
+	mu::Instance::ID IDToRelease;
+
 	//!
 	bool IsBuildParameterDecorations() const
 	{
@@ -410,6 +417,8 @@ struct FMutableOperationData
 
 	/** This option comes from the operation request */
 	bool bNeverStream = false;
+	/** When this option is enabled it will reuse the Mutable core instance and its temp data between updates.  */
+	bool bLiveUpdateMode = false;
 	/** This option comes from the operation request. It is used to reduce the number of mipmaps that mutable must generate for images.  */
 	int32 MipsToSkip = 0;
 
@@ -687,8 +696,11 @@ public:
 	// Init the async Skeletal Mesh creation/update
 	void InitUpdateSkeletalMesh(UCustomizableObjectInstance& Public, FMutableQueueElem::EQueuePriorityType Priority);
 		
-	// Init an async and safe release of the UE4 and Mutable resources used by the instance without actually destroying the instance, for example if it's very far away
+	// Init an async and safe release of the UE and Mutable resources used by the instance without actually destroying the instance, for example if it's very far away
 	void InitDiscardResourcesSkeletalMesh(UCustomizableObjectInstance* InCustomizableObjectInstance);
+
+	// Init the async release of a Mutable Core Instance ID and all the temp resources associated with it
+	void InitInstanceIDRelease(mu::Instance::ID IDToRelease);
 	
 	bool IsReplaceDiscardedWithReferenceMeshEnabled() const { return bReplaceDiscardedWithReferenceMesh; }
 	void SetReplaceDiscardedWithReferenceMeshEnabled(bool bIsEnabled) { bReplaceDiscardedWithReferenceMesh = bIsEnabled; }
diff --git a/Engine/Plugins/Experimental/Mutable/Source/CustomizableObject/Public/MuCO/CustomizableObjectUIData.h b/Engine/Plugins/Experimental/Mutable/Source/CustomizableObject/Public/MuCO/CustomizableObjectUIData.h
--- a/Engine/Plugins/Experimental/Mutable/Source/CustomizableObject/Public/MuCO/CustomizableObjectUIData.h
+++ b/Engine/Plugins/Experimental/Mutable/Source/CustomizableObject/Public/MuCO/CustomizableObjectUIData.h
@@ -261,12 +261,20 @@ struct FParameterUIData
 	UPROPERTY(BlueprintReadWrite, Category = CustomizableObject)
 	bool bDontCompressRuntimeTextures = false; // Only useful for state metadata
 
+	/** In this mode instances and their temp data will be reused between updates. It will be much faster but spend as much as ten times the memory.
+	    Useful for customization lockers with few characters that are going to have their parameters changed many times, not for in-game */
+	UPROPERTY(BlueprintReadWrite, Category = CustomizableObject)
+	bool bLiveUpdateMode = false;
+
 	UPROPERTY(BlueprintReadWrite, Category = CustomizableObject)
 	TMap<FString, FString> ForcedParameterValues;
 
 	bool operator ==(const FParameterUIData& Other) const
 	{
-		if (Name != Other.Name || ParamUIMetadata != Other.ParamUIMetadata || Type != Other.Type || ArrayIntegerParameterOption != Other.ArrayIntegerParameterOption || IntegerParameterGroupType != Other.IntegerParameterGroupType || !ForcedParameterValues.OrderIndependentCompareEqual(Other.ForcedParameterValues))
+		if (Name != Other.Name || ParamUIMetadata != Other.ParamUIMetadata || Type != Other.Type 
+			|| ArrayIntegerParameterOption != Other.ArrayIntegerParameterOption || IntegerParameterGroupType != Other.IntegerParameterGroupType 
+			|| bLiveUpdateMode != Other.bLiveUpdateMode || !ForcedParameterValues.OrderIndependentCompareEqual(Other.ForcedParameterValues)
+			)
 		{
 			return false;
 		}
@@ -282,6 +290,7 @@ struct FParameterUIData
 		Ar << UIData.ArrayIntegerParameterOption;
 		Ar << UIData.IntegerParameterGroupType;
 		Ar << UIData.bDontCompressRuntimeTextures;
+		Ar << UIData.bLiveUpdateMode;
 		Ar << UIData.ForcedParameterValues;
 
 		return Ar;
diff --git a/Engine/Plugins/Experimental/Mutable/Source/CustomizableObjectEditor/Classes/MuCOE/Nodes/CustomizableObjectNodeObject.h b/Engine/Plugins/Experimental/Mutable/Source/CustomizableObjectEditor/Classes/MuCOE/Nodes/CustomizableObjectNodeObject.h
--- a/Engine/Plugins/Experimental/Mutable/Source/CustomizableObjectEditor/Classes/MuCOE/Nodes/CustomizableObjectNodeObject.h
+++ b/Engine/Plugins/Experimental/Mutable/Source/CustomizableObjectEditor/Classes/MuCOE/Nodes/CustomizableObjectNodeObject.h
@@ -41,6 +41,10 @@ struct CUSTOMIZABLEOBJECTEDITOR_API FCustomizableObjectState
 	UPROPERTY(EditAnywhere, Category = CustomizableObject)
 	bool bDontCompressRuntimeTextures = false;
 
+	/** LiveUpdateMode will reuse instance temp. data between updates and speed up update times, but spend much more memory. Good for customization screens, not for actual gameplay modes. */
+	UPROPERTY(EditAnywhere, Category = CustomizableObject)
+	bool bLiveUpdateMode = false;
+
 	UPROPERTY(EditAnywhere, Category = CustomizableObject)
 	bool bBuildOnlyFirstLOD = false;
 
diff --git a/Engine/Plugins/Experimental/Mutable/Source/CustomizableObjectEditor/Private/MuCOE/GenerateMutableSource/GenerateMutableSource.cpp b/Engine/Plugins/Experimental/Mutable/Source/CustomizableObjectEditor/Private/MuCOE/GenerateMutableSource/GenerateMutableSource.cpp
--- a/Engine/Plugins/Experimental/Mutable/Source/CustomizableObjectEditor/Private/MuCOE/GenerateMutableSource/GenerateMutableSource.cpp
+++ b/Engine/Plugins/Experimental/Mutable/Source/CustomizableObjectEditor/Private/MuCOE/GenerateMutableSource/GenerateMutableSource.cpp
@@ -742,6 +742,7 @@ mu::NodeObjectPtr GenerateMutableSource(const UEdGraphPin * Pin, FMutableGraphGe
 			// UI Data
 			FParameterUIData ParameterUIData(State.Name, State.StateUIMetadata, EMutableParameterType::None);
 			ParameterUIData.bDontCompressRuntimeTextures = State.bDontCompressRuntimeTextures;
+			ParameterUIData.bLiveUpdateMode = State.bLiveUpdateMode;
 			ParameterUIData.ForcedParameterValues = State.ForcedParameterValues;
 
 			GenerationContext.StateUIDataMap.Add(State.Name, ParameterUIData);
