[mutable] Added blueprint callable method to get the amount of components present in a CO.

diff --git a/Engine/Plugins/Experimental/Mutable/Source/CustomizableObject/Classes/MuCO/CustomizableObject.h b/Engine/Plugins/Experimental/Mutable/Source/CustomizableObject/Classes/MuCO/CustomizableObject.h
--- a/Engine/Plugins/Experimental/Mutable/Source/CustomizableObject/Classes/MuCO/CustomizableObject.h
+++ b/Engine/Plugins/Experimental/Mutable/Source/CustomizableObject/Classes/MuCO/CustomizableObject.h
@@ -1148,6 +1148,16 @@ public:
 
 #endif // WITH_EDITORONLY_DATA
 
+	/** Amount of components in this CO. Set at the end of the model compilation process. */
+	UPROPERTY()
+	int32 NumMeshComponentsInRoot = 0;
+	
+	/** Method to query the amount of components this Customizable Object has.
+	 * @warning It must be invoked from a COInstance to ensure that the CO has been compiled
+	 */
+	UFUNCTION(BlueprintCallable, Category = CustomizableObject)
+	int32 GetComponentCount() const;
+	
 	// Object parameters interface. This is used to query static data about the parameters available
 	// in instances of this object.
 
@@ -1223,7 +1233,7 @@ private:
 	// This is a manual version number for the binary blobs in this asset.
 	// Increasing it invalidates all the previously compiled models.
 	// Warning: If while merging code both versions have changed, take the highest+1.
-	static const int32 CurrentSupportedVersion = 362;
+	static const int32 CurrentSupportedVersion = 363;
 
 public:
 
diff --git a/Engine/Plugins/Experimental/Mutable/Source/CustomizableObject/Private/MuCO/CustomizableObject.cpp b/Engine/Plugins/Experimental/Mutable/Source/CustomizableObject/Private/MuCO/CustomizableObject.cpp
--- a/Engine/Plugins/Experimental/Mutable/Source/CustomizableObject/Private/MuCO/CustomizableObject.cpp
+++ b/Engine/Plugins/Experimental/Mutable/Source/CustomizableObject/Private/MuCO/CustomizableObject.cpp
@@ -452,6 +452,7 @@ void UCustomizableObject::SaveCompiledData(FArchive& MemoryWriter, bool bSkipEdi
 	}
 
 	MemoryWriter << LODSettings.NumLODsInRoot;
+	MemoryWriter << NumMeshComponentsInRoot;
 
 	MemoryWriter << LODSettings.FirstLODAvailable;
 
@@ -542,6 +543,7 @@ void UCustomizableObject::LoadCompiledData(FArchive& MemoryReader, bool bSkipEdi
 		}
 
 		MemoryReader << LODSettings.NumLODsInRoot;
+		MemoryReader << NumMeshComponentsInRoot;
 
 		MemoryReader << LODSettings.FirstLODAvailable;
 
@@ -1110,6 +1112,19 @@ int32 UCustomizableObject::GetNumLODs() const
 	return LODSettings.NumLODsInRoot;
 }
 
+int32 UCustomizableObject::GetComponentCount() const
+{
+	if (!IsCompiled())
+	{
+		UE_LOG(LogMutable, Warning,
+		       TEXT(
+			       "You are trying to get the component count of a non compiled CO. This will always return 0 as value."
+		       ));
+		return 0;
+	}
+	
+	return NumMeshComponentsInRoot;
+}
 
 int UCustomizableObject::GetParameterCount() const
 {
diff --git a/Engine/Plugins/Experimental/Mutable/Source/CustomizableObjectEditor/Private/MuCOE/CustomizableObjectCompiler.cpp b/Engine/Plugins/Experimental/Mutable/Source/CustomizableObjectEditor/Private/MuCOE/CustomizableObjectCompiler.cpp
--- a/Engine/Plugins/Experimental/Mutable/Source/CustomizableObjectEditor/Private/MuCOE/CustomizableObjectCompiler.cpp
+++ b/Engine/Plugins/Experimental/Mutable/Source/CustomizableObjectEditor/Private/MuCOE/CustomizableObjectCompiler.cpp
@@ -1139,7 +1139,9 @@ void FCustomizableObjectCompiler::CompileInternal(UCustomizableObject* Object, c
 #endif
 
 		Object->LODSettings.NumLODsInRoot = GenerationContext.NumLODsInRoot;
-
+		
+		Object->NumMeshComponentsInRoot = GenerationContext.NumMeshComponentsInRoot;
+		
 		Object->LODSettings.FirstLODAvailable = GenerationContext.FirstLODAvailable;
 
 		Object->LODSettings.bLODStreamingEnabled = GenerationContext.bEnableLODStreaming;
