#rnx

diff --git a/Engine/Plugins/Experimental/Mutable/Source/CustomizableObject/Private/MuCO/CustomizableObjectMipDataProvider.cpp b/Engine/Plugins/Experimental/Mutable/Source/CustomizableObject/Private/MuCO/CustomizableObjectMipDataProvider.cpp
--- a/Engine/Plugins/Experimental/Mutable/Source/CustomizableObject/Private/MuCO/CustomizableObjectMipDataProvider.cpp
+++ b/Engine/Plugins/Experimental/Mutable/Source/CustomizableObject/Private/MuCO/CustomizableObjectMipDataProvider.cpp
@@ -88,7 +88,7 @@ namespace impl
 
 		// For now, we are forcing the recreation of mutable-side instances with every update.
 		mu::Instance::ID InstanceID = System->NewInstance(Model);
-		UE_LOG(LogMutable, Log, TEXT("Creating instance with id [%d] "), InstanceID)
+		UE_LOG(LogMutable, Verbose, TEXT("Creating Mutable instance with id [%d] for a single UpdateImage"), InstanceID)
 
 		const mu::Instance* Instance = nullptr;
 
diff --git a/Engine/Plugins/Experimental/Mutable/Source/CustomizableObject/Private/MuCO/CustomizableObjectSystem.cpp b/Engine/Plugins/Experimental/Mutable/Source/CustomizableObject/Private/MuCO/CustomizableObjectSystem.cpp
--- a/Engine/Plugins/Experimental/Mutable/Source/CustomizableObject/Private/MuCO/CustomizableObjectSystem.cpp
+++ b/Engine/Plugins/Experimental/Mutable/Source/CustomizableObject/Private/MuCO/CustomizableObjectSystem.cpp
@@ -739,6 +739,14 @@ static FAutoConsoleVariableRef CVarEnableMutableProgressiveMipStreaming(
 	ECVF_Default);
 
 
+int32 FCustomizableObjectSystemPrivate::EnableMutableLiveUpdate = 1;
+
+static FAutoConsoleVariableRef CVarEnableMutableLiveUpdate(
+	TEXT("mutable.EnableMutableLiveUpdate"), FCustomizableObjectSystemPrivate::EnableMutableLiveUpdate,
+	TEXT("If set to 1 or greater Mutable can use the live update mode if set in the current Mutable state. If disabled, it will never use live update mode even if set in the current Mutable state."),
+	ECVF_Default);
+
+
 /** Update the given Instance Skeletal Meshes and call its callbacks. */
 void UpdateSkeletalMesh(UCustomizableObjectInstance* CustomizableObjectInstance)
 {
@@ -2172,9 +2180,14 @@ namespace impl
 
 		CandidateInstance->CommitMinMaxLOD();
 
-		FString StateName = CandidateInstance->GetCustomizableObject()->GetStateName(CandidateInstance->GetState());
-		const FParameterUIData* StateData = CandidateInstance->GetCustomizableObject()->StateUIDataMap.Find(StateName);
-		bool bLiveUpdateMode = StateData ? StateData->bLiveUpdateMode : false;
+		bool bLiveUpdateMode = false;
+
+		if (SystemPrivateData->EnableMutableLiveUpdate)
+		{
+			FString StateName = CandidateInstance->GetCustomizableObject()->GetStateName(CandidateInstance->GetState());
+			const FParameterUIData* StateData = CandidateInstance->GetCustomizableObject()->StateUIDataMap.Find(StateName);
+			bLiveUpdateMode = StateData ? StateData->bLiveUpdateMode : false;
+		}
 
 		if (bLiveUpdateMode && (!Operation->bNeverStream || Operation->MipsToSkip > 0))
 		{
diff --git a/Engine/Plugins/Experimental/Mutable/Source/CustomizableObject/Private/MuCO/CustomizableObjectSystemPrivate.h b/Engine/Plugins/Experimental/Mutable/Source/CustomizableObject/Private/MuCO/CustomizableObjectSystemPrivate.h
--- a/Engine/Plugins/Experimental/Mutable/Source/CustomizableObject/Private/MuCO/CustomizableObjectSystemPrivate.h
+++ b/Engine/Plugins/Experimental/Mutable/Source/CustomizableObject/Private/MuCO/CustomizableObjectSystemPrivate.h
@@ -490,6 +490,7 @@ public:
 	TQueue<FMutableTask> PendingTasks;
 
 	static int32 EnableMutableProgressiveMipStreaming;
+	static int32 EnableMutableLiveUpdate;
 	static int32 EnableMutableAnimInfoDebugging;
 
 	/** */
