[mutable] Removed unecessary image resize operations in texture parameters. Added variables to set the parameter texture size.

#rnx

diff --git a/Engine/Plugins/Experimental/Mutable/Source/CustomizableObjectEditor/Classes/MuCOE/Nodes/CustomizableObjectNodeTextureParameter.h b/Engine/Plugins/Experimental/Mutable/Source/CustomizableObjectEditor/Classes/MuCOE/Nodes/CustomizableObjectNodeTextureParameter.h
--- a/Engine/Plugins/Experimental/Mutable/Source/CustomizableObjectEditor/Classes/MuCOE/Nodes/CustomizableObjectNodeTextureParameter.h
+++ b/Engine/Plugins/Experimental/Mutable/Source/CustomizableObjectEditor/Classes/MuCOE/Nodes/CustomizableObjectNodeTextureParameter.h
@@ -33,6 +33,14 @@ public:
 	UPROPERTY(EditAnywhere, Category = UI, meta = (DisplayName = "Parameter UI Metadata"))
 	FMutableParamUIMetadata ParamUIMetadata;
 
+	/** Set the width of the Texture when there is no texture reference.*/
+	UPROPERTY(EditAnywhere, Category = CustomizableObject)
+	int32 TextureSizeX = 0;
+
+	/** Set the height of the Texture when there is no texture reference.*/
+	UPROPERTY(EditAnywhere, Category = CustomizableObject)
+	int32 TextureSizeY = 0;
+
 	// Begin EdGraphNode interface
 	FText GetNodeTitle(ENodeTitleType::Type TitleType) const override;
 	FLinearColor GetNodeTitleColor() const override;
diff --git a/Engine/Plugins/Experimental/Mutable/Source/CustomizableObjectEditor/Private/MuCOE/GenerateMutableSource/GenerateMutableSource.cpp b/Engine/Plugins/Experimental/Mutable/Source/CustomizableObjectEditor/Private/MuCOE/GenerateMutableSource/GenerateMutableSource.cpp
--- a/Engine/Plugins/Experimental/Mutable/Source/CustomizableObjectEditor/Private/MuCOE/GenerateMutableSource/GenerateMutableSource.cpp
+++ b/Engine/Plugins/Experimental/Mutable/Source/CustomizableObjectEditor/Private/MuCOE/GenerateMutableSource/GenerateMutableSource.cpp
@@ -1364,7 +1364,7 @@ void PopulateReferenceSkeletalMeshesData(FMutableGraphGenerationContext& Generat
 
 
 int32 ComputeLODBias(const FMutableGraphGenerationContext& GenerationContext, const UTexture2D* ReferenceTexture, int32 MaxTextureSize,
-	const UCustomizableObjectNodeMaterial* MaterialNode, const int32 ImageIndex)
+	const UCustomizableObjectNodeMaterial* MaterialNode, const int32 ImageIndex, bool bUseLODAsBias)
 {
 	int32 LODBias = 0;
 
@@ -1408,7 +1408,7 @@ int32 ComputeLODBias(const FMutableGraphGenerationContext& GenerationContext, co
 		GenerationContext.CurrentAutoLODStrategy == ECustomizableObjectAutomaticLODStrategy::AutomaticFromMesh)
 	{
 		// Only if the texture actually uses a layout. Otherwise it could be a special texture we shouldn't scale.
-		if (MaterialNode && MaterialNode->GetImageUVLayout(ImageIndex) >= 0)
+		if ((MaterialNode && MaterialNode->GetImageUVLayout(ImageIndex) >= 0) || bUseLODAsBias)
 		{
 			// \todo: make it an object property to be tweaked
 			int MipsToSkipPerLOD = 1;
diff --git a/Engine/Plugins/Experimental/Mutable/Source/CustomizableObjectEditor/Private/MuCOE/GenerateMutableSource/GenerateMutableSource.h b/Engine/Plugins/Experimental/Mutable/Source/CustomizableObjectEditor/Private/MuCOE/GenerateMutableSource/GenerateMutableSource.h
--- a/Engine/Plugins/Experimental/Mutable/Source/CustomizableObjectEditor/Private/MuCOE/GenerateMutableSource/GenerateMutableSource.h
+++ b/Engine/Plugins/Experimental/Mutable/Source/CustomizableObjectEditor/Private/MuCOE/GenerateMutableSource/GenerateMutableSource.h
@@ -816,7 +816,7 @@ FString GenerateGameplayTag(const FString& GameplayTag);
 // Computes the LOD bias for a texture given the current mesh LOD and automatic LOD settings, the reference texture settings
 // and whether it's being built for a server or not
 int32 ComputeLODBias(const FMutableGraphGenerationContext& GenerationContext, const UTexture2D* ReferenceTexture, int32 MaxTextureSize,
-	const UCustomizableObjectNodeMaterial* MaterialNode, const int32 ImageIndex);
+	const UCustomizableObjectNodeMaterial* MaterialNode, const int32 ImageIndex, bool bUseLODAsBias = true);
 
 
 int32 GetMaxTextureSize(const UTexture2D* ReferenceTexture, const FMutableGraphGenerationContext& GenerationContext);
diff --git a/Engine/Plugins/Experimental/Mutable/Source/CustomizableObjectEditor/Private/MuCOE/GenerateMutableSource/GenerateMutableSourceImage.cpp b/Engine/Plugins/Experimental/Mutable/Source/CustomizableObjectEditor/Private/MuCOE/GenerateMutableSource/GenerateMutableSourceImage.cpp
--- a/Engine/Plugins/Experimental/Mutable/Source/CustomizableObjectEditor/Private/MuCOE/GenerateMutableSource/GenerateMutableSourceImage.cpp
+++ b/Engine/Plugins/Experimental/Mutable/Source/CustomizableObjectEditor/Private/MuCOE/GenerateMutableSource/GenerateMutableSourceImage.cpp
@@ -175,6 +175,7 @@ mu::NodeImagePtr GenerateMutableSourceImage(const UEdGraphPin* Pin, FMutableGrap
 
 		TextureNode->SetName(TCHAR_TO_ANSI(*TypedNodeParam->ParameterName));
 		TextureNode->SetUid(TCHAR_TO_ANSI(*GenerationContext.GetNodeIdUnique(Node).ToString()));
+		
 		if (NodeRange)
 		{
 			TextureNode->SetRangeCount(1);
@@ -203,8 +204,6 @@ mu::NodeImagePtr GenerateMutableSourceImage(const UEdGraphPin* Pin, FMutableGrap
 		}
 		FormatNode->SetSource(TextureNode);
 
-		
-
 		mu::NodeImageResizePtr ResizeNode = new mu::NodeImageResize();
 		ResizeNode->SetBase(FormatNode);
 		ResizeNode->SetRelative(false);
@@ -212,19 +211,16 @@ mu::NodeImagePtr GenerateMutableSourceImage(const UEdGraphPin* Pin, FMutableGrap
 		const UTexture2D* ReferenceTexture = TypedNodeParam->DefaultValue;
 		if (ReferenceTexture)
 		{
-			const int32 LODBias = ComputeLODBias(GenerationContext, ReferenceTexture, MaxTextureSize, nullptr, INDEX_NONE);
-
-			int32 Width = MaxTextureSize > 0 ? FMath::Min(MaxTextureSize, ReferenceTexture->GetImportedSize().X) : ReferenceTexture->GetImportedSize().X;
-			int32 Height = MaxTextureSize > 0 ? FMath::Min(MaxTextureSize, ReferenceTexture->GetImportedSize().Y) : ReferenceTexture->GetImportedSize().Y;
-			Width = Width >> LODBias;
-			Height = Height >> LODBias;
-
-			ResizeNode->SetSize(FMath::Max(Width,1), FMath::Max(Height, 1));
+			ResizeNode->SetSize(FMath::Max(ReferenceTexture->GetImportedSize().X,1), FMath::Max(ReferenceTexture->GetImportedSize().X, 1));
 		}
 		else
 		{
-			// \TODO: Let the user specify this in the node?
-			ResizeNode->SetSize(1024, 1024);
+			if (TypedNodeParam->TextureSizeX <= 0 || TypedNodeParam->TextureSizeY <= 0)
+			{
+				GenerationContext.Compiler->CompilerLog(LOCTEXT("TextureParameterSize0", "Texture size not specified. Add a reference texture or set a valid value to the Texture Size variables."), Node);
+			}
+
+			ResizeNode->SetSize(FMath::Max(TypedNodeParam->TextureSizeX, 1), FMath::Max(TypedNodeParam->TextureSizeY, 1));
 		}
 
 		Result = ResizeNode;
@@ -579,7 +575,7 @@ mu::NodeImagePtr GenerateMutableSourceImage(const UEdGraphPin* Pin, FMutableGrap
 		// Calculating Texture size using Reference texture parameters
 		if (TypedNodeProject->ReferenceTexture)
 		{
-			int32 LODBias = ComputeLODBias(GenerationContext, TypedNodeProject->ReferenceTexture, TypedNodeProject->ReferenceTexture->MaxTextureSize, nullptr, INDEX_NONE);
+			int32 LODBias = ComputeLODBias(GenerationContext, TypedNodeProject->ReferenceTexture, TypedNodeProject->ReferenceTexture->MaxTextureSize, nullptr, INDEX_NONE, false);
 
 			if (TextureSize.X > 0 && TextureSize.Y > 0)
 			{
