[mutable] Fix bug with code optimiser transformation used for forcing uncompressed formats, etc.
Fix potential crash when malformed bone data is found. Add support for additional bone data format to skinning optimization.
#rnx

diff --git a/Engine/Plugins/Experimental/Mutable/Source/MutableRuntime/Private/MuR/OpMeshOptimizeSkinning.h b/Engine/Plugins/Experimental/Mutable/Source/MutableRuntime/Private/MuR/OpMeshOptimizeSkinning.h
--- a/Engine/Plugins/Experimental/Mutable/Source/MutableRuntime/Private/MuR/OpMeshOptimizeSkinning.h
+++ b/Engine/Plugins/Experimental/Mutable/Source/MutableRuntime/Private/MuR/OpMeshOptimizeSkinning.h
@@ -135,12 +135,29 @@ namespace mu
 							{
 							case MESH_BUFFER_FORMAT::MBF_UINT8:
 							{
-								*Data = (uint8)RemapedBoneIndices[*SourceData];
+								int32 SourceIndex = *SourceData;
+								if (SourceIndex < RemapedBoneIndices.Num())
+								{
+									*Data = (uint8)RemapedBoneIndices[SourceIndex];
+								}
 								break;
 							}
 							case MESH_BUFFER_FORMAT::MBF_UINT16:
 							{
-								*Data = (uint8)RemapedBoneIndices[*((uint16*)SourceData)];
+								int32 SourceIndex = *((uint16*)SourceData);
+								if (SourceIndex< RemapedBoneIndices.Num())
+								{
+									*Data = (uint8)RemapedBoneIndices[SourceIndex];
+								}
+								break;
+							}
+							case MESH_BUFFER_FORMAT::MBF_INT32:
+							{
+								int32 SourceIndex = *((int32*)SourceData);
+								if (SourceIndex < RemapedBoneIndices.Num())
+								{
+									*Data = (uint8)RemapedBoneIndices[SourceIndex];
+								}
 								break;
 							}
 							default:
diff --git a/Engine/Plugins/Experimental/Mutable/Source/MutableTools/Private/MuT/AST.cpp b/Engine/Plugins/Experimental/Mutable/Source/MutableTools/Private/MuT/AST.cpp
--- a/Engine/Plugins/Experimental/Mutable/Source/MutableTools/Private/MuT/AST.cpp
+++ b/Engine/Plugins/Experimental/Mutable/Source/MutableTools/Private/MuT/AST.cpp
@@ -601,8 +601,7 @@ void Visitor_TopDown_Unique_Cloning::Process()
                 // Proceed with children
                 if (processChildren)
                 {
-                    // TODO: Shouldn't we recurse newAt?
-					m_pending.Add({ true,at });
+					m_pending.Add({ true, newAt });
 
                     at->ForEachChild( [&](ASTChild& ref)
                     {
diff --git a/Engine/Plugins/Experimental/Mutable/Source/MutableTools/Private/MuT/CodeOptimiser_Parameters.cpp b/Engine/Plugins/Experimental/Mutable/Source/MutableTools/Private/MuT/CodeOptimiser_Parameters.cpp
--- a/Engine/Plugins/Experimental/Mutable/Source/MutableTools/Private/MuT/CodeOptimiser_Parameters.cpp
+++ b/Engine/Plugins/Experimental/Mutable/Source/MutableTools/Private/MuT/CodeOptimiser_Parameters.cpp
@@ -141,8 +141,8 @@ namespace mu
                     case OP_TYPE::PR_PARAMETER:
                     case OP_TYPE::IM_PARAMETER:
                     {
-                        auto typed = dynamic_cast<const ASTOpParameter*>(at.get());
-                        const auto& params = m_pState->nodeState.m_runtimeParams;
+						const ASTOpParameter* typed = dynamic_cast<const ASTOpParameter*>(at.get());
+                        const TArray<string>& params = m_pState->nodeState.m_runtimeParams;
                         if ( params.Find( typed->parameter.m_name)
                              !=
                              INDEX_NONE )
@@ -155,7 +155,7 @@ namespace mu
 
                     case OP_TYPE::ME_INTERPOLATE:
                     {
-                        auto typed = dynamic_cast<const ASTOpFixed*>(at.get());
+						const ASTOpFixed* typed = dynamic_cast<const ASTOpFixed*>(at.get());
                         PENDING_ITEM childItem;
                         childItem.itemType = 0;
                         childItem.onlyLayoutsRelevant = item.onlyLayoutsRelevant;
@@ -380,7 +380,7 @@ namespace mu
         //-------------------------------------------------------------------------------------
         case OP_TYPE::IM_CONDITIONAL:
         {
-            auto typedAt = dynamic_cast<const ASTOpConditional*>(at.get());
+			const ASTOpConditional* typedAt = dynamic_cast<const ASTOpConditional*>(at.get());
 
             // If the condition is not runtime, but the branches are, try to move the
             // conditional down
@@ -414,27 +414,27 @@ namespace mu
                              )
                             {
                                 // Move the conditional down
-                                auto compOp = mu::Clone<ASTOpImageCompose>(typedYes);
+                                Ptr<ASTOpImageCompose> compOp = mu::Clone<ASTOpImageCompose>(typedYes);
 
-                                auto baseCond = mu::Clone<ASTOpConditional>(typedAt);
+                                Ptr<ASTOpConditional> baseCond = mu::Clone<ASTOpConditional>(typedAt);
                                 baseCond->yes = typedYes->Base.child();
                                 baseCond->no = typedNo->Base.child();
                                 compOp->Base = baseCond;
 
-                                auto blockCond = mu::Clone<ASTOpConditional>(typedAt);
+                                Ptr<ASTOpConditional> blockCond = mu::Clone<ASTOpConditional>(typedAt);
                                 blockCond->yes = typedYes->BlockImage.child();
                                 blockCond->no = typedNo->BlockImage.child();
                                 compOp->BlockImage = blockCond;
 
                                 if (typedYes->Mask)
                                 {
-                                    auto maskCond = mu::Clone<ASTOpConditional>(typedAt);
+                                    Ptr<ASTOpConditional> maskCond = mu::Clone<ASTOpConditional>(typedAt);
                                     maskCond->yes = typedYes->Mask.child();
                                     maskCond->no = typedNo->Mask.child();
                                     compOp->Mask = maskCond;
                                 }
 
-                                auto layCond = mu::Clone<ASTOpConditional>(typedAt);
+                                Ptr<ASTOpConditional> layCond = mu::Clone<ASTOpConditional>(typedAt);
                                 layCond->type = OP_TYPE::LA_CONDITIONAL;
                                 layCond->yes = typedYes->Layout.child();
                                 layCond->no = typedNo->Layout.child();
@@ -482,7 +482,7 @@ namespace mu
                             resizeOp->SetChild( resizeOp->op.args.ImageResizeLike.source, plainOp );
                             resizeOp->SetChild( resizeOp->op.args.ImageResizeLike.sizeSource, typedYes->base );
 
-                            auto maskOp = mu::Clone<ASTOpConditional>(typedAt);
+                            Ptr<ASTOpConditional> maskOp = mu::Clone<ASTOpConditional>(typedAt);
                             maskOp->no = resizeOp;
 
                             // If there is no mask (because it is optional), we need to make a
@@ -852,9 +852,9 @@ namespace mu
         {
 			const ASTOpImageCompose* typedAt = dynamic_cast<const ASTOpImageCompose*>(at.get());
 
-            auto blockAt = typedAt->BlockImage.child();
-            auto baseAt = typedAt->Base.child();
-            auto layoutAt = typedAt->Layout.child();
+            Ptr<ASTOp> blockAt = typedAt->BlockImage.child();
+			Ptr<ASTOp> baseAt = typedAt->Base.child();
+			Ptr<ASTOp> layoutAt = typedAt->Layout.child();
 
 			if (!blockAt)
 			{
@@ -887,10 +887,10 @@ namespace mu
 
                         // The mask is a compose of the block mask on the base mask, but if none has
                         // a mask we don't need to make one.
-                        auto baseImage = typedBaseAt->base.child();
-                        auto baseMask = typedBaseAt->mask.child();
-                        auto blockImage = typedBlockAt->base.child();
-                        auto blockMask = typedBlockAt->mask.child();
+						Ptr<ASTOp> baseImage = typedBaseAt->base.child();
+						Ptr<ASTOp> baseMask = typedBaseAt->mask.child();
+						Ptr<ASTOp> blockImage = typedBlockAt->base.child();
+						Ptr<ASTOp> blockMask = typedBlockAt->mask.child();
 
                         Ptr<ASTOpImageCompose> maskOp;
                         if (baseMask || blockMask)
@@ -898,8 +898,8 @@ namespace mu
                             // \TODO: BLEH! This may create a discrepancy of number of mips between
                             // the base image and the mask This is for now solved with emergy fix
                             // c36adf47-e40d-490f-b709-41142bafad78
-                            auto newBaseMask = EnsureValidMask(baseMask, baseImage);
-                            auto newBlockMask = EnsureValidMask(blockMask, blockImage);
+							Ptr<ASTOp> newBaseMask = EnsureValidMask(baseMask, baseImage);
+							Ptr<ASTOp> newBlockMask = EnsureValidMask(blockMask, blockImage);
 
                             maskOp = mu::Clone<ASTOpImageCompose>(typedAt);
                             maskOp->Base = newBaseMask;
@@ -907,7 +907,7 @@ namespace mu
                         }
 
                         // The base is composition of the bases of both layer effect
-                        auto baseOp = mu::Clone<ASTOpImageCompose>(typedAt);
+						Ptr<ASTOpImageCompose> baseOp = mu::Clone<ASTOpImageCompose>(typedAt);
                         baseOp->Base = baseImage;
                         baseOp->BlockImage = blockImage;
 
@@ -929,12 +929,12 @@ namespace mu
 
                         // The mask is a compose of the block mask on the base mask, but if none has
                         // a mask we don't need to make one.
-                        auto baseImage = typedBaseAt->base.child();
-                        auto baseBlended = typedBaseAt->blend.child();
-                        auto baseMask = typedBaseAt->mask.child();
-                        auto blockImage = typedBlockAt->base.child();
-                        auto blockBlended = typedBlockAt->blend.child();
-                        auto blockMask = typedBlockAt->mask.child();
+						Ptr<ASTOp> baseImage = typedBaseAt->base.child();
+						Ptr<ASTOp> baseBlended = typedBaseAt->blend.child();
+						Ptr<ASTOp> baseMask = typedBaseAt->mask.child();
+						Ptr<ASTOp> blockImage = typedBlockAt->base.child();
+						Ptr<ASTOp> blockBlended = typedBlockAt->blend.child();
+						Ptr<ASTOp> blockMask = typedBlockAt->mask.child();
 
                         Ptr<ASTOpImageCompose> maskOp;
                         if (baseMask || blockMask)
@@ -942,8 +942,8 @@ namespace mu
                             // \TODO: BLEH! This may create a discrepancy of number of mips between
                             // the base image and the mask This is for now solved with emergy fix
                             // c36adf47-e40d-490f-b709-41142bafad78
-                            auto newBaseMask = EnsureValidMask(baseMask, baseImage);
-                            auto newBlockMask = EnsureValidMask(blockMask, blockImage);
+							Ptr<ASTOp> newBaseMask = EnsureValidMask(baseMask, baseImage);
+							Ptr<ASTOp> newBlockMask = EnsureValidMask(blockMask, blockImage);
 
                             maskOp = mu::Clone<ASTOpImageCompose>(typedAt);
                             maskOp->Base = newBaseMask;
@@ -951,12 +951,12 @@ namespace mu
                         }
 
                         // The base is composition of the bases of both layer effect
-                        auto baseOp = mu::Clone<ASTOpImageCompose>(typedAt);
+						Ptr<ASTOpImageCompose> baseOp = mu::Clone<ASTOpImageCompose>(typedAt);
                         baseOp->Base = baseImage;
                         baseOp->BlockImage = blockImage;
 
                         // The base is composition of the bases of both layer effect
-                        auto blendedOp = mu::Clone<ASTOpImageCompose>(typedAt);
+						Ptr<ASTOpImageCompose> blendedOp = mu::Clone<ASTOpImageCompose>(typedAt);
                         blendedOp->Base = baseBlended;
                         blendedOp->BlockImage = blockBlended;
 
@@ -982,19 +982,19 @@ namespace mu
                  &&
                  baseType == OP_TYPE::IM_COMPOSE )
             {
-                auto typedBaseAt = dynamic_cast<const ASTOpImageCompose*>(baseAt.get());
+				const ASTOpImageCompose* typedBaseAt = dynamic_cast<const ASTOpImageCompose*>(baseAt.get());
 
-                auto baseBlockAt = typedBaseAt->BlockImage.child();
+				Ptr<ASTOp> baseBlockAt = typedBaseAt->BlockImage.child();
                 bool baseBlockHasAny = m_hasRuntimeParamVisitor.HasAny( baseBlockAt );
                 if ( baseBlockHasAny )
                 {
                     optimised = true;
 
                     // Swap
-                    auto childCompose = mu::Clone<ASTOpImageCompose>(at);
+					Ptr<ASTOpImageCompose> childCompose = mu::Clone<ASTOpImageCompose>(at);
                     childCompose->Base = typedBaseAt->Base.child();
 
-                    auto parentCompose = mu::Clone<ASTOpImageCompose>(baseAt);
+					Ptr<ASTOpImageCompose> parentCompose = mu::Clone<ASTOpImageCompose>(baseAt);
                     parentCompose->Base = childCompose;
 
                     at = parentCompose;
@@ -1021,8 +1021,8 @@ namespace mu
 
 					const ASTOpImageLayerColor* typedBlockAt = dynamic_cast<const ASTOpImageLayerColor*>(blockAt.get());
 
-                    auto blockImage = typedBlockAt->base.child();
-                    auto blockMask = typedBlockAt->mask.child();
+					Ptr<ASTOp> blockImage = typedBlockAt->base.child();
+					Ptr<ASTOp> blockMask = typedBlockAt->mask.child();
 
                     // The mask is a compose of the layer mask on a black image, however if there is
                     // no mask and the base of the layer opertation is a blanklayout, we can skip
@@ -1031,7 +1031,7 @@ namespace mu
                     if (blockMask || baseType!=OP_TYPE::IM_BLANKLAYOUT)
                     {
                         maskOp = mu::Clone<ASTOpImageCompose>(at);
-                        auto newMaskBlock = EnsureValidMask(blockMask, blockImage);
+						Ptr<ASTOp> newMaskBlock = EnsureValidMask(blockMask, blockImage);
                         maskOp->BlockImage = newMaskBlock;
 
                         Ptr<ASTOpFixed> blackOp = new ASTOpFixed;
@@ -1075,9 +1075,9 @@ namespace mu
 
                     const ASTOpImageLayer* typedBlockAt = dynamic_cast<const ASTOpImageLayer*>(blockAt.get());
 
-                    auto blockImage = typedBlockAt->base.child();
-                    auto blockBlended = typedBlockAt->blend.child();
-                    auto blockMask = typedBlockAt->mask.child();
+					Ptr<ASTOp> blockImage = typedBlockAt->base.child();
+					Ptr<ASTOp> blockBlended = typedBlockAt->blend.child();
+					Ptr<ASTOp> blockMask = typedBlockAt->mask.child();
 
                     // The mask is a compose of the layer mask on a black image, however if there is
                     // no mask and the base of the layer opertation is a blanklayout, we can skip
@@ -1086,7 +1086,7 @@ namespace mu
                     if (blockMask || baseType != OP_TYPE::IM_BLANKLAYOUT)
                     {
                         maskOp = mu::Clone<ASTOpImageCompose>(at);
-                        auto newMaskBlock = EnsureValidMask(blockMask, blockImage);
+						Ptr<ASTOp> newMaskBlock = EnsureValidMask(blockMask, blockImage);
                         maskOp->BlockImage = newMaskBlock;
 
 
@@ -1113,7 +1113,7 @@ namespace mu
                     }
 
                     // The blended is a compose of the blended image on a blank image
-                    auto blendedOp = mu::Clone<ASTOpImageCompose>(at);
+					Ptr<ASTOpImageCompose> blendedOp = mu::Clone<ASTOpImageCompose>(at);
                     {
                         blendedOp->BlockImage = blockBlended;
 
@@ -1199,7 +1199,7 @@ namespace mu
 
                     const ASTOpImageLayerColor* typedBaseAt = dynamic_cast<const ASTOpImageLayerColor*>(baseAt.get());
 
-                    auto maskOp = mu::Clone<ASTOpImageCompose>(at);
+					Ptr<ASTOpImageCompose> maskOp = mu::Clone<ASTOpImageCompose>(at);
                     {
                         Ptr<ASTOpFixed> blackOp = new ASTOpFixed;
                         blackOp->op.type = OP_TYPE::CO_CONSTANT;
@@ -1221,13 +1221,13 @@ namespace mu
                         blockResizeOp->SetChild( blockResizeOp->op.args.ImageResizeLike.source, plainOp );
 
                         // Blank out the block from the mask
-                        auto newMaskBase = EnsureValidMask( typedBaseAt->mask.child(), baseAt );
+						Ptr<ASTOp> newMaskBase = EnsureValidMask( typedBaseAt->mask.child(), baseAt );
                         maskOp->Base = newMaskBase;
                         maskOp->BlockImage = blockResizeOp;
                     }
 
                     // The base is composition of the softlight base on the compose base
-                    auto baseOp = mu::Clone<ASTOpImageCompose>(at);
+					Ptr<ASTOpImageCompose> baseOp = mu::Clone<ASTOpImageCompose>(at);
                     baseOp->Base = typedBaseAt->base.child();
 
                     Ptr<ASTOpImageLayerColor> nop = mu::Clone<ASTOpImageLayerColor>(baseAt);
@@ -1245,7 +1245,7 @@ namespace mu
 
                     const ASTOpImageLayer* typedBaseAt = dynamic_cast<const ASTOpImageLayer*>(baseAt.get());
 
-                    auto maskOp = mu::Clone<ASTOpImageCompose>(at);
+					Ptr<ASTOpImageCompose> maskOp = mu::Clone<ASTOpImageCompose>(at);
                     {
                         Ptr<ASTOpFixed> blackOp = new ASTOpFixed;
                         blackOp->op.type = OP_TYPE::CO_CONSTANT;
@@ -1416,9 +1416,9 @@ namespace mu
         // Sink the mipmap if worth it.
         case OP_TYPE::IM_MIPMAP:
         {
-            auto typedAt = dynamic_cast<const ASTOpImageMipmap*>(at.get());
+			const ASTOpImageMipmap* typedAt = dynamic_cast<const ASTOpImageMipmap*>(at.get());
 
-            auto sourceOp = typedAt->Source.child();
+			Ptr<ASTOp> sourceOp = typedAt->Source.child();
 
             switch ( sourceOp->GetOpType() )
             {
@@ -1432,16 +1432,16 @@ namespace mu
                 {
                     m_modified = true;
 
-                    auto top = mu::Clone<ASTOpImageLayerColor>(sourceOp);
+					Ptr<ASTOpImageLayerColor> top = mu::Clone<ASTOpImageLayerColor>(sourceOp);
 
-                    auto baseOp = mu::Clone<ASTOpImageMipmap>(at);
+					Ptr<ASTOpImageMipmap> baseOp = mu::Clone<ASTOpImageMipmap>(at);
                     baseOp->Source = typedSource->base.child();
                     top->base = baseOp;
 
-                    auto sourceMaskOp = typedSource->mask.child();
+					Ptr<ASTOp> sourceMaskOp = typedSource->mask.child();
                     if (sourceMaskOp)
                     {
-                        auto maskOp = mu::Clone<ASTOpImageMipmap>(at);
+						Ptr<ASTOpImageMipmap> maskOp = mu::Clone<ASTOpImageMipmap>(at);
                         maskOp->Source = sourceMaskOp;
                         top->mask = maskOp;
                     }
@@ -1495,7 +1495,7 @@ namespace mu
             // Remove unsupported formats
             if (GetOpDataType( at->GetOpType() )==DT_IMAGE)
             {
-                auto it = m_supportedFormats.Find(at);
+				std::array<uint8_t, (size_t)EImageFormat::IF_COUNT>* it = m_supportedFormats.Find(at);
                 if (!it)
                 {
                     // Default to all supported
@@ -1556,7 +1556,7 @@ namespace mu
 
             case OP_TYPE::IM_DISPLACE:
             {
-                auto typedAt = dynamic_cast<const ASTOpFixed*>(at.get());
+				const ASTOpFixed* typedAt = dynamic_cast<const ASTOpFixed*>(at.get());
 
                 RecurseWithCurrentState( typedAt->children[typedAt->op.args.ImageDisplace.source].child() );
 
@@ -1580,7 +1580,7 @@ namespace mu
 
         bool IsSupportedFormat( Ptr<ASTOp> at, EImageFormat format ) const
         {
-            auto it = m_supportedFormats.Find(at);
+            const std::array<uint8_t, (size_t)EImageFormat::IF_COUNT>* it = m_supportedFormats.Find(at);
             if (!it)
             {
                 return false;
@@ -1609,7 +1609,7 @@ namespace mu
     void SubtreeRelevantParametersVisitorAST::Run( Ptr<ASTOp> root )
     {
         // Cached?
-        auto it = m_resultCache.find( STATE(root,false) );
+		std::unordered_map< STATE, std::unordered_set< string >, state_hash >::iterator it = m_resultCache.find( STATE(root,false) );
         if (it!=m_resultCache.end())
         {
             m_params = it->second;
@@ -1637,7 +1637,7 @@ namespace mu
                 case OP_TYPE::PR_PARAMETER:
                 case OP_TYPE::IM_PARAMETER:
                 {
-                    auto typedAt = dynamic_cast<const ASTOpParameter*>(at.get());
+					const ASTOpParameter* typedAt = dynamic_cast<const ASTOpParameter*>(at.get());
                     m_params.insert(typedAt->parameter.m_name);
 
                     // Not interested in the parameters from the parameters decorators.
@@ -1665,24 +1665,24 @@ namespace mu
                     // Manually choose how to recurse this op
 					const ASTOpMeshMorph* pTyped = dynamic_cast<const ASTOpMeshMorph*>( at.get() );
 
-                    if ( auto& base = pTyped->Base )
+                    if ( pTyped->Base )
                     {
-						pending.Add({ base.m_child, state });
+						pending.Add({ pTyped->Base.m_child, state });
                     }
 
                     // Mesh morphs don't modify the layouts, so we can ignore the factor and morphs
                     if (!state)
                     {
-                        if ( auto& factor = pTyped->Factor )
+                        if ( pTyped->Factor )
                         {
-							pending.Add({ factor.m_child, state });
+							pending.Add({ pTyped->Factor.m_child, state });
                         }
 
                         for (int32 t=0;t<pTyped->Targets.Num(); ++t)
                         {
-                            if ( auto& target = pTyped->Targets[t] )
+                            if ( pTyped->Targets[t] )
                             {
-								pending.Add({ target.m_child, state });
+								pending.Add({ pTyped->Targets[t].m_child, state });
                             }
                         }
                     }
@@ -1724,7 +1724,7 @@ namespace mu
             pState->m_updateCache.Empty();
             pState->m_dynamicResources.Empty();
 
-            for( const auto& i : m_cache )
+            for( const TPair<Ptr<ASTOp>, bool>& i : m_cache )
             {
                 if ( i.Value )
                 {
@@ -1732,7 +1732,7 @@ namespace mu
                 }
             }
 
-            for(const auto& i : m_dynamicResourceRoot )
+            for(const TPair<Ptr<ASTOp>, bool>& i : m_dynamicResourceRoot )
             {
                 if ( i.Value )
                 {
@@ -1742,7 +1742,7 @@ namespace mu
 
 					// Temp copy
 					TArray<mu::string> ParamCopy;
-					for ( const auto& e: subtreeParams.m_params )
+					for ( const string& e: subtreeParams.m_params )
 					{
 						ParamCopy.Add(e);
 					}
@@ -1793,7 +1793,7 @@ namespace mu
                 case OP_TYPE::IN_ADDSCALAR:
                 case OP_TYPE::IN_ADDSTRING:
                 {
-                    auto typedAt = dynamic_cast<const ASTOpInstanceAdd*>(at.get());
+					const ASTOpInstanceAdd* typedAt = dynamic_cast<const ASTOpInstanceAdd*>(at.get());
 
                     TPair<bool,bool> newState;
                     newState.Key = false; //resource root
@@ -1918,8 +1918,8 @@ namespace mu
         // textures? We don't want them uncompressed.
         if( type==OP_TYPE::IN_ADDIMAGE )
         {
-            auto typedAt = dynamic_cast<ASTOpInstanceAdd*>(at.get());
-            auto imageAt = typedAt->value.child();
+			ASTOpInstanceAdd* typedAt = dynamic_cast<ASTOpInstanceAdd*>(at.get());
+			Ptr<ASTOp> imageAt = typedAt->value.child();
 
             // Does it have a runtime parameter in its subtree?
             bool hasRuntimeParameter = m_hasRuntimeParamVisitor.HasAny( imageAt );
@@ -1969,7 +1969,7 @@ namespace mu
 
         if( at->GetOpType()==OP_TYPE::IN_ADDLOD )
         {
-            auto* typedAt = dynamic_cast<ASTOpAddLOD*>(at.get());
+			ASTOpAddLOD* typedAt = dynamic_cast<ASTOpAddLOD*>(at.get());
 
             if (typedAt->lods.Num()>(size_t)m_lodCount)
             {
@@ -2028,8 +2028,7 @@ namespace mu
 
                     UE_LOG(LogMutableCore, Verbose, TEXT(" - before parameter optimiser"));
 
-                    ParameterOptimiserAST param( m_states[s],
-                                                 m_options->GetPrivate()->m_optimisationOptions );
+                    ParameterOptimiserAST param( m_states[s], m_options->GetPrivate()->m_optimisationOptions );
                     modified = param.Apply();
 
                     TArray<Ptr<ASTOp>> roots;
@@ -2040,14 +2039,12 @@ namespace mu
 
                     // All kind of optimisations that depend on the meaning of each operation
                     UE_LOG(LogMutableCore, Verbose, TEXT(" - semantic optimiser"));
-                    modified |= SemanticOptimiserAST(
-                        roots, m_options->GetPrivate()->m_optimisationOptions );
+                    modified |= SemanticOptimiserAST( roots, m_options->GetPrivate()->m_optimisationOptions );
 					UE_LOG(LogMutableCore, Verbose, TEXT("(int) %s : %ld"), TEXT("ast size"), int64(ASTOp::CountNodes(roots)));
 					//ASTOp::LogHistogram(roots);
 
                     UE_LOG(LogMutableCore, Verbose, TEXT(" - sink optimiser"));
-                    modified |=
-                        SinkOptimiserAST( roots, m_options->GetPrivate()->m_optimisationOptions );
+                    modified |= SinkOptimiserAST( roots, m_options->GetPrivate()->m_optimisationOptions );
 					UE_LOG(LogMutableCore, Verbose, TEXT("(int) %s : %ld"), TEXT("ast size"), int64(ASTOp::CountNodes(roots)));
 					//ASTOp::LogHistogram(roots);
 
@@ -2079,7 +2076,7 @@ namespace mu
         // necessary at this stage before GPU optimisation.
         {
             TArray<Ptr<ASTOp>> roots;
-            for(const auto& s:m_states)
+            for(const STATE_COMPILATION_DATA& s:m_states)
             {
                 roots.Add(s.root);
             }
@@ -2087,7 +2084,7 @@ namespace mu
             AccumulateAllImageFormatsOpAST opFormats;
             opFormats.Run(roots);
 
-            for ( auto& s: m_states )
+            for (STATE_COMPILATION_DATA& s: m_states )
             {
                 {
                     UE_LOG(LogMutableCore, Verbose, TEXT(" - state cache"));
@@ -2131,7 +2128,7 @@ namespace mu
             while (modified && (!m_optimizeIterationsMax || m_optimizeIterationsLeft>0 || !numIterations ))
             {
                 TArray<Ptr<ASTOp>> roots;
-                for(const auto& s:m_states)
+                for(const STATE_COMPILATION_DATA& s:m_states)
                 {
                     roots.Add(s.root);
                 }
@@ -2154,14 +2151,14 @@ namespace mu
 				UE_LOG(LogMutableCore, Verbose, TEXT("(int) %s : %ld"), TEXT("ast size"), int64(ASTOp::CountNodes(roots)));
 			}
 
-            for(auto& s:m_states)
+            for(STATE_COMPILATION_DATA& s:m_states)
             {
                 UE_LOG(LogMutableCore, Verbose, TEXT(" - constant optimiser"));
                 ConstantGeneratorAST( m_options->GetPrivate(), s.root );
             }
 
             TArray<Ptr<ASTOp>> roots;
-            for(const auto& s:m_states)
+            for(const STATE_COMPILATION_DATA& s:m_states)
             {
                 roots.Add(s.root);
             }
@@ -2178,7 +2175,7 @@ namespace mu
 
         // Gather all the current roots
         TArray<Ptr<ASTOp>> roots;
-        for(const auto& s:m_states)
+        for(const STATE_COMPILATION_DATA& s:m_states)
         {
             roots.Add(s.root);
         }
@@ -2196,7 +2193,7 @@ namespace mu
 
             // Update the marks for the instructions that don't depend on runtime parameters to be
             // cached.
-            for (auto& s:m_states)
+            for (STATE_COMPILATION_DATA& s:m_states)
             {
                 StateCacheDetectorAST c( &s );
             }
diff --git a/Engine/Plugins/Experimental/Mutable/Source/MutableTools/Private/MuT/Compiler.cpp b/Engine/Plugins/Experimental/Mutable/Source/MutableTools/Private/MuT/Compiler.cpp
--- a/Engine/Plugins/Experimental/Mutable/Source/MutableTools/Private/MuT/Compiler.cpp
+++ b/Engine/Plugins/Experimental/Mutable/Source/MutableTools/Private/MuT/Compiler.cpp
@@ -240,7 +240,7 @@ namespace mu
         }
 
         // Set the runtime parameter indices.
-        for( auto& s: states )
+        for( STATE_COMPILATION_DATA& s: states )
         {
             for ( int32 p=0; p<s.nodeState.m_runtimeParams.Num(); ++p )
             {
@@ -261,16 +261,13 @@ namespace mu
                 }
                 else
                 {
-                    char temp[256];
-                    mutable_snprintf( temp, 256,
-                                      "The state [%s] refers to a parameter [%s] "
-                                      "that has not been found in the model. This error can be "
-                                      "safely dismissed in case of partial compilation.",
-                                      s.nodeState.m_name.c_str(),
-                                      s.nodeState.m_runtimeParams[p].c_str()
-                                      );
-                    m_pD->m_pErrorLog->GetPrivate()->Add
-                            ( temp, ELMT_WARNING, pNode->GetBasePrivate()->m_errorContext );
+					FString Temp = FString::Printf(TEXT(
+						"The state [%s] refers to a parameter [%s] "
+						"that has not been found in the model. This error can be "
+						"safely dismissed in case of partial compilation."), 
+						s.nodeState.m_name.c_str(),
+						s.nodeState.m_runtimeParams[p].c_str());
+                    m_pD->m_pErrorLog->GetPrivate()->Add(Temp, ELMT_WARNING, pNode->GetBasePrivate()->m_errorContext );
                 }
             }
 
@@ -327,16 +324,6 @@ namespace mu
 
 		UE_LOG(LogMutableCore, Verbose, TEXT("(int) %s : %ld"), TEXT("program size"), int64(program.m_opAddress.Num()));
 
-        // Optimize the generated code
-        {        
-//            CodeOptimiser optimiser( m_pD->m_gpuPlatformProps,
-//                                     m_pD->m_optimisationOptions,
-//                                     states );
-//            optimiser.Optimise( pResult.get() );
-
-//            m_pD->m_pModelReport = optimiser.m_pModelReport;
-        }
-
         // Merge the log in the right order
         genErrorLog->Merge( m_pD->m_pErrorLog.get() );
         m_pD->m_pErrorLog = genErrorLog.get();
