[Engine/Plugins]
* Removed more includes using iwyu and marked up includes that is removed because of bugs in iwyu.

diff --git a/Engine/Plugins/AI/EnvironmentQueryEditor/Source/EnvironmentQueryEditor/Private/EnvironmentQueryEditor.cpp b/Engine/Plugins/AI/EnvironmentQueryEditor/Source/EnvironmentQueryEditor/Private/EnvironmentQueryEditor.cpp
--- a/Engine/Plugins/AI/EnvironmentQueryEditor/Source/EnvironmentQueryEditor/Private/EnvironmentQueryEditor.cpp
+++ b/Engine/Plugins/AI/EnvironmentQueryEditor/Source/EnvironmentQueryEditor/Private/EnvironmentQueryEditor.cpp
@@ -2,7 +2,6 @@
 
 #include "EnvironmentQueryEditor.h"
 #include "DetailsViewArgs.h"
-#include "Widgets/Text/STextBlock.h"
 #include "EdGraph/EdGraphSchema.h"
 #include "EnvironmentQuery/EnvQueryGenerator.h"
 #include "EnvironmentQuery/EnvQueryOption.h"
diff --git a/Engine/Plugins/AI/EnvironmentQueryEditor/Source/EnvironmentQueryEditor/Private/SEnvQueryProfiler.cpp b/Engine/Plugins/AI/EnvironmentQueryEditor/Source/EnvironmentQueryEditor/Private/SEnvQueryProfiler.cpp
--- a/Engine/Plugins/AI/EnvironmentQueryEditor/Source/EnvironmentQueryEditor/Private/SEnvQueryProfiler.cpp
+++ b/Engine/Plugins/AI/EnvironmentQueryEditor/Source/EnvironmentQueryEditor/Private/SEnvQueryProfiler.cpp
@@ -14,7 +14,6 @@
 #include "AssetRegistry/AssetRegistryModule.h"
 #include "EnvironmentQuery/EnvQuery.h"
 #include "Subsystems/AssetEditorSubsystem.h"
-#include "Widgets/Views/SListView.h"
 
 #define LOCTEXT_NAMESPACE "EnvironmentQueryEditor"
 
diff --git a/Engine/Plugins/AI/EnvironmentQueryEditor/Source/EnvironmentQueryEditor/Private/SEnvQueryProfiler.h b/Engine/Plugins/AI/EnvironmentQueryEditor/Source/EnvironmentQueryEditor/Private/SEnvQueryProfiler.h
--- a/Engine/Plugins/AI/EnvironmentQueryEditor/Source/EnvironmentQueryEditor/Private/SEnvQueryProfiler.h
+++ b/Engine/Plugins/AI/EnvironmentQueryEditor/Source/EnvironmentQueryEditor/Private/SEnvQueryProfiler.h
@@ -2,7 +2,6 @@
 
 #pragma once
 
-#include "Widgets/Views/ITableRow.h"
 #include "Widgets/Views/SListView.h" // IWYU pragma: keep
 
 class SEnvQueryLoadGraph;
diff --git a/Engine/Plugins/Compression/OodleNetwork/Source/Private/OodleNetworkTrainerCommandlet.cpp b/Engine/Plugins/Compression/OodleNetwork/Source/Private/OodleNetworkTrainerCommandlet.cpp
--- a/Engine/Plugins/Compression/OodleNetwork/Source/Private/OodleNetworkTrainerCommandlet.cpp
+++ b/Engine/Plugins/Compression/OodleNetwork/Source/Private/OodleNetworkTrainerCommandlet.cpp
@@ -1,6 +1,7 @@
 // Copyright Epic Games, Inc. All Rights Reserved.
 
 #include "OodleNetworkTrainerCommandlet.h"
+#include "HAL/PlatformFile.h"
 #include "Misc/MessageDialog.h"
 #include "HAL/FileManager.h"
 #include "Misc/ConfigCacheIni.h"
diff --git a/Engine/Plugins/Developer/TraceSourceFiltering/Source/SourceFilteringEditor/Private/ISessionSourceFilterService.h b/Engine/Plugins/Developer/TraceSourceFiltering/Source/SourceFilteringEditor/Private/ISessionSourceFilterService.h
--- a/Engine/Plugins/Developer/TraceSourceFiltering/Source/SourceFilteringEditor/Private/ISessionSourceFilterService.h
+++ b/Engine/Plugins/Developer/TraceSourceFiltering/Source/SourceFilteringEditor/Private/ISessionSourceFilterService.h
@@ -3,7 +3,6 @@
 #pragma once
 
 #include "Delegates/Delegate.h"
-#include "Delegates/DelegateCombinations.h"
 
 class FExtender;
 
diff --git a/Engine/Plugins/Developer/TraceSourceFiltering/Source/SourceFilteringTrace/Private/ResultCache.cpp b/Engine/Plugins/Developer/TraceSourceFiltering/Source/SourceFilteringTrace/Private/ResultCache.cpp
--- a/Engine/Plugins/Developer/TraceSourceFiltering/Source/SourceFilteringTrace/Private/ResultCache.cpp
+++ b/Engine/Plugins/Developer/TraceSourceFiltering/Source/SourceFilteringTrace/Private/ResultCache.cpp
@@ -2,7 +2,6 @@
 
 #include "ResultCache.h"
 #include "DataSourceFiltering.h"
-#include "SourceFilter.h"
 #include "SourceFilterSet.h"
 
 void FResultCache::Reset()
diff --git a/Engine/Plugins/Developer/TraceSourceFiltering/Source/SourceFilteringTrace/Private/ResultCache.h b/Engine/Plugins/Developer/TraceSourceFiltering/Source/SourceFilteringTrace/Private/ResultCache.h
--- a/Engine/Plugins/Developer/TraceSourceFiltering/Source/SourceFilteringTrace/Private/ResultCache.h
+++ b/Engine/Plugins/Developer/TraceSourceFiltering/Source/SourceFilteringTrace/Private/ResultCache.h
@@ -2,8 +2,11 @@
 
 #pragma once
 
-#include "SourceFilterSet.h"
-#include "SourceFilter.h"
+#include "Containers/Map.h"
+
+class AActor;
+struct FFilter;
+struct FFilterSet;
 
 /** Whether or not to use a BitArray to store per-actor/filter results */
 #define USE_BITSET_STORAGE 1
diff --git a/Engine/Plugins/Developer/TraceSourceFiltering/Source/SourceFilteringTrace/Private/SourceFilterSet.h b/Engine/Plugins/Developer/TraceSourceFiltering/Source/SourceFilteringTrace/Private/SourceFilterSet.h
--- a/Engine/Plugins/Developer/TraceSourceFiltering/Source/SourceFilteringTrace/Private/SourceFilterSet.h
+++ b/Engine/Plugins/Developer/TraceSourceFiltering/Source/SourceFilteringTrace/Private/SourceFilterSet.h
@@ -2,8 +2,10 @@
 
 #pragma once
 
-#include "SourceFilter.h"
-#include "DataSourceFiltering.h"
+#include "Containers/Array.h"
+#include "SourceFilter.h" // IWYU pragma: keep
+
+enum class EFilterSetMode : uint8;
 
 typedef TInlineAllocator<8> FilterSetAllocator;
 
@@ -31,4 +33,4 @@ struct FFilterSet
 	uint8 bContainsAsyncFilter : 1;
 	/** Expected value (true/false) for whether or not an applied filter within this set passes */
 	uint8 bInitialPassingValue : 1;	
-};
\ No newline at end of file
+};
diff --git a/Engine/Plugins/Developer/TraceSourceFiltering/Source/SourceFilteringTrace/Private/SourceFilterSetup.h b/Engine/Plugins/Developer/TraceSourceFiltering/Source/SourceFilteringTrace/Private/SourceFilterSetup.h
--- a/Engine/Plugins/Developer/TraceSourceFiltering/Source/SourceFilteringTrace/Private/SourceFilterSetup.h
+++ b/Engine/Plugins/Developer/TraceSourceFiltering/Source/SourceFilteringTrace/Private/SourceFilterSetup.h
@@ -5,7 +5,6 @@
 #include "Delegates/Delegate.h"
 #include "ResultCache.h"
 
-#include "SourceFilter.h"
 #include "SourceFilterSet.h"
 
 struct FActorClassFilter;
diff --git a/Engine/Plugins/Editor/ChangelistReview/Source/Private/SSourceControlReview.cpp b/Engine/Plugins/Editor/ChangelistReview/Source/Private/SSourceControlReview.cpp
--- a/Engine/Plugins/Editor/ChangelistReview/Source/Private/SSourceControlReview.cpp
+++ b/Engine/Plugins/Editor/ChangelistReview/Source/Private/SSourceControlReview.cpp
@@ -9,12 +9,10 @@
 #include "ISourceControlModule.h"
 #include "Misc/PackageName.h"
 #include "SourceControlOperations.h"
-#include "Misc/Paths.h"
 #include "SSourceControlReviewEntry.h"
 #include "AssetRegistry/AssetRegistryModule.h"
 #include "Kismet/KismetSystemLibrary.h"
 #include "Styling/StarshipCoreStyle.h"
-#include "Widgets/Input/SEditableText.h"
 #include "Widgets/Input/SMultiLineEditableTextBox.h"
 #include "Widgets/Input/SEditableTextBox.h"
 #include "Widgets/Notifications/SProgressBar.h"
@@ -25,9 +23,7 @@
 #include "Widgets/Layout/SGridPanel.h"
 #include "Brushes/SlateRoundedBoxBrush.h"
 #include "Internationalization/Regex.h"
-#include "Widgets/Views/SListView.h"
 #include "Widgets/Input/SComboBox.h"
-#include "CoreGlobals.h"
 #include "Misc/ConfigCacheIni.h"
 
 #define LOCTEXT_NAMESPACE "SourceControlReview"
diff --git a/Engine/Plugins/Editor/ChangelistReview/Source/Private/SSourceControlReview.h b/Engine/Plugins/Editor/ChangelistReview/Source/Private/SSourceControlReview.h
--- a/Engine/Plugins/Editor/ChangelistReview/Source/Private/SSourceControlReview.h
+++ b/Engine/Plugins/Editor/ChangelistReview/Source/Private/SSourceControlReview.h
@@ -4,7 +4,8 @@
 #include "ISourceControlProvider.h"
 #include "SourceControlOperations.h"
 #include "Widgets/Views/SHeaderRow.h"
-#include "Widgets/Input/SComboBox.h"
+
+template <typename OptionType> class SComboBox;
 
 class SEditableText;
 class STableViewBase;
diff --git a/Engine/Plugins/Editor/EngineAssetDefinitions/Source/Private/Animation/AssetDefinition_SkeletalMesh.cpp b/Engine/Plugins/Editor/EngineAssetDefinitions/Source/Private/Animation/AssetDefinition_SkeletalMesh.cpp
--- a/Engine/Plugins/Editor/EngineAssetDefinitions/Source/Private/Animation/AssetDefinition_SkeletalMesh.cpp
+++ b/Engine/Plugins/Editor/EngineAssetDefinitions/Source/Private/Animation/AssetDefinition_SkeletalMesh.cpp
@@ -4,9 +4,7 @@
 
 #include "AnimationEditorUtils.h"
 #include "AssetNotifications.h"
-#include "AssetToolsModule.h"
 #include "ContentBrowserMenuContexts.h"
-#include "ContentBrowserModule.h"
 #include "EditorReimportHandler.h"
 #include "FbxMeshUtils.h"
 #include "ISkeletalMeshEditorModule.h"
@@ -15,10 +13,10 @@
 #include "Factories/SkeletonFactory.h"
 #include "ThumbnailRendering/SceneThumbnailInfo.h"
 #include "Misc/MessageDialog.h"
-#include "Animation/Skeleton.h"
 #include "Factories/FbxSkeletalMeshImportData.h"
 #include "Factories/PhysicsAssetFactory.h"
 #include "PhysicsEngine/PhysicsAsset.h"
+#include "ToolMenuSection.h"
 #include "Widgets/Layout/SBorder.h"
 #include "Widgets/Images/SImage.h"
 #include "ToolMenu.h"
diff --git a/Engine/Plugins/Editor/EngineAssetDefinitions/Source/Private/Animation/AssetDefinition_SkeletalMesh.h b/Engine/Plugins/Editor/EngineAssetDefinitions/Source/Private/Animation/AssetDefinition_SkeletalMesh.h
--- a/Engine/Plugins/Editor/EngineAssetDefinitions/Source/Private/Animation/AssetDefinition_SkeletalMesh.h
+++ b/Engine/Plugins/Editor/EngineAssetDefinitions/Source/Private/Animation/AssetDefinition_SkeletalMesh.h
@@ -2,12 +2,13 @@
 
 #pragma once
 
-#include "CoreMinimal.h"
 #include "Engine/SkeletalMesh.h"
 #include "AssetDefinitionDefault.h"
 
 #include "AssetDefinition_SkeletalMesh.generated.h"
 
+struct EVisibility;
+
 UCLASS()
 class UAssetDefinition_SkeletalMesh : public UAssetDefinitionDefault
 {
diff --git a/Engine/Plugins/Editor/EngineAssetDefinitions/Source/Private/Animation/AssetDefinition_Skeleton.cpp b/Engine/Plugins/Editor/EngineAssetDefinitions/Source/Private/Animation/AssetDefinition_Skeleton.cpp
--- a/Engine/Plugins/Editor/EngineAssetDefinitions/Source/Private/Animation/AssetDefinition_Skeleton.cpp
+++ b/Engine/Plugins/Editor/EngineAssetDefinitions/Source/Private/Animation/AssetDefinition_Skeleton.cpp
@@ -1,28 +1,12 @@
 // Copyright Epic Games, Inc. All Rights Reserved.
 
 #include "Animation/AssetDefinition_Skeleton.h"
-#include "Widgets/Text/STextBlock.h"
+#include "ToolMenu.h"
 #include "ToolMenus.h"
-#include "Engine/SkeletalMesh.h"
-#include "Animation/AnimationAsset.h"
-#include "Animation/AnimSequenceBase.h"
-#include "Animation/AnimSequence.h"
-#include "Widgets/Layout/SBorder.h"
-#include "Misc/MessageDialog.h"
-#include "HAL/FileManager.h"
-#include "UObject/UObjectHash.h"
-#include "UObject/UObjectIterator.h"
-#include "Framework/Application/SlateApplication.h"
-#include "Widgets/Input/SButton.h"
-#include "Widgets/Input/SCheckBox.h"
-#include "Styling/AppStyle.h"
-#include "Components/SkeletalMeshComponent.h"
 #include "AnimationEditorUtils.h"
 #include "ISkeletonEditorModule.h"
-#include "Algo/Transform.h"
 #include "ContentBrowserMenuContexts.h"
-#include "Subsystems/AssetEditorSubsystem.h"
-#include "Editor.h"
+#include "ToolMenuSection.h"
 
 #define LOCTEXT_NAMESPACE "AssetTypeActions"
 
diff --git a/Engine/Plugins/Editor/EngineAssetDefinitions/Source/Private/Animation/AssetDefinition_Skeleton.h b/Engine/Plugins/Editor/EngineAssetDefinitions/Source/Private/Animation/AssetDefinition_Skeleton.h
--- a/Engine/Plugins/Editor/EngineAssetDefinitions/Source/Private/Animation/AssetDefinition_Skeleton.h
+++ b/Engine/Plugins/Editor/EngineAssetDefinitions/Source/Private/Animation/AssetDefinition_Skeleton.h
@@ -2,7 +2,6 @@
 
 #pragma once
 
-#include "CoreMinimal.h"
 #include "Animation/Skeleton.h"
 #include "AssetDefinitionDefault.h"
 
diff --git a/Engine/Plugins/Editor/EngineAssetDefinitions/Source/Private/Animation/SkeletalMeshUtilities.cpp b/Engine/Plugins/Editor/EngineAssetDefinitions/Source/Private/Animation/SkeletalMeshUtilities.cpp
--- a/Engine/Plugins/Editor/EngineAssetDefinitions/Source/Private/Animation/SkeletalMeshUtilities.cpp
+++ b/Engine/Plugins/Editor/EngineAssetDefinitions/Source/Private/Animation/SkeletalMeshUtilities.cpp
@@ -1,37 +1,15 @@
 // Copyright Epic Games, Inc. All Rights Reserved.
 
 #include "Animation/SkeletalMeshUtilities.h"
-#include "Animation/Skeleton.h"
-#include "Widgets/Text/STextBlock.h"
-#include "ToolMenus.h"
-#include "Misc/PackageName.h"
-#include "Widgets/Layout/SBorder.h"
+#include "Engine/SkeletalMesh.h"
 #include "Misc/MessageDialog.h"
 #include "SlateOptMacros.h"
 #include "Framework/Application/SlateApplication.h"
 #include "Widgets/Layout/SSeparator.h"
-#include "Widgets/Layout/SUniformGridPanel.h"
-#include "Widgets/Input/SButton.h"
 #include "Widgets/Layout/SScrollBox.h"
 #include "Widgets/Input/SCheckBox.h"
-#include "Styling/AppStyle.h"
-#include "Factories/SkeletonFactory.h"
-#include "EditorFramework/AssetImportData.h"
-#include "ThumbnailRendering/SceneThumbnailInfo.h"
-#include "AssetRegistry/AssetRegistryModule.h"
 #include "SSkeletonWidget.h"
-#include "ContentBrowserModule.h"
-#include "AnimationEditorUtils.h"
-#include "FbxMeshUtils.h"
 #include "AssetNotifications.h"
-#include "PhysicsEngine/PhysicsAsset.h"
-#include "ISkeletalMeshEditorModule.h"
-#include "Algo/Transform.h"
-#include "Factories/PhysicsAssetFactory.h"
-#include "EditorReimportHandler.h"
-#include "Widgets/Images/SImage.h"
-#include "Factories/FbxSkeletalMeshImportData.h"
-#include "Subsystems/AssetEditorSubsystem.h"
 #include "Editor.h"
 
 #define LOCTEXT_NAMESPACE "AssetTypeActions"
@@ -510,4 +488,4 @@ namespace UE::SkeletalMeshUtilities
 	}
 }
 
-#undef LOCTEXT_NAMESPACE
\ No newline at end of file
+#undef LOCTEXT_NAMESPACE
diff --git a/Engine/Plugins/Editor/EngineAssetDefinitions/Source/Private/Animation/SkeletalMeshUtilities.h b/Engine/Plugins/Editor/EngineAssetDefinitions/Source/Private/Animation/SkeletalMeshUtilities.h
--- a/Engine/Plugins/Editor/EngineAssetDefinitions/Source/Private/Animation/SkeletalMeshUtilities.h
+++ b/Engine/Plugins/Editor/EngineAssetDefinitions/Source/Private/Animation/SkeletalMeshUtilities.h
@@ -2,11 +2,10 @@
 
 #pragma once
 
-#include "CoreMinimal.h"
 
 class USkeletalMesh;
 
 namespace UE::SkeletalMeshUtilities
 {
 	void AssignSkeletonToMesh(USkeletalMesh* SkelMesh);
-}
\ No newline at end of file
+}
diff --git a/Engine/Plugins/Editor/EngineAssetDefinitions/Source/Private/AssetDefinition_DataAsset.cpp b/Engine/Plugins/Editor/EngineAssetDefinitions/Source/Private/AssetDefinition_DataAsset.cpp
--- a/Engine/Plugins/Editor/EngineAssetDefinitions/Source/Private/AssetDefinition_DataAsset.cpp
+++ b/Engine/Plugins/Editor/EngineAssetDefinitions/Source/Private/AssetDefinition_DataAsset.cpp
@@ -10,7 +10,6 @@
 #include "ToolMenus.h"
 #include "Kismet2/SClassPickerDialog.h"
 #include "Engine/Engine.h"
-#include "Misc/AssetFilterData.h"
 #include "ToolMenuSection.h"
 
 #define LOCTEXT_NAMESPACE "UAssetDefinition_DataAsset"
diff --git a/Engine/Plugins/Enterprise/VariantManager/Source/VariantManager/Private/SVariantManagerTableRow.cpp b/Engine/Plugins/Enterprise/VariantManager/Source/VariantManager/Private/SVariantManagerTableRow.cpp
--- a/Engine/Plugins/Enterprise/VariantManager/Source/VariantManager/Private/SVariantManagerTableRow.cpp
+++ b/Engine/Plugins/Enterprise/VariantManager/Source/VariantManager/Private/SVariantManagerTableRow.cpp
@@ -10,7 +10,6 @@
 
 #include "GameFramework/Actor.h"
 #include "Styling/SlateIconFinder.h"
-#include "Widgets/Views/SListView.h"
 
 #define LOCTEXT_NAMESPACE "SVariantManagerTableRow"
 
diff --git a/Engine/Plugins/Enterprise/VariantManager/Source/VariantManager/Private/VariantManagerNodeTree.h b/Engine/Plugins/Enterprise/VariantManager/Source/VariantManager/Private/VariantManagerNodeTree.h
--- a/Engine/Plugins/Enterprise/VariantManager/Source/VariantManager/Private/VariantManagerNodeTree.h
+++ b/Engine/Plugins/Enterprise/VariantManager/Source/VariantManager/Private/VariantManagerNodeTree.h
@@ -3,7 +3,6 @@
 #pragma once
 
 #include "Delegates/Delegate.h"
-#include "Templates/SharedPointer.h"
 
 class FVariantManager;
 class FVariantManagerDisplayNode;
diff --git a/Engine/Plugins/Enterprise/VariantManager/Source/VariantManager/Private/VariantManagerSelection.h b/Engine/Plugins/Enterprise/VariantManager/Source/VariantManager/Private/VariantManagerSelection.h
--- a/Engine/Plugins/Enterprise/VariantManager/Source/VariantManager/Private/VariantManagerSelection.h
+++ b/Engine/Plugins/Enterprise/VariantManager/Source/VariantManager/Private/VariantManagerSelection.h
@@ -3,7 +3,6 @@
 #pragma once
 
 #include "Delegates/Delegate.h"
-#include "Templates/SharedPointer.h"
 
 class FVariantManagerActorNode;
 class FVariantManagerDisplayNode;
diff --git a/Engine/Plugins/Experimental/ChaosNiagara/Source/ChaosNiagara/Private/NiagaraDataInterfaceChaosDestruction.cpp b/Engine/Plugins/Experimental/ChaosNiagara/Source/ChaosNiagara/Private/NiagaraDataInterfaceChaosDestruction.cpp
--- a/Engine/Plugins/Experimental/ChaosNiagara/Source/ChaosNiagara/Private/NiagaraDataInterfaceChaosDestruction.cpp
+++ b/Engine/Plugins/Experimental/ChaosNiagara/Source/ChaosNiagara/Private/NiagaraDataInterfaceChaosDestruction.cpp
@@ -18,8 +18,6 @@
 #include "PBDRigidsSolver.h"
 #include "Physics/Experimental/PhysScene_Chaos.h"
 #include "RenderGraphBuilder.h"
-#include "RenderGraphTrace.h"
-#include "RenderGraphValidation.h"
 #include "RenderingThread.h"
 
 #include UE_INLINE_GENERATED_CPP_BY_NAME(NiagaraDataInterfaceChaosDestruction)
diff --git a/Engine/Plugins/Experimental/CommonConversation/Source/CommonConversationEditor/Private/ConversationEditorTabFactories.cpp b/Engine/Plugins/Experimental/CommonConversation/Source/CommonConversationEditor/Private/ConversationEditorTabFactories.cpp
--- a/Engine/Plugins/Experimental/CommonConversation/Source/CommonConversationEditor/Private/ConversationEditorTabFactories.cpp
+++ b/Engine/Plugins/Experimental/CommonConversation/Source/CommonConversationEditor/Private/ConversationEditorTabFactories.cpp
@@ -4,7 +4,6 @@
 #include "ConversationEditor.h"
 #include "Engine/Blueprint.h"
 #include "ConversationEditorTabs.h"
-#include "EdGraph/EdGraph.h"
 #include "Widgets/Docking/SDockTab.h"
 #include "ConversationDatabase.h"
 
diff --git a/Engine/Plugins/Experimental/GameFeatures/Source/GameFeaturesEditor/Private/GameFeatureDataDetailsCustomization.h b/Engine/Plugins/Experimental/GameFeatures/Source/GameFeaturesEditor/Private/GameFeatureDataDetailsCustomization.h
--- a/Engine/Plugins/Experimental/GameFeatures/Source/GameFeaturesEditor/Private/GameFeatureDataDetailsCustomization.h
+++ b/Engine/Plugins/Experimental/GameFeatures/Source/GameFeaturesEditor/Private/GameFeatureDataDetailsCustomization.h
@@ -4,7 +4,7 @@
 
 #include "IDetailCustomization.h"
 #include "GameFeatureTypesFwd.h"
-#include "UObject/WeakObjectPtr.h"
+#include "UObject/WeakObjectPtrTemplates.h"
 
 namespace UE::GameFeatures { struct FResult; }
 
diff --git a/Engine/Plugins/Experimental/Mutable/Source/CustomizableObject/Private/MuCO/CustomizableObject.cpp b/Engine/Plugins/Experimental/Mutable/Source/CustomizableObject/Private/MuCO/CustomizableObject.cpp
--- a/Engine/Plugins/Experimental/Mutable/Source/CustomizableObject/Private/MuCO/CustomizableObject.cpp
+++ b/Engine/Plugins/Experimental/Mutable/Source/CustomizableObject/Private/MuCO/CustomizableObject.cpp
@@ -8,6 +8,7 @@
 #include "EdGraph/EdGraph.h"
 #include "Engine/SkeletalMesh.h"
 #include "HAL/FileManager.h"
+#include "HAL/PlatformFile.h"
 #include "HAL/PlatformFileManager.h"
 #include "Input/Reply.h"
 #include "Interfaces/ITargetPlatform.h"
diff --git a/Engine/Plugins/Experimental/Mutable/Source/CustomizableObject/Private/MuCO/LogBenchmarkUtil.cpp b/Engine/Plugins/Experimental/Mutable/Source/CustomizableObject/Private/MuCO/LogBenchmarkUtil.cpp
--- a/Engine/Plugins/Experimental/Mutable/Source/CustomizableObject/Private/MuCO/LogBenchmarkUtil.cpp
+++ b/Engine/Plugins/Experimental/Mutable/Source/CustomizableObject/Private/MuCO/LogBenchmarkUtil.cpp
@@ -2,6 +2,7 @@
 
 #include "MuCO/LogBenchmarkUtil.h"
 
+#include "HAL/PlatformFile.h"
 #include "HAL/PlatformFileManager.h"
 #include "Misc/FileHelper.h"
 #include "MuCO/UnrealPortabilityHelpers.h"
diff --git a/Engine/Plugins/Experimental/Mutable/Source/CustomizableObject/Private/MuCO/LogBenchmarkUtil.h b/Engine/Plugins/Experimental/Mutable/Source/CustomizableObject/Private/MuCO/LogBenchmarkUtil.h
--- a/Engine/Plugins/Experimental/Mutable/Source/CustomizableObject/Private/MuCO/LogBenchmarkUtil.h
+++ b/Engine/Plugins/Experimental/Mutable/Source/CustomizableObject/Private/MuCO/LogBenchmarkUtil.h
@@ -3,8 +3,6 @@
 #pragma once
 
 #include "Containers/Map.h"
-#include "GenericPlatform/GenericPlatformMath.h"
-#include "HAL/Platform.h"
 #include "UObject/NameTypes.h"
 
 class LogBenchmarkUtil
diff --git a/Engine/Plugins/Experimental/Mutable/Source/CustomizableObject/Private/MuCO/UnrealMutableModelDiskStreamer.cpp b/Engine/Plugins/Experimental/Mutable/Source/CustomizableObject/Private/MuCO/UnrealMutableModelDiskStreamer.cpp
--- a/Engine/Plugins/Experimental/Mutable/Source/CustomizableObject/Private/MuCO/UnrealMutableModelDiskStreamer.cpp
+++ b/Engine/Plugins/Experimental/Mutable/Source/CustomizableObject/Private/MuCO/UnrealMutableModelDiskStreamer.cpp
@@ -3,6 +3,7 @@
 #include "MuCO/UnrealMutableModelDiskStreamer.h"
 
 #include "Async/AsyncFileHandle.h"
+#include "HAL/PlatformFile.h"
 #include "MuCO/CustomizableObjectSystem.h"
 #include "MuR/Model.h"
 #include "MuR/MutableTrace.h"
diff --git a/Engine/Plugins/Experimental/Mutable/Source/CustomizableObjectEditor/Private/MuCOE/CustomizableObjectEditor.cpp b/Engine/Plugins/Experimental/Mutable/Source/CustomizableObjectEditor/Private/MuCOE/CustomizableObjectEditor.cpp
--- a/Engine/Plugins/Experimental/Mutable/Source/CustomizableObjectEditor/Private/MuCOE/CustomizableObjectEditor.cpp
+++ b/Engine/Plugins/Experimental/Mutable/Source/CustomizableObjectEditor/Private/MuCOE/CustomizableObjectEditor.cpp
@@ -2,57 +2,20 @@
 
 #include "MuCOE/CustomizableObjectEditor.h"
 
-#include "AdvancedPreviewSceneModule.h"
-#include "Animation/AnimSingleNodeInstance.h"
 #include "Animation/DebugSkelMeshComponent.h"
-#include "Animation/PoseAsset.h"
-#include "Animation/SmartName.h"
 #include "AssetRegistry/ARFilter.h"
-#include "AssetRegistry/AssetData.h"
 #include "AssetRegistry/AssetRegistryModule.h"
-#include "AssetRegistry/IAssetRegistry.h"
-#include "Components/SkeletalMeshComponent.h"
-#include "Containers/Set.h"
 #include "DetailsViewArgs.h"
-#include "EdGraph/EdGraph.h"
-#include "EdGraph/EdGraphPin.h"
-#include "EdGraph/EdGraphSchema.h"
-#include "EdGraphNode_Comment.h"
 #include "EdGraphUtilities.h"
-#include "Editor.h"
-#include "Editor/EditorEngine.h"
-#include "Engine/Engine.h"
-#include "Engine/EngineTypes.h"
 #include "Engine/StaticMesh.h"
-#include "Engine/StreamableManager.h"
 #include "FileHelpers.h"
 #include "Framework/Commands/GenericCommands.h"
-#include "Framework/Commands/UIAction.h"
-#include "Framework/Commands/UICommandList.h"
-#include "Framework/Docking/TabManager.h"
 #include "Framework/MultiBox/MultiBoxBuilder.h"
-#include "Framework/MultiBox/MultiBoxExtender.h"
-#include "Framework/Notifications/NotificationManager.h"
-#include "Framework/SlateDelegates.h"
 #include "GraphEditorActions.h"
 #include "HAL/PlatformApplicationMisc.h"
-#include "HAL/PlatformTime.h"
 #include "IDetailsView.h"
-#include "Internationalization/Internationalization.h"
-#include "Layout/BasicLayoutWidgetSlot.h"
-#include "Layout/ChildrenBase.h"
-#include "Layout/SlateRect.h"
-#include "Logging/LogCategory.h"
-#include "Logging/LogMacros.h"
-#include "Logging/TokenizedMessage.h"
 #include "MaterialGraph/MaterialGraphNode.h"
 #include "MaterialGraph/MaterialGraphNode_Comment.h"
-#include "Misc/AssertionMacros.h"
-#include "Misc/Attribute.h"
-#include "Misc/CString.h"
-#include "Misc/Guid.h"
-#include "Modules/ModuleManager.h"
-#include "MuCO/CustomizableObject.h"
 #include "MuCO/CustomizableObjectInstance.h"
 #include "MuCO/CustomizableObjectSystem.h"
 #include "MuCO/CustomizableSkeletalComponent.h"
@@ -69,70 +32,36 @@
 #include "MuCOE/CustomizableObjectPreviewScene.h"
 #include "MuCOE/EdGraphSchema_CustomizableObject.h"
 #include "MuCOE/GraphTraversal.h"
-#include "MuCOE/Nodes/CustomizableObjectNode.h"
-#include "MuCOE/Nodes/CustomizableObjectNodeColorArithmeticOp.h"
-#include "MuCOE/Nodes/CustomizableObjectNodeEditMaterial.h"
 #include "MuCOE/Nodes/CustomizableObjectNodeEnumParameter.h"
 #include "MuCOE/Nodes/CustomizableObjectNodeFloatParameter.h"
 #include "MuCOE/Nodes/CustomizableObjectNodeGroupProjectorParameter.h"
-#include "MuCOE/Nodes/CustomizableObjectNodeLayoutBlocks.h"
 #include "MuCOE/Nodes/CustomizableObjectNodeMaterial.h"
 #include "MuCOE/Nodes/CustomizableObjectNodeMaterialVariation.h"
 #include "MuCOE/Nodes/CustomizableObjectNodeMeshClipMorph.h"
 #include "MuCOE/Nodes/CustomizableObjectNodeMeshClipWithMesh.h"
-#include "MuCOE/Nodes/CustomizableObjectNodeObject.h"
 #include "MuCOE/Nodes/CustomizableObjectNodeObjectGroup.h"
 #include "MuCOE/Nodes/CustomizableObjectNodeProjectorConstant.h"
-#include "MuCOE/Nodes/CustomizableObjectNodeProjectorParameter.h"
-#include "MuCOE/Nodes/CustomizableObjectNodeRemoveMeshBlocks.h"
 #include "MuCOE/Nodes/CustomizableObjectNodeSkeletalMesh.h"
 #include "MuCOE/Nodes/CustomizableObjectNodeStaticMesh.h"
-#include "MuCOE/Nodes/CustomizableObjectNodeTable.h"
 #include "MuCOE/SCustomizableObjectEditorAdvancedPreviewSettings.h"
 #include "MuCOE/SCustomizableObjectEditorPerformanceReport.h"
 #include "MuCOE/SCustomizableObjectEditorTagExplorer.h"
 #include "MuCOE/SCustomizableObjectEditorTextureAnalyzer.h"
 #include "MuCOE/SCustomizableObjectEditorViewport.h"
-#include "MuCOE/SCustomizableObjectNodeLayoutBlocksEditor.h"
-#include "MuCOE/SCustomizableObjectNodeLayoutBlocksSelector.h"
 #include "PropertyEditorModule.h"
-#include "SNodePanel.h"
 #include "ScopedTransaction.h"
-#include "SlotBase.h"
-#include "Styling/AppStyle.h"
-#include "Styling/ISlateStyle.h"
-#include "Templates/Casts.h"
-#include "Templates/Function.h"
-#include "Templates/SubclassOf.h"
-#include "Textures/SlateIcon.h"
-#include "Toolkits/AssetEditorToolkit.h"
-#include "Trace/Detail/Channel.h"
-#include "Types/SlateEnums.h"
-#include "Types/SlateStructs.h"
-#include "UObject/Class.h"
 #include "UObject/EnumProperty.h"
-#include "UObject/Field.h"
-#include "UObject/Package.h"
-#include "UObject/TopLevelAssetPath.h"
-#include "UObject/UnrealNames.h"
-#include "UObject/UnrealType.h"
-#include "Widgets/DeclarativeSyntaxSupport.h"
 #include "Widgets/Docking/SDockTab.h"
 #include "Widgets/Input/SSearchBox.h"
 #include "Widgets/Input/STextComboBox.h"
-#include "Widgets/Layout/SBorder.h"
-#include "Widgets/Layout/SBox.h"
-#include "Widgets/Layout/SScrollBar.h"
 #include "Widgets/Layout/SScrollBox.h"
-#include "Widgets/Layout/SSplitter.h"
 #include "Widgets/Notifications/SNotificationList.h"
-#include "Widgets/SBoxPanel.h"
 
 class FAdvancedPreviewScene;
 class FWorkspaceItem;
 class IToolkitHost;
 class SWidget;
-
+enum class EColorArithmeticOperation : uint8;
 
 #define LOCTEXT_NAMESPACE "CustomizableObjectEditor"
 
diff --git a/Engine/Plugins/Experimental/Mutable/Source/CustomizableObjectEditor/Private/MuCOE/CustomizableObjectLayout.cpp b/Engine/Plugins/Experimental/Mutable/Source/CustomizableObjectEditor/Private/MuCOE/CustomizableObjectLayout.cpp
--- a/Engine/Plugins/Experimental/Mutable/Source/CustomizableObjectEditor/Private/MuCOE/CustomizableObjectLayout.cpp
+++ b/Engine/Plugins/Experimental/Mutable/Source/CustomizableObjectEditor/Private/MuCOE/CustomizableObjectLayout.cpp
@@ -4,7 +4,6 @@
 
 #include "EdGraph/EdGraph.h"
 #include "Engine/StaticMesh.h"
-#include "Materials/MaterialInterface.h"
 #include "MuCOE/CustomizableObjectCompiler.h"
 #include "MuCOE/GenerateMutableSource/GenerateMutableSourceMesh.h"
 #include "MuCOE/GraphTraversal.h"
diff --git a/Engine/Plugins/Experimental/Mutable/Source/CustomizableObjectEditor/Private/MuCOE/GenerateMutableSource/GenerateMutableSource.cpp b/Engine/Plugins/Experimental/Mutable/Source/CustomizableObjectEditor/Private/MuCOE/GenerateMutableSource/GenerateMutableSource.cpp
--- a/Engine/Plugins/Experimental/Mutable/Source/CustomizableObjectEditor/Private/MuCOE/GenerateMutableSource/GenerateMutableSource.cpp
+++ b/Engine/Plugins/Experimental/Mutable/Source/CustomizableObjectEditor/Private/MuCOE/GenerateMutableSource/GenerateMutableSource.cpp
@@ -3,32 +3,9 @@
 #include "MuCOE/GenerateMutableSource/GenerateMutableSource.h"
 
 #include "Animation/AnimInstance.h"
-#include "Animation/Skeleton.h"
-#include "Containers/IndirectArray.h"
-#include "Containers/StringConv.h"
-#include "EdGraph/EdGraph.h"
-#include "EdGraph/EdGraphNode.h"
-#include "Engine/SkeletalMesh.h"
 #include "Engine/SkeletalMeshSocket.h"
-#include "Engine/Texture2D.h"
 #include "Engine/TextureLODSettings.h"
-#include "HAL/PlatformCrt.h"
-#include "HAL/PlatformMisc.h"
 #include "Interfaces/ITargetPlatform.h"
-#include "Internationalization/Internationalization.h"
-#include "Internationalization/Text.h"
-#include "Logging/LogCategory.h"
-#include "Logging/LogMacros.h"
-#include "Logging/TokenizedMessage.h"
-#include "Math/BoxSphereBounds.h"
-#include "Math/Rotator.h"
-#include "Math/Vector.h"
-#include "Misc/AssertionMacros.h"
-#include "MuCO/CustomizableObject.h"
-#include "MuCO/CustomizableObjectIdentifier.h"
-#include "MuCO/CustomizableObjectParameterTypeDefinitions.h"
-#include "MuCO/CustomizableObjectUIData.h"
-#include "MuCO/UnrealPortabilityHelpers.h"
 #include "MuCOE/CustomizableObjectCompiler.h"
 #include "MuCOE/EdGraphSchema_CustomizableObject.h"
 #include "MuCOE/GenerateMutableSource/GenerateMutableSourceGroupProjector.h"
@@ -40,9 +17,7 @@
 #include "MuCOE/Nodes/CustomizableObjectNodeMaterial.h"
 #include "MuCOE/Nodes/CustomizableObjectNodeMaterialVariation.h"
 #include "MuCOE/Nodes/CustomizableObjectNodeMesh.h"
-#include "MuCOE/Nodes/CustomizableObjectNodeMeshClipMorph.h"
 #include "MuCOE/Nodes/CustomizableObjectNodeMeshClipWithMesh.h"
-#include "MuCOE/Nodes/CustomizableObjectNodeModifierBase.h"
 #include "MuCOE/Nodes/CustomizableObjectNodeObjectGroup.h"
 #include "MuCOE/Nodes/CustomizableObjectNodeTable.h"
 #include "MuCOE/Nodes/CustomizableObjectNodeTexture.h"
@@ -59,32 +34,14 @@
 #include "MuCOE/Nodes/CustomizableObjectNodeTextureTransform.h"
 #include "MuCOE/Nodes/CustomizableObjectNodeTextureVariation.h"
 #include "MuCOE/UnrealEditorPortabilityHelpers.h"
-#include "MuR/Mesh.h"
-#include "MuR/Ptr.h"
-#include "MuR/Skeleton.h"
-#include "MuT/Node.h"
-#include "MuT/NodeComponent.h"
 #include "MuT/NodeComponentEdit.h"
-#include "MuT/NodeComponentNew.h"
 #include "MuT/NodeLOD.h"
-#include "MuT/NodeMeshApplyPose.h"
 #include "MuT/NodeMeshConstant.h"
-#include "MuT/NodeModifierMeshClipWithMesh.h"
 #include "MuT/NodeObjectGroup.h"
 #include "MuT/NodeObjectNew.h"
 #include "MuT/NodeSurfaceEdit.h"
-#include "MuT/NodeSurfaceNew.h"
-#include "PerPlatformProperties.h"
-#include "PerQualityLevelProperties.h"
 #include "PhysicsEngine/PhysicsAsset.h"
 #include "PlatformInfo.h"
-#include "Rendering/SkeletalMeshLODRenderData.h"
-#include "Templates/Casts.h"
-#include "Templates/ChooseClass.h"
-#include "Templates/SubclassOf.h"
-#include "Trace/Detail/Channel.h"
-#include "UObject/ObjectPtr.h"
-#include "UObject/UObjectGlobals.h"
 
 #define LOCTEXT_NAMESPACE "CustomizableObjectEditor"
 
@@ -162,6 +119,7 @@ FMutableGraphGenerationContext::FMutableGraphGenerationContext(UCustomizableObje
 	MeshGenerationFlags.Push(EMutableMeshConversionFlags::None);
 }
 
+FMutableGraphGenerationContext::~FMutableGraphGenerationContext() = default;
 
 mu::MeshPtr FMutableGraphGenerationContext::FindGeneratedMesh( const FGeneratedMeshData::FKey& Key )
 {
diff --git a/Engine/Plugins/Experimental/Mutable/Source/CustomizableObjectEditor/Private/MuCOE/GenerateMutableSource/GenerateMutableSource.h b/Engine/Plugins/Experimental/Mutable/Source/CustomizableObjectEditor/Private/MuCOE/GenerateMutableSource/GenerateMutableSource.h
--- a/Engine/Plugins/Experimental/Mutable/Source/CustomizableObjectEditor/Private/MuCOE/GenerateMutableSource/GenerateMutableSource.h
+++ b/Engine/Plugins/Experimental/Mutable/Source/CustomizableObjectEditor/Private/MuCOE/GenerateMutableSource/GenerateMutableSource.h
@@ -3,29 +3,12 @@
 #pragma once
 
 #include "Animation/Skeleton.h"
-#include "Containers/Array.h"
-#include "Containers/EnumAsByte.h"
-#include "Containers/Map.h"
-#include "Containers/Set.h"
-#include "Containers/SparseArray.h"
-#include "Containers/UnrealString.h"
-#include "EdGraph/EdGraphPin.h"
 #include "Engine/SkeletalMesh.h"
-#include "Engine/Texture.h"
-#include "Engine/TextureDefines.h"
-#include "Math/NumericLimits.h"
-#include "Math/Transform.h"
-#include "Math/UnrealMathSSE.h"
-#include "Misc/Guid.h"
-#include "MuCO/CustomizableObject.h"
-#include "MuCO/CustomizableObjectClothingTypes.h"
-#include "MuCOE/Nodes/CustomizableObjectNode.h"
 #include "MuCOE/Nodes/CustomizableObjectNodeMaterialBase.h"
 #include "MuCOE/Nodes/CustomizableObjectNodeObject.h"
-#include "MuR/Mesh.h"
-#include "MuR/RefCounted.h"
-#include "MuT/Node.h"
+#include "MuR/MutableMemory.h"
 #include "MuT/NodeComponentNew.h"
+#include "MuR/Ptr.h"
 #include "MuT/NodeImageConstant.h"
 #include "MuT/NodeMeshApplyPose.h"
 #include "MuT/NodeModifierMeshClipWithMesh.h"
@@ -35,13 +18,11 @@
 #include "MuT/NodeScalarParameter.h"
 #include "MuT/NodeSurfaceNew.h"
 #include "MuT/Table.h"
-#include "RHIDefinitions.h"
-#include "ReferenceSkeleton.h"
-#include "Templates/TypeHash.h"
-#include "UObject/NameTypes.h"
 #include "UObject/Package.h"
-#include "UObject/SoftObjectPtr.h"
-#include "UObject/UnrealNames.h"
+
+class UCustomizableObjectNodeMeshClipWithMesh;
+class UCustomizableObjectNodeTable;
+struct FCustomizableObjectClothingAssetData;
 
 class FCustomizableObjectCompiler;
 class UAnimInstance;
@@ -542,6 +523,8 @@ private:
 struct FMutableGraphGenerationContext
 {
 	FMutableGraphGenerationContext(UCustomizableObject* CustomizableObject, class FCustomizableObjectCompiler* InCompiler, const FCompilationOptions& InOptions);
+	~FMutableGraphGenerationContext();
+
 	UCustomizableObject* Object = nullptr;
 
 	// Non-owned reference to the compiler object
diff --git a/Engine/Plugins/Experimental/Mutable/Source/CustomizableObjectEditor/Private/MuCOE/SCustomizableObjectNodeSkeletalMeshRTMorphSelector.cpp b/Engine/Plugins/Experimental/Mutable/Source/CustomizableObjectEditor/Private/MuCOE/SCustomizableObjectNodeSkeletalMeshRTMorphSelector.cpp
--- a/Engine/Plugins/Experimental/Mutable/Source/CustomizableObjectEditor/Private/MuCOE/SCustomizableObjectNodeSkeletalMeshRTMorphSelector.cpp
+++ b/Engine/Plugins/Experimental/Mutable/Source/CustomizableObjectEditor/Private/MuCOE/SCustomizableObjectNodeSkeletalMeshRTMorphSelector.cpp
@@ -8,7 +8,6 @@
 #include "ScopedTransaction.h"
 #include "Widgets/Input/SCheckBox.h"
 #include "Widgets/Input/SSearchBox.h"
-#include "Widgets/Views/SListView.h"
 
 class ITableRow;
 class STableViewBase;
diff --git a/Engine/Plugins/Experimental/PCG/Source/PCG/Private/Elements/PCGTransformPoints.cpp b/Engine/Plugins/Experimental/PCG/Source/PCG/Private/Elements/PCGTransformPoints.cpp
--- a/Engine/Plugins/Experimental/PCG/Source/PCG/Private/Elements/PCGTransformPoints.cpp
+++ b/Engine/Plugins/Experimental/PCG/Source/PCG/Private/Elements/PCGTransformPoints.cpp
@@ -2,13 +2,14 @@
 
 #include "Elements/PCGTransformPoints.h"
 
+#include "Data/PCGSpatialData.h"
 #include "Math/RandomStream.h"
+#include "Helpers/PCGAsync.h"
 #include "PCGHelpers.h"
 #include "Data/PCGPointData.h"
 #include "Helpers/PCGSettingsHelpers.h"
 #include "PCGContext.h"
 #include "PCGPin.h"
-#include "PCGPoint.h"
 
 #include UE_INLINE_GENERATED_CPP_BY_NAME(PCGTransformPoints)
 
diff --git a/Engine/Plugins/Experimental/PCG/Source/PCG/Private/Graph/PCGGraphExecutor.cpp b/Engine/Plugins/Experimental/PCG/Source/PCG/Private/Graph/PCGGraphExecutor.cpp
--- a/Engine/Plugins/Experimental/PCG/Source/PCG/Private/Graph/PCGGraphExecutor.cpp
+++ b/Engine/Plugins/Experimental/PCG/Source/PCG/Private/Graph/PCGGraphExecutor.cpp
@@ -1,7 +1,6 @@
 // Copyright Epic Games, Inc. All Rights Reserved.
 
 #include "PCGGraphExecutor.h"
-#include "GameFramework/Actor.h"
 #include "PCGComponent.h"
 #include "Graph/PCGGraphCache.h"
 #include "PCGGraph.h"
diff --git a/Engine/Plugins/Experimental/PCG/Source/PCG/Private/Tests/PCGTestsCommon.cpp b/Engine/Plugins/Experimental/PCG/Source/PCG/Private/Tests/PCGTestsCommon.cpp
--- a/Engine/Plugins/Experimental/PCG/Source/PCG/Private/Tests/PCGTestsCommon.cpp
+++ b/Engine/Plugins/Experimental/PCG/Source/PCG/Private/Tests/PCGTestsCommon.cpp
@@ -3,18 +3,18 @@
 #include "Tests/PCGTestsCommon.h"
 
 #include "PCGComponent.h"
+#include "PCGContext.h"
 #include "PCGGraph.h"
 #include "PCGParamData.h"
-#include "PCGSettings.h"
 #include "Data/PCGPointData.h"
-#include "Data/PCGPolyLineData.h"
-#include "Data/PCGPrimitiveData.h"
-#include "Data/PCGSurfaceData.h"
+#include "Data/PCGPolyLineData.h" // IWYU pragma: keep
+#include "Data/PCGPrimitiveData.h" // IWYU pragma: keep
+#include "Data/PCGSurfaceData.h" // IWYU pragma: keep
 #include "Data/PCGVolumeData.h"
+#include "PCGPin.h"
 
 #if WITH_EDITOR
 #include "Editor.h"
-#include "Engine/World.h"
 #endif
 
 namespace PCGTestsCommon
diff --git a/Engine/Plugins/Experimental/PCG/Source/PCG/Private/Utils/PCGExtraCapture.cpp b/Engine/Plugins/Experimental/PCG/Source/PCG/Private/Utils/PCGExtraCapture.cpp
--- a/Engine/Plugins/Experimental/PCG/Source/PCG/Private/Utils/PCGExtraCapture.cpp
+++ b/Engine/Plugins/Experimental/PCG/Source/PCG/Private/Utils/PCGExtraCapture.cpp
@@ -3,6 +3,7 @@
 #include "Utils/PCGExtraCapture.h"
 
 #include "PCGComponent.h"
+#include "PCGContext.h"
 
 #if WITH_EDITOR
 
@@ -107,4 +108,4 @@ void PCGUtils::FScopedCall::Serialize(const TCHAR* V, ELogVerbosity::Type Verbos
 	CapturedMessages.Add(PCGUtils::FCapturedMessage{ MessageCounter++, Category, V, Verbosity });
 }
 
-#endif // WITH_EDITOR
\ No newline at end of file
+#endif // WITH_EDITOR
diff --git a/Engine/Plugins/Experimental/PCG/Source/PCGEditor/Private/SPCGEditorGraphLogView.cpp b/Engine/Plugins/Experimental/PCG/Source/PCGEditor/Private/SPCGEditorGraphLogView.cpp
--- a/Engine/Plugins/Experimental/PCG/Source/PCGEditor/Private/SPCGEditorGraphLogView.cpp
+++ b/Engine/Plugins/Experimental/PCG/Source/PCGEditor/Private/SPCGEditorGraphLogView.cpp
@@ -10,7 +10,6 @@
 #include "PCGGraph.h"
 #include "PCGSubgraph.h"
 
-#include "PCGSettings.h"
 #include "Widgets/Layout/SScrollBox.h"
 
 #include "Styling/StyleColors.h"
diff --git a/Engine/Plugins/Experimental/PCG/Source/PCGEditor/Private/SPCGEditorGraphProfilingView.cpp b/Engine/Plugins/Experimental/PCG/Source/PCGEditor/Private/SPCGEditorGraphProfilingView.cpp
--- a/Engine/Plugins/Experimental/PCG/Source/PCGEditor/Private/SPCGEditorGraphProfilingView.cpp
+++ b/Engine/Plugins/Experimental/PCG/Source/PCGEditor/Private/SPCGEditorGraphProfilingView.cpp
@@ -2,6 +2,7 @@
 
 #include "SPCGEditorGraphProfilingView.h"
 
+#include "Framework/Views/TableViewMetadata.h"
 #include "PCGComponent.h"
 #include "PCGEditor.h"
 #include "PCGEditorGraph.h"
@@ -10,13 +11,9 @@
 #include "PCGSubgraph.h"
 #include "PCGSubsystem.h"
 
-#include "Styling/AppStyle.h"
 #include "Widgets/Input/SButton.h"
-#include "Widgets/SBoxPanel.h"
-#include "Widgets/SWidget.h"
 #include "Widgets/Input/SCheckBox.h"
 #include "Widgets/Layout/SScrollBox.h"
-#include "Widgets/Text/STextBlock.h"
 #include "Widgets/Views/SListView.h"
 
 #define LOCTEXT_NAMESPACE "SPCGEditorGraphProfilingView"
diff --git a/Engine/Plugins/Experimental/PlatformCrypto/Source/PlatformCryptoOpenSSL/Private/PlatformCryptoOpenSSLTypes.h b/Engine/Plugins/Experimental/PlatformCrypto/Source/PlatformCryptoOpenSSL/Private/PlatformCryptoOpenSSLTypes.h
--- a/Engine/Plugins/Experimental/PlatformCrypto/Source/PlatformCryptoOpenSSL/Private/PlatformCryptoOpenSSLTypes.h
+++ b/Engine/Plugins/Experimental/PlatformCrypto/Source/PlatformCryptoOpenSSL/Private/PlatformCryptoOpenSSLTypes.h
@@ -2,13 +2,11 @@
 
 #pragma once
 
-#include "EncryptionContextOpenSSL.h"
+#include "EncryptionContextOpenSSL.h" // IWYU pragma: keep
 
 THIRD_PARTY_INCLUDES_START
 #include <openssl/evp.h>
-#include <openssl/rsa.h>
-#include <openssl/rand.h>
-#include <openssl/sha.h>
+#include <openssl/ossl_typ.h>
 THIRD_PARTY_INCLUDES_END
 
 DECLARE_LOG_CATEGORY_EXTERN(LogPlatformCryptoOpenSSL, Warning, All);
diff --git a/Engine/Plugins/Experimental/RemoteSession/Source/RemoteSession/Private/Channels/RemoteSessionLiveLinkChannel.cpp b/Engine/Plugins/Experimental/RemoteSession/Source/RemoteSession/Private/Channels/RemoteSessionLiveLinkChannel.cpp
--- a/Engine/Plugins/Experimental/RemoteSession/Source/RemoteSession/Private/Channels/RemoteSessionLiveLinkChannel.cpp
+++ b/Engine/Plugins/Experimental/RemoteSession/Source/RemoteSession/Private/Channels/RemoteSessionLiveLinkChannel.cpp
@@ -2,6 +2,7 @@
 
 #include "Channels/RemoteSessionLiveLinkChannel.h"
 #include "BackChannel/IBackChannelConnection.h"
+#include "Channels/RemoteSessionChannel.h"
 #include "RemoteSession.h"
 #include "BackChannel/Protocol/OSC/BackChannelOSCMessage.h"
 
diff --git a/Engine/Plugins/Messaging/TcpMessaging/Source/TcpMessaging/Private/Transport/TcpMessageTransport.cpp b/Engine/Plugins/Messaging/TcpMessaging/Source/TcpMessaging/Private/Transport/TcpMessageTransport.cpp
--- a/Engine/Plugins/Messaging/TcpMessaging/Source/TcpMessaging/Private/Transport/TcpMessageTransport.cpp
+++ b/Engine/Plugins/Messaging/TcpMessaging/Source/TcpMessaging/Private/Transport/TcpMessageTransport.cpp
@@ -6,6 +6,7 @@
 #include "IMessageTransportHandler.h"
 #include "TcpMessagingPrivate.h"
 
+#include "TcpSerializedMessage.h"
 #include "Transport/TcpDeserializedMessage.h"
 #include "Transport/TcpMessageTransportConnection.h"
 #include "Transport/TcpSerializeMessageTask.h"
diff --git a/Engine/Plugins/Online/OnlineSubsystemNull/Source/Private/OnlineAchievementsInterfaceNull.cpp b/Engine/Plugins/Online/OnlineSubsystemNull/Source/Private/OnlineAchievementsInterfaceNull.cpp
--- a/Engine/Plugins/Online/OnlineSubsystemNull/Source/Private/OnlineAchievementsInterfaceNull.cpp
+++ b/Engine/Plugins/Online/OnlineSubsystemNull/Source/Private/OnlineAchievementsInterfaceNull.cpp
@@ -1,8 +1,8 @@
 // Copyright Epic Games, Inc. All Rights Reserved.
 
 #include "OnlineAchievementsInterfaceNull.h"
-#include "OnlineSubsystem.h"
-#include "OnlineSubsystemNullTypes.h"
+#include "OnlineSubsystem.h" // IWYU pragma: keep
+#include "OnlineSubsystemNullTypes.h" // IWYU pragma: keep
 
 FOnlineAchievementsNull::FOnlineAchievementsNull(class FOnlineSubsystemNull* InSubsystem)
 	:	NullSubsystem(InSubsystem)
diff --git a/Engine/Plugins/Online/OnlineSubsystemNull/Source/Private/OnlineAchievementsInterfaceNull.h b/Engine/Plugins/Online/OnlineSubsystemNull/Source/Private/OnlineAchievementsInterfaceNull.h
--- a/Engine/Plugins/Online/OnlineSubsystemNull/Source/Private/OnlineAchievementsInterfaceNull.h
+++ b/Engine/Plugins/Online/OnlineSubsystemNull/Source/Private/OnlineAchievementsInterfaceNull.h
@@ -2,12 +2,8 @@
 
 #pragma once
  
-#include "CoreMinimal.h"
-#include "OnlineSubsystemTypes.h"
 #include "Misc/ConfigCacheIni.h"
 #include "Interfaces/OnlineAchievementsInterface.h"
-#include "OnlineSubsystemNullPackage.h"
-#include "OnlineSubsystemNullTypes.h"
 
 /**
  *	IOnlineAchievements - Interface class for acheivements
diff --git a/Engine/Plugins/Online/OnlineSubsystemNull/Source/Private/OnlineLeaderboardInterfaceNull.h b/Engine/Plugins/Online/OnlineSubsystemNull/Source/Private/OnlineLeaderboardInterfaceNull.h
--- a/Engine/Plugins/Online/OnlineSubsystemNull/Source/Private/OnlineLeaderboardInterfaceNull.h
+++ b/Engine/Plugins/Online/OnlineSubsystemNull/Source/Private/OnlineLeaderboardInterfaceNull.h
@@ -3,7 +3,6 @@
 #pragma once
 
 #include "Interfaces/OnlineLeaderboardInterface.h"
-#include "OnlineSubsystemNullPackage.h"
 #include "OnlineSubsystemNullTypes.h"
 
 class FOnlineSubsystemNull;
diff --git a/Engine/Plugins/Runtime/InputDebugging/Source/InputDebugging/Private/InputDeviceDebugTools.h b/Engine/Plugins/Runtime/InputDebugging/Source/InputDebugging/Private/InputDeviceDebugTools.h
--- a/Engine/Plugins/Runtime/InputDebugging/Source/InputDebugging/Private/InputDeviceDebugTools.h
+++ b/Engine/Plugins/Runtime/InputDebugging/Source/InputDebugging/Private/InputDeviceDebugTools.h
@@ -2,7 +2,7 @@
 
 #pragma once
 
-#include "Misc/Build.h" // IWYU: pragma: keep
+#include "Misc/Build.h" // IWYU pragma: keep
 
 // Allow you to enable this debug drawing in a build.cs if you want, but default to when we are not in shipping.
 #if !defined(SUPPORT_INPUT_DEVICE_DEBUGGING)
diff --git a/Engine/Plugins/Runtime/ModelViewViewModel/Source/ModelViewViewModel/Private/ViewModel/MVVMViewModelBlueprintGeneratedClass.cpp b/Engine/Plugins/Runtime/ModelViewViewModel/Source/ModelViewViewModel/Private/ViewModel/MVVMViewModelBlueprintGeneratedClass.cpp
--- a/Engine/Plugins/Runtime/ModelViewViewModel/Source/ModelViewViewModel/Private/ViewModel/MVVMViewModelBlueprintGeneratedClass.cpp
+++ b/Engine/Plugins/Runtime/ModelViewViewModel/Source/ModelViewViewModel/Private/ViewModel/MVVMViewModelBlueprintGeneratedClass.cpp
@@ -2,7 +2,7 @@
 
 #include "ViewModel/MVVMViewModelBlueprintGeneratedClass.h"
 #include "MVVMViewModelBase.h"
-#include "UObject/Package.h"
+#include "UObject/Package.h" // IWYU pragma: keep
 
 #include UE_INLINE_GENERATED_CPP_BY_NAME(MVVMViewModelBlueprintGeneratedClass)
 
diff --git a/Engine/Plugins/Runtime/ModelViewViewModel/Source/ModelViewViewModelEditor/Private/Styling/MVVMEditorStyle.cpp b/Engine/Plugins/Runtime/ModelViewViewModel/Source/ModelViewViewModelEditor/Private/Styling/MVVMEditorStyle.cpp
--- a/Engine/Plugins/Runtime/ModelViewViewModel/Source/ModelViewViewModelEditor/Private/Styling/MVVMEditorStyle.cpp
+++ b/Engine/Plugins/Runtime/ModelViewViewModel/Source/ModelViewViewModelEditor/Private/Styling/MVVMEditorStyle.cpp
@@ -1,6 +1,7 @@
 // Copyright Epic Games, Inc. All Rights Reserved.
 
 #include "Styling/MVVMEditorStyle.h"
+#include "Brushes/SlateColorBrush.h"
 #include "Brushes/SlateImageBrush.h"
 #include "Styling/SlateStyleMacros.h"
 #include "Brushes/SlateRoundedBoxBrush.h"
diff --git a/Engine/Plugins/Runtime/ModelViewViewModel/Source/ModelViewViewModelEditor/Private/Tabs/MVVMViewModelSummoner.cpp b/Engine/Plugins/Runtime/ModelViewViewModel/Source/ModelViewViewModelEditor/Private/Tabs/MVVMViewModelSummoner.cpp
--- a/Engine/Plugins/Runtime/ModelViewViewModel/Source/ModelViewViewModelEditor/Private/Tabs/MVVMViewModelSummoner.cpp
+++ b/Engine/Plugins/Runtime/ModelViewViewModel/Source/ModelViewViewModelEditor/Private/Tabs/MVVMViewModelSummoner.cpp
@@ -13,7 +13,7 @@ namespace UE::MVVM
 
 const FName FViewModelSummoner::TabID(TEXT("ViewModelPanel"));
 
-FViewModelSummoner::FViewModelSummoner(TSharedPtr<class FWidgetBlueprintEditor> InBlueprintEditor)
+FViewModelSummoner::FViewModelSummoner(TSharedPtr<FWidgetBlueprintEditor> InBlueprintEditor)
 		: FWorkflowTabFactory(TabID, InBlueprintEditor)
 		, BlueprintEditor(InBlueprintEditor)
 {
diff --git a/Engine/Plugins/Runtime/ModelViewViewModel/Source/ModelViewViewModelEditor/Private/Tabs/MVVMViewModelSummoner.h b/Engine/Plugins/Runtime/ModelViewViewModel/Source/ModelViewViewModelEditor/Private/Tabs/MVVMViewModelSummoner.h
--- a/Engine/Plugins/Runtime/ModelViewViewModel/Source/ModelViewViewModelEditor/Private/Tabs/MVVMViewModelSummoner.h
+++ b/Engine/Plugins/Runtime/ModelViewViewModel/Source/ModelViewViewModelEditor/Private/Tabs/MVVMViewModelSummoner.h
@@ -2,10 +2,9 @@
 
 #pragma once
 
-#include "CoreMinimal.h"
-#include "Widgets/SWidget.h"
 #include "WorkflowOrientedApp/WorkflowTabFactory.h"
-#include "WidgetBlueprintEditor.h"
+
+class FWidgetBlueprintEditor;
 
 namespace UE::MVVM
 {
@@ -16,12 +15,12 @@ public:
 	static const FName TabID;
 	
 public:
-	FViewModelSummoner(TSharedPtr<class FWidgetBlueprintEditor> InBlueprintEditor);
+	FViewModelSummoner(TSharedPtr<FWidgetBlueprintEditor> InBlueprintEditor);
 	
 	virtual TSharedRef<SWidget> CreateTabBody(const FWorkflowTabSpawnInfo& Info) const override;
 	
 protected:
-	TWeakPtr<class FWidgetBlueprintEditor> BlueprintEditor;
+	TWeakPtr<FWidgetBlueprintEditor> BlueprintEditor;
 };
 
 } // namespace
diff --git a/Engine/Plugins/Runtime/ModelViewViewModel/Source/ModelViewViewModelEditor/Private/Widgets/SMVVMFieldSelector.h b/Engine/Plugins/Runtime/ModelViewViewModel/Source/ModelViewViewModelEditor/Private/Widgets/SMVVMFieldSelector.h
--- a/Engine/Plugins/Runtime/ModelViewViewModel/Source/ModelViewViewModelEditor/Private/Widgets/SMVVMFieldSelector.h
+++ b/Engine/Plugins/Runtime/ModelViewViewModel/Source/ModelViewViewModelEditor/Private/Widgets/SMVVMFieldSelector.h
@@ -2,21 +2,19 @@
 
 #pragma once
 
-#include "Misc/TVariant.h"
 #include "MVVMPropertyPath.h"
 #include "Styling/CoreStyle.h"
-#include "Templates/ValueOrError.h"
-#include "Types/MVVMBindingMode.h"
+#include "Styling/SlateTypes.h"
 #include "Types/MVVMBindingSource.h"
-#include "Types/MVVMFieldVariant.h"
-#include "UObject/UnrealType.h"
-#include "Widgets/DeclarativeSyntaxSupport.h"
-#include "Widgets/Input/SComboBox.h"
+#include "Styling/SlateWidgetStyleAsset.h"
 #include "Widgets/SCompoundWidget.h"
-#include "Widgets/SMVVMFieldEntry.h"
-#include "Widgets/SMVVMViewModelBindingListWidget.h"
-#include "Widgets/Text/STextBlock.h"
-#include "Widgets/Views/STreeView.h"
+
+enum class EMVVMBindingMode : uint8;
+namespace ESelectInfo { enum Type : int; }
+namespace UE::MVVM { class SFieldEntry; }
+namespace UE::MVVM { enum class EFieldVisibility : uint8; }
+template <typename ItemType> class SListView;
+template <typename ItemType> class STreeView;
 
 class SComboButton;
 class SSearchBox;
@@ -183,4 +181,4 @@ private:
 	bool bIsMenuInitialized = false;
 }; 
 
-} // namespace UE::MVVM
\ No newline at end of file
+} // namespace UE::MVVM
diff --git a/Engine/Plugins/Runtime/ModelViewViewModel/Source/ModelViewViewModelEditor/Private/Widgets/SMVVMViewBindingListView.cpp b/Engine/Plugins/Runtime/ModelViewViewModel/Source/ModelViewViewModelEditor/Private/Widgets/SMVVMViewBindingListView.cpp
--- a/Engine/Plugins/Runtime/ModelViewViewModel/Source/ModelViewViewModelEditor/Private/Widgets/SMVVMViewBindingListView.cpp
+++ b/Engine/Plugins/Runtime/ModelViewViewModel/Source/ModelViewViewModelEditor/Private/Widgets/SMVVMViewBindingListView.cpp
@@ -2,45 +2,31 @@
 
 #include "Widgets/SMVVMViewBindingListView.h"
 
-#include "Algo/Transform.h"
 #include "Bindings/MVVMBindingHelper.h"
 #include "BlueprintEditor.h"
-#include "Blueprint/WidgetTree.h"
 #include "Dialog/SCustomDialog.h"
-#include "Editor.h"
-#include "Engine/Engine.h"
-#include "Features/IModularFeatures.h"
 #include "Framework/MultiBox/MultiBoxBuilder.h"
-#include "IPropertyAccessEditor.h"
-#include "K2Node_CallFunction.h"
+#include "Framework/Views/TableViewMetadata.h"
 #include "Misc/MessageDialog.h"
 #include "MVVMBlueprintView.h"
-#include "MVVMBlueprintViewBinding.h"
 #include "MVVMEditorSubsystem.h"
-#include "MVVMSubsystem.h"
 #include "MVVMWidgetBlueprintExtension_View.h"
 #include "ScopedTransaction.h"
 #include "SEnumCombo.h"
-#include "SPrimaryButton.h"
 #include "SSimpleButton.h" 
-#include "Styling/AppStyle.h"
 #include "Styling/MVVMEditorStyle.h"
 #include "Styling/StyleColors.h"
-#include "Widgets/SMVVMConversionPath.h"
+#include "Widgets/SMVVMFieldEntry.h"
 #include "Widgets/SMVVMFieldSelector.h"
 #include "Widgets/SMVVMFunctionParameter.h"
-#include "Widgets/SMVVMPropertyPath.h"
 #include "Widgets/SMVVMSourceSelector.h"
 #include "Widgets/SMVVMViewBindingPanel.h"
-#include "Widgets/SNullWidget.h"
 #include "Widgets/Layout/SSpacer.h"
-#include "Widgets/Images/SImage.h"
 #include "Widgets/Input/SCheckBox.h"
 #include "Widgets/Input/SComboBox.h"
 #include "Widgets/Input/SEditableTextBox.h"
-#include "Widgets/Text/STextBlock.h"
+#include "Widgets/SMVVMViewModelPanel.h" // IWYU pragma: keep
 #include "Widgets/Views/STreeView.h"
-#include "Widgets/Views/STableViewBase.h"
 
 #define LOCTEXT_NAMESPACE "BindingListView"
 
diff --git a/Engine/Plugins/Runtime/ModelViewViewModel/Source/ModelViewViewModelEditor/Private/Widgets/SMVVMViewBindingListView.h b/Engine/Plugins/Runtime/ModelViewViewModel/Source/ModelViewViewModelEditor/Private/Widgets/SMVVMViewBindingListView.h
--- a/Engine/Plugins/Runtime/ModelViewViewModel/Source/ModelViewViewModelEditor/Private/Widgets/SMVVMViewBindingListView.h
+++ b/Engine/Plugins/Runtime/ModelViewViewModel/Source/ModelViewViewModelEditor/Private/Widgets/SMVVMViewBindingListView.h
@@ -2,11 +2,11 @@
 
 #pragma once
 
-#include "CoreMinimal.h"
 
-#include "Widgets/DeclarativeSyntaxSupport.h"
 #include "Widgets/SCompoundWidget.h"
 
+namespace ESelectInfo { enum Type : int; }
+
 class ITableRow;
 template <typename ItemType> class STreeView;
 class STableViewBase;
diff --git a/Engine/Plugins/Runtime/ModelViewViewModel/Source/ModelViewViewModelEditor/Private/Widgets/SMVVMViewModelPanel.cpp b/Engine/Plugins/Runtime/ModelViewViewModel/Source/ModelViewViewModelEditor/Private/Widgets/SMVVMViewModelPanel.cpp
--- a/Engine/Plugins/Runtime/ModelViewViewModel/Source/ModelViewViewModelEditor/Private/Widgets/SMVVMViewModelPanel.cpp
+++ b/Engine/Plugins/Runtime/ModelViewViewModel/Source/ModelViewViewModelEditor/Private/Widgets/SMVVMViewModelPanel.cpp
@@ -2,39 +2,25 @@
 
 #include "Widgets/SMVVMViewModelPanel.h"
 
-#include "Bindings/MVVMBindingHelper.h"
-#include "Editor.h"
-#include "Editor/EditorEngine.h"
+#include "DetailsViewArgs.h"
 #include "IStructureDetailsView.h"
-#include "Kismet2/Kismet2NameValidators.h"
 #include "MVVMBlueprintView.h"
-#include "MVVMBlueprintViewModelContext.h"
 #include "MVVMEditorSubsystem.h"
-#include "MVVMSubsystem.h"
-#include "MVVMViewModelBase.h"
 #include "PropertyEditorModule.h"
 #include "WidgetBlueprint.h"
 #include "WidgetBlueprintEditor.h"
 
 #include "Framework/Commands/GenericCommands.h"
-#include "Framework/Commands/UICommandList.h"
 #include "Framework/MultiBox/MultiBoxBuilder.h"
 
-#include "Styling/AppStyle.h"
-#include "Styling/SlateIconFinder.h"
-#include "Styling/MVVMEditorStyle.h"
 
-#include "Widgets/SBoxPanel.h"
 #include "Widgets/Images/SImage.h"
-#include "Widgets/Input/SSearchBox.h"
 #include "Widgets/Layout/SBorder.h"
-#include "Widgets/Layout/SBox.h"
+#include "Widgets/Layout/SSplitter.h"
 #include "Widgets/PropertyViewer/SFieldIcon.h"
-#include "Widgets/PropertyViewer/SFieldName.h"
-#include "Widgets/PropertyViewer/SPropertyViewer.h"
 #include "Widgets/SMVVMSelectViewModel.h"
+#include "Widgets/SMVVMViewModelBindingListWidget.h"
 #include "Widgets/Text/SInlineEditableTextBlock.h"
-#include "Widgets/Text/STextBlock.h"
 #include "SPositiveActionButton.h"
 
 #define LOCTEXT_NAMESPACE "ViewModelPanel"
@@ -117,6 +103,7 @@ void SMVVMViewModelPanel::Construct(const FArguments& InArgs, TSharedPtr<FWidget
 	];
 }
 
+SMVVMViewModelPanel::SMVVMViewModelPanel() = default;
 
 SMVVMViewModelPanel::~SMVVMViewModelPanel()
 {
@@ -599,4 +586,4 @@ void SMVVMViewModelPanel::NotifyPostChange(const FPropertyChangedEvent& Property
 
 } // namespace
 
-#undef LOCTEXT_NAMESPACE
\ No newline at end of file
+#undef LOCTEXT_NAMESPACE
diff --git a/Engine/Plugins/Runtime/ModelViewViewModel/Source/ModelViewViewModelEditor/Private/Widgets/SMVVMViewModelPanel.h b/Engine/Plugins/Runtime/ModelViewViewModel/Source/ModelViewViewModelEditor/Private/Widgets/SMVVMViewModelPanel.h
--- a/Engine/Plugins/Runtime/ModelViewViewModel/Source/ModelViewViewModelEditor/Private/Widgets/SMVVMViewModelPanel.h
+++ b/Engine/Plugins/Runtime/ModelViewViewModel/Source/ModelViewViewModelEditor/Private/Widgets/SMVVMViewModelPanel.h
@@ -2,12 +2,13 @@
 
 #pragma once
 
-#include "CoreMinimal.h"
 
 #include "Misc/NotifyHook.h"
 #include "Widgets/PropertyViewer/SPropertyViewer.h"
-#include "Widgets/SMVVMViewModelBindingListWidget.h"
-#include "Widgets/SCompoundWidget.h"
+
+namespace ETextCommit { enum Type : int; }
+namespace UE::MVVM { class FFieldIterator_Bindable; }
+namespace UE::PropertyViewer { class FFieldExpander_Default; }
 
 class FWidgetBlueprintEditor;
 class SInlineEditableTextBlock;
@@ -30,6 +31,7 @@ public:
 	SLATE_END_ARGS()
 
 	void Construct(const FArguments& InArgs, TSharedPtr<FWidgetBlueprintEditor> Editor);
+	SMVVMViewModelPanel();
 	virtual ~SMVVMViewModelPanel();
 
 	void OpenAddViewModelMenu();
diff --git a/Engine/Plugins/Runtime/NetworkPredictionInsights/Source/NetworkPredictionInsights/Private/NetworkPredictionAnalyzer.h b/Engine/Plugins/Runtime/NetworkPredictionInsights/Source/NetworkPredictionInsights/Private/NetworkPredictionAnalyzer.h
--- a/Engine/Plugins/Runtime/NetworkPredictionInsights/Source/NetworkPredictionInsights/Private/NetworkPredictionAnalyzer.h
+++ b/Engine/Plugins/Runtime/NetworkPredictionInsights/Source/NetworkPredictionInsights/Private/NetworkPredictionAnalyzer.h
@@ -3,7 +3,6 @@
 #pragma once
 
 #include "Trace/Analyzer.h"
-#include "TraceServices/Model/AnalysisSession.h"
 #include "Containers/Map.h"
 
 class FNetworkPredictionProvider;
@@ -70,4 +69,4 @@ private:
 
 	TMap<uint32, FThreadState*> ThreadStatesMap;
 	FThreadState& GetThreadState(uint32 ThreadId);
-};
\ No newline at end of file
+};
diff --git a/Engine/Plugins/Runtime/NetworkPredictionInsights/Source/NetworkPredictionInsights/Private/NetworkPredictionInsightsModule.cpp b/Engine/Plugins/Runtime/NetworkPredictionInsights/Source/NetworkPredictionInsights/Private/NetworkPredictionInsightsModule.cpp
--- a/Engine/Plugins/Runtime/NetworkPredictionInsights/Source/NetworkPredictionInsights/Private/NetworkPredictionInsightsModule.cpp
+++ b/Engine/Plugins/Runtime/NetworkPredictionInsights/Source/NetworkPredictionInsights/Private/NetworkPredictionInsightsModule.cpp
@@ -2,21 +2,14 @@
 
 #include "NetworkPredictionInsightsModule.h"
 
-#include "Containers/Ticker.h"
 #include "Features/IModularFeatures.h"
-#include "Framework/Docking/LayoutExtender.h"
-#include "HAL/LowLevelMemTracker.h"
-#include "HAL/PlatformApplicationMisc.h"
-#include "Insights/ITimingViewExtender.h"
+#include "INetworkPredictionInsightsModule.h" // IWYU pragma: keep
 #include "Insights/IUnrealInsightsModule.h"
 #include "Misc/CommandLine.h"
-#include "Modules/ModuleManager.h"
+#include "Misc/CoreDelegates.h"
 #include "ProfilingDebugging/TraceAuxiliary.h"
-#include "Stats/Stats.h"
 #include "String/ParseTokens.h"
-#include "Styling/AppStyle.h"
-#include "Trace/StoreClient.h"
-#include "TraceServices/ITraceServicesModule.h"
+#include "Misc/Paths.h"
 #include "Widgets/Docking/SDockTab.h"
 #include "WorkspaceMenuStructure.h"
 #include "WorkspaceMenuStructureModule.h"
diff --git a/Engine/Plugins/Runtime/NetworkPredictionInsights/Source/NetworkPredictionInsights/Private/NetworkPredictionInsightsModule.h b/Engine/Plugins/Runtime/NetworkPredictionInsights/Source/NetworkPredictionInsights/Private/NetworkPredictionInsightsModule.h
--- a/Engine/Plugins/Runtime/NetworkPredictionInsights/Source/NetworkPredictionInsights/Private/NetworkPredictionInsightsModule.h
+++ b/Engine/Plugins/Runtime/NetworkPredictionInsights/Source/NetworkPredictionInsights/Private/NetworkPredictionInsightsModule.h
@@ -2,11 +2,9 @@
 
 #pragma once
 
-#include "CoreMinimal.h"
+#include "Containers/Ticker.h"
 #include "Modules/ModuleInterface.h"
 #include "NetworkPredictionTraceModule.h"
-#include "Framework/Docking/TabManager.h"
-#include "INetworkPredictionInsightsModule.h"
 #include "Trace/StoreService.h"
 
 struct FInsightsMajorTabExtender;
diff --git a/Engine/Plugins/Runtime/ProceduralMeshComponent/Source/ProceduralMeshComponentEditor/Private/ProceduralMeshComponentDetails.h b/Engine/Plugins/Runtime/ProceduralMeshComponent/Source/ProceduralMeshComponentEditor/Private/ProceduralMeshComponentDetails.h
--- a/Engine/Plugins/Runtime/ProceduralMeshComponent/Source/ProceduralMeshComponentEditor/Private/ProceduralMeshComponentDetails.h
+++ b/Engine/Plugins/Runtime/ProceduralMeshComponent/Source/ProceduralMeshComponentEditor/Private/ProceduralMeshComponentDetails.h
@@ -3,7 +3,7 @@
 #pragma once
 
 #include "IDetailCustomization.h"
-#include "UObject/WeakObjectPtr.h"
+#include "UObject/WeakObjectPtrTemplates.h"
 
 class FReply;
 
diff --git a/Engine/Plugins/Runtime/SoundCueTemplates/Source/SoundCueTemplatesEditor/Private/AssetDefinition_SoundCueTemplate.cpp b/Engine/Plugins/Runtime/SoundCueTemplates/Source/SoundCueTemplatesEditor/Private/AssetDefinition_SoundCueTemplate.cpp
--- a/Engine/Plugins/Runtime/SoundCueTemplates/Source/SoundCueTemplatesEditor/Private/AssetDefinition_SoundCueTemplate.cpp
+++ b/Engine/Plugins/Runtime/SoundCueTemplates/Source/SoundCueTemplatesEditor/Private/AssetDefinition_SoundCueTemplate.cpp
@@ -3,17 +3,17 @@
 #include "AssetDefinition_SoundCueTemplate.h"
 
 #include "AssetToolsModule.h"
-#include "SoundCueTemplate.h"
+#include "IAssetTools.h"
 #include "SoundCueTemplateFactory.h"
 
 #include "AudioEditorSettings.h"
-#include "Components/AudioComponent.h"
+#include "Misc/PackageName.h"
 #include "ToolMenus.h"
+#include "ToolMenu.h"
 #include "ToolMenuSection.h"
 #include "ContentBrowserModule.h"
 #include "ContentBrowserMenuContexts.h"
 #include "IContentBrowserSingleton.h"
-#include "ObjectEditorUtils.h"
 #include "Sound/SoundWaveProcedural.h"
 
 #define LOCTEXT_NAMESPACE "AssetTypeActions"
diff --git a/Engine/Plugins/Runtime/SoundCueTemplates/Source/SoundCueTemplatesEditor/Private/AssetDefinition_SoundCueTemplate.h b/Engine/Plugins/Runtime/SoundCueTemplates/Source/SoundCueTemplatesEditor/Private/AssetDefinition_SoundCueTemplate.h
--- a/Engine/Plugins/Runtime/SoundCueTemplates/Source/SoundCueTemplatesEditor/Private/AssetDefinition_SoundCueTemplate.h
+++ b/Engine/Plugins/Runtime/SoundCueTemplates/Source/SoundCueTemplatesEditor/Private/AssetDefinition_SoundCueTemplate.h
@@ -2,12 +2,13 @@
 
 #pragma once
 
-#include "CoreMinimal.h"
 #include "SoundCueTemplate.h"
 #include "AssetTypeActions/AssetDefinition_SoundBase.h"
 
 #include "AssetDefinition_SoundCueTemplate.generated.h"
 
+struct FAssetCategoryPath;
+
 class USoundCueTemplate;
 struct FToolMenuContext;
 
diff --git a/Engine/Plugins/Runtime/SoundCueTemplates/Source/SoundCueTemplatesEditor/Private/SoundCueTemplatesEditorModule.cpp b/Engine/Plugins/Runtime/SoundCueTemplates/Source/SoundCueTemplatesEditor/Private/SoundCueTemplatesEditorModule.cpp
--- a/Engine/Plugins/Runtime/SoundCueTemplates/Source/SoundCueTemplatesEditor/Private/SoundCueTemplatesEditorModule.cpp
+++ b/Engine/Plugins/Runtime/SoundCueTemplates/Source/SoundCueTemplatesEditor/Private/SoundCueTemplatesEditorModule.cpp
@@ -2,9 +2,7 @@
 
 #include "SoundCueTemplatesEditorModule.h"
 
-#include "Modules/ModuleManager.h"
 #include "ToolMenus.h"
-#include "HAL/LowLevelMemTracker.h"
 
 void FSoundCueTemplatesEditorModule::StartupModule()
 {
diff --git a/Engine/Plugins/Runtime/WorldConditions/Source/WorldConditionsEditor/Private/WorldConditionsEditorModule.cpp b/Engine/Plugins/Runtime/WorldConditions/Source/WorldConditionsEditor/Private/WorldConditionsEditorModule.cpp
--- a/Engine/Plugins/Runtime/WorldConditions/Source/WorldConditionsEditor/Private/WorldConditionsEditorModule.cpp
+++ b/Engine/Plugins/Runtime/WorldConditions/Source/WorldConditionsEditor/Private/WorldConditionsEditorModule.cpp
@@ -2,9 +2,8 @@
 
 #include "WorldConditionsEditorModule.h"
 
+#include "ComponentVisualizer.h" // IWYU pragma: keep
 #include "PropertyEditorModule.h"
-#include "UnrealEdGlobals.h"
-#include "Editor/UnrealEdEngine.h"
 #include "WorldConditionEditorStyle.h"
 #include "Customizations/WorldConditionEditableDetails.h"
 #include "Customizations/WorldConditionQueryDefinitionDetails.h"
