[Mutable] Fix TaskGraph oversubscription hangs by doing async external image load.

diff --git a/Engine/Plugins/Experimental/Mutable/Source/CustomizableObject/Private/MuCO/UnrealMutableImageProvider.cpp b/Engine/Plugins/Experimental/Mutable/Source/CustomizableObject/Private/MuCO/UnrealMutableImageProvider.cpp
--- a/Engine/Plugins/Experimental/Mutable/Source/CustomizableObject/Private/MuCO/UnrealMutableImageProvider.cpp
+++ b/Engine/Plugins/Experimental/Mutable/Source/CustomizableObject/Private/MuCO/UnrealMutableImageProvider.cpp
@@ -115,8 +115,7 @@ mu::EImageFormat GetMutablePixelFormat(EPixelFormat InTextureFormat)
 
 
 //-------------------------------------------------------------------------------------------------
-
-mu::ImagePtr FUnrealMutableImageProvider::GetImage(mu::EXTERNAL_IMAGE_ID id, uint8 MipmapsToSkip )
+TTuple<FGraphEventRef, TFunction<void()>> FUnrealMutableImageProvider::GetImageAsync(mu::EXTERNAL_IMAGE_ID id, uint8 MipmapsToSkip, TFunction<void (mu::Ptr<mu::Image>)>& ResultCallback)
 {
 	// Thread: worker
 	MUTABLE_CPUPROFILER_SCOPE(FUnrealMutableImageProvider::GetImage);
@@ -125,7 +124,7 @@ mu::ImagePtr FUnrealMutableImageProvider::GetImage(mu::EXTERNAL_IMAGE_ID id, uin
 	mu::ImagePtr Image;
 
 	// Some data that may have to be copied from the GlobalExternalImages while it's locked
-	TUniquePtr<IBulkDataIORequest> IORequest;
+	IBulkDataIORequest* IORequest = nullptr;
 	const int32 LODs = 1;
 
 	EPixelFormat Format = EPixelFormat::PF_Unknown;
@@ -133,15 +132,25 @@ mu::ImagePtr FUnrealMutableImageProvider::GetImage(mu::EXTERNAL_IMAGE_ID id, uin
 
 	mu::EImageFormat MutImageFormat = mu::EImageFormat::IF_NONE;
 	int32 MutImageDataSize = 0;
-	
+
+	auto TrivialReturn = []() -> TTuple<FGraphEventRef, TFunction<void()>>
+	{
+		FGraphEventRef CompletionEvent = FGraphEvent::CreateGraphEvent();
+		CompletionEvent->DispatchSubsequents();
+
+		return MakeTuple(CompletionEvent, []() -> void {});
+	};
+
 	{
 		FScopeLock Lock(&ExternalImagesLock);
 		// Inside this scope it's safe to access GlobalExternalImages
 
+
 		if (!GlobalExternalImages.Contains(id))
 		{
 			// Null case, no image was provided
-			return CreateDummy();
+			ResultCallback(CreateDummy());
+			return Invoke(TrivialReturn);
 		}
 
 		const FUnrealMutableImageInfo& ImageInfo = GlobalExternalImages[id];
@@ -149,7 +158,8 @@ mu::ImagePtr FUnrealMutableImageProvider::GetImage(mu::EXTERNAL_IMAGE_ID id, uin
 		if (ImageInfo.Image)
 		{
 			// Easy case where the image was directly provided
-			return ImageInfo.Image;
+			ResultCallback(ImageInfo.Image);
+			return Invoke(TrivialReturn);
 		}
 		else if (UTexture2D* TextureToLoad = ImageInfo.TextureToLoad)
 		{
@@ -170,7 +180,8 @@ mu::ImagePtr FUnrealMutableImageProvider::GetImage(mu::EXTERNAL_IMAGE_ID id, uin
 
 #if WITH_EDITOR
 			// In the editor the src data can be directly accessed
-			return ConvertTextureUnrealToMutable(TextureToLoad, MipIndex);
+			ResultCallback(ConvertTextureUnrealToMutable(TextureToLoad, MipIndex));
+			return Invoke(TrivialReturn);
 #else
 			// Texture format and the equivalent mutable format
 			Format = TextureToLoad->GetPlatformData()->PixelFormat;
@@ -180,7 +191,8 @@ mu::ImagePtr FUnrealMutableImageProvider::GetImage(mu::EXTERNAL_IMAGE_ID id, uin
 			if (MutImageFormat == mu::EImageFormat::IF_NONE)
 			{
 				UE_LOG(LogMutable, Warning, TEXT("Failed to get external image. Unexpected image format. EImageFormat [%s]."), GetPixelFormatString(Format));
-				return CreateDummy();
+				ResultCallback(CreateDummy());
+				return Invoke(TrivialReturn);
 			}
 
 			int SizeX = TextureToLoad->GetSizeX() >> MipIndex;
@@ -201,13 +213,106 @@ mu::ImagePtr FUnrealMutableImageProvider::GetImage(mu::EXTERNAL_IMAGE_ID id, uin
 				UE_LOG(LogMutable, Warning, TEXT("Failed to get external image. Bulk data size is different than the expected size. BulkData size [%d]. Mutable image data size [%d]."),
 					BulkDataSize, MutImageDataSize);
 
-				return CreateDummy();
+				ResultCallback(CreateDummy());
+				return Invoke(TrivialReturn);
 			}
 
 			// Create a streaming request if the data is not loaded or copy the mip data
 			if (!BulkData.IsBulkDataLoaded())
 			{
-				IORequest.Reset(BulkData.CreateStreamingRequest(EAsyncIOPriorityAndFlags::AIOP_High, nullptr, Image->GetData()));
+				FGraphEventRef IORequestCompletionEvent = FGraphEvent::CreateGraphEvent();
+
+				TFunction<void(bool, IBulkDataIORequest*)> IOCallback =
+					[
+						MutImageDataSize,
+						MutImageFormat,
+						Format,
+						Image,
+						BulkDataSize,	
+						ResultCallback, // Notice ResultCallback is captured by copy
+						IORequestCompletionEvent
+					](bool bWasCancelled, IBulkDataIORequest* IORequest)
+				{
+					ON_SCOPE_EXIT
+					{
+						if (IORequestCompletionEvent.IsValid())
+						{
+							IORequestCompletionEvent->DispatchSubsequents();
+						}
+					};
+					
+					// Should we do someting different than returning a dummy image if cancelled?
+					if (bWasCancelled)
+					{
+						ResultCallback(CreateDummy());
+						return;
+					}
+
+					uint8* Results = IORequest->GetReadResults(); // required?
+
+					if (Results && Image->GetDataSize() == (int32)IORequest->GetSize())
+					{
+						check(BulkDataSize == (int32)IORequest->GetSize());
+						check(Results == Image->GetData());
+
+						ResultCallback(Image);
+						return;
+					}
+
+					if (!Results)
+					{
+						UE_LOG(LogMutable, Warning, TEXT("Failed to get external image. IO Request failed. Request results [%hhd]. Format: [%s]. MutableFormat: [%d]."),
+							(Results != nullptr),
+							GetPixelFormatString(Format),
+							(int32)MutImageFormat);
+					}
+					else if (MutImageDataSize != (int32)IORequest->GetSize())
+					{
+						UE_LOG(LogMutable, Warning, TEXT("Failed to get external image. Requested size is different than the expected size. RequestSize: [%lld]. ExpectedSize: [%d]. Format: [%s]. MutableFormat: [%d]."),
+							IORequest->GetSize(),
+							Image->GetDataSize(),
+							GetPixelFormatString(Format),
+							(int32)MutImageFormat);
+					}
+					else
+					{
+						UE_LOG(LogMutable, Warning, TEXT("Failed to get external image."));
+					}
+
+					// Something failed when loading the bulk data, just return a dummy
+					ResultCallback(CreateDummy());
+				};
+			
+				// Is the resposability of the CreateStreamingRequest caller to delete the IORequest. 
+				// This can *not* be done in the IOCallback because it would cause a deadlock so it is deferred to the returned
+				// cleanup function. Another solution could be to spwan a new task that depends on the 
+				// IORequestComplitionEvent which deletes it.
+				IORequest = BulkData.CreateStreamingRequest(EAsyncIOPriorityAndFlags::AIOP_High, &IOCallback, Image->GetData());
+
+				if (IORequest)
+				{
+					// Make the lambda mutable and set the IORequest pointer to null when deleted so it is safer 
+					// agains multiple calls.
+					const auto DeleteIORequest = [IORequest]() mutable -> void
+					{
+						if (IORequest)
+						{
+							delete IORequest;
+						}
+						
+						IORequest = nullptr;
+					};
+
+					return MakeTuple(IORequestCompletionEvent, DeleteIORequest);
+
+				}
+				else
+				{
+					UE_LOG(LogMutable, Warning, TEXT("Failed to create an IORequest for a UTexture2D BulkData for an application-specific image parameter."));
+
+					ResultCallback(CreateDummy());
+					return Invoke(TrivialReturn);
+				}
 			}
 			else
 			{
@@ -219,12 +324,14 @@ mu::ImagePtr FUnrealMutableImageProvider::GetImage(mu::EXTERNAL_IMAGE_ID id, uin
 					FMemory::Memcpy(Image->GetData(), Data, BulkDataSize);
 
 					BulkData.Unlock();
-					return Image;
+					ResultCallback(Image);
+					return Invoke(TrivialReturn);
 				}
 				else
 				{
 					UE_LOG(LogMutable, Warning, TEXT("Failed to get external image. Bulk data already locked or null."));
-					return CreateDummy();
+					ResultCallback(CreateDummy());
+					return Invoke(TrivialReturn);
 				}
 			}
 #endif
@@ -233,52 +340,83 @@ mu::ImagePtr FUnrealMutableImageProvider::GetImage(mu::EXTERNAL_IMAGE_ID id, uin
 		{
 			// No UTexture2D was provided, cannot do anything, just provide a dummy texture
 			UE_LOG(LogMutable, Warning, TEXT("No UTexture2D was provided for an application-specific image parameter."));
-			return CreateDummy();
+			ResultCallback(CreateDummy());
+			return Invoke(TrivialReturn);
 		}
 	}
 
+	// Make sure the returned event is dispatched at some point for all code paths, 
+	// in this case returning Invoke(TrivialReturn) or through the IORequest callback.
+}
 
-	if (IORequest)
+
+// This should mantain parity with the descriptor of the images generated by GetImageAsync 
+mu::FImageDesc FUnrealMutableImageProvider::GetImageDesc(mu::EXTERNAL_IMAGE_ID id, uint8 MipmapsToSkip) 
+{
+	MUTABLE_CPUPROFILER_SCOPE(FUnrealMutableImageProvider::GetImageDesc);
+
+	mu::FImageDesc Result;	
 	{
-		// It's OK to wait since this must run on the Mutable thread
-		bool bComplete = IORequest->WaitCompletion();
-		uint8* Results = IORequest->GetReadResults(); // required?
+		FScopeLock Lock(&ExternalImagesLock);
+		// Inside this scope it's safe to access GlobalExternalImages
 
-		if (bComplete && Results && Image->GetDataSize() == (int32)IORequest->GetSize())
+		if (!GlobalExternalImages.Contains(id))
 		{
-			check(BulkDataSize == (int32)IORequest->GetSize());
-			check(Results == Image->GetData());
-			return Image;
+			// Null case, no image was provided
+			return CreateDummyDesc();
 		}
 
-		if (!bComplete || !Results)
+		const FUnrealMutableImageInfo& ImageInfo = GlobalExternalImages[id];
+
+		if (ImageInfo.Image)
 		{
-			UE_LOG(LogMutable, Warning, TEXT("Failed to get external image. IO Request failed. Request completed: [%hhd]. Request results [%hhd]. Format: [%s]. MutableFormat: [%d]."),
-				bComplete,
-				(Results != nullptr),
-				GetPixelFormatString(Format),
-				(int32)MutImageFormat);
+			// Easy case where the image was directly provided
+			Result = mu::FImageDesc(ImageInfo.Image->GetSize(), ImageInfo.Image->GetFormat(), ImageInfo.Image->GetLODCount());
 		}
-		else if (MutImageDataSize != (int32)IORequest->GetSize())
+		else if (UTexture2D* TextureToLoad = ImageInfo.TextureToLoad)
 		{
-			UE_LOG(LogMutable, Warning, TEXT("Failed to get external image. Requested size is different than the expected size. RequestSize: [%lld]. ExpectedSize: [%d]. Format: [%s]. MutableFormat: [%d]."),
-				IORequest->GetSize(),
-				Image->GetDataSize(),
-				GetPixelFormatString(Format),
-				(int32)MutImageFormat);
+			// It's safe to access TextureToLoad because ExternalImagesLock guarantees that the data in GlobalExternalImages is valid,
+			// not being modified by the game thread at the moment and the texture cannot be GCed because of the AddReferencedObjects
+			// in the FUnrealMutableImageProvider
+
+			uint8 MipIndex = MipmapsToSkip < TextureToLoad->GetNumMips() ? MipmapsToSkip : TextureToLoad->GetNumMips() - 1;
+
+			// Mips in the mip tail are inlined and can't be streamed, find the smallest mip available.
+			for (; MipIndex > 0; --MipIndex)
+			{
+				if (TextureToLoad->GetPlatformData()->Mips[MipIndex].BulkData.CanLoadFromDisk())
+				{
+					break;
+				}
+			}
+
+			// Texture format and the equivalent mutable format
+			const EPixelFormat Format = TextureToLoad->GetPlatformData()->PixelFormat;
+			const mu::EImageFormat MutableFormat = GetMutablePixelFormat(Format);
+
+			// Check if it's a format we support
+			if (MutableFormat == mu::EImageFormat::IF_NONE)
+			{
+				UE_LOG(LogMutable, Warning, TEXT("Failed to get external image descriptor. Unexpected image format. EImageFormat [%s]."), GetPixelFormatString(Format));
+				return CreateDummyDesc();
+			}
+
+			const mu::FImageSize ImageSize = mu::FImageSize(
+					TextureToLoad->GetSizeX() >> MipIndex,
+					TextureToLoad->GetSizeY() >> MipIndex);
+
+			const int32 Lods = 1;
+
+			Result = mu::FImageDesc(ImageSize, MutableFormat, Lods);
 		}
 		else
 		{
-			UE_LOG(LogMutable, Warning, TEXT("Failed to get external image."));
+			// No UTexture2D was provided, cannot do anything, just provide a dummy texture
+			UE_LOG(LogMutable, Warning, TEXT("No UTexture2D was provided for an application-specific image parameter descriptor."));
+			return CreateDummyDesc();
 		}
 
-		// Something failed when loading the bulk data, just return a dummy
-		return CreateDummy();
-	}
-	else
-	{
-		UE_LOG(LogMutable, Warning, TEXT("Failed to create an IORequest for a UTexture2D BulkData for an application-specific image parameter."));
-		return CreateDummy();
+		return Result;
 	}
 }
 
@@ -391,12 +529,12 @@ void FUnrealMutableImageProvider::ClearCache()
 mu::ImagePtr FUnrealMutableImageProvider::CreateDummy()
 {
 	// Create a dummy image
-	const int size = 32;
+	const int size = DUMMY_IMAGE_DESC.m_size[0];
 	const int checkerSize = 4;
 	constexpr int checkerTileCount = 2;
 	uint8_t colours[checkerTileCount][4] = { { 255,255,0,255 },{ 0,0,255,255 } };
 
-	mu::ImagePtr pResult = new mu::Image(size, size, 1, mu::EImageFormat::IF_RGBA_UBYTE);
+	mu::ImagePtr pResult = new mu::Image(size, size, DUMMY_IMAGE_DESC.m_lods, DUMMY_IMAGE_DESC.m_format);
 
 	uint8_t* pData = pResult->GetData();
 	for (int x = 0; x < size; ++x)
@@ -416,6 +554,12 @@ mu::ImagePtr FUnrealMutableImageProvider::CreateDummy()
 }
 
 
+mu::FImageDesc FUnrealMutableImageProvider::CreateDummyDesc()
+{
+	return DUMMY_IMAGE_DESC;
+}
+
+
 void FUnrealMutableImageProvider::AddReferencedObjects(FReferenceCollector& Collector)
 {
 	for (TPair<uint64, FUnrealMutableImageInfo>& Image : GlobalExternalImages)
diff --git a/Engine/Plugins/Experimental/Mutable/Source/CustomizableObject/Private/MuCO/UnrealMutableImageProvider.h b/Engine/Plugins/Experimental/Mutable/Source/CustomizableObject/Private/MuCO/UnrealMutableImageProvider.h
--- a/Engine/Plugins/Experimental/Mutable/Source/CustomizableObject/Private/MuCO/UnrealMutableImageProvider.h
+++ b/Engine/Plugins/Experimental/Mutable/Source/CustomizableObject/Private/MuCO/UnrealMutableImageProvider.h
@@ -12,12 +12,16 @@ class UTexture2D;
 /** Implementation of a mutable core provider for image parameters that are application-specific. */
 class FUnrealMutableImageProvider : public mu::ImageParameterGenerator, public FGCObject
 {
+
 public:
 
 	// mu::ImageParameterGenerator interface
 	// Thread: worker
-	mu::ImagePtr GetImage(mu::EXTERNAL_IMAGE_ID id, uint8 MipmapsToSkip) override;
+	virtual TTuple<FGraphEventRef, TFunction<void()>> GetImageAsync(mu::EXTERNAL_IMAGE_ID Id, uint8 MipmapsToSkip, TFunction<void (mu::Ptr<mu::Image>)>& ResultCallback) override;
+
+	virtual mu::FImageDesc GetImageDesc(mu::EXTERNAL_IMAGE_ID Id, uint8 MipmapsToSkip) override;
 
+	
 	// Own interface
 	// Thread: Game
 	void CacheImage(mu::EXTERNAL_IMAGE_ID id);
@@ -49,8 +53,12 @@ private:
 		UTexture2D* TextureToLoad = nullptr;
 	};
 
+	static inline const mu::FImageDesc DUMMY_IMAGE_DESC = 
+			mu::FImageDesc(mu::FImageSize(32, 32), mu::EImageFormat::IF_RGBA_UBYTE, 1);
+
 	/** This will be called if an image Id has been requested by Mutable core but it has not been provided by any provider. */
-	mu::ImagePtr CreateDummy();
+	static mu::ImagePtr CreateDummy();
+	static mu::FImageDesc CreateDummyDesc();
 
 	/** Map of Ids to external textures that may be required for any instance or Mutable texture mip under construction.
 	* This is only safely written from the game thread protected by the following critical section, and it
diff --git a/Engine/Plugins/Experimental/Mutable/Source/CustomizableObjectEditor/Private/MuCOE/SMutableCodeViewer.cpp b/Engine/Plugins/Experimental/Mutable/Source/CustomizableObjectEditor/Private/MuCOE/SMutableCodeViewer.cpp
--- a/Engine/Plugins/Experimental/Mutable/Source/CustomizableObjectEditor/Private/MuCOE/SMutableCodeViewer.cpp
+++ b/Engine/Plugins/Experimental/Mutable/Source/CustomizableObjectEditor/Private/MuCOE/SMutableCodeViewer.cpp
@@ -2142,17 +2142,19 @@ namespace
 	/** Test implementation to provide image parameters. It will generate some images of a fixed size and format. */
 	class TestImageProvider : public mu::ImageParameterGenerator
 	{
+		static inline const mu::FImageDesc IMAGE_DESC = 
+			mu::FImageDesc(mu::FImageSize(1024, 1024), mu::EImageFormat::IF_RGBA_UBYTE, 1);
+		
 	public:
-		mu::Ptr<mu::Image> GetImage(mu::EXTERNAL_IMAGE_ID id, uint8 MipmapsToSkip) override
+		TTuple<FGraphEventRef, TFunction<void()>> GetImageAsync(mu::EXTERNAL_IMAGE_ID id, uint8 MipmapsToSkip, TFunction<void (mu::Ptr<mu::Image>)>& ResultCallback) override
 		{
 			MUTABLE_CPUPROFILER_SCOPE(TestImageProvider_GetImage);
 
-			int32 Size = 1024;
+			int32 Size = IMAGE_DESC.m_size[0];
 			Size = FMath::Max(4, Size / (1 << MipmapsToSkip));
 
-			const int32 LODCount = 1;
-			const mu::EImageFormat Format = mu::EImageFormat::IF_RGBA_UBYTE;
-			mu::Ptr<mu::Image> Image = new mu::Image(Size, Size, LODCount, Format);
+			mu::Ptr<mu::Image> Image = 
+					new mu::Image(IMAGE_DESC.m_size[0], IMAGE_DESC.m_size[1], IMAGE_DESC.m_lods, IMAGE_DESC.m_format);
 
 			// Generate an alphatested circle with an horizontal gradient color.
 			uint8* Data = Image->GetData();
@@ -2174,9 +2176,17 @@ namespace
 				}
 			}
 
-			return Image;
+			ResultCallback(Image);
+			FGraphEventRef CompletionEvent = FGraphEvent::CreateGraphEvent();
+			CompletionEvent->DispatchSubsequents();
+
+			return MakeTuple(CompletionEvent, []() -> void {});
 		}
 
+		mu::FImageDesc GetImageDesc(mu::EXTERNAL_IMAGE_ID Id, uint8 MipmapsToSkip) override
+		{
+			return IMAGE_DESC;
+		}
 	};
 }
 
diff --git a/Engine/Plugins/Experimental/Mutable/Source/MutableRuntime/Private/MuR/CodeRunner.cpp b/Engine/Plugins/Experimental/Mutable/Source/MutableRuntime/Private/MuR/CodeRunner.cpp
--- a/Engine/Plugins/Experimental/Mutable/Source/MutableRuntime/Private/MuR/CodeRunner.cpp
+++ b/Engine/Plugins/Experimental/Mutable/Source/MutableRuntime/Private/MuR/CodeRunner.cpp
@@ -111,17 +111,15 @@ namespace mu
 
 
     //---------------------------------------------------------------------------------------------
-    ImagePtr CodeRunner::LoadExternalImage( EXTERNAL_IMAGE_ID Id, uint8 MipmapsToSkip )
+	TTuple<FGraphEventRef, TFunction<void()>> CodeRunner::LoadExternalImageAsync(EXTERNAL_IMAGE_ID Id, uint8 MipmapsToSkip, TFunction<void(Ptr<Image>)>& ResultCallback)
     {
-		MUTABLE_CPUPROFILER_SCOPE(LoadExternalImage);
+		MUTABLE_CPUPROFILER_SCOPE(LoadExternalImageAsync);
 
 		check(m_pSystem);
 
-		ImagePtr pResult;
-
 		if (m_pSystem->m_pImageParameterGenerator)
 		{
-			pResult = m_pSystem->m_pImageParameterGenerator->GetImage(Id, MipmapsToSkip);
+			return m_pSystem->m_pImageParameterGenerator->GetImageAsync(Id, MipmapsToSkip, ResultCallback);
 
 			// Don't cache for now. Need to figure out how to invalidate them.
 			// \TODO: Like constants? attached to a cache level?
@@ -133,10 +131,35 @@ namespace mu
 			check(false);
 		}
 
-		return pResult;
+		// Not needed as it should never reach this point, but added for correctness.
+		FGraphEventRef CompletionEvent = FGraphEvent::CreateGraphEvent();
+		CompletionEvent->DispatchSubsequents();
+
+		return MakeTuple(CompletionEvent, []() -> void {});
 	}
 
+	
+    //---------------------------------------------------------------------------------------------
+	FImageDesc CodeRunner::GetExternalImageDesc(EXTERNAL_IMAGE_ID Id, uint8 MipmapsToSkip)
+	{
+		MUTABLE_CPUPROFILER_SCOPE(GetExternalImageDesc);
+
+		check(m_pSystem);
+
+		if (m_pSystem->m_pImageParameterGenerator)
+		{
+			return m_pSystem->m_pImageParameterGenerator->GetImageDesc(Id, MipmapsToSkip);
+		}
+		else
+		{
+			// Not found and there is no generator!
+			check(false);
+		}
+
+		return FImageDesc();
+	}
 
+	
     //---------------------------------------------------------------------------------------------
     void CodeRunner::RunCode_Conditional( FScheduledOp& item,
                                           const Model* pModel
@@ -2422,19 +2445,6 @@ namespace mu
 		switch (type)
         {
 
-        case OP_TYPE::IM_PARAMETER:
-        {
-			OP::ParameterArgs args = pModel->GetPrivate()->m_program.GetOpArgs<OP::ParameterArgs>(item.At);
-			Ptr<RangeIndex> Index = BuildCurrentOpRangeIndex(item, pParams, pModel, args.variable);
-
-			EXTERNAL_IMAGE_ID id = pParams->GetImageValue(args.variable, Index);
-
-			uint8 MipmapsToSkip = item.ExecutionOptions;
-			ImagePtr pResult = LoadExternalImage(id, MipmapsToSkip);
-			GetMemory().SetImage(item, pResult);
-            break;
-        }
-
         case OP_TYPE::IM_LAYERCOLOUR:
         {
 			OP::ImageLayerColourArgs args = pModel->GetPrivate()->m_program.GetOpArgs<OP::ImageLayerColourArgs>(item.At);
@@ -5302,10 +5312,7 @@ namespace mu
 			OP::ParameterArgs args = program.GetOpArgs<OP::ParameterArgs>(item.At);
 			EXTERNAL_IMAGE_ID id = pParams->GetImageValue(args.variable);
 			uint8 MipsToSkip = item.ExecutionOptions;
-			ImagePtr pResult = LoadExternalImage(id, MipsToSkip);
-			m_heapImageDesc[item.CustomState].m_format = pResult->GetFormat();
-			m_heapImageDesc[item.CustomState].m_size = pResult->GetSize();
-			m_heapImageDesc[item.CustomState].m_lods = pResult->GetLODCount();
+			m_heapImageDesc[item.CustomState] = GetExternalImageDesc(id, MipsToSkip);
 			GetMemory().SetValidDesc(item);
 			break;
 		}
diff --git a/Engine/Plugins/Experimental/Mutable/Source/MutableRuntime/Private/MuR/CodeRunner.h b/Engine/Plugins/Experimental/Mutable/Source/MutableRuntime/Private/MuR/CodeRunner.h
--- a/Engine/Plugins/Experimental/Mutable/Source/MutableRuntime/Private/MuR/CodeRunner.h
+++ b/Engine/Plugins/Experimental/Mutable/Source/MutableRuntime/Private/MuR/CodeRunner.h
@@ -20,6 +20,7 @@
 #include "MuR/Types.h"
 #include "Templates/RefCounting.h"
 #include "Templates/SharedPointer.h"
+#include "Templates/Tuple.h"
 
 // This define could come from MuR/System.h
 #ifdef MUTABLE_USE_NEW_TASKGRAPH
@@ -101,9 +102,14 @@ namespace  mu
 
 		void RunCodeImageDesc(FScheduledOp&, const Parameters*, const Model*, uint32 LodMask);
 
+    public:
 
-        ImagePtr LoadExternalImage( EXTERNAL_IMAGE_ID, uint8 MipmapsToSkip );
+		//! Load an external image asynchronously, retuns an event to wait for complition and a cleanup function 
+		//! that must be called once the event has completed.
+        TTuple<FGraphEventRef, TFunction<void()>> LoadExternalImageAsync(EXTERNAL_IMAGE_ID Id, uint8 MipmapsToSkip, TFunction<void (Ptr<Image>)>& ResultCallback);
+ 	    mu::FImageDesc GetExternalImageDesc(EXTERNAL_IMAGE_ID Id, uint8 MipmapsToSkip);
 
+    protected:
         //! Heap of intermediate data pushed by some instructions and referred by others.
         //! It is not released until no operations are pending.
 		TArray< FScheduledOpData > m_heapData;
diff --git a/Engine/Plugins/Experimental/Mutable/Source/MutableRuntime/Private/MuR/CodeRunnerTasked.cpp b/Engine/Plugins/Experimental/Mutable/Source/MutableRuntime/Private/MuR/CodeRunnerTasked.cpp
--- a/Engine/Plugins/Experimental/Mutable/Source/MutableRuntime/Private/MuR/CodeRunnerTasked.cpp
+++ b/Engine/Plugins/Experimental/Mutable/Source/MutableRuntime/Private/MuR/CodeRunnerTasked.cpp
@@ -1469,6 +1469,70 @@ namespace mu
 	}
 
 
+	//---------------------------------------------------------------------------------------------
+	//---------------------------------------------------------------------------------------------
+	//---------------------------------------------------------------------------------------------
+	class FImageExternalLoadTask : public CodeRunner::FIssuedTask
+	{
+	public:
+		FImageExternalLoadTask(FScheduledOp InItem, uint8 InMipmapsToSkip, EXTERNAL_IMAGE_ID InId);
+
+		// FIssuedTask interface
+		virtual bool Prepare(CodeRunner*, const TSharedPtr<const Model>&, bool& bOutFailed) override;
+		virtual void Complete(CodeRunner* Runner) override;
+		
+	private:
+		FScheduledOp Item;
+		uint8 MipmapsToSkip;
+		EXTERNAL_IMAGE_ID Id;
+
+		Ptr<Image> Result;
+		
+		TFunction<void()> ExternalCleanUpFunc;
+	};
+
+
+	//---------------------------------------------------------------------------------------------
+	FImageExternalLoadTask::FImageExternalLoadTask(FScheduledOp InItem,  uint8 InMipmapsToSkip, EXTERNAL_IMAGE_ID InId)
+	{
+		Item = InItem;
+		MipmapsToSkip = InMipmapsToSkip;
+		Id = InId;
+	}
+
+	//---------------------------------------------------------------------------------------------
+	bool FImageExternalLoadTask::Prepare(CodeRunner* Runner, const TSharedPtr<const Model>&, bool& bOutFailed)
+	{
+		// LoadExternalImageAsync will always generate some image even if it is a dummy one.
+		bOutFailed = false;
+
+		// Capturing this here should not be a problem. The lifetime of the callback lambda is tied to 
+		// the task and the later will always outlive the former. 
+		
+		// Probably we could simply pass a reference to the result image. 
+		TFunction<void (Ptr<Image>)> ResultCallback = [this](Ptr<Image> InResult)
+		{
+			Result = InResult;
+		};
+
+		Tie(Event, ExternalCleanUpFunc) = Runner->LoadExternalImageAsync(Id, MipmapsToSkip, ResultCallback);
+
+		// return false indicating there is no work to do so Event is not overriden by a DoWork task.
+		return false;
+	}
+
+	//---------------------------------------------------------------------------------------------
+	void FImageExternalLoadTask::Complete(CodeRunner* Runner)
+	{
+		if (ExternalCleanUpFunc)
+		{
+			Invoke(ExternalCleanUpFunc);
+		}
+
+		Runner->GetMemory().SetImage(Item, Result);
+	}	
+
+	
 	//---------------------------------------------------------------------------------------------
 	//---------------------------------------------------------------------------------------------
 	//---------------------------------------------------------------------------------------------
@@ -1534,6 +1598,19 @@ namespace mu
 			break;
 		}		
 
+		case OP_TYPE::IM_PARAMETER:
+		{
+			OP::ParameterArgs args = program.GetOpArgs<OP::ParameterArgs>(item.At);
+			Ptr<RangeIndex> Index = BuildCurrentOpRangeIndex(item, m_pParams, m_pModel.Get(), args.variable);
+
+			const EXTERNAL_IMAGE_ID Id = m_pParams->GetImageValue(args.variable, Index);
+			const uint8 MipmapsToSkip = item.ExecutionOptions;
+			
+			Issued = MakeShared<FImageExternalLoadTask>(item, MipmapsToSkip, Id);
+
+			break;
+		}
+			
 		case OP_TYPE::IM_PIXELFORMAT:
 		{
 			if (item.Stage == 1)
diff --git a/Engine/Plugins/Experimental/Mutable/Source/MutableRuntime/Public/MuR/System.h b/Engine/Plugins/Experimental/Mutable/Source/MutableRuntime/Public/MuR/System.h
--- a/Engine/Plugins/Experimental/Mutable/Source/MutableRuntime/Public/MuR/System.h
+++ b/Engine/Plugins/Experimental/Mutable/Source/MutableRuntime/Public/MuR/System.h
@@ -69,8 +69,10 @@ namespace mu
         //! Ensure virtual destruction
         virtual ~ImageParameterGenerator() = default;
 
-        //!
-        virtual Ptr<Image> GetImage( EXTERNAL_IMAGE_ID id, uint8 MipmapsToSkip) = 0;
+        //! Returns the completion event and a cleanup function that must be called once event is completed.
+        virtual TTuple<FGraphEventRef, TFunction<void()>> GetImageAsync(EXTERNAL_IMAGE_ID Id, uint8 MipmapsToSkip, TFunction<void (Ptr<Image>)>& ResultCallback) = 0;
+
+        virtual mu::FImageDesc GetImageDesc(EXTERNAL_IMAGE_ID Id, uint8 MipmapsToSkip) = 0;
     };
 
 
