[mutable] Code optimizations for resize-rel vs additional operations.
#rnx

diff --git a/Engine/Plugins/Experimental/Mutable/Source/MutableTools/Private/MuT/CodeOptimiser_Semantic.cpp b/Engine/Plugins/Experimental/Mutable/Source/MutableTools/Private/MuT/CodeOptimiser_Semantic.cpp
--- a/Engine/Plugins/Experimental/Mutable/Source/MutableTools/Private/MuT/CodeOptimiser_Semantic.cpp
+++ b/Engine/Plugins/Experimental/Mutable/Source/MutableTools/Private/MuT/CodeOptimiser_Semantic.cpp
@@ -22,6 +22,7 @@
 #include "MuT/ASTOpImagePixelFormat.h"
 #include "MuT/ASTOpImageMakeGrowMap.h"
 #include "MuT/ASTOpImageSwizzle.h"
+#include "MuT/ASTOpImageTransform.h"
 #include "MuT/ASTOpMeshMorph.h"
 #include "MuT/ASTOpSwitch.h"
 #include "MuT/CodeOptimiser.h"
@@ -1162,7 +1163,7 @@ namespace mu
 				bool acceptable = false;
 				{
 					auto typedAt = dynamic_cast<const ASTOpImageCompose*>(at.get());
-					auto originalBaseOp = typedAt->Base.child();
+					Ptr<ASTOp> originalBaseOp = typedAt->Base.child();
 
 					// \todo: recursion-proof cache?
 					int layoutBlockPixelsX = 0;
@@ -1185,10 +1186,10 @@ namespace mu
 				{
 					auto newOp = mu::Clone<ASTOpImageCompose>(at);
 
-					auto baseOp = newOp->Base.child();
+					Ptr<ASTOp> baseOp = newOp->Base.child();
 					newOp->Base = Visit(baseOp, currentSinkingOp);
 
-					auto blockOp = newOp->BlockImage.child();
+					Ptr<ASTOp> blockOp = newOp->BlockImage.child();
 					newOp->BlockImage = Visit(blockOp, currentSinkingOp);
 
 					newAt = newOp;
@@ -1216,7 +1217,7 @@ namespace mu
 			case OP_TYPE::IM_MIPMAP:
 			{
 				auto newOp = mu::Clone<ASTOpImageMipmap>(at.get());
-				auto baseOp = newOp->Source.child();
+				Ptr<ASTOp> baseOp = newOp->Source.child();
 				newOp->Source = Visit(baseOp, currentSinkingOp);
 				newAt = newOp;
 				break;
@@ -1224,12 +1225,12 @@ namespace mu
 
 			case OP_TYPE::IM_INTERPOLATE:
 			{
-				auto newOp = mu::Clone<ASTOpFixed>(at);
+				Ptr<ASTOpFixed> newOp = mu::Clone<ASTOpFixed>(at);
 
 				for (int v = 0; v < MUTABLE_OP_MAX_INTERPOLATE_COUNT; ++v)
 				{
-					auto child = newOp->children[newOp->op.args.ImageInterpolate.targets[v]].child();
-					auto bOp = Visit(child, currentSinkingOp);
+					Ptr<ASTOp> child = newOp->children[newOp->op.args.ImageInterpolate.targets[v]].child();
+					Ptr<ASTOp> bOp = Visit(child, currentSinkingOp);
 					newOp->SetChild(newOp->op.args.ImageInterpolate.targets[v], bOp);
 				}
 
@@ -1240,15 +1241,15 @@ namespace mu
 
 			case OP_TYPE::IM_INTERPOLATE3:
 			{
-				auto newOp = mu::Clone<ASTOpFixed>(at);
+				Ptr<ASTOpFixed> newOp = mu::Clone<ASTOpFixed>(at);
 
-				auto top0 = newOp->children[newOp->op.args.ImageInterpolate3.target0].child();
+				Ptr<ASTOp> top0 = newOp->children[newOp->op.args.ImageInterpolate3.target0].child();
 				newOp->SetChild(newOp->op.args.ImageInterpolate3.target0, Visit(top0, currentSinkingOp));
 
-				auto top1 = newOp->children[newOp->op.args.ImageInterpolate3.target1].child();
+				Ptr<ASTOp> top1 = newOp->children[newOp->op.args.ImageInterpolate3.target1].child();
 				newOp->SetChild(newOp->op.args.ImageInterpolate3.target1, Visit(top1, currentSinkingOp));
 
-				auto top2 = newOp->children[newOp->op.args.ImageInterpolate3.target2].child();
+				Ptr<ASTOp> top2 = newOp->children[newOp->op.args.ImageInterpolate3.target2].child();
 				newOp->SetChild(newOp->op.args.ImageInterpolate3.target2, Visit(top2, currentSinkingOp));
 
 				newAt = newOp;
@@ -1257,7 +1258,7 @@ namespace mu
 
 			case OP_TYPE::IM_MULTILAYER:
 			{
-				auto nop = mu::Clone<ASTOpImageMultiLayer>(at);
+				Ptr<ASTOpImageMultiLayer> nop = mu::Clone<ASTOpImageMultiLayer>(at);
 				nop->base = Visit(nop->base.child(), currentSinkingOp);
 				nop->mask = Visit(nop->mask.child(), currentSinkingOp);
 				nop->blend = Visit(nop->blend.child(), currentSinkingOp);
@@ -1267,15 +1268,15 @@ namespace mu
 
 			case OP_TYPE::IM_LAYER:
 			{
-				auto nop = mu::Clone<ASTOpImageLayer>(at);
+				Ptr<ASTOpImageLayer> nop = mu::Clone<ASTOpImageLayer>(at);
 
-				auto aOp = nop->base.child();
+				Ptr<ASTOp> aOp = nop->base.child();
 				nop->base = Visit(aOp, currentSinkingOp);
 
-				auto bOp = nop->blend.child();
+				Ptr<ASTOp> bOp = nop->blend.child();
 				nop->blend = Visit(bOp, currentSinkingOp);
 
-				auto mOp = nop->mask.child();
+				Ptr<ASTOp> mOp = nop->mask.child();
 				nop->mask = Visit(mOp, currentSinkingOp);
 
 				newAt = nop;
@@ -1284,12 +1285,12 @@ namespace mu
 
 			case OP_TYPE::IM_LAYERCOLOUR:
 			{
-				auto nop = mu::Clone<ASTOpImageLayerColor>(at);
+				Ptr<ASTOpImageLayerColor> nop = mu::Clone<ASTOpImageLayerColor>(at);
 
-				auto aOp = nop->base.child();
+				Ptr<ASTOp> aOp = nop->base.child();
 				nop->base = Visit(aOp, currentSinkingOp);
 
-				auto mOp = nop->mask.child();
+				Ptr<ASTOp> mOp = nop->mask.child();
 				nop->mask = Visit(mOp, currentSinkingOp);
 
 				newAt = nop;
@@ -1298,12 +1299,12 @@ namespace mu
 
 			case OP_TYPE::IM_DISPLACE:
 			{
-				auto nop = mu::Clone<ASTOpFixed>(at);
+				Ptr<ASTOpFixed> nop = mu::Clone<ASTOpFixed>(at);
 
-				auto sourceOp = nop->children[nop->op.args.ImageDisplace.source].child();
+				Ptr<ASTOp> sourceOp = nop->children[nop->op.args.ImageDisplace.source].child();
 				nop->SetChild(nop->op.args.ImageDisplace.source, Visit(sourceOp, currentSinkingOp));
 
-				auto mapOp = nop->children[nop->op.args.ImageDisplace.displacementMap].child();
+				Ptr<ASTOp> mapOp = nop->children[nop->op.args.ImageDisplace.displacementMap].child();
 				nop->SetChild(nop->op.args.ImageDisplace.displacementMap, Visit(mapOp, currentSinkingOp));
 
 				// Make sure we don't scale a constant displacement map, which is very wrong.
@@ -1321,17 +1322,44 @@ namespace mu
 
 			case OP_TYPE::IM_MAKEGROWMAP:
 			{
-				auto nop = mu::Clone<ASTOpImageMakeGrowMap>(at);
-				auto maskOp = nop->Mask.child();
+				Ptr<ASTOpImageMakeGrowMap> nop = mu::Clone<ASTOpImageMakeGrowMap>(at);
+				Ptr<ASTOp> maskOp = nop->Mask.child();
 				nop->Mask = Visit(maskOp, currentSinkingOp);
 				newAt = nop;
 				break;
 			}
 
+			case OP_TYPE::IM_INVERT:
+			{
+				Ptr<ASTOpFixed> nop = mu::Clone<ASTOpFixed>(at);
+				Ptr<ASTOp> maskOp = nop->children[nop->op.args.ImageInvert.base].child();
+				nop->SetChild(nop->op.args.ImageInvert.base, Visit(maskOp, currentSinkingOp));
+				newAt = nop;
+				break;
+			}
+
+			case OP_TYPE::IM_SATURATE:
+			{
+				Ptr<ASTOpFixed> nop = mu::Clone<ASTOpFixed>(at);
+				Ptr<ASTOp> maskOp = nop->children[nop->op.args.ImageSaturate.base].child();
+				nop->SetChild(nop->op.args.ImageSaturate.base, Visit(maskOp, currentSinkingOp));
+				newAt = nop;
+				break;
+			}
+
+			case OP_TYPE::IM_TRANSFORM:
+			{
+				Ptr<ASTOpImageTransform> nop = mu::Clone<ASTOpImageTransform>(at);
+				Ptr<ASTOp> maskOp = nop->base.child();
+				nop->base = Visit(maskOp, currentSinkingOp);
+				newAt = nop;
+				break;
+			}
+
 			case OP_TYPE::IM_RASTERMESH:
 			{
-				auto nop = mu::Clone<ASTOpFixed>(at);
-				auto maskOp = nop->children[nop->op.args.ImageRasterMesh.mask].child();
+				Ptr<ASTOpFixed> nop = mu::Clone<ASTOpFixed>(at);
+				Ptr<ASTOp> maskOp = nop->children[nop->op.args.ImageRasterMesh.mask].child();
 				nop->SetChild(nop->op.args.ImageRasterMesh.mask, Visit(maskOp, currentSinkingOp));
 
 				// Resize the image to project as well, assuming that since the target has a different resolution
