Changed an ensureAlwaysMsgf to an ensureMsgf and added a UE_LOG Error whenever it fails, so we will get full logging every time without interfering with debugging and wasting time redumping the callstack constantly.

Made this change because continuous ensures were interfering with my testing; for that reason and to prevent destroying perf or effectively "crashing" servers (due to time taken to dump the stack) unnecessarily, ensureAlways should very rarely be used.

#UE5 #RNX #NoReleaseNotes

diff --git a/Engine/Plugins/Experimental/Mutable/Source/CustomizableObject/Private/MuCO/CustomizableObjectInstanceDescriptor.cpp b/Engine/Plugins/Experimental/Mutable/Source/CustomizableObject/Private/MuCO/CustomizableObjectInstanceDescriptor.cpp
--- a/Engine/Plugins/Experimental/Mutable/Source/CustomizableObject/Private/MuCO/CustomizableObjectInstanceDescriptor.cpp
+++ b/Engine/Plugins/Experimental/Mutable/Source/CustomizableObject/Private/MuCO/CustomizableObjectInstanceDescriptor.cpp
@@ -2111,11 +2111,15 @@ void FCustomizableObjectInstanceDescriptor::SetCurrentState(const FString& State
 {
 	check(CustomizableObject);
 
-	if (const int32 Result = CustomizableObject->FindState(StateName);
-		ensureAlwaysMsgf(Result != -1, TEXT("Unknown %s state."), *StateName))
+	const int32 Result = CustomizableObject->FindState(StateName);
+	if (ensureMsgf(Result != -1, TEXT("Unknown %s state."), *StateName))
 	{
 		SetState(Result);
-	}	
+	}
+	else
+	{
+		UE_LOG(LogMutable, Error, TEXT("%s: Unknown %s state."), ANSI_TO_TCHAR(__FUNCTION__), *StateName);
+	}
 }
 
 
