[Mutable] Fix low end platforms hitches when Mutable is performing background tasks.

diff --git a/Engine/Plugins/Experimental/Mutable/Source/MutableRuntime/Private/MuR/CodeRunner.h b/Engine/Plugins/Experimental/Mutable/Source/MutableRuntime/Private/MuR/CodeRunner.h
--- a/Engine/Plugins/Experimental/Mutable/Source/MutableRuntime/Private/MuR/CodeRunner.h
+++ b/Engine/Plugins/Experimental/Mutable/Source/MutableRuntime/Private/MuR/CodeRunner.h
@@ -28,6 +28,18 @@
 #else
 #endif
 
+static int32 CoreRunnerTaskPriority =
+#ifdef PLATFORM_SWITCH // Required due to low end platforms hitches. Mutable is occupying all threads with long tasks preventing Niagara from executing its time sensitive tasks.
+	5;
+#else
+	0;
+#endif
+
+static FAutoConsoleVariableRef CVarCoreRunnerTaskPriority(
+	TEXT("mutable.CoreRunnerTaskPriority"),
+	CoreRunnerTaskPriority,
+	TEXT("0 AnyThread. 1 to 6, from hight to low priority."),
+	ECVF_Default);
 
 namespace  mu
 {
diff --git a/Engine/Plugins/Experimental/Mutable/Source/MutableRuntime/Private/MuR/CodeRunnerTasked.cpp b/Engine/Plugins/Experimental/Mutable/Source/MutableRuntime/Private/MuR/CodeRunnerTasked.cpp
--- a/Engine/Plugins/Experimental/Mutable/Source/MutableRuntime/Private/MuR/CodeRunnerTasked.cpp
+++ b/Engine/Plugins/Experimental/Mutable/Source/MutableRuntime/Private/MuR/CodeRunnerTasked.cpp
@@ -195,11 +195,39 @@ namespace mu
 								[IssuedTask]() { IssuedTask->DoWork(); },
 								UE::Tasks::ETaskPriority::Inherit );
 #else
+
+							ENamedThreads::Type Priority;
+							switch (CoreRunnerTaskPriority)
+							{
+							default:
+							case 0:
+								Priority = ENamedThreads::AnyThread;
+								break;
+							case 1:
+								Priority = ENamedThreads::AnyHiPriThreadHiPriTask;
+								break;
+							case 2:
+								Priority = ENamedThreads::AnyHiPriThreadNormalTask;
+								break;
+							case 3:
+								Priority = ENamedThreads::AnyNormalThreadHiPriTask;
+								break;
+							case 4:
+								Priority = ENamedThreads::AnyNormalThreadNormalTask;
+								break;
+							case 5:
+								Priority = ENamedThreads::AnyBackgroundHiPriTask;
+								break;
+							case 6:
+								Priority = ENamedThreads::AnyBackgroundThreadNormalTask;
+								break;								
+							}
+							
 							IssuedTask->Event = FFunctionGraphTask::CreateAndDispatchWhenReady(
 								[IssuedTask]() { IssuedTask->DoWork(); },
 								TStatId{},
 								nullptr,
-								ENamedThreads::AnyThread);
+								Priority);
 #endif
 						}
 					}
