[Mutable] Moved CVar declaration to a source file.

diff --git a/Engine/Plugins/Experimental/Mutable/Source/MutableRuntime/Private/MuR/CodeRunner.h b/Engine/Plugins/Experimental/Mutable/Source/MutableRuntime/Private/MuR/CodeRunner.h
--- a/Engine/Plugins/Experimental/Mutable/Source/MutableRuntime/Private/MuR/CodeRunner.h
+++ b/Engine/Plugins/Experimental/Mutable/Source/MutableRuntime/Private/MuR/CodeRunner.h
@@ -28,18 +28,6 @@
 #else
 #endif
 
-static int32 CoreRunnerTaskPriority =
-#ifdef PLATFORM_SWITCH // Required due to low end platforms hitches. Mutable is occupying all threads with long tasks preventing Niagara from executing its time sensitive tasks.
-	5;
-#else
-	0;
-#endif
-
-static FAutoConsoleVariableRef CVarCoreRunnerTaskPriority(
-	TEXT("mutable.CoreRunnerTaskPriority"),
-	CoreRunnerTaskPriority,
-	TEXT("0 AnyThread. 1 to 6, from hight to low priority."),
-	ECVF_Default);
 
 namespace  mu
 {
diff --git a/Engine/Plugins/Experimental/Mutable/Source/MutableRuntime/Private/MuR/CodeRunnerTasked.cpp b/Engine/Plugins/Experimental/Mutable/Source/MutableRuntime/Private/MuR/CodeRunnerTasked.cpp
--- a/Engine/Plugins/Experimental/Mutable/Source/MutableRuntime/Private/MuR/CodeRunnerTasked.cpp
+++ b/Engine/Plugins/Experimental/Mutable/Source/MutableRuntime/Private/MuR/CodeRunnerTasked.cpp
@@ -44,6 +44,21 @@
 
 DECLARE_CYCLE_STAT(TEXT("MutableCoreTask"), STAT_MutableCoreTask, STATGROUP_Game);
 
+
+int32 CoreRunnerTaskPriority =
+#ifdef PLATFORM_SWITCH // Required due to low end platforms hitches. Mutable is occupying all threads with long tasks preventing Niagara from executing its time sensitive tasks.
+		5;
+#else
+		0;
+#endif
+
+static FAutoConsoleVariableRef CVarCoreRunnerTaskPriority(
+	TEXT("mutable.CoreRunnerTaskPriority"),
+	CoreRunnerTaskPriority,
+	TEXT("0 AnyThread. 1 to 6, from hight to low priority."),
+	ECVF_Default);
+
+
 namespace mu
 {
 
