[Mutable] temporal fix to prevent crash on mutable compilations with textures with high lod bias.

#rnx

diff --git a/Engine/Plugins/Experimental/Mutable/Source/CustomizableObjectEditor/Private/MuCOE/GenerateMutableSource/GenerateMutableSource.cpp b/Engine/Plugins/Experimental/Mutable/Source/CustomizableObjectEditor/Private/MuCOE/GenerateMutableSource/GenerateMutableSource.cpp
--- a/Engine/Plugins/Experimental/Mutable/Source/CustomizableObjectEditor/Private/MuCOE/GenerateMutableSource/GenerateMutableSource.cpp
+++ b/Engine/Plugins/Experimental/Mutable/Source/CustomizableObjectEditor/Private/MuCOE/GenerateMutableSource/GenerateMutableSource.cpp
@@ -1383,7 +1383,7 @@ int32 ComputeLODBias(const FMutableGraphGenerationContext& GenerationContext, co
 		UE_LOG(LogMutable, Verbose, TEXT("Compiling texture without reference will have LOD Bias %d."), LODBias);
 	}
 
-	return LODBias;
+	return FMath::Min(6, LODBias);
 }
 
 
