[mutable] Image project node: add support to select angle fading influence on RGB and A channels separately.
#rnx

diff --git a/Engine/Plugins/Experimental/Mutable/Source/CustomizableObject/Classes/MuCO/CustomizableObject.h b/Engine/Plugins/Experimental/Mutable/Source/CustomizableObject/Classes/MuCO/CustomizableObject.h
--- a/Engine/Plugins/Experimental/Mutable/Source/CustomizableObject/Classes/MuCO/CustomizableObject.h
+++ b/Engine/Plugins/Experimental/Mutable/Source/CustomizableObject/Classes/MuCO/CustomizableObject.h
@@ -180,6 +180,7 @@ struct FCompilationOptions
 
 	/** Force a very big number on the mips to skip during compilation. Useful to debug special cooks of the data. */
 	bool bForceLargeLODBias = false;
+	int32 DebugBias = 0;
 
 };
 
@@ -1202,7 +1203,7 @@ private:
 	// This is a manual version number for the binary blobs in this asset.
 	// Increasing it invalidates all the previously compiled models.
 	// Warning: If while merging code both versions have changed, take the highest+1.
-	static const int32 CurrentSupportedVersion = 363;
+	static const int32 CurrentSupportedVersion = 364;
 
 public:
 
diff --git a/Engine/Plugins/Experimental/Mutable/Source/CustomizableObjectEditor/Classes/MuCOE/Nodes/CustomizableObjectNodeTextureProject.h b/Engine/Plugins/Experimental/Mutable/Source/CustomizableObjectEditor/Classes/MuCOE/Nodes/CustomizableObjectNodeTextureProject.h
--- a/Engine/Plugins/Experimental/Mutable/Source/CustomizableObjectEditor/Classes/MuCOE/Nodes/CustomizableObjectNodeTextureProject.h
+++ b/Engine/Plugins/Experimental/Mutable/Source/CustomizableObjectEditor/Classes/MuCOE/Nodes/CustomizableObjectNodeTextureProject.h
@@ -28,6 +28,12 @@ public:
 	// Layout to use for the generated images.
 	UPROPERTY(EditAnywhere, Category = CustomizableObject)
 	int32 Layout = 0;
+	
+	UPROPERTY(EditAnywhere, Category = CustomizableObject)
+	bool bEnableAngleFadeOutForRGB = true;
+
+	UPROPERTY(EditAnywhere, Category = CustomizableObject)
+	bool bEnableAngleFadeOutForAlpha = true;
 
 	/** Set the width of the Texture. If the reference texture is not null, it overrdies the Size X property */
 	UPROPERTY(EditAnywhere, Category = CustomizableObject)
diff --git a/Engine/Plugins/Experimental/Mutable/Source/CustomizableObjectEditor/Private/MuCOE/GenerateMutableSource/GenerateMutableSource.cpp b/Engine/Plugins/Experimental/Mutable/Source/CustomizableObjectEditor/Private/MuCOE/GenerateMutableSource/GenerateMutableSource.cpp
--- a/Engine/Plugins/Experimental/Mutable/Source/CustomizableObjectEditor/Private/MuCOE/GenerateMutableSource/GenerateMutableSource.cpp
+++ b/Engine/Plugins/Experimental/Mutable/Source/CustomizableObjectEditor/Private/MuCOE/GenerateMutableSource/GenerateMutableSource.cpp
@@ -1329,7 +1329,7 @@ int32 ComputeLODBias(const FMutableGraphGenerationContext& GenerationContext, co
 	if (GenerationContext.Options.bForceLargeLODBias)
 	{
 		// This seems to be the highest LODBias we get during cook.
-		return 11;
+		return GenerationContext.Options.DebugBias;
 	}
 
 	// We used to calculate the lod bias directly from the group like this:
diff --git a/Engine/Plugins/Experimental/Mutable/Source/CustomizableObjectEditor/Private/MuCOE/GenerateMutableSource/GenerateMutableSourceImage.cpp b/Engine/Plugins/Experimental/Mutable/Source/CustomizableObjectEditor/Private/MuCOE/GenerateMutableSource/GenerateMutableSourceImage.cpp
--- a/Engine/Plugins/Experimental/Mutable/Source/CustomizableObjectEditor/Private/MuCOE/GenerateMutableSource/GenerateMutableSourceImage.cpp
+++ b/Engine/Plugins/Experimental/Mutable/Source/CustomizableObjectEditor/Private/MuCOE/GenerateMutableSource/GenerateMutableSourceImage.cpp
@@ -591,6 +591,8 @@ mu::NodeImagePtr GenerateMutableSourceImage(const UEdGraphPin* Pin, FMutableGrap
 
 		ImageNode->SetImageSize(TextureSize);
 
+		ImageNode->SetAngleFadeChannels(TypedNodeProject->bEnableAngleFadeOutForRGB, TypedNodeProject->bEnableAngleFadeOutForAlpha);
+
 		if (const UEdGraphPin* ConnectedPin = FollowInputPin(*TypedNodeProject->AngleFadeStartPin()))
 		{
 			mu::NodeScalarPtr ScalarNode = GenerateMutableSourceFloat(ConnectedPin, GenerationContext);
diff --git a/Engine/Plugins/Experimental/Mutable/Source/CustomizableObjectEditor/Private/MuCOE/SMutableObjectViewer.cpp b/Engine/Plugins/Experimental/Mutable/Source/CustomizableObjectEditor/Private/MuCOE/SMutableObjectViewer.cpp
--- a/Engine/Plugins/Experimental/Mutable/Source/CustomizableObjectEditor/Private/MuCOE/SMutableObjectViewer.cpp
+++ b/Engine/Plugins/Experimental/Mutable/Source/CustomizableObjectEditor/Private/MuCOE/SMutableObjectViewer.cpp
@@ -237,6 +237,30 @@ void SMutableObjectViewer::GenerateMutableGraphPressed()
 
 void SMutableObjectViewer::CompileMutableCodePressed()
 {
+	if (CompileOptions.bForceLargeLODBias)
+	{
+		// Debug compile with many different biasses
+		const int32 MaxBias = 15;
+		for (int32 Bias = 0; Bias < MaxBias; ++Bias)
+		{
+			CompileOptions.DebugBias = Bias;
+				
+			mu::NodePtr RootNode = Compiler.Export(CustomizableObject, CompileOptions);
+			if (!RootNode)
+			{
+				// TODO: Show errors
+				return;
+			}
+
+			// Do the compilation to Mutable Code synchronously.
+			TSharedPtr<FCustomizableObjectCompileRunnable> CompileTask = MakeShareable(new FCustomizableObjectCompileRunnable(RootNode, false));
+			CompileTask->Options = CompileOptions;
+			CompileTask->Init();
+			CompileTask->Run();
+		}
+
+	}
+
 	// Convert from Unreal graph to Mutable graph.
 	mu::NodePtr RootNode = Compiler.Export(CustomizableObject, CompileOptions);
 	if (!RootNode)
diff --git a/Engine/Plugins/Experimental/Mutable/Source/MutableRuntime/Private/MuR/CodeRunner.cpp b/Engine/Plugins/Experimental/Mutable/Source/MutableRuntime/Private/MuR/CodeRunner.cpp
--- a/Engine/Plugins/Experimental/Mutable/Source/MutableRuntime/Private/MuR/CodeRunner.cpp
+++ b/Engine/Plugins/Experimental/Mutable/Source/MutableRuntime/Private/MuR/CodeRunner.cpp
@@ -3926,19 +3926,17 @@ namespace mu
                     // moving the right layout channel to the 0.
                     int layout = 0;
 
-                    float projectionAngle = 0;
                     if ( args.projector && pSource && pSource->GetSizeX()>0 && pSource->GetSizeY()>0 )
                     {
                         Ptr<const Projector> pProjector = GetMemory().GetProjector(FScheduledOp::FromOpAndOptions(args.projector, item, 0));
                         if (pProjector)
                         {
-                            projectionAngle = pProjector->m_value.projectionAngle;
-
                             switch (pProjector->m_value.type)
                             {
                             case PROJECTOR_TYPE::PLANAR:
                                 ImageRasterProjectedPlanar( pMesh.get(), pNew.get(),
                                                             pSource.get(), pMask.get(),
+															args.bIsRGBFadingEnabled, args.bIsAlphaFadingEnabled,
                                                             fadeStart, fadeEnd,
                                                             layout, args.blockIndex,
                                                             &scratch );
@@ -3947,7 +3945,8 @@ namespace mu
                             case PROJECTOR_TYPE::WRAPPING:
                                 ImageRasterProjectedWrapping( pMesh.get(), pNew.get(),
                                                               pSource.get(), pMask.get(),
-                                                              fadeStart, fadeEnd,
+															  args.bIsRGBFadingEnabled, args.bIsAlphaFadingEnabled,
+															  fadeStart, fadeEnd,
                                                               layout, args.blockIndex,
                                                               &scratch );
                                 break;
@@ -3955,9 +3954,10 @@ namespace mu
                             case PROJECTOR_TYPE::CYLINDRICAL:
                                 ImageRasterProjectedCylindrical( pMesh.get(), pNew.get(),
                                                                  pSource.get(), pMask.get(),
-                                                                 fadeStart, fadeEnd,
+																 args.bIsRGBFadingEnabled, args.bIsAlphaFadingEnabled,
+																 fadeStart, fadeEnd,
                                                                  layout,
-                                                                 projectionAngle,
+																 pProjector->m_value.projectionAngle,
                                                                  &scratch );
                                 break;
 
diff --git a/Engine/Plugins/Experimental/Mutable/Source/MutableRuntime/Private/MuR/OpImageProject.cpp b/Engine/Plugins/Experimental/Mutable/Source/MutableRuntime/Private/MuR/OpImageProject.cpp
--- a/Engine/Plugins/Experimental/Mutable/Source/MutableRuntime/Private/MuR/OpImageProject.cpp
+++ b/Engine/Plugins/Experimental/Mutable/Source/MutableRuntime/Private/MuR/OpImageProject.cpp
@@ -31,386 +31,391 @@ namespace mu
 namespace
 {
 
-	class RaycastPixelProcessor
-	{
-	public:
-
-		RaycastPixelProcessor( const FProjector& projector,
-							   const Image* pSource,
-                               const uint8* pTargetData,
-                               const uint8* pMaskData,
-							   float fadeStart,
-							   float fadeEnd )
-		{
-			projector.GetDirectionSideUp( m_direction, m_side, m_up );
-
-			m_position[0] = projector.position[0];
-			m_position[1] = projector.position[1];
-			m_position[2] = projector.position[2];
-
-			m_scale[0] = projector.scale[0];
-            m_scale[1] = projector.scale[1];
-            m_scale[2] = projector.scale[2];
-
-			m_pSource = pSource;
-
-            m_fadeStartCos = cosf(fadeStart);
-            m_fadeEndCos = cosf(fadeEnd);
-
-			m_pTargetData = pTargetData;
-			m_pMaskData = pMaskData;
-
-			EImageFormat format = m_pSource->GetFormat();
-			m_pixelSize = GetImageFormatData( format ).m_bytesPerBlock;
-
-		}
-
-
-		//-----------------------------------------------------------------------------------------
-        inline void ProcessPixel( unsigned char* pBufferPos, float varying[4] ) const
-		{
-			float factor = 1;
-
-            // depth clamp
-            if (varying[2]<0.0f || varying[2]>1.0f)
-            {
-                factor = 0;
-            }
-
-            if ( factor>0 && m_pMaskData )
-			{
-                uint8 maskFactor = m_pMaskData[ (pBufferPos-m_pTargetData) / m_pixelSize ];
-				factor = float(maskFactor) / 255.0f;
-			}
-
-			if ( factor>0 )
-			{
-//				vec3<float> normal( varying[2], varying[3], varying[4] );
-//				normal = normalise_approx( normal );
-
-//				float angleCos = dot( normal, m_direction * -1.0f );
-                float angleCos = varying[3];
-
-                if ( angleCos<m_fadeStartCos && angleCos>m_fadeEndCos )
-				{
-                    factor *= ( angleCos - m_fadeEndCos ) / ( m_fadeStartCos - m_fadeEndCos );
-				}
-                else if ( angleCos<=m_fadeEndCos )
-				{
-					factor = 0;
-				}
-			}
-
-
-			if ( factor>0 )
-			{
-				// Transform the origin to the projected image texture coordinates
-				float u = varying[0];
-				float v = varying[1];
-
-				if ( u>=0 && u<1 && v>=0 && v<1 )
-				{
-
-                    const uint8* pPixel = m_pSource->GetData();
-					// TODO: clamp?
-					pPixel += ( m_pSource->GetSizeX() * int(m_pSource->GetSizeY() * v)
-								+ int( m_pSource->GetSizeX() * u ) )
-							* m_pixelSize;
-
-					// Write result
-					switch ( m_pSource->GetFormat() )
-					{
-					case EImageFormat::IF_L_UBYTE:
-                        pBufferPos[0] = uint8( pPixel[0] * factor );
-						break;
-
-					case EImageFormat::IF_RGB_UBYTE:
-                        pBufferPos[0] = uint8( pPixel[0] * factor );
-                        pBufferPos[1] = uint8( pPixel[1] * factor );
-                        pBufferPos[2] = uint8( pPixel[2] * factor );
-						break;
-
-                    case EImageFormat::IF_BGRA_UBYTE:
-                    case EImageFormat::IF_RGBA_UBYTE:
-                        pBufferPos[0] = uint8( pPixel[0] * factor );
-                        pBufferPos[1] = uint8( pPixel[1] * factor );
-                        pBufferPos[2] = uint8( pPixel[2] * factor );
-                        pBufferPos[3] = uint8( pPixel[3] * factor );
-						break;
-
-					default:
-						// Not implemented
-						check( false );
-						break;
-					}
-				}
-			}
-		}
-
-
-		const Image* m_pSource = nullptr;
-		FVector3f m_direction;
-		FVector3f m_up;
-		FVector3f m_side;
-		FVector3f m_position;
-		FVector3f m_scale;
-
-        //! Cosine of the fading angle range
-		float m_fadeStartCos = 0.0f, m_fadeEndCos = 0.0f;
-
-        const uint8* m_pTargetData = nullptr;
-        const uint8* m_pMaskData = nullptr;
-		int m_pixelSize = 0;
-
-	};
+//	class RaycastPixelProcessor
+//	{
+//	public:
+//
+//		RaycastPixelProcessor( const FProjector& projector,
+//							   const Image* pSource,
+//                               const uint8* pTargetData,
+//                               const uint8* pMaskData,
+//							   float fadeStart,
+//							   float fadeEnd )
+//		{
+//			projector.GetDirectionSideUp( m_direction, m_side, m_up );
+//
+//			m_position[0] = projector.position[0];
+//			m_position[1] = projector.position[1];
+//			m_position[2] = projector.position[2];
+//
+//			m_scale[0] = projector.scale[0];
+//            m_scale[1] = projector.scale[1];
+//            m_scale[2] = projector.scale[2];
+//
+//			m_pSource = pSource;
+//
+//            m_fadeStartCos = cosf(fadeStart);
+//            m_fadeEndCos = cosf(fadeEnd);
+//
+//			m_pTargetData = pTargetData;
+//			m_pMaskData = pMaskData;
+//
+//			EImageFormat format = m_pSource->GetFormat();
+//			m_pixelSize = GetImageFormatData( format ).m_bytesPerBlock;
+//
+//		}
+//
+//
+//		//-----------------------------------------------------------------------------------------
+//        inline void ProcessPixel( unsigned char* pBufferPos, float varying[4] ) const
+//		{
+//			float factor = 1;
+//
+//            // depth clamp
+//            if (varying[2]<0.0f || varying[2]>1.0f)
+//            {
+//                factor = 0;
+//            }
+//
+//            if ( factor>0 && m_pMaskData )
+//			{
+//                uint8 maskFactor = m_pMaskData[ (pBufferPos-m_pTargetData) / m_pixelSize ];
+//				factor = float(maskFactor) / 255.0f;
+//			}
+//
+//			if ( factor>0 )
+//			{
+////				vec3<float> normal( varying[2], varying[3], varying[4] );
+////				normal = normalise_approx( normal );
+//
+////				float angleCos = dot( normal, m_direction * -1.0f );
+//                float angleCos = varying[3];
+//
+//                if ( angleCos<m_fadeStartCos && angleCos>m_fadeEndCos )
+//				{
+//                    factor *= ( angleCos - m_fadeEndCos ) / ( m_fadeStartCos - m_fadeEndCos );
+//				}
+//                else if ( angleCos<=m_fadeEndCos )
+//				{
+//					factor = 0;
+//				}
+//			}
+//
+//
+//			if ( factor>0 )
+//			{
+//				// Transform the origin to the projected image texture coordinates
+//				float u = varying[0];
+//				float v = varying[1];
+//
+//				if ( u>=0 && u<1 && v>=0 && v<1 )
+//				{
+//
+//                    const uint8* pPixel = m_pSource->GetData();
+//					// TODO: clamp?
+//					pPixel += ( m_pSource->GetSizeX() * int(m_pSource->GetSizeY() * v)
+//								+ int( m_pSource->GetSizeX() * u ) )
+//							* m_pixelSize;
+//
+//					// Write result
+//					switch ( m_pSource->GetFormat() )
+//					{
+//					case EImageFormat::IF_L_UBYTE:
+//                        pBufferPos[0] = uint8( pPixel[0] * factor );
+//						break;
+//
+//					case EImageFormat::IF_RGB_UBYTE:
+//                        pBufferPos[0] = uint8( pPixel[0] * factor );
+//                        pBufferPos[1] = uint8( pPixel[1] * factor );
+//                        pBufferPos[2] = uint8( pPixel[2] * factor );
+//						break;
+//
+//                    case EImageFormat::IF_BGRA_UBYTE:
+//                    case EImageFormat::IF_RGBA_UBYTE:
+//                        pBufferPos[0] = uint8( pPixel[0] * factor );
+//                        pBufferPos[1] = uint8( pPixel[1] * factor );
+//                        pBufferPos[2] = uint8( pPixel[2] * factor );
+//                        pBufferPos[3] = uint8( pPixel[3] * factor );
+//						break;
+//
+//					default:
+//						// Not implemented
+//						check( false );
+//						break;
+//					}
+//				}
+//			}
+//		}
+//
+//
+//		const Image* m_pSource = nullptr;
+//		FVector3f m_direction;
+//		FVector3f m_up;
+//		FVector3f m_side;
+//		FVector3f m_position;
+//		FVector3f m_scale;
+//
+//        //! Cosine of the fading angle range
+//		float m_fadeStartCos = 0.0f, m_fadeEndCos = 0.0f;
+//
+//        const uint8* m_pTargetData = nullptr;
+//        const uint8* m_pMaskData = nullptr;
+//		int m_pixelSize = 0;
+//
+//	};
 
 
 	//---------------------------------------------------------------------------------------------
-    template<int PIXEL_SIZE,bool HAS_MASK=true>
-	class RaycastPixelProcessor_UBYTE
-	{
-	public:
-
-		RaycastPixelProcessor_UBYTE( const FProjector& projector,
-									 const Image* pSource,
-                                     const uint8* pTargetData,
-                                     const uint8* pMaskData,
-									 float fadeStart,
-									 float fadeEnd )
-		{
-			projector.GetDirectionSideUp( m_direction, m_side, m_up );
-
-			m_position[0] = projector.position[0];
-			m_position[1] = projector.position[1];
-			m_position[2] = projector.position[2];
-
-			m_scale[0] = projector.scale[0];
-            m_scale[1] = projector.scale[1];
-            m_scale[2] = projector.scale[2];
-
-			m_pSourceData = pSource->GetData();
-			m_sourceSizeX = pSource->GetSizeX();
-			m_sourceSizeY = pSource->GetSizeY();
-
-            m_fadeStartCos = cosf(fadeStart);
-            m_fadeEndCos = cosf(fadeEnd);
-
-			m_pTargetData = pTargetData;
-			m_pMaskData = pMaskData;
-
-			check( GetImageFormatData( pSource->GetFormat() ).m_bytesPerBlock
-							==
-							PIXEL_SIZE );
-		}
-
-
-		//-----------------------------------------------------------------------------------------
-        inline void ProcessPixel( unsigned char* pBufferPos, float varying[4] ) const
-		{
-            int32 factor_8 = 256;
-
-            // depth clamp
-            if (varying[2]<0.0f || varying[2]>1.0f)
-            {
-                factor_8 = 0;
-            }
-
-//				vec3<float> normal( varying[2], varying[3], varying[4] );
-//				normal = normalise_approx( normal );
-
-//				float angleCos = dot( normal, m_direction * -1.0f );
-            float angleCos = varying[3];
-
-            if (HAS_MASK)
-            {
-                if ( factor_8>0 && m_pMaskData )
-                {
-                    uint8 maskFactor = m_pMaskData[ (pBufferPos-m_pTargetData) / PIXEL_SIZE ];
-                    factor_8 = maskFactor;
-                }
-
-                if ( factor_8>0 )
-                {
-                    if ( angleCos<m_fadeStartCos && angleCos>m_fadeEndCos )
-                    {
-                        factor_8 = ( factor_8
-                                     *
-                                     int32( 256.0f * ( angleCos - m_fadeEndCos ) / ( m_fadeStartCos - m_fadeEndCos ) )
-                                     ) >> 8;
-                    }
-                    else if ( angleCos<=m_fadeEndCos )
-                    {
-                        factor_8 = 0;
-                    }
-                }
-            }
-            else
-            {
-                if ( angleCos<m_fadeStartCos && angleCos>m_fadeEndCos )
-                {
-                    factor_8 = int32( 256.0f * ( angleCos - m_fadeEndCos ) / ( m_fadeStartCos - m_fadeEndCos ) );
-                }
-                else if ( angleCos<=m_fadeEndCos )
-                {
-                    factor_8 = 0;
-                }
-            }
-
-
-			if ( factor_8>0 )
-			{
-				float u = varying[0];
-				float v = varying[1];
-
-				if ( u>=0 && u<1 && v>=0 && v<1 )
-				{
-
-                    const uint8* pPixel = m_pSourceData;
-					// TODO: clamp?
-					pPixel += ( m_sourceSizeX * int(m_sourceSizeY * v)
-								+ int( m_sourceSizeX * u ) )
-							* PIXEL_SIZE;
-
-					// Write result
-					for ( int i=0; i<PIXEL_SIZE; ++i )
-					{
-						pBufferPos[i] = uint8( ( pPixel[i] * factor_8 ) >> 8 );
-					}
-				}
-			}
-		}
-
-
-        const uint8* m_pSourceData;
-		int m_sourceSizeX, m_sourceSizeY;
-
-		FVector3f m_direction;
-		FVector3f m_up;
-		FVector3f m_side;
-		FVector3f m_position;
-		FVector3f m_scale;
-
-        //! Cosine of the fading angle range
-        float m_fadeStartCos, m_fadeEndCos;
-
-        const uint8* m_pTargetData;
-        const uint8* m_pMaskData;
-	};
-
-
-
-    class RasterProjectedPixelProcessor
-    {
-    public:
-
-        RasterProjectedPixelProcessor( const Image* pSource,
-                               const uint8* pTargetData,
-                               const uint8* pMaskData,
-                               float fadeStart,
-                               float fadeEnd )
-        {
-            m_pSource = pSource;
-
-            m_fadeStartCos = cosf(fadeStart);
-            m_fadeEndCos = cosf(fadeEnd);
-
-            m_pTargetData = pTargetData;
-            m_pMaskData = pMaskData;
-
-			EImageFormat format = m_pSource->GetFormat();
-            m_pixelSize = GetImageFormatData( format ).m_bytesPerBlock;
-
-        }
-
-
-        //-----------------------------------------------------------------------------------------
-        inline void ProcessPixel( unsigned char* pBufferPos, float varying[4] ) const
-        {
-            float factor = 1;
-
-            // depth clamp
-            if (varying[2]<0.0f || varying[2]>1.0f)
-            {
-                factor = 0;
-            }
-
-            if ( factor>0 && m_pMaskData )
-            {
-                uint8 maskFactor = m_pMaskData[ (pBufferPos-m_pTargetData) / m_pixelSize ];
-                factor = float(maskFactor) / 255.0f;
-            }
-
-            if ( factor>0 )
-            {
-                float angleCos = varying[3];
-
-                if ( angleCos<m_fadeStartCos && angleCos>m_fadeEndCos )
-                {
-                    factor *= ( angleCos - m_fadeEndCos ) / ( m_fadeStartCos - m_fadeEndCos );
-                }
-                else if ( angleCos<=m_fadeEndCos )
-                {
-                    factor = 0;
-                }
-            }
-
-
-            if ( factor>0 )
-            {
-                // Transform the origin to the projected image texture coordinates
-                float u = varying[0];
-                float v = varying[1];
-
-                if ( u>=0 && u<1 && v>=0 && v<1 )
-                {
-
-                    const uint8* pPixel = m_pSource->GetData();
-                    // TODO: clamp?
-                    pPixel += ( m_pSource->GetSizeX() * int(m_pSource->GetSizeY() * v)
-                                + int( m_pSource->GetSizeX() * u ) )
-                            * m_pixelSize;
-
-                    // Write result
-                    switch ( m_pSource->GetFormat() )
-                    {
-                    case EImageFormat::IF_L_UBYTE:
-                        pBufferPos[0] = uint8( pPixel[0] * factor );
-                        break;
-
-                    case EImageFormat::IF_RGB_UBYTE:
-                        pBufferPos[0] = uint8( pPixel[0] * factor );
-                        pBufferPos[1] = uint8( pPixel[1] * factor );
-                        pBufferPos[2] = uint8( pPixel[2] * factor );
-                        break;
-
-                    case EImageFormat::IF_BGRA_UBYTE:
-                    case EImageFormat::IF_RGBA_UBYTE:
-                        pBufferPos[0] = uint8( pPixel[0] * factor );
-                        pBufferPos[1] = uint8( pPixel[1] * factor );
-                        pBufferPos[2] = uint8( pPixel[2] * factor );
-                        pBufferPos[3] = uint8( pPixel[3] * factor );
-                        break;
-
-                    default:
-                        // Not implemented
-                        check( false );
-                        break;
-                    }
-                }
-            }
-        }
-
-
-        const Image* m_pSource;
-
-        //! Cosine of the fading angle range
-        float m_fadeStartCos, m_fadeEndCos;
-
-        const uint8* m_pTargetData;
-        const uint8* m_pMaskData;
-        int m_pixelSize;
-
-    };
+//    template<int PIXEL_SIZE,bool HAS_MASK=true>
+//	class RaycastPixelProcessor_UBYTE
+//	{
+//	public:
+//
+//		RaycastPixelProcessor_UBYTE( const FProjector& projector,
+//									 const Image* pSource,
+//                                     const uint8* pTargetData,
+//                                     const uint8* pMaskData,
+//									 float fadeStart,
+//									 float fadeEnd )
+//		{
+//			projector.GetDirectionSideUp( m_direction, m_side, m_up );
+//
+//			m_position[0] = projector.position[0];
+//			m_position[1] = projector.position[1];
+//			m_position[2] = projector.position[2];
+//
+//			m_scale[0] = projector.scale[0];
+//            m_scale[1] = projector.scale[1];
+//            m_scale[2] = projector.scale[2];
+//
+//			m_pSourceData = pSource->GetData();
+//			m_sourceSizeX = pSource->GetSizeX();
+//			m_sourceSizeY = pSource->GetSizeY();
+//
+//            m_fadeStartCos = cosf(fadeStart);
+//            m_fadeEndCos = cosf(fadeEnd);
+//
+//			m_pTargetData = pTargetData;
+//			m_pMaskData = pMaskData;
+//
+//			check( GetImageFormatData( pSource->GetFormat() ).m_bytesPerBlock
+//							==
+//							PIXEL_SIZE );
+//		}
+//
+//
+//		//-----------------------------------------------------------------------------------------
+//        inline void ProcessPixel( unsigned char* pBufferPos, float varying[4] ) const
+//		{
+//            int32 factor_8 = 256;
+//
+//            // depth clamp
+//            if (varying[2]<0.0f || varying[2]>1.0f)
+//            {
+//                factor_8 = 0;
+//            }
+//
+////				vec3<float> normal( varying[2], varying[3], varying[4] );
+////				normal = normalise_approx( normal );
+//
+////				float angleCos = dot( normal, m_direction * -1.0f );
+//            float angleCos = varying[3];
+//
+//            if (HAS_MASK)
+//            {
+//                if ( factor_8>0 && m_pMaskData )
+//                {
+//                    uint8 maskFactor = m_pMaskData[ (pBufferPos-m_pTargetData) / PIXEL_SIZE ];
+//                    factor_8 = maskFactor;
+//                }
+//
+//                if ( factor_8>0 )
+//                {
+//                    if ( angleCos<m_fadeStartCos && angleCos>m_fadeEndCos )
+//                    {
+//                        factor_8 = ( factor_8
+//                                     *
+//                                     int32( 256.0f * ( angleCos - m_fadeEndCos ) / ( m_fadeStartCos - m_fadeEndCos ) )
+//                                     ) >> 8;
+//                    }
+//                    else if ( angleCos<=m_fadeEndCos )
+//                    {
+//                        factor_8 = 0;
+//                    }
+//                }
+//            }
+//            else
+//            {
+//                if ( angleCos<m_fadeStartCos && angleCos>m_fadeEndCos )
+//                {
+//                    factor_8 = int32( 256.0f * ( angleCos - m_fadeEndCos ) / ( m_fadeStartCos - m_fadeEndCos ) );
+//                }
+//                else if ( angleCos<=m_fadeEndCos )
+//                {
+//                    factor_8 = 0;
+//                }
+//            }
+//
+//
+//			if ( factor_8>0 )
+//			{
+//				float u = varying[0];
+//				float v = varying[1];
+//
+//				if ( u>=0 && u<1 && v>=0 && v<1 )
+//				{
+//
+//                    const uint8* pPixel = m_pSourceData;
+//					// TODO: clamp?
+//					pPixel += ( m_sourceSizeX * int(m_sourceSizeY * v)
+//								+ int( m_sourceSizeX * u ) )
+//							* PIXEL_SIZE;
+//
+//					// Write result
+//					for ( int i=0; i<PIXEL_SIZE; ++i )
+//					{
+//						pBufferPos[i] = uint8( ( pPixel[i] * factor_8 ) >> 8 );
+//					}
+//				}
+//			}
+//		}
+//
+//
+//        const uint8* m_pSourceData;
+//		int m_sourceSizeX, m_sourceSizeY;
+//
+//		FVector3f m_direction;
+//		FVector3f m_up;
+//		FVector3f m_side;
+//		FVector3f m_position;
+//		FVector3f m_scale;
+//
+//        //! Cosine of the fading angle range
+//        float m_fadeStartCos, m_fadeEndCos;
+//
+//        const uint8* m_pTargetData;
+//        const uint8* m_pMaskData;
+//	};
+
+
+
+  //  class RasterProjectedPixelProcessor
+  //  {
+  //  public:
+
+  //      RasterProjectedPixelProcessor( const Image* pSource,
+		//	const uint8* pTargetData,
+		//	const uint8* pMaskData,
+		//	bool bIsRGBFadingEnabled, bool bIsAlphaFadingEnabled,
+		//	float fadeStart, float fadeEnd )
+  //      {
+  //          m_pSource = pSource;
+
+		//	// \TODO: to template parameters?
+		//	m_bIsRGBFadingEnabled = bIsRGBFadingEnabled;
+		//	m_bIsAlphaFadingEnabled = bIsAlphaFadingEnabled;
+
+  //          m_fadeStartCos = cosf(fadeStart);
+  //          m_fadeEndCos = cosf(fadeEnd);
+
+  //          m_pTargetData = pTargetData;
+  //          m_pMaskData = pMaskData;
+
+		//	EImageFormat format = m_pSource->GetFormat();
+  //          m_pixelSize = GetImageFormatData( format ).m_bytesPerBlock;
+
+  //      }
+
+
+  //      //-----------------------------------------------------------------------------------------
+  //      inline void ProcessPixel( unsigned char* pBufferPos, float varying[4] ) const
+  //      {
+  //          float factor = 1;
+
+  //          // depth clamp
+  //          if (varying[2]<0.0f || varying[2]>1.0f)
+  //          {
+  //              factor = 0;
+  //          }
+
+  //          if ( factor>0 && m_pMaskData )
+  //          {
+  //              uint8 maskFactor = m_pMaskData[ (pBufferPos-m_pTargetData) / m_pixelSize ];
+  //              factor = float(maskFactor) / 255.0f;
+  //          }
+
+  //          if ( factor>0 )
+  //          {
+  //              float angleCos = varying[3];
+
+  //              if ( angleCos<m_fadeStartCos && angleCos>m_fadeEndCos )
+  //              {
+  //                  factor *= ( angleCos - m_fadeEndCos ) / ( m_fadeStartCos - m_fadeEndCos );
+  //              }
+  //              else if ( angleCos<=m_fadeEndCos )
+  //              {
+  //                  factor = 0;
+  //              }
+  //          }
+
+
+  //          if ( factor>0 )
+  //          {
+  //              // Transform the origin to the projected image texture coordinates
+  //              float u = varying[0];
+  //              float v = varying[1];
+
+  //              if ( u>=0 && u<1 && v>=0 && v<1 )
+  //              {
+
+  //                  const uint8* pPixel = m_pSource->GetData();
+  //                  // TODO: clamp?
+  //                  pPixel += ( m_pSource->GetSizeX() * int(m_pSource->GetSizeY() * v)
+  //                              + int( m_pSource->GetSizeX() * u ) )
+  //                          * m_pixelSize;
+
+  //                  // Write result
+  //                  switch ( m_pSource->GetFormat() )
+  //                  {
+  //                  case EImageFormat::IF_L_UBYTE:
+  //                      pBufferPos[0] = uint8( pPixel[0] * factor );
+  //                      break;
+
+  //                  case EImageFormat::IF_RGB_UBYTE:
+  //                      pBufferPos[0] = uint8( pPixel[0] * factor );
+  //                      pBufferPos[1] = uint8( pPixel[1] * factor );
+  //                      pBufferPos[2] = uint8( pPixel[2] * factor );
+  //                      break;
+
+  //                  case EImageFormat::IF_BGRA_UBYTE:
+  //                  case EImageFormat::IF_RGBA_UBYTE:
+  //                      pBufferPos[0] = uint8( pPixel[0] * factor );
+  //                      pBufferPos[1] = uint8( pPixel[1] * factor );
+  //                      pBufferPos[2] = uint8( pPixel[2] * factor );
+  //                      pBufferPos[3] = uint8( pPixel[3] * factor );
+  //                      break;
+
+  //                  default:
+  //                      // Not implemented
+  //                      check( false );
+  //                      break;
+  //                  }
+  //              }
+  //          }
+  //      }
+
+
+  //      const Image* m_pSource;
+
+  //      //! Cosine of the fading angle range
+  //      float m_fadeStartCos, m_fadeEndCos;
+		//bool m_bIsRGBFadingEnabled, m_bIsAlphaFadingEnabled;
+
+  //      const uint8* m_pTargetData;
+  //      const uint8* m_pMaskData;
+  //      int m_pixelSize;
+
+  //  };
 
 
     //---------------------------------------------------------------------------------------------
@@ -420,15 +425,19 @@ namespace
     public:
 
         RasterProjectedPixelProcessor_UBYTE( const Image* pSource,
-                                     const uint8* pTargetData,
-                                     const uint8* pMaskData,
-                                     float fadeStart,
-                                     float fadeEnd )
+			const uint8* pTargetData,
+			const uint8* pMaskData,
+			bool bIsRGBFadingEnabled, bool bIsAlphaFadingEnabled,
+			float fadeStart, float fadeEnd )
         {
             m_pSourceData = pSource->GetData();
             m_sourceSizeX = pSource->GetSizeX();
             m_sourceSizeY = pSource->GetSizeY();
 
+			// \TODO: to template parameters?
+			m_bIsRGBFadingEnabled = bIsRGBFadingEnabled;
+			m_bIsAlphaFadingEnabled = bIsAlphaFadingEnabled;
+
             m_fadeStartCos = cosf(fadeStart);
             m_fadeEndCos = cosf(fadeEnd);
 
@@ -444,9 +453,11 @@ namespace
         //-----------------------------------------------------------------------------------------
         inline void ProcessPixel( unsigned char* pBufferPos, float varying[4] ) const
         {
+			// \TODO: This is very bad for the CPU. Please optimize.
+
 			if (!m_pSourceData) return;
 
-            int32 factor_8 = 256;
+			int32 factor_8 = 256;
 
             // depth clamp
             if (varying[2]<0.0f || varying[2]>1.0f)
@@ -454,6 +465,8 @@ namespace
                 factor_8 = 0;
             }
 
+			int32 factor_unfaded_8 = factor_8;
+
             float angleCos = varying[3];
 
             if (HAS_MASK)
@@ -464,6 +477,8 @@ namespace
                     factor_8 = maskFactor;
                 }
 
+				factor_unfaded_8 = factor_8;
+
                 if ( factor_8>0 )
                 {
                     if ( angleCos<m_fadeStartCos && angleCos>m_fadeEndCos )
@@ -476,6 +491,7 @@ namespace
                     else if ( angleCos<=m_fadeEndCos )
                     {
                         factor_8 = 0;
+						factor_unfaded_8 = 0;
                     }
                 }
             }
@@ -488,6 +504,7 @@ namespace
                 else if ( angleCos<=m_fadeEndCos )
                 {
                     factor_8 = 0;
+					factor_unfaded_8 = 0;
                 }
             }
 
@@ -502,15 +519,52 @@ namespace
 
                     const uint8* pPixel = m_pSourceData;
                     // TODO: clamp?
-                    pPixel += ( m_sourceSizeX * int(m_sourceSizeY * v)
-                                + int( m_sourceSizeX * u ) )
-                            * PIXEL_SIZE;
+                    pPixel += ( m_sourceSizeX * int(m_sourceSizeY * v) + int( m_sourceSizeX * u ) ) * PIXEL_SIZE;
 
                     // Write result
-                    for ( int i=0; i<PIXEL_SIZE; ++i )
-                    {
-                        pBufferPos[i] = uint8( ( pPixel[i] * factor_8 ) >> 8 );
-                    }
+					if (!m_bIsRGBFadingEnabled && !m_bIsAlphaFadingEnabled)
+					{
+						for (int i = 0; i < PIXEL_SIZE; ++i)
+						{
+							pBufferPos[i] = uint8((pPixel[i] * factor_unfaded_8) >> 8);
+						}
+					}
+					else if (m_bIsRGBFadingEnabled && m_bIsAlphaFadingEnabled)
+					{
+						for (int i = 0; i < PIXEL_SIZE; ++i)
+						{
+							pBufferPos[i] = uint8((pPixel[i] * factor_8) >> 8);
+						}
+					}
+					else
+					{
+						if (m_bIsRGBFadingEnabled)
+						{
+							for (int i = 0; i < PIXEL_SIZE && i < 3; ++i)
+							{
+								pBufferPos[i] = uint8((pPixel[i] * factor_8) >> 8);
+							}
+						}
+						else
+						{
+							for (int i = 0; i < PIXEL_SIZE && i < 3; ++i)
+							{
+								pBufferPos[i] = uint8((pPixel[i] * factor_unfaded_8) >> 8);
+							}
+						}
+
+						if (PIXEL_SIZE >= 3)
+						{
+							if (m_bIsAlphaFadingEnabled)
+							{
+								pBufferPos[3] = uint8((pPixel[3] * factor_8) >> 8);
+							}
+							else
+							{
+								pBufferPos[3] = uint8((pPixel[3] * factor_unfaded_8) >> 8);
+							}
+						}
+					}
                 }
             }
         }
@@ -521,6 +575,7 @@ namespace
 
         //! Cosine of the fading angle range
         float m_fadeStartCos, m_fadeEndCos;
+		bool m_bIsRGBFadingEnabled, m_bIsAlphaFadingEnabled;
 
         const uint8* m_pTargetData;
         const uint8* m_pMaskData;
@@ -655,161 +710,161 @@ namespace
 
 
 //-------------------------------------------------------------------------------------------------
-void ImageRasterProjected_Generic( const Mesh* pMesh,
-                           Image* pImage,
-                           const Image* pSource,
-                           const Image* pMask,
-                           float fadeStart,
-                           float fadeEnd,
-                           int layout,
-                           int block )
-{
-	MUTABLE_CPUPROFILER_SCOPE(ImageRasterProjected_Generic);
-
-	EImageFormat format = pImage->GetFormat();
-    int pixelSize = GetImageFormatData( format ).m_bytesPerBlock;
-
-    int sizeX = pImage->GetSizeX();
-    int sizeY = pImage->GetSizeY();
-
-    RasterProjectedPixelProcessor pixelProc( pSource,
-                                     pImage->GetData(),
-                                     pMask ? pMask->GetData() : 0,
-                                     fadeStart, fadeEnd );
-
-    // Get the vertices
-    int vertexCount = pMesh->GetVertexCount();
-	TArray< RasterVertex<4> > vertices;
-	vertices.SetNum(vertexCount);
-
-    UntypedMeshBufferIteratorConst texIt( pMesh->GetVertexBuffers(), MBS_TEXCOORDS, layout );
-    UntypedMeshBufferIteratorConst posIt( pMesh->GetVertexBuffers(), MBS_POSITION, 0 );
-    UntypedMeshBufferIteratorConst norIt( pMesh->GetVertexBuffers(), MBS_NORMAL, 0 );
-    for ( int v=0; v<vertexCount; ++v )
-    {
-        float uv[2]={0.0f,0.0f};
-        ConvertData( 0, uv, MBF_FLOAT32, texIt.ptr(), texIt.GetFormat() );
-        ConvertData( 1, uv, MBF_FLOAT32, texIt.ptr(), texIt.GetFormat() );
-
-        vertices[v].x = uv[0] * sizeX;
-        vertices[v].y = uv[1] * sizeY;
-
-        FVector3f pos={0.0f,0.0f,0.0f};
-        ConvertData( 0, &pos[0], MBF_FLOAT32, posIt.ptr(), posIt.GetFormat() );
-        ConvertData( 1, &pos[0], MBF_FLOAT32, posIt.ptr(), posIt.GetFormat() );
-        ConvertData( 2, &pos[0], MBF_FLOAT32, posIt.ptr(), posIt.GetFormat() );
-        vertices[v].interpolators[0] = pos[0];
-        vertices[v].interpolators[1] = pos[1];
-        vertices[v].interpolators[2] = pos[2];
-
-        float nor=0.0f;
-        ConvertData( 0, &nor, MBF_FLOAT32, norIt.ptr(), norIt.GetFormat() );
-        vertices[v].interpolators[3] = nor;
-
-        ++texIt;
-        ++posIt;
-        ++norIt;
-    }
-
-    // Get the indices
-    int faceCount = pMesh->GetFaceCount();
-	TArray<int> indices;
-	indices.SetNum(faceCount * 3);
-
-    UntypedMeshBufferIteratorConst indIt( pMesh->GetIndexBuffers(), MBS_VERTEXINDEX, 0 );
-    for ( int i=0; i<faceCount*3; ++i )
-    {
-        uint32 index=0;
-        ConvertData( 0, &index, MBF_UINT32, indIt.ptr(), indIt.GetFormat() );
-
-        indices[i] = index;
-        ++indIt;
-    }
-
-    // Get the block per face
-	TArray<int> blocks;
-	blocks.SetNum(vertexCount);
-
-    UntypedMeshBufferIteratorConst bloIt( pMesh->GetVertexBuffers(), MBS_LAYOUTBLOCK, layout );
-    if (bloIt.ptr())
-    {
-        for ( int i=0; i<vertexCount; ++i )
-        {
-            uint16 index=0;
-            ConvertData( 0, &index, MBF_UINT16, bloIt.ptr(), bloIt.GetFormat() );
-
-            blocks[i] = index;
-            ++bloIt;
-        }
-    }
-
-    if ( block<0 )
-    {
-        // We are rastering the entire layout in an image, so we need to apply the block->layout
-        // transform to the UVs
-        const Layout* pLayout = layout<pMesh->GetLayoutCount() ? pMesh->GetLayout( layout ) : nullptr;
-        if (pLayout && bloIt.ptr())
-        {
-			TArray< box< vec2<float> > > transforms;
-			transforms.SetNum(pLayout->GetBlockCount());
-            for ( int b=0; b<pLayout->GetBlockCount(); ++b )
-            {
-                FIntPoint grid = pLayout->GetGridSize();
-
-                box< vec2<int> > blockRect;
-                pLayout->GetBlock( b, &blockRect.min[0], &blockRect.min[1], &blockRect.size[0], &blockRect.size[1] );
-
-                box< vec2<float> > rect;
-                rect.min[0] = ( (float)blockRect.min[0] ) / (float) grid[0];
-                rect.min[1] = ( (float)blockRect.min[1] ) / (float) grid[1];
-                rect.size[0] = ( (float)blockRect.size[0] ) / (float) grid[0];
-                rect.size[1] = ( (float)blockRect.size[1] ) / (float) grid[1];
-                transforms[b] = rect;
-            }
-
-            for ( int i=0; i<vertexCount; ++i )
-            {
-                int relBlock = pLayout->FindBlock( blocks[i] );
-                vertices[i].x = vertices[i].x * transforms[ relBlock ].size[0]
-                        + transforms[ relBlock ].min[0] * sizeX;
-                vertices[i].y = vertices[i].y * transforms[ relBlock ].size[1]
-                        + transforms[ relBlock ].min[1] * sizeY;
-            }
-        }
-
-        // Raster all the faces
-        for ( int f=0; f<faceCount; ++f )
-        {
-            Triangle( pImage->GetData(),
-                      sizeX, sizeY,
-                      pixelSize,
-                      vertices[indices[f*3+0]],
-                      vertices[indices[f*3+1]],
-                      vertices[indices[f*3+2]],
-                      pixelProc,
-                      false );
-        }
-    }
-    else
-    {
-        // Raster only the faces in the selected block
-        for ( int f=0; f<faceCount; ++f )
-        {
-            if ( blocks[ indices[f*3+0] ]==block )
-            {
-                Triangle( pImage->GetData(),
-                          sizeX, sizeY,
-                          pixelSize,
-                          vertices[indices[f*3+0]],
-                          vertices[indices[f*3+1]],
-                          vertices[indices[f*3+2]],
-                          pixelProc,
-                          false );
-            }
-        }
-    }
-}
-
+//void ImageRasterProjected_Generic( const Mesh* pMesh,
+//                           Image* pImage,
+//                           const Image* pSource,
+//                           const Image* pMask,
+//                           float fadeStart,
+//                           float fadeEnd,
+//                           int layout,
+//                           int block )
+//{
+//	MUTABLE_CPUPROFILER_SCOPE(ImageRasterProjected_Generic);
+//
+//	EImageFormat format = pImage->GetFormat();
+//    int pixelSize = GetImageFormatData( format ).m_bytesPerBlock;
+//
+//    int sizeX = pImage->GetSizeX();
+//    int sizeY = pImage->GetSizeY();
+//
+//    RasterProjectedPixelProcessor pixelProc( pSource,
+//                                     pImage->GetData(),
+//                                     pMask ? pMask->GetData() : 0,
+//                                     fadeStart, fadeEnd );
+//
+//    // Get the vertices
+//    int vertexCount = pMesh->GetVertexCount();
+//	TArray< RasterVertex<4> > vertices;
+//	vertices.SetNum(vertexCount);
+//
+//    UntypedMeshBufferIteratorConst texIt( pMesh->GetVertexBuffers(), MBS_TEXCOORDS, layout );
+//    UntypedMeshBufferIteratorConst posIt( pMesh->GetVertexBuffers(), MBS_POSITION, 0 );
+//    UntypedMeshBufferIteratorConst norIt( pMesh->GetVertexBuffers(), MBS_NORMAL, 0 );
+//    for ( int v=0; v<vertexCount; ++v )
+//    {
+//        float uv[2]={0.0f,0.0f};
+//        ConvertData( 0, uv, MBF_FLOAT32, texIt.ptr(), texIt.GetFormat() );
+//        ConvertData( 1, uv, MBF_FLOAT32, texIt.ptr(), texIt.GetFormat() );
+//
+//        vertices[v].x = uv[0] * sizeX;
+//        vertices[v].y = uv[1] * sizeY;
+//
+//        FVector3f pos={0.0f,0.0f,0.0f};
+//        ConvertData( 0, &pos[0], MBF_FLOAT32, posIt.ptr(), posIt.GetFormat() );
+//        ConvertData( 1, &pos[0], MBF_FLOAT32, posIt.ptr(), posIt.GetFormat() );
+//        ConvertData( 2, &pos[0], MBF_FLOAT32, posIt.ptr(), posIt.GetFormat() );
+//        vertices[v].interpolators[0] = pos[0];
+//        vertices[v].interpolators[1] = pos[1];
+//        vertices[v].interpolators[2] = pos[2];
+//
+//        float nor=0.0f;
+//        ConvertData( 0, &nor, MBF_FLOAT32, norIt.ptr(), norIt.GetFormat() );
+//        vertices[v].interpolators[3] = nor;
+//
+//        ++texIt;
+//        ++posIt;
+//        ++norIt;
+//    }
+//
+//    // Get the indices
+//    int faceCount = pMesh->GetFaceCount();
+//	TArray<int> indices;
+//	indices.SetNum(faceCount * 3);
+//
+//    UntypedMeshBufferIteratorConst indIt( pMesh->GetIndexBuffers(), MBS_VERTEXINDEX, 0 );
+//    for ( int i=0; i<faceCount*3; ++i )
+//    {
+//        uint32 index=0;
+//        ConvertData( 0, &index, MBF_UINT32, indIt.ptr(), indIt.GetFormat() );
+//
+//        indices[i] = index;
+//        ++indIt;
+//    }
+//
+//    // Get the block per face
+//	TArray<int> blocks;
+//	blocks.SetNum(vertexCount);
+//
+//    UntypedMeshBufferIteratorConst bloIt( pMesh->GetVertexBuffers(), MBS_LAYOUTBLOCK, layout );
+//    if (bloIt.ptr())
+//    {
+//        for ( int i=0; i<vertexCount; ++i )
+//        {
+//            uint16 index=0;
+//            ConvertData( 0, &index, MBF_UINT16, bloIt.ptr(), bloIt.GetFormat() );
+//
+//            blocks[i] = index;
+//            ++bloIt;
+//        }
+//    }
+//
+//    if ( block<0 )
+//    {
+//        // We are rastering the entire layout in an image, so we need to apply the block->layout
+//        // transform to the UVs
+//        const Layout* pLayout = layout<pMesh->GetLayoutCount() ? pMesh->GetLayout( layout ) : nullptr;
+//        if (pLayout && bloIt.ptr())
+//        {
+//			TArray< box< vec2<float> > > transforms;
+//			transforms.SetNum(pLayout->GetBlockCount());
+//            for ( int b=0; b<pLayout->GetBlockCount(); ++b )
+//            {
+//                FIntPoint grid = pLayout->GetGridSize();
+//
+//                box< vec2<int> > blockRect;
+//                pLayout->GetBlock( b, &blockRect.min[0], &blockRect.min[1], &blockRect.size[0], &blockRect.size[1] );
+//
+//                box< vec2<float> > rect;
+//                rect.min[0] = ( (float)blockRect.min[0] ) / (float) grid[0];
+//                rect.min[1] = ( (float)blockRect.min[1] ) / (float) grid[1];
+//                rect.size[0] = ( (float)blockRect.size[0] ) / (float) grid[0];
+//                rect.size[1] = ( (float)blockRect.size[1] ) / (float) grid[1];
+//                transforms[b] = rect;
+//            }
+//
+//            for ( int i=0; i<vertexCount; ++i )
+//            {
+//                int relBlock = pLayout->FindBlock( blocks[i] );
+//                vertices[i].x = vertices[i].x * transforms[ relBlock ].size[0]
+//                        + transforms[ relBlock ].min[0] * sizeX;
+//                vertices[i].y = vertices[i].y * transforms[ relBlock ].size[1]
+//                        + transforms[ relBlock ].min[1] * sizeY;
+//            }
+//        }
+//
+//        // Raster all the faces
+//        for ( int f=0; f<faceCount; ++f )
+//        {
+//            Triangle( pImage->GetData(),
+//                      sizeX, sizeY,
+//                      pixelSize,
+//                      vertices[indices[f*3+0]],
+//                      vertices[indices[f*3+1]],
+//                      vertices[indices[f*3+2]],
+//                      pixelProc,
+//                      false );
+//        }
+//    }
+//    else
+//    {
+//        // Raster only the faces in the selected block
+//        for ( int f=0; f<faceCount; ++f )
+//        {
+//            if ( blocks[ indices[f*3+0] ]==block )
+//            {
+//                Triangle( pImage->GetData(),
+//                          sizeX, sizeY,
+//                          pixelSize,
+//                          vertices[indices[f*3+0]],
+//                          vertices[indices[f*3+1]],
+//                          vertices[indices[f*3+2]],
+//                          pixelProc,
+//                          false );
+//            }
+//        }
+//    }
+//}
+//
 
 //-------------------------------------------------------------------------------------------------
 //! This format is the one we assumee the meshes optimised for planar and cylindrical projection
@@ -1003,14 +1058,13 @@ void ImageRasterProjected_OptimisedWrapping( const Mesh* pMesh,
 
 //-------------------------------------------------------------------------------------------------
 void ImageRasterProjectedPlanar( const Mesh* pMesh,
-                                     Image* pImage,
-                                     const Image* pSource,
-                                     const Image* pMask,
-                                     float fadeStart,
-                                     float fadeEnd,
-                                     int layout,
-                                     int block,
-                                     SCRATCH_IMAGE_PROJECT* scratch )
+	Image* pImage,
+	const Image* pSource,
+	const Image* pMask,
+	bool bIsRGBFadingEnabled, bool bIsAlphaFadingEnabled,
+	float fadeStart, float fadeEnd,
+	int layout, int block,
+	SCRATCH_IMAGE_PROJECT* scratch )
 {
     check( !pMask || pMask->GetSizeX() == pImage->GetSizeX() );
     check( !pMask || pMask->GetSizeY() == pImage->GetSizeY() );
@@ -1023,10 +1077,12 @@ void ImageRasterProjectedPlanar( const Mesh* pMesh,
         // Mesh-optimised version
         if ( pSource->GetFormat() == EImageFormat::IF_RGB_UBYTE )
         {
-            RasterProjectedPixelProcessor_UBYTE<3> pixelProc(  pSource,
-                                                       pImage->GetData(),
-                                                       pMask ? pMask->GetData() : 0,
-                                                       fadeStart, fadeEnd );
+            RasterProjectedPixelProcessor_UBYTE<3> pixelProc( 
+				pSource,
+				pImage->GetData(),
+				pMask ? pMask->GetData() : 0,
+				bIsRGBFadingEnabled, bIsAlphaFadingEnabled,
+				fadeStart, fadeEnd );
 
             ImageRasterProjected_Optimised( pMesh, pImage, pixelProc, fadeEnd, scratch );
         }
@@ -1037,17 +1093,19 @@ void ImageRasterProjectedPlanar( const Mesh* pMesh,
             if (pMask!=nullptr)
             {
                 RasterProjectedPixelProcessor_UBYTE<4,true> pixelProc( pSource,
-                                                           pImage->GetData(),
-                                                           pMask->GetData(),
-                                                           fadeStart, fadeEnd );
+                    pImage->GetData(),
+                    pMask->GetData(),
+					bIsRGBFadingEnabled, bIsAlphaFadingEnabled,
+					fadeStart, fadeEnd );
                 ImageRasterProjected_Optimised( pMesh, pImage, pixelProc, fadeEnd, scratch );
             }
             else
             {
                 RasterProjectedPixelProcessor_UBYTE<4,false> pixelProc( pSource,
-                                                           pImage->GetData(),
-                                                           nullptr,
-                                                           fadeStart, fadeEnd );
+                    pImage->GetData(),
+                    nullptr,
+					bIsRGBFadingEnabled, bIsAlphaFadingEnabled,
+					fadeStart, fadeEnd );
                 ImageRasterProjected_Optimised( pMesh, pImage, pixelProc, fadeEnd, scratch );
             }
 
@@ -1056,21 +1114,24 @@ void ImageRasterProjectedPlanar( const Mesh* pMesh,
         else if ( pSource->GetFormat() == EImageFormat::IF_L_UBYTE )
         {
             RasterProjectedPixelProcessor_UBYTE<1> pixelProc( pSource,
-                                                     pImage->GetData(),
-                                                     pMask ? pMask->GetData() : 0,
-                                                     fadeStart, fadeEnd );
+				pImage->GetData(),
+				pMask ? pMask->GetData() : 0,
+				bIsRGBFadingEnabled, bIsAlphaFadingEnabled,
+				fadeStart, fadeEnd );
 
             ImageRasterProjected_Optimised( pMesh, pImage, pixelProc, fadeEnd, scratch );
         }
 
         else
         {
-            RasterProjectedPixelProcessor pixelProc( pSource,
-                                             pImage->GetData(),
-                                             pMask ? pMask->GetData() : 0,
-                                             fadeStart, fadeEnd );
-
-            ImageRasterProjected_Optimised( pMesh, pImage, pixelProc, fadeEnd, scratch );
+			check(false);
+    //        RasterProjectedPixelProcessor pixelProc( pSource,
+    //                                         pImage->GetData(),
+    //                                         pMask ? pMask->GetData() : 0, 
+				//bIsRGBFadingEnabled, bIsAlphaFadingEnabled,
+    //                                         fadeStart, fadeEnd );
+
+    //        ImageRasterProjected_Optimised( pMesh, pImage, pixelProc, fadeEnd, scratch );
         }
     }
 
@@ -1078,10 +1139,9 @@ void ImageRasterProjectedPlanar( const Mesh* pMesh,
     {
         check(false);
         // Generic version
-        // TODO: It will allocate temp memory.
-        ImageRasterProjected_Generic( pMesh, pImage, pSource, pMask,
-                              fadeStart, fadeEnd,
-                              layout, block );
+        //ImageRasterProjected_Generic( pMesh, pImage, pSource, pMask,
+        //                      fadeStart, fadeEnd,
+        //                      layout, block );
     }
 
 }
@@ -1092,8 +1152,8 @@ void ImageRasterProjectedWrapping( const Mesh* pMesh,
                                        Image* pImage,
                                        const Image* pSource,
                                        const Image* pMask,
-                                       float fadeStart,
-                                       float fadeEnd,
+									   bool bIsRGBFadingEnabled, bool bIsAlphaFadingEnabled,
+                                       float fadeStart, float fadeEnd,
                                        int layout,
                                        int block,
                                        SCRATCH_IMAGE_PROJECT* scratch )
@@ -1111,9 +1171,10 @@ void ImageRasterProjectedWrapping( const Mesh* pMesh,
         if ( pSource->GetFormat() == EImageFormat::IF_RGB_UBYTE )
         {
             RasterProjectedPixelProcessor_UBYTE<3> pixelProc(  pSource,
-                                                       pImage->GetData(),
-                                                       pMask ? pMask->GetData() : 0,
-                                                       fadeStart, fadeEnd );
+				pImage->GetData(),
+				pMask ? pMask->GetData() : 0,
+				bIsRGBFadingEnabled, bIsAlphaFadingEnabled,
+				fadeStart, fadeEnd );
 
             ImageRasterProjected_OptimisedWrapping( pMesh, pImage, pixelProc, fadeEnd, block, scratch );
         }
@@ -1124,17 +1185,19 @@ void ImageRasterProjectedWrapping( const Mesh* pMesh,
             if (pMask!=nullptr)
             {
                 RasterProjectedPixelProcessor_UBYTE<4,true> pixelProc( pSource,
-                                                           pImage->GetData(),
-                                                           pMask->GetData(),
-                                                           fadeStart, fadeEnd );
+					pImage->GetData(),
+					pMask->GetData(),
+					bIsRGBFadingEnabled, bIsAlphaFadingEnabled,
+					fadeStart, fadeEnd );
                 ImageRasterProjected_OptimisedWrapping( pMesh, pImage, pixelProc, fadeEnd, block, scratch );
             }
             else
             {
                 RasterProjectedPixelProcessor_UBYTE<4,false> pixelProc( pSource,
-                                                           pImage->GetData(),
-                                                           nullptr,
-                                                           fadeStart, fadeEnd );
+					pImage->GetData(),
+					nullptr,
+					bIsRGBFadingEnabled, bIsAlphaFadingEnabled,
+					fadeStart, fadeEnd );
                 ImageRasterProjected_OptimisedWrapping( pMesh, pImage, pixelProc, fadeEnd, block, scratch );
             }
 
@@ -1143,21 +1206,24 @@ void ImageRasterProjectedWrapping( const Mesh* pMesh,
         else if ( pSource->GetFormat() == EImageFormat::IF_L_UBYTE )
         {
             RasterProjectedPixelProcessor_UBYTE<1> pixelProc( pSource,
-                                                     pImage->GetData(),
-                                                     pMask ? pMask->GetData() : 0,
-                                                     fadeStart, fadeEnd );
+				pImage->GetData(),
+				pMask ? pMask->GetData() : 0,
+				bIsRGBFadingEnabled, bIsAlphaFadingEnabled,
+				fadeStart, fadeEnd );
 
             ImageRasterProjected_OptimisedWrapping( pMesh, pImage, pixelProc, fadeEnd, block, scratch );
         }
 
         else
         {
-            RasterProjectedPixelProcessor pixelProc( pSource,
-                                             pImage->GetData(),
-                                             pMask ? pMask->GetData() : 0,
-                                             fadeStart, fadeEnd );
-
-            ImageRasterProjected_OptimisedWrapping( pMesh, pImage, pixelProc, fadeEnd, block, scratch );
+			check(false);
+    //        RasterProjectedPixelProcessor pixelProc( pSource,
+				//pImage->GetData(),
+				//pMask ? pMask->GetData() : 0,
+				//bIsRGBFadingEnabled, bIsAlphaFadingEnabled,
+				//fadeStart, fadeEnd );
+
+    //        ImageRasterProjected_OptimisedWrapping( pMesh, pImage, pixelProc, fadeEnd, block, scratch );
         }
     }
 
@@ -1165,10 +1231,9 @@ void ImageRasterProjectedWrapping( const Mesh* pMesh,
     {
         check(false);
         // Generic version
-        // TODO: It will allocate temp memory.
-        ImageRasterProjected_Generic( pMesh, pImage, pSource, pMask,
-                              fadeStart, fadeEnd,
-                              layout, block );
+        //ImageRasterProjected_Generic( pMesh, pImage, pSource, pMask,
+        //                      fadeStart, fadeEnd,
+        //                      layout, block );
     }
 
 }
@@ -1176,14 +1241,14 @@ void ImageRasterProjectedWrapping( const Mesh* pMesh,
 
 //-------------------------------------------------------------------------------------------------
 void ImageRasterProjectedCylindrical( const Mesh* pMesh,
-                                          Image* pImage,
-                                          const Image* pSource,
-                                          const Image* pMask,
-                                          float fadeStart,
-                                          float fadeEnd,
-                                          int /*layout*/,
-                                          float projectionAngle,
-                                          SCRATCH_IMAGE_PROJECT* scratch )
+                                        Image* pImage,
+                                        const Image* pSource,
+                                        const Image* pMask,
+										bool bIsRGBFadingEnabled, bool bIsAlphaFadingEnabled,
+                                        float fadeStart, float fadeEnd,
+                                        int /*layout*/,
+                                        float projectionAngle,
+                                        SCRATCH_IMAGE_PROJECT* scratch )
 {
     check( !pMask || pMask->GetSizeX() == pImage->GetSizeX() );
     check( !pMask || pMask->GetSizeY() == pImage->GetSizeY() );
diff --git a/Engine/Plugins/Experimental/Mutable/Source/MutableRuntime/Private/MuR/OpImageProject.h b/Engine/Plugins/Experimental/Mutable/Source/MutableRuntime/Private/MuR/OpImageProject.h
--- a/Engine/Plugins/Experimental/Mutable/Source/MutableRuntime/Private/MuR/OpImageProject.h
+++ b/Engine/Plugins/Experimental/Mutable/Source/MutableRuntime/Private/MuR/OpImageProject.h
@@ -17,24 +17,12 @@ namespace mu
 		TArray<uint8> culledVertex;
     };
 
-
-//    extern void ImageProject( const Mesh* pMesh,
-//                              Image* pTargetImage,
-//                              const Image* pSource,
-//                              const Image* pMask,
-//                              const FProjector& projector,
-//                              float fadeStart,
-//                              float fadeEnd,
-//                              int layout,
-//                              int block,
-//                              SCRATCH_IMAGE_PROJECT* scratch );
-
     extern void ImageRasterProjectedPlanar( const Mesh* pMesh,
                                             Image* pTargetImage,
                                             const Image* pSource,
                                             const Image* pMask,
-                                            float fadeStart,
-                                            float fadeEnd,
+											bool bIsRGBFadingEnabled, bool bIsAlphaFadingEnabled,
+											float fadeStart, float fadeEnd,
                                             int layout,
                                             int block,
                                             SCRATCH_IMAGE_PROJECT* scratch );
@@ -43,8 +31,8 @@ namespace mu
                                       Image* pTargetImage,
                                       const Image* pSource,
                                       const Image* pMask,
-                                      float fadeStart,
-                                      float fadeEnd,
+									  bool bIsRGBFadingEnabled, bool bIsAlphaFadingEnabled,
+									  float fadeStart, float fadeEnd,
                                       int layout,
                                       float projectionAngle,
                                       SCRATCH_IMAGE_PROJECT* scratch );
@@ -53,9 +41,9 @@ namespace mu
                                       Image* pTargetImage,
                                       const Image* pSource,
                                       const Image* pMask,
-                                      float fadeStart,
-                                      float fadeEnd,
-                                      int layout,
+									  bool bIsRGBFadingEnabled, bool bIsAlphaFadingEnabled,
+									  float fadeStart, float fadeEnd,
+									  int layout,
                                       int block,
                                       SCRATCH_IMAGE_PROJECT* scratch );
 
diff --git a/Engine/Plugins/Experimental/Mutable/Source/MutableRuntime/Private/MuR/Operations.cpp b/Engine/Plugins/Experimental/Mutable/Source/MutableRuntime/Private/MuR/Operations.cpp
--- a/Engine/Plugins/Experimental/Mutable/Source/MutableRuntime/Private/MuR/Operations.cpp
+++ b/Engine/Plugins/Experimental/Mutable/Source/MutableRuntime/Private/MuR/Operations.cpp
@@ -325,14 +325,6 @@ namespace mu
             f(&op.args.ImageCrop.source );
             break;
 
-        case OP_TYPE::IM_RASTERMESH:
-            f(&op.args.ImageRasterMesh.mesh );
-            f(&op.args.ImageRasterMesh.image );
-            f(&op.args.ImageRasterMesh.mask );
-            f(&op.args.ImageRasterMesh.angleFadeProperties );
-            f(&op.args.ImageRasterMesh.projector );
-            break;
-
         case OP_TYPE::IM_DISPLACE:
             f(&op.args.ImageDisplace.source );
             f(&op.args.ImageDisplace.displacementMap );
diff --git a/Engine/Plugins/Experimental/Mutable/Source/MutableRuntime/Private/MuR/Operations.h b/Engine/Plugins/Experimental/Mutable/Source/MutableRuntime/Private/MuR/Operations.h
--- a/Engine/Plugins/Experimental/Mutable/Source/MutableRuntime/Private/MuR/Operations.h
+++ b/Engine/Plugins/Experimental/Mutable/Source/MutableRuntime/Private/MuR/Operations.h
@@ -749,19 +749,22 @@ namespace mu
         struct ImageRasterMeshArgs
         {
             ADDRESS mesh;
-            uint16 sizeX, sizeY;
-            int32 blockIndex;
 
-            // These are used in case of projected mesh raster.
-            ADDRESS image;
-            ADDRESS angleFadeProperties;
+			//! These are used in case of projected mesh raster.
+			ADDRESS image;
+			ADDRESS angleFadeProperties;
 
-            //! Mask selecting the pixels in the destination image that may receive projection.
-            ADDRESS mask;
+			//! Mask selecting the pixels in the destination image that may receive projection.
+			ADDRESS mask;
 
-            //! A projector may be needed for some kind of per-pixel raster operations
-            //! like cylindrical projections.
-            ADDRESS projector;
+			//! A projector may be needed for some kind of per-pixel raster operations
+			//! like cylindrical projections.
+			ADDRESS projector;
+			
+			int32 blockIndex;
+			uint16 sizeX, sizeY;
+			uint8 bIsRGBFadingEnabled : 1;
+			uint8 bIsAlphaFadingEnabled : 1;
         };
 
         struct ImageMakeGrowMapArgs
@@ -1106,7 +1109,7 @@ namespace mu
             ImagePlainColourArgs ImagePlainColour;
             ImageGPUArgs ImageGPU;
             ImageCropArgs ImageCrop;
-            ImageRasterMeshArgs ImageRasterMesh;
+//            ImageRasterMeshArgs ImageRasterMesh;
             ImageDisplaceArgs ImageDisplace;
 			ImageInvertArgs ImageInvert;
 
diff --git a/Engine/Plugins/Experimental/Mutable/Source/MutableRuntime/Public/MuR/Model.h b/Engine/Plugins/Experimental/Mutable/Source/MutableRuntime/Public/MuR/Model.h
--- a/Engine/Plugins/Experimental/Mutable/Source/MutableRuntime/Public/MuR/Model.h
+++ b/Engine/Plugins/Experimental/Mutable/Source/MutableRuntime/Public/MuR/Model.h
@@ -13,7 +13,7 @@
 //! data structures. Compiled models are not necessarily compatible when the runtime is updated,
 //! so this version number can be used externally to verify this. It is not used internally, and
 //! serializing models from different versions than this runtime will probably result in a crash.
-#define MUTABLE_COMPILED_MODEL_CODE_VERSION		uint32( 53 )
+#define MUTABLE_COMPILED_MODEL_CODE_VERSION		uint32( 54 )
 #define MUTABLE_PARAMETERS_VERSION              uint32( 1 )
 
 
diff --git a/Engine/Plugins/Experimental/Mutable/Source/MutableTools/Private/MuT/AST.cpp b/Engine/Plugins/Experimental/Mutable/Source/MutableTools/Private/MuT/AST.cpp
--- a/Engine/Plugins/Experimental/Mutable/Source/MutableTools/Private/MuT/AST.cpp
+++ b/Engine/Plugins/Experimental/Mutable/Source/MutableTools/Private/MuT/AST.cpp
@@ -1213,12 +1213,6 @@ FImageDesc ASTOpFixed::GetImageDesc( bool returnBestOption, FGetImageDescContext
         res.m_format = EImageFormat::IF_L_UBYTE;
         break;
 
-    case OP_TYPE::IM_RASTERMESH:
-        res = GetImageDesc( op.args.ImageRasterMesh.image, returnBestOption, context );
-        res.m_size[0]=op.args.ImageRasterMesh.sizeX;
-        res.m_size[1]=op.args.ImageRasterMesh.sizeY;
-        break;
-
     case OP_TYPE::IM_DISPLACE:
         res = GetImageDesc( op.args.ImageDisplace.source, returnBestOption, context );
         break;
@@ -1720,12 +1714,6 @@ mu::Ptr<ImageSizeExpression> ASTOpFixed::GetImageSizeExpression() const
 			pRes = children[op.args.ImageInvert.base].child()->GetImageSizeExpression();
 		}
 		break;
-
-	case OP_TYPE::IM_RASTERMESH:
-		pRes->type = ImageSizeExpression::ISET_CONSTANT;
-		pRes->size[0] = op.args.ImageRasterMesh.sizeX ? op.args.ImageRasterMesh.sizeX : 256;
-		pRes->size[1] = op.args.ImageRasterMesh.sizeY ? op.args.ImageRasterMesh.sizeY : 256;
-		break;
     	
     case OP_TYPE::IM_CROP:
 		pRes->type = ImageSizeExpression::ISET_CONSTANT;
diff --git a/Engine/Plugins/Experimental/Mutable/Source/MutableTools/Private/MuT/ASTOpImageLayer.cpp b/Engine/Plugins/Experimental/Mutable/Source/MutableTools/Private/MuT/ASTOpImageLayer.cpp
--- a/Engine/Plugins/Experimental/Mutable/Source/MutableTools/Private/MuT/ASTOpImageLayer.cpp
+++ b/Engine/Plugins/Experimental/Mutable/Source/MutableTools/Private/MuT/ASTOpImageLayer.cpp
@@ -6,6 +6,7 @@
 #include "MuT/ASTOpImagePatch.h"
 #include "MuT/ASTOpImageLayerColor.h"
 #include "MuT/ASTOpImageSwizzle.h"
+#include "MuT/ASTOpImageRasterMesh.h"
 #include "MuT/ASTOpSwitch.h"
 #include "MuR/ModelPrivate.h"
 #include "MuR/RefCounted.h"
@@ -304,23 +305,23 @@ namespace mu
 
 				case OP_TYPE::IM_RASTERMESH:
 				{
-					const ASTOpFixed* MaskRaster = dynamic_cast<const ASTOpFixed*>(CurrentMask.get());
-					const ASTOpFixed* BlendRaster = dynamic_cast<const ASTOpFixed*>(CurrentBlend.get());
+					const ASTOpImageRasterMesh* MaskRaster = dynamic_cast<const ASTOpImageRasterMesh*>(CurrentMask.get());
+					const ASTOpImageRasterMesh* BlendRaster = dynamic_cast<const ASTOpImageRasterMesh*>(CurrentBlend.get());
 					if (MaskRaster && BlendRaster
 						&&
-						MaskRaster->children[MaskRaster->op.args.ImageRasterMesh.mesh].child() == BlendRaster->children[BlendRaster->op.args.ImageRasterMesh.mesh].child()
+						MaskRaster->mesh.child() == BlendRaster->mesh.child()
 						&&
-						MaskRaster->children[MaskRaster->op.args.ImageRasterMesh.projector].child() == BlendRaster->children[BlendRaster->op.args.ImageRasterMesh.projector].child()
+						MaskRaster->projector.child() == BlendRaster->projector.child()
 						&&
-						MaskRaster->children[MaskRaster->op.args.ImageRasterMesh.mask].child() == BlendRaster->children[BlendRaster->op.args.ImageRasterMesh.mask].child()
+						MaskRaster->mask.child() == BlendRaster->mask.child()
 						&&
-						MaskRaster->children[MaskRaster->op.args.ImageRasterMesh.angleFadeProperties].child() == BlendRaster->children[BlendRaster->op.args.ImageRasterMesh.angleFadeProperties].child()
+						MaskRaster->angleFadeProperties.child() == BlendRaster->angleFadeProperties.child()
 						&&
-						MaskRaster->op.args.ImageRasterMesh.blockIndex == BlendRaster->op.args.ImageRasterMesh.blockIndex
+						MaskRaster->blockIndex == BlendRaster->blockIndex
 						)
 					{
-						CurrentMask = MaskRaster->children[MaskRaster->op.args.ImageRasterMesh.image].child();
-						CurrentBlend = BlendRaster->children[BlendRaster->op.args.ImageRasterMesh.image].child();
+						CurrentMask = MaskRaster->image.child();
+						CurrentBlend = BlendRaster->image.child();
 						bMatchingAlphaExpression = true;
 					}
 					break;
diff --git a/Engine/Plugins/Experimental/Mutable/Source/MutableTools/Private/MuT/ASTOpImageLayerColor.cpp b/Engine/Plugins/Experimental/Mutable/Source/MutableTools/Private/MuT/ASTOpImageLayerColor.cpp
--- a/Engine/Plugins/Experimental/Mutable/Source/MutableTools/Private/MuT/ASTOpImageLayerColor.cpp
+++ b/Engine/Plugins/Experimental/Mutable/Source/MutableTools/Private/MuT/ASTOpImageLayerColor.cpp
@@ -4,6 +4,7 @@
 
 #include "MuT/ASTOpSwitch.h"
 #include "MuT/ASTOpImageSwizzle.h"
+#include "MuT/ASTOpImageRasterMesh.h"
 #include "Containers/Map.h"
 #include "HAL/PlatformMath.h"
 #include "MuR/ModelPrivate.h"
@@ -15,7 +16,6 @@
 namespace mu
 {
 
-
 	ASTOpImageLayerColor::ASTOpImageLayerColor()
 		: base(this)
 		, color(this)
@@ -396,12 +396,12 @@ namespace mu
 				break;
 			}
 
-			Ptr<ASTOpFixed> NewRaster = mu::Clone<ASTOpFixed>(baseAt);
+			Ptr<ASTOpImageRasterMesh> NewRaster = mu::Clone<ASTOpImageRasterMesh>(baseAt);
 
-			Ptr<ASTOp> sourceOp = NewRaster->children[NewRaster->op.args.ImageRasterMesh.image].child();
+			Ptr<ASTOp> sourceOp = NewRaster->image.child();
 			Ptr<ASTOpImageLayerColor> NewSource = mu::Clone<ASTOpImageLayerColor>(this);
 			NewSource->base = sourceOp;
-			NewRaster->SetChild(NewRaster->op.args.ImageRasterMesh.image, NewSource);
+			NewRaster->image = NewSource;
 
 			at = NewRaster;
 			break;
diff --git a/Engine/Plugins/Experimental/Mutable/Source/MutableTools/Private/MuT/ASTOpImagePixelFormat.cpp b/Engine/Plugins/Experimental/Mutable/Source/MutableTools/Private/MuT/ASTOpImagePixelFormat.cpp
--- a/Engine/Plugins/Experimental/Mutable/Source/MutableTools/Private/MuT/ASTOpImagePixelFormat.cpp
+++ b/Engine/Plugins/Experimental/Mutable/Source/MutableTools/Private/MuT/ASTOpImagePixelFormat.cpp
@@ -16,6 +16,7 @@
 #include "MuT/ASTOpImagePatch.h"
 #include "MuT/ASTOpImageLayer.h"
 #include "MuT/ASTOpImageLayerColor.h"
+#include "MuT/ASTOpImageRasterMesh.h"
 #include "MuT/ASTOpSwitch.h"
 #include "MuT/StreamsPrivate.h"
 
@@ -144,13 +145,13 @@ namespace mu
 			// This op doesn't support compressed formats
 			if (!bIsCompressedFormat
 				&&
-				dynamic_cast<const ASTOpFixed*>(sourceAt.get())->op.args.ImageRasterMesh.image)
+				dynamic_cast<const ASTOpImageRasterMesh*>(sourceAt.get())->image)
 			{
-				auto newOp = mu::Clone<ASTOpFixed>(sourceAt);
+				Ptr<ASTOpImageRasterMesh> newOp = mu::Clone<ASTOpImageRasterMesh>(sourceAt);
 
 				mu::Ptr<ASTOpImagePixelFormat> fop = mu::Clone<ASTOpImagePixelFormat>(this);
-				fop->Source = newOp->children[newOp->op.args.ImageRasterMesh.image].child();
-				newOp->SetChild(newOp->op.args.ImageRasterMesh.image, fop);
+				fop->Source = newOp->image.child();
+				newOp->image = fop;
 
 				at = newOp;
 			}
diff --git a/Engine/Plugins/Experimental/Mutable/Source/MutableTools/Private/MuT/ASTOpImageRasterMesh.cpp b/Engine/Plugins/Experimental/Mutable/Source/MutableTools/Private/MuT/ASTOpImageRasterMesh.cpp
new file mode 100644
--- /dev/null
+++ b/Engine/Plugins/Experimental/Mutable/Source/MutableTools/Private/MuT/ASTOpImageRasterMesh.cpp
@@ -0,0 +1,360 @@
+// Copyright Epic Games, Inc. All Rights Reserved.
+
+#include "MuT/ASTOpImageRasterMesh.h"
+
+#include "MuT/ASTOpSwitch.h"
+#include "MuT/ASTOpImageSwizzle.h"
+#include "MuT/ASTOpMeshMorph.h"
+#include "MuT/ASTOpConditional.h"
+#include "MuT/ASTOpSwitch.h"
+#include "Containers/Map.h"
+#include "HAL/PlatformMath.h"
+#include "MuR/ModelPrivate.h"
+#include "MuR/RefCounted.h"
+#include "MuR/Types.h"
+#include "MuT/StreamsPrivate.h"
+
+
+namespace mu
+{
+
+
+	ASTOpImageRasterMesh::ASTOpImageRasterMesh()
+		: mesh(this)
+		, image(this)
+		, angleFadeProperties(this)
+		, mask(this)
+		, projector(this)
+	{
+		blockIndex = -1;
+		sizeX = 0;
+		sizeY = 0;
+		bIsRGBFadingEnabled = 1;
+		bIsAlphaFadingEnabled = 1;
+	}
+
+
+	//-------------------------------------------------------------------------------------------------
+	ASTOpImageRasterMesh::~ASTOpImageRasterMesh()
+	{
+		// Explicit call needed to avoid recursive destruction
+		ASTOp::RemoveChildren();
+	}
+
+
+	//-------------------------------------------------------------------------------------------------
+	bool ASTOpImageRasterMesh::IsEqual(const ASTOp& InOtherUntyped) const
+	{
+		if (const ASTOpImageRasterMesh* Other = dynamic_cast<const ASTOpImageRasterMesh*>(&InOtherUntyped))
+		{
+			return mesh == Other->mesh &&
+				image == Other->image &&
+				angleFadeProperties == Other->angleFadeProperties &&
+				mask == Other->mask &&
+				projector == Other->projector &&
+				blockIndex == Other->blockIndex &&
+				sizeX == Other->sizeX &&
+				sizeY == Other->sizeY &&
+				bIsRGBFadingEnabled == Other->bIsRGBFadingEnabled &&
+				bIsAlphaFadingEnabled == Other->bIsAlphaFadingEnabled;
+		}
+		return false;
+	}
+
+
+	//-------------------------------------------------------------------------------------------------
+	uint64 ASTOpImageRasterMesh::Hash() const
+	{
+		uint64 res = std::hash<OP_TYPE>()(GetOpType());
+		hash_combine(res, mesh.child().get());
+		hash_combine(res, image.child().get());
+		hash_combine(res, angleFadeProperties.child().get());
+		hash_combine(res, mask.child().get());
+		hash_combine(res, projector.child().get());
+		hash_combine(res, blockIndex);
+		return res;
+	}
+
+
+	//-------------------------------------------------------------------------------------------------
+	mu::Ptr<ASTOp> ASTOpImageRasterMesh::Clone(MapChildFuncRef mapChild) const
+	{
+		Ptr<ASTOpImageRasterMesh> n = new ASTOpImageRasterMesh();
+		n->mesh = mapChild(mesh.child());
+		n->image = mapChild(image.child());
+		n->angleFadeProperties = mapChild(angleFadeProperties.child());
+		n->mask = mapChild(mask.child());
+		n->projector = mapChild(projector.child());
+		n->blockIndex = blockIndex;
+		n->sizeX = sizeX;
+		n->sizeY = sizeY;
+		n->bIsRGBFadingEnabled = bIsRGBFadingEnabled;
+		n->bIsAlphaFadingEnabled = bIsAlphaFadingEnabled;
+		return n;
+	}
+
+
+	//-------------------------------------------------------------------------------------------------
+	void ASTOpImageRasterMesh::ForEachChild(const TFunctionRef<void(ASTChild&)> f)
+	{
+		f(mesh);
+		f(image);
+		f(angleFadeProperties);
+		f(mask);
+		f(projector);
+	}
+
+
+	//-------------------------------------------------------------------------------------------------
+	void ASTOpImageRasterMesh::Link(FProgram& program, const FLinkerOptions*)
+	{
+		// Already linked?
+		if (!linkedAddress)
+		{
+			OP::ImageRasterMeshArgs args;
+			FMemory::Memzero(&args, sizeof(args));
+
+			args.blockIndex = blockIndex;
+			args.sizeX = sizeX;
+			args.sizeY = sizeY;
+			args.bIsRGBFadingEnabled = bIsRGBFadingEnabled;
+			args.bIsAlphaFadingEnabled = bIsAlphaFadingEnabled;
+
+			if (mesh) args.mesh = mesh->linkedAddress;
+			if (image) args.image = image->linkedAddress;
+			if (angleFadeProperties) args.angleFadeProperties = angleFadeProperties->linkedAddress;
+			if (mask) args.mask = mask->linkedAddress;
+			if (projector) args.projector = projector->linkedAddress;
+
+			linkedAddress = (OP::ADDRESS)program.m_opAddress.Num();
+			program.m_opAddress.Add((uint32)program.m_byteCode.Num());
+			AppendCode(program.m_byteCode, GetOpType());
+			AppendCode(program.m_byteCode, args);
+		}
+	}
+
+
+	//-------------------------------------------------------------------------------------------------
+	FImageDesc ASTOpImageRasterMesh::GetImageDesc(bool returnBestOption, FGetImageDescContext* context) const 
+	{
+		FImageDesc res;
+
+		// Local context in case it is necessary
+		FGetImageDescContext localContext;
+		if (!context)
+		{
+			context = &localContext;
+		}
+		else
+		{
+			// Cached result?
+			FImageDesc* PtrValue = context->m_results.Find(this);
+			if (PtrValue)
+			{
+				return *PtrValue;
+			}
+		}
+
+		// Actual work
+		if (image)
+		{
+			res = image->GetImageDesc( returnBestOption, context);
+			res.m_size[0] = sizeX;
+			res.m_size[1] = sizeY;
+		}
+
+
+		// Cache the result
+		if (context)
+		{
+			context->m_results.Add(this, res);
+		}
+
+		return res;
+	}
+
+
+	//-------------------------------------------------------------------------------------------------
+	Ptr<ImageSizeExpression> ASTOpImageRasterMesh::GetImageSizeExpression() const
+	{
+		Ptr<ImageSizeExpression> pRes = new ImageSizeExpression;
+		pRes->type = ImageSizeExpression::ISET_CONSTANT;
+		pRes->size[0] = sizeX ? sizeX : 256;
+		pRes->size[1] = sizeY ? sizeY : 256;
+
+		return nullptr;
+	}
+
+
+	//---------------------------------------------------------------------------------------------
+	Ptr<ASTOp> ASTOpImageRasterMesh::OptimiseSemantic(const FModelOptimizationOptions& options) const
+	{
+		Ptr<ASTOp> at;
+
+		// TODO
+
+		return at;
+	}
+
+
+	//-------------------------------------------------------------------------------------------------
+	Ptr<ASTOp> ASTOpImageRasterMesh::OptimiseSink(const FModelOptimizationOptions& options, FOptimizeSinkContext& context) const
+	{
+		Ptr<ASTOp> at;
+
+		Ptr<ASTOp> OriginalAt = at;
+		Ptr<ASTOp> sourceAt = mesh.child();
+		Ptr<ASTOp> imageAt = image.child();
+
+		OP_TYPE sourceType = sourceAt->GetOpType();
+		switch (sourceType)
+		{
+
+		case OP_TYPE::ME_PROJECT:
+		{
+			// If we are rastering just the UV layout (to create a mask) we don't care about
+			// mesh project operations, which modify only the positions.
+			// This optimisation helps with states removing fake dependencies on projector
+			// parameters that may be runtime.
+			if (!imageAt)
+			{
+				// We remove the project from the raster children
+				// \todo: maybe this clone is not necessary
+				const ASTOpFixed* typedSource = dynamic_cast<const ASTOpFixed*>(sourceAt.get());
+				Ptr<ASTOpImageRasterMesh> nop = mu::Clone<ASTOpImageRasterMesh>(this);
+				nop->mesh = typedSource->children[typedSource->op.args.MeshProject.mesh].child();
+				at = nop;
+			}
+			break;
+		}
+
+		case OP_TYPE::ME_INTERPOLATE:
+		{
+			auto typedSource = dynamic_cast<const ASTOpFixed*>(sourceAt.get());
+			Ptr<ASTOpImageRasterMesh> rasterOp = mu::Clone<ASTOpImageRasterMesh>(this);
+			rasterOp->mesh = typedSource->children[typedSource->op.args.MeshInterpolate.base].child();
+			at = rasterOp;
+			break;
+		}
+
+		case OP_TYPE::ME_MORPH2:
+		{
+			const ASTOpMeshMorph* typedSource = dynamic_cast<const ASTOpMeshMorph*>(sourceAt.get());
+			Ptr<ASTOpImageRasterMesh> rasterOp = mu::Clone<ASTOpImageRasterMesh>(this);
+			rasterOp->mesh = typedSource->Base.child();
+			at = rasterOp;
+			break;
+		}
+
+		case OP_TYPE::ME_CONDITIONAL:
+		{
+			auto nop = mu::Clone<ASTOpConditional>(sourceAt.get());
+			nop->type = OP_TYPE::IM_CONDITIONAL;
+
+			Ptr<ASTOpImageRasterMesh> aOp = mu::Clone<ASTOpImageRasterMesh>(this);
+			aOp->mesh = nop->yes.child();
+			nop->yes = aOp;
+
+			Ptr<ASTOpImageRasterMesh> bOp = mu::Clone<ASTOpImageRasterMesh>(this);
+			bOp->mesh = nop->no.child();
+			nop->no = bOp;
+
+			at = nop;
+			break;
+		}
+
+		case OP_TYPE::ME_SWITCH:
+		{
+			// Make an image for every path
+			auto nop = mu::Clone<ASTOpSwitch>(sourceAt.get());
+			nop->type = OP_TYPE::IM_SWITCH;
+
+			if (nop->def)
+			{
+				Ptr<ASTOpImageRasterMesh> defOp = mu::Clone<ASTOpImageRasterMesh>(this);
+				defOp->mesh = nop->def.child();
+				nop->def = defOp;
+			}
+
+			// We need to copy the options because we change them
+			for (size_t o = 0; o < nop->cases.Num(); ++o)
+			{
+				if (nop->cases[o].branch)
+				{
+					Ptr<ASTOpImageRasterMesh> bOp = mu::Clone<ASTOpImageRasterMesh>(this);
+					bOp->mesh = nop->cases[o].branch.child();
+					nop->cases[o].branch = bOp;
+				}
+			}
+
+			at = nop;
+			break;
+		}
+
+		default:
+			break;
+		}
+
+		// If we didn't optimize the mesh child, try to optimize the image child.
+		if (OriginalAt == at && imageAt)
+		{
+			OP_TYPE imageType = imageAt->GetOpType();
+			switch (imageType)
+			{
+
+				// TODO: Implement for image conditionals.
+				//case OP_TYPE::ME_CONDITIONAL:
+				//{
+				//	auto nop = mu::Clone<ASTOpConditional>(sourceAt.get());
+				//	nop->type = OP_TYPE::IM_CONDITIONAL;
+
+				//	Ptr<ASTOpFixed> aOp = mu::Clone<ASTOpFixed>(this);
+				//	aOp->SetChild(aOp->op.args.ImageRasterMesh.mesh, nop->yes);
+				//	nop->yes = aOp;
+
+				//	Ptr<ASTOpFixed> bOp = mu::Clone<ASTOpFixed>(this);
+				//	bOp->SetChild(bOp->op.args.ImageRasterMesh.mesh, nop->no);
+				//	nop->no = bOp;
+
+				//	at = nop;
+				//	break;
+				//}
+
+			case OP_TYPE::IM_SWITCH:
+			{
+				// TODO: Do this only if the projector is constant?
+
+				// Make a project for every path
+				auto nop = mu::Clone<ASTOpSwitch>(imageAt.get());
+
+				if (nop->def)
+				{
+					Ptr<ASTOpImageRasterMesh> defOp = mu::Clone<ASTOpImageRasterMesh>(this);
+					defOp->image = nop->def.child();
+					nop->def = defOp;
+				}
+
+				// We need to copy the options because we change them
+				for (size_t o = 0; o < nop->cases.Num(); ++o)
+				{
+					if (nop->cases[o].branch)
+					{
+						Ptr<ASTOpImageRasterMesh> bOp = mu::Clone<ASTOpImageRasterMesh>(this);
+						bOp->image = nop->cases[o].branch.child();
+						nop->cases[o].branch = bOp;
+					}
+				}
+
+				at = nop;
+				break;
+			}
+
+			default:
+				break;
+			}
+		}
+
+		return at;
+	}
+
+}
diff --git a/Engine/Plugins/Experimental/Mutable/Source/MutableTools/Private/MuT/ASTOpImageRasterMesh.h b/Engine/Plugins/Experimental/Mutable/Source/MutableTools/Private/MuT/ASTOpImageRasterMesh.h
new file mode 100644
--- /dev/null
+++ b/Engine/Plugins/Experimental/Mutable/Source/MutableTools/Private/MuT/ASTOpImageRasterMesh.h
@@ -0,0 +1,51 @@
+// Copyright Epic Games, Inc. All Rights Reserved.
+
+#pragma once
+
+#include "HAL/Platform.h"
+#include "MuR/Image.h"
+#include "MuR/Operations.h"
+#include "MuR/Ptr.h"
+#include "MuT/AST.h"
+
+
+namespace mu
+{
+	struct FProgram;
+
+	class ASTOpImageRasterMesh : public ASTOp
+	{
+	public:
+
+		ASTChild mesh;
+		ASTChild image;
+		ASTChild angleFadeProperties;
+		ASTChild mask;
+		ASTChild projector;
+
+		int32 blockIndex;
+		uint16 sizeX, sizeY;
+		uint8 bIsRGBFadingEnabled : 1;
+		uint8 bIsAlphaFadingEnabled : 1;
+
+	public:
+
+		ASTOpImageRasterMesh();
+		ASTOpImageRasterMesh(const ASTOpImageRasterMesh&) = delete;
+		~ASTOpImageRasterMesh();
+
+		OP_TYPE GetOpType() const override { return OP_TYPE::IM_RASTERMESH; }
+		uint64 Hash() const override;
+		bool IsEqual(const ASTOp& otherUntyped) const override;
+		Ptr<ASTOp> Clone(MapChildFuncRef mapChild) const override;
+		void ForEachChild(const TFunctionRef<void(ASTChild&)>) override;
+		void Link(FProgram& program, const FLinkerOptions* Options) override;
+		FImageDesc GetImageDesc(bool returnBestOption, FGetImageDescContext* context) const override;
+		Ptr<ImageSizeExpression> GetImageSizeExpression() const override;
+		Ptr<ASTOp> OptimiseSemantic(const FModelOptimizationOptions& options) const;
+		Ptr<ASTOp> OptimiseSink(const FModelOptimizationOptions&, FOptimizeSinkContext&) const;
+
+	};
+
+}
+
diff --git a/Engine/Plugins/Experimental/Mutable/Source/MutableTools/Private/MuT/ASTOpImageSwizzle.cpp b/Engine/Plugins/Experimental/Mutable/Source/MutableTools/Private/MuT/ASTOpImageSwizzle.cpp
--- a/Engine/Plugins/Experimental/Mutable/Source/MutableTools/Private/MuT/ASTOpImageSwizzle.cpp
+++ b/Engine/Plugins/Experimental/Mutable/Source/MutableTools/Private/MuT/ASTOpImageSwizzle.cpp
@@ -18,6 +18,7 @@
 #include "MuT/ASTOpImageMultiLayer.h"
 #include "MuT/ASTOpImageLayer.h"
 #include "MuT/ASTOpImageLayerColor.h"
+#include "MuT/ASTOpImageRasterMesh.h"
 #include "MuT/ASTOpSwitch.h"
 #include "MuT/StreamsPrivate.h"
 
@@ -349,10 +350,10 @@ namespace mu
 
 			case OP_TYPE::IM_RASTERMESH:
 			{
-				Ptr<ASTOpFixed> NewRaster = mu::Clone<ASTOpFixed>(channelSourceAt);
+				Ptr<ASTOpImageRasterMesh> NewRaster = mu::Clone<ASTOpImageRasterMesh>(channelSourceAt);
 				Ptr<ASTOpImageSwizzle> NewSwizzle = mu::Clone<ASTOpImageSwizzle>(this);
-				ReplaceAllSources(NewSwizzle, NewRaster->children[NewRaster->op.args.ImageRasterMesh.image].child());
-				NewRaster->SetChild(NewRaster->op.args.ImageRasterMesh.image, NewSwizzle);
+				ReplaceAllSources(NewSwizzle, NewRaster->image.child());
+				NewRaster->image = NewSwizzle;
 				at = NewRaster;
 				break;
 			}
@@ -622,7 +623,7 @@ namespace mu
 			// Swizzle down compatible raster meshes.
 			if (!at && sourceType == OP_TYPE::IM_RASTERMESH)
 			{
-				const ASTOpFixed* FirstRasterMesh = dynamic_cast<const ASTOpFixed*>(Sources[0].child().get());
+				const ASTOpImageRasterMesh* FirstRasterMesh = dynamic_cast<const ASTOpImageRasterMesh*>(Sources[0].child().get());
 				check(FirstRasterMesh);
 
 				bool bAreAllRasterMeshesCompatible = true;
@@ -630,13 +631,20 @@ namespace mu
 				{
 					if (Sources[c])
 					{
-						const ASTOpFixed* Typed = dynamic_cast<const ASTOpFixed*>(Sources[c].child().get());
+						const ASTOpImageRasterMesh* Typed = dynamic_cast<const ASTOpImageRasterMesh*>(Sources[c].child().get());
 						check(Typed);
-						// Compare all args but the source image
-						OP::ImageRasterMeshArgs ArgCopy = FirstRasterMesh->op.args.ImageRasterMesh;
-						ArgCopy.image = Typed->op.args.ImageRasterMesh.image;
 
-						if (FMemory::Memcmp(&ArgCopy, &Typed->op.args.ImageRasterMesh, sizeof(OP::ImageRasterMeshArgs)) != 0)
+						// Compare all args but the source image
+						if (Typed->mesh.child() != FirstRasterMesh->mesh.child()
+							|| Typed->angleFadeProperties.child() != FirstRasterMesh->angleFadeProperties.child()
+							|| Typed->mask.child() != FirstRasterMesh->mask.child()
+							|| Typed->projector.child() != FirstRasterMesh->projector.child()
+							|| Typed->blockIndex != FirstRasterMesh->blockIndex
+							|| Typed->sizeX != FirstRasterMesh->sizeX
+							|| Typed->sizeY != FirstRasterMesh->sizeY
+							|| Typed->bIsRGBFadingEnabled != FirstRasterMesh->bIsRGBFadingEnabled
+							|| Typed->bIsAlphaFadingEnabled != FirstRasterMesh->bIsAlphaFadingEnabled
+							)
 						{
 							bAreAllRasterMeshesCompatible = false;
 							break;
@@ -647,19 +655,19 @@ namespace mu
 				if (bAreAllRasterMeshesCompatible)
 				{
 					// Move the swizzle down all the paths
-					Ptr<ASTOpFixed> NewRaster = mu::Clone<ASTOpFixed>(FirstRasterMesh);
+					Ptr<ASTOpImageRasterMesh> NewRaster = mu::Clone<ASTOpImageRasterMesh>(FirstRasterMesh);
 
 					Ptr<ASTOpImageSwizzle> NewSwizzle = mu::Clone<ASTOpImageSwizzle>(this);
 					for (int c = 0; c < MUTABLE_OP_MAX_SWIZZLE_CHANNELS; ++c)
 					{
-						const ASTOpFixed* ChannelRaster = dynamic_cast<const ASTOpFixed*>(Sources[c].child().get());
+						const ASTOpImageRasterMesh* ChannelRaster = dynamic_cast<const ASTOpImageRasterMesh*>(Sources[c].child().get());
 						if (ChannelRaster)
 						{
-							NewSwizzle->Sources[c] = ChannelRaster->children[ChannelRaster->op.args.ImageRasterMesh.image].child();
+							NewSwizzle->Sources[c] = ChannelRaster->image.child();
 						}
 					}
 
-					NewRaster->SetChild(NewRaster->op.args.ImageRasterMesh.image, NewSwizzle);
+					NewRaster->image = NewSwizzle;
 
 					at = NewRaster;
 				}
diff --git a/Engine/Plugins/Experimental/Mutable/Source/MutableTools/Private/MuT/CodeGenerator_Image.cpp b/Engine/Plugins/Experimental/Mutable/Source/MutableTools/Private/MuT/CodeGenerator_Image.cpp
--- a/Engine/Plugins/Experimental/Mutable/Source/MutableTools/Private/MuT/CodeGenerator_Image.cpp
+++ b/Engine/Plugins/Experimental/Mutable/Source/MutableTools/Private/MuT/CodeGenerator_Image.cpp
@@ -29,6 +29,7 @@
 #include "MuT/ASTOpImageTransform.h"
 #include "MuT/ASTOpImageMakeGrowMap.h"
 #include "MuT/ASTOpImageSwizzle.h"
+#include "MuT/ASTOpImageRasterMesh.h"
 #include "MuT/ASTOpMeshExtractLayoutBlocks.h"
 #include "MuT/ASTOpMeshFormat.h"
 #include "MuT/ASTOpParameter.h"
@@ -1526,10 +1527,9 @@ namespace mu
 
         // Image raster operation
         //------------------------------
-        Ptr<ASTOpFixed> op = new ASTOpFixed();
-        op->op.type = OP_TYPE::IM_RASTERMESH;
-        op->SetChild( op->op.args.ImageRasterMesh.mesh, lastMeshOp);
-        op->SetChild( op->op.args.ImageRasterMesh.projector, projectorResult.op);
+        Ptr<ASTOpImageRasterMesh> op = new ASTOpImageRasterMesh();
+        op->mesh = lastMeshOp;
+        op->projector = projectorResult.op;
 
         // Image
         if ( node.m_pImage )
@@ -1547,7 +1547,7 @@ namespace mu
             m_imageState.Add( newState );
 
             // Generate
-            op->SetChild( op->op.args.ImageRasterMesh.image, Generate( node.m_pImage.get() ) );
+            op->image = Generate( node.m_pImage.get() );
 
             // Restore rect
             m_imageState.Pop();
@@ -1556,55 +1556,54 @@ namespace mu
         else
         {
             // This argument is required
-            op->SetChild( op->op.args.ImageRasterMesh.image,
-                    GenerateMissingImageCode( "Projector image", EImageFormat::IF_RGB_UBYTE, node.m_errorContext ) );
+            op->image = GenerateMissingImageCode( "Projector image", EImageFormat::IF_RGB_UBYTE, node.m_errorContext );
         }
 
         // Image size, from the current block being generated
-        op->op.args.ImageRasterMesh.sizeX = (uint16)m_imageState.Last().m_imageRect.size[0];
-        op->op.args.ImageRasterMesh.sizeY = (uint16)m_imageState.Last().m_imageRect.size[1];
+        op->sizeX = (uint16)m_imageState.Last().m_imageRect.size[0];
+        op->sizeY = (uint16)m_imageState.Last().m_imageRect.size[1];
 		// \TODO: Review naming of arg
-        op->op.args.ImageRasterMesh.blockIndex = GeneratedLayoutBlockId;
+        op->blockIndex = GeneratedLayoutBlockId;
 
-        // Fading properties are optional, and stored in a colour
-        if (node.m_pAngleFadeStart||node.m_pAngleFadeEnd)
-        {
-            NodeScalarConstantPtr pDefaultFade = new NodeScalarConstant();
-            pDefaultFade->SetValue( 180.0f );
+		op->bIsRGBFadingEnabled = node.bIsRGBFadingEnabled;
+		op->bIsAlphaFadingEnabled = node.bIsAlphaFadingEnabled;
+
+		// Fading angles are optional, and stored in a colour. If one exists, we generate both.
+		if (node.m_pAngleFadeStart || node.m_pAngleFadeEnd)
+		{
+			NodeScalarConstantPtr pDefaultFade = new NodeScalarConstant();
+			pDefaultFade->SetValue(180.0f);
 
-            NodeColourFromScalarsPtr pPropsNode = new NodeColourFromScalars();
+			NodeColourFromScalarsPtr pPropsNode = new NodeColourFromScalars();
 
-            if (node.m_pAngleFadeStart) pPropsNode->SetX(node.m_pAngleFadeStart);
-            else pPropsNode->SetX(pDefaultFade);
+			if (node.m_pAngleFadeStart) pPropsNode->SetX(node.m_pAngleFadeStart);
+			else pPropsNode->SetX(pDefaultFade);
 
-            if (node.m_pAngleFadeEnd) pPropsNode->SetY(node.m_pAngleFadeEnd);
-            else pPropsNode->SetY(pDefaultFade);
+			if (node.m_pAngleFadeEnd) pPropsNode->SetY(node.m_pAngleFadeEnd);
+			else pPropsNode->SetY(pDefaultFade);
 
-            op->SetChild( op->op.args.ImageRasterMesh.angleFadeProperties, Generate( pPropsNode ) );
-        }
+			op->angleFadeProperties = Generate(pPropsNode);
+		}
 
         // Target mask
         if ( node.m_pMask )
         {
             auto mask = Generate( node.m_pMask.get() );
             mask = GenerateImageFormat( mask, EImageFormat::IF_L_UBYTE );
-            op->SetChild( op->op.args.ImageRasterMesh.mask, GenerateImageSize
-                    ( mask,
+            op->mask = GenerateImageSize( mask,
                       FImageSize( (uint16)m_imageState.Last().m_imageRect.size[0],
-                                  (uint16)m_imageState.Last().m_imageRect.size[1]) )
-                              );
+                                  (uint16)m_imageState.Last().m_imageRect.size[1]) );
         }
 
         // Seam correction operations
         //------------------------------
-        Ptr<ASTOpFixed> rasterop = new ASTOpFixed();
-        rasterop->op.type = OP_TYPE::IM_RASTERMESH;
-        rasterop->SetChild( rasterop->op.args.ImageRasterMesh.mesh, op->children[op->op.args.ImageRasterMesh.mesh].child() );
-        rasterop->op.args.ImageRasterMesh.image = 0;
-        rasterop->op.args.ImageRasterMesh.mask = 0;
-        rasterop->op.args.ImageRasterMesh.blockIndex = op->op.args.ImageRasterMesh.blockIndex;
-        rasterop->op.args.ImageRasterMesh.sizeX = op->op.args.ImageRasterMesh.sizeX;
-        rasterop->op.args.ImageRasterMesh.sizeY = op->op.args.ImageRasterMesh.sizeY;
+        Ptr<ASTOpImageRasterMesh> rasterop = new ASTOpImageRasterMesh();
+        rasterop->mesh = op->mesh.child();
+        rasterop->image = 0;
+        rasterop->mask = 0;
+        rasterop->blockIndex = op->blockIndex;
+        rasterop->sizeX = op->sizeX;
+        rasterop->sizeY = op->sizeY;
 
         Ptr<ASTOpImageMakeGrowMap> MakeGrowMapOp = new ASTOpImageMakeGrowMap();
 		MakeGrowMapOp->Mask = rasterop;
diff --git a/Engine/Plugins/Experimental/Mutable/Source/MutableTools/Private/MuT/CodeOptimiser.cpp b/Engine/Plugins/Experimental/Mutable/Source/MutableTools/Private/MuT/CodeOptimiser.cpp
--- a/Engine/Plugins/Experimental/Mutable/Source/MutableTools/Private/MuT/CodeOptimiser.cpp
+++ b/Engine/Plugins/Experimental/Mutable/Source/MutableTools/Private/MuT/CodeOptimiser.cpp
@@ -19,6 +19,7 @@
 #include "MuT/ASTOpMeshBindShape.h"
 #include "MuT/ASTOpMeshClipMorphPlane.h"
 #include "MuT/ASTOpMeshApplyPose.h"
+#include "MuT/ASTOpImageRasterMesh.h"
 
 #include <unordered_set>
 
@@ -702,8 +703,8 @@ namespace mu
 	//                else
 					if (n->GetOpType()==OP_TYPE::IM_RASTERMESH)
 				{
-					const auto& args = dynamic_cast<const ASTOpFixed*>(n.get())->op.args.ImageRasterMesh;
-					if ( !args.image )
+					const ASTOpImageRasterMesh* Raster = dynamic_cast<const ASTOpImageRasterMesh*>(n.get());
+					if ( !Raster->image )
 					{
 						specialCase = true;
 					}
diff --git a/Engine/Plugins/Experimental/Mutable/Source/MutableTools/Private/MuT/CodeOptimiser_Semantic.cpp b/Engine/Plugins/Experimental/Mutable/Source/MutableTools/Private/MuT/CodeOptimiser_Semantic.cpp
--- a/Engine/Plugins/Experimental/Mutable/Source/MutableTools/Private/MuT/CodeOptimiser_Semantic.cpp
+++ b/Engine/Plugins/Experimental/Mutable/Source/MutableTools/Private/MuT/CodeOptimiser_Semantic.cpp
@@ -23,6 +23,7 @@
 #include "MuT/ASTOpImageMakeGrowMap.h"
 #include "MuT/ASTOpImageSwizzle.h"
 #include "MuT/ASTOpImageTransform.h"
+#include "MuT/ASTOpImageRasterMesh.h"
 #include "MuT/ASTOpMeshMorph.h"
 #include "MuT/ASTOpSwitch.h"
 #include "MuT/CodeOptimiser.h"
@@ -716,176 +717,14 @@ namespace mu
         //-----------------------------------------------------------------------------------------
 		case OP_TYPE::IM_SWIZZLE:
 		case OP_TYPE::LA_REMOVEBLOCKS:
-        {
+		case OP_TYPE::IM_RASTERMESH:
+		{
 			// Moved to their own operation class
 			check(false);
             break;
         }
 
 
-        //-----------------------------------------------------------------------------------------
-        //-----------------------------------------------------------------------------------------
-        //-----------------------------------------------------------------------------------------
-        case OP_TYPE::IM_RASTERMESH:
-        {
-			Ptr<ASTOp> OriginalAt = at;
-			Ptr<ASTOp> sourceAt = children[op.args.ImageRasterMesh.mesh].child();
-			Ptr<ASTOp> imageAt = children[op.args.ImageRasterMesh.image].child();
-
-            OP_TYPE sourceType = sourceAt->GetOpType();
-            switch ( sourceType )
-            {
-
-            case OP_TYPE::ME_PROJECT:
-            {
-                // If we are rastering just the UV layout (to create a mask) we don't care about
-                // mesh project operations, which modify only the positions.
-                // This optimisation helps with states removing fake dependencies on projector
-                // parameters that may be runtime.
-                if (!imageAt)
-                {
-                    // We remove the project from the raster children
-                    // \todo: maybe this clone is not necessary
-					const ASTOpFixed* typedSource = dynamic_cast<const ASTOpFixed*>(sourceAt.get());
-                    Ptr<ASTOpFixed> nop = mu::Clone<ASTOpFixed>(this);
-                    nop->SetChild( nop->op.args.ImageRasterMesh.mesh,
-                                   typedSource->children[typedSource->op.args.MeshProject.mesh]);
-                    at = nop;
-                }
-                break;
-            }
-
-            case OP_TYPE::ME_INTERPOLATE:
-            {
-                auto typedSource = dynamic_cast<const ASTOpFixed*>(sourceAt.get());
-				Ptr<ASTOpFixed> rasterOp = mu::Clone<ASTOpFixed>(this);
-                rasterOp->SetChild( rasterOp->op.args.ImageRasterMesh.mesh,
-                        typedSource->children[typedSource->op.args.MeshInterpolate.base] );
-                at = rasterOp;
-                break;
-            }
-
-            case OP_TYPE::ME_MORPH2:
-            {
-				const ASTOpMeshMorph* typedSource = dynamic_cast<const ASTOpMeshMorph*>(sourceAt.get());
-				Ptr<ASTOpFixed> rasterOp = mu::Clone<ASTOpFixed>(this);
-                rasterOp->SetChild( rasterOp->op.args.ImageRasterMesh.mesh, typedSource->Base );
-                at = rasterOp;
-                break;
-            }
-
-            case OP_TYPE::ME_CONDITIONAL:
-            {
-                auto nop = mu::Clone<ASTOpConditional>(sourceAt.get());
-                nop->type = OP_TYPE::IM_CONDITIONAL;
-
-				Ptr<ASTOpFixed> aOp = mu::Clone<ASTOpFixed>(this);
-                aOp->SetChild( aOp->op.args.ImageRasterMesh.mesh, nop->yes );
-                nop->yes = aOp;
-
-				Ptr<ASTOpFixed> bOp = mu::Clone<ASTOpFixed>(this);
-                bOp->SetChild( bOp->op.args.ImageRasterMesh.mesh, nop->no );
-                nop->no = bOp;
-
-                at = nop;
-                break;
-            }
-
-            case OP_TYPE::ME_SWITCH:
-            {
-                // Make an image for every path
-                auto nop = mu::Clone<ASTOpSwitch>(sourceAt.get());
-                nop->type = OP_TYPE::IM_SWITCH;
-
-                if (nop->def)
-                {
-					Ptr<ASTOpFixed> defOp = mu::Clone<ASTOpFixed>(this);
-                    defOp->SetChild( defOp->op.args.ImageRasterMesh.mesh, nop->def );
-                    nop->def = defOp;
-                }
-
-                // We need to copy the options because we change them
-                for ( size_t o=0; o<nop->cases.Num(); ++o )
-                {
-                    if ( nop->cases[o].branch )
-                    {
-						Ptr<ASTOpFixed> bOp = mu::Clone<ASTOpFixed>(this);
-                        bOp->SetChild( bOp->op.args.ImageRasterMesh.mesh, nop->cases[o].branch );
-                        nop->cases[o].branch = bOp;
-                    }
-                }
-
-                at = nop;
-                break;
-            }
-
-            default:
-                break;
-            }
-
-			// If we didn't optimize the mesh child, try to optimize the image child.
-			if (OriginalAt == at && imageAt)
-			{
-				OP_TYPE imageType = imageAt->GetOpType();
-				switch (imageType)
-				{
-
-				// TODO: Implement for image conditionals.
-				//case OP_TYPE::ME_CONDITIONAL:
-				//{
-				//	auto nop = mu::Clone<ASTOpConditional>(sourceAt.get());
-				//	nop->type = OP_TYPE::IM_CONDITIONAL;
-
-				//	Ptr<ASTOpFixed> aOp = mu::Clone<ASTOpFixed>(this);
-				//	aOp->SetChild(aOp->op.args.ImageRasterMesh.mesh, nop->yes);
-				//	nop->yes = aOp;
-
-				//	Ptr<ASTOpFixed> bOp = mu::Clone<ASTOpFixed>(this);
-				//	bOp->SetChild(bOp->op.args.ImageRasterMesh.mesh, nop->no);
-				//	nop->no = bOp;
-
-				//	at = nop;
-				//	break;
-				//}
-
-				case OP_TYPE::IM_SWITCH:
-				{
-					// TODO: Do this only if the projector is constant?
-					
-					// Make a project for every path
-					auto nop = mu::Clone<ASTOpSwitch>(imageAt.get());
-
-					if (nop->def)
-					{
-						Ptr<ASTOpFixed> defOp = mu::Clone<ASTOpFixed>(this);
-						defOp->SetChild(defOp->op.args.ImageRasterMesh.image, nop->def);
-						nop->def = defOp;
-					}
-
-					// We need to copy the options because we change them
-					for (size_t o = 0; o < nop->cases.Num(); ++o)
-					{
-						if (nop->cases[o].branch)
-						{
-							Ptr<ASTOpFixed> bOp = mu::Clone<ASTOpFixed>(this);
-							bOp->SetChild(bOp->op.args.ImageRasterMesh.image, nop->cases[o].branch);
-							nop->cases[o].branch = bOp;
-						}
-					}
-
-					at = nop;
-					break;
-				}
-
-				default:
-					break;
-				}
-			}
-
-            break;
-        }
-
-
         //-----------------------------------------------------------------------------------------
         //-----------------------------------------------------------------------------------------
         //-----------------------------------------------------------------------------------------
@@ -1358,18 +1197,18 @@ namespace mu
 
 			case OP_TYPE::IM_RASTERMESH:
 			{
-				Ptr<ASTOpFixed> nop = mu::Clone<ASTOpFixed>(at);
-				Ptr<ASTOp> maskOp = nop->children[nop->op.args.ImageRasterMesh.mask].child();
-				nop->SetChild(nop->op.args.ImageRasterMesh.mask, Visit(maskOp, currentSinkingOp));
+				Ptr<ASTOpImageRasterMesh> nop = mu::Clone<ASTOpImageRasterMesh>(at);
+				Ptr<ASTOp> maskOp = nop->mask.child();
+				nop->mask = Visit(maskOp, currentSinkingOp);
 
 				// Resize the image to project as well, assuming that since the target has a different resolution
 				// it make sense for the source image to have a similar resize.
 				// Actually, don't do it because the LODBias will be applied separetely at graph generation time.
-				//auto imageOp = nop->children[nop->op.args.ImageRasterMesh.image].child();
-				//nop->SetChild(nop->op.args.ImageRasterMesh.image, Visit(imageOp, currentSinkingOp));
+				//auto imageOp = nop->image.child();
+				//nop->image = Visit(imageOp, currentSinkingOp);
 
-				nop->op.args.ImageRasterMesh.sizeX = uint16(nop->op.args.ImageRasterMesh.sizeX * scaleX + 0.5f);
-				nop->op.args.ImageRasterMesh.sizeY = uint16(nop->op.args.ImageRasterMesh.sizeY * scaleY + 0.5f);
+				nop->sizeX = uint16(nop->sizeX * scaleX + 0.5f);
+				nop->sizeY = uint16(nop->sizeY * scaleY + 0.5f);
 				newAt = nop;
 				break;
 			}
@@ -1504,20 +1343,20 @@ namespace mu
 
             case OP_TYPE::IM_RASTERMESH:
             {
-                Ptr<ASTOpFixed> newOp = mu::Clone<ASTOpFixed>(sourceAt);
+                Ptr<ASTOpImageRasterMesh> newOp = mu::Clone<ASTOpImageRasterMesh>(sourceAt);
 
                 //if ( newOp->op.args.ImageRasterMesh.sizeX != op.args.ImageResize.size[0]
                 //     ||
                 //     newOp->op.args.ImageRasterMesh.sizeY != op.args.ImageResize.size[1] )
                 {
-                    newOp->op.args.ImageRasterMesh.sizeX = op.args.ImageResize.size[0];
-                    newOp->op.args.ImageRasterMesh.sizeY = op.args.ImageResize.size[1];
+                    newOp->sizeX = op.args.ImageResize.size[0];
+                    newOp->sizeY = op.args.ImageResize.size[1];
 
-					if (newOp->op.args.ImageRasterMesh.mask)
+					if (newOp->mask)
 					{
 						Ptr<ASTOpFixed> mop = mu::Clone<ASTOpFixed>(this);
-						mop->SetChild(mop->op.args.ImageResize.source, newOp->children[newOp->op.args.ImageRasterMesh.mask].child());
-						newOp->SetChild(newOp->op.args.ImageRasterMesh.mask, mop);
+						mop->SetChild(mop->op.args.ImageResize.source, newOp->mask.child());
+						newOp->mask = mop;
 					}
 
 					// Don't apply absolute resizes to the image to raster: it could even enlarge it.
diff --git a/Engine/Plugins/Experimental/Mutable/Source/MutableTools/Private/MuT/NodeImageProject.cpp b/Engine/Plugins/Experimental/Mutable/Source/MutableTools/Private/MuT/NodeImageProject.cpp
--- a/Engine/Plugins/Experimental/Mutable/Source/MutableTools/Private/MuT/NodeImageProject.cpp
+++ b/Engine/Plugins/Experimental/Mutable/Source/MutableTools/Private/MuT/NodeImageProject.cpp
@@ -142,26 +142,20 @@ namespace mu
 
 
 	//---------------------------------------------------------------------------------------------
-	NodeScalarPtr NodeImageProject::GetAngleFadeStart() const
+	void NodeImageProject::SetAngleFadeChannels(bool bFadeRGB, bool bFadeA)
 	{
-		return m_pD->m_pAngleFadeStart.get();
+		m_pD->bIsRGBFadingEnabled = bFadeRGB;
+		m_pD->bIsAlphaFadingEnabled = bFadeA;
 	}
 
 
 	//---------------------------------------------------------------------------------------------
-	void NodeImageProject::SetAngleFadeStart( NodeScalarPtr pNode )
+	void NodeImageProject::SetAngleFadeStart(NodeScalarPtr pNode)
 	{
 		m_pD->m_pAngleFadeStart = pNode;
 	}
 
 
-	//---------------------------------------------------------------------------------------------
-	NodeScalarPtr NodeImageProject::GetAngleFadeEnd() const
-	{
-		return m_pD->m_pAngleFadeEnd.get();
-	}
-
-
 	//---------------------------------------------------------------------------------------------
 	void NodeImageProject::SetAngleFadeEnd( NodeScalarPtr pNode )
 	{
diff --git a/Engine/Plugins/Experimental/Mutable/Source/MutableTools/Private/MuT/NodeImageProjectPrivate.h b/Engine/Plugins/Experimental/Mutable/Source/MutableTools/Private/MuT/NodeImageProjectPrivate.h
--- a/Engine/Plugins/Experimental/Mutable/Source/MutableTools/Private/MuT/NodeImageProjectPrivate.h
+++ b/Engine/Plugins/Experimental/Mutable/Source/MutableTools/Private/MuT/NodeImageProjectPrivate.h
@@ -14,7 +14,6 @@
 namespace mu
 {
 
-
 	class NodeImageProject::Private : public NodeImage::Private
 	{
 	public:
@@ -31,13 +30,15 @@ namespace mu
 		NodeScalarPtr m_pAngleFadeEnd;
 		NodeImagePtr m_pImage;
 		NodeImagePtr m_pMask;
-        uint8_t m_layout = 0;
 		FUintVector2 m_imageSize;
+		uint8 m_layout = 0;
+		bool bIsRGBFadingEnabled = true;
+		bool bIsAlphaFadingEnabled = true;
 
 		//!
 		void Serialise( OutputArchive& arch ) const
 		{
-            uint32_t ver = 2;
+            uint32 ver = 3;
 			arch << ver;
 
 			arch << m_pProjector;
@@ -48,14 +49,16 @@ namespace mu
             arch << m_pMask;
             arch << m_layout;
 			arch << m_imageSize;
-        }
+			arch << bIsRGBFadingEnabled;
+			arch << bIsAlphaFadingEnabled;
+		}
 
 		//!
 		void Unserialise( InputArchive& arch )
 		{
-            uint32_t ver;
+            uint32 ver;
 			arch >> ver;
-            check(ver==2);
+            check(ver>=2 && ver<=3);
 
 			arch >> m_pProjector;
 			arch >> m_pMesh;
@@ -65,6 +68,11 @@ namespace mu
 			arch >> m_pMask;
             arch >> m_layout;
 			arch >> m_imageSize;
+			if (ver >= 3)
+			{
+				arch >> bIsRGBFadingEnabled;
+				arch >> bIsAlphaFadingEnabled;
+			}
 		}
 	};
 
diff --git a/Engine/Plugins/Experimental/Mutable/Source/MutableTools/Public/MuT/NodeImageProject.h b/Engine/Plugins/Experimental/Mutable/Source/MutableTools/Public/MuT/NodeImageProject.h
--- a/Engine/Plugins/Experimental/Mutable/Source/MutableTools/Public/MuT/NodeImageProject.h
+++ b/Engine/Plugins/Experimental/Mutable/Source/MutableTools/Public/MuT/NodeImageProject.h
@@ -58,8 +58,6 @@ namespace mu
 		// Node Interface
 		//-----------------------------------------------------------------------------------------
 
-        
-
         const NODE_TYPE* GetType() const override;
 		static const NODE_TYPE* GetStaticType();
 
@@ -83,12 +81,16 @@ namespace mu
 		NodeImagePtr GetTargetMask() const;
 		void SetTargetMask( NodeImagePtr );
 
-        //! Get the node generating the threshold angle of the projection. In degrees.
-		NodeScalarPtr GetAngleFadeStart() const;
+		//! Set the angle-based fading behaviour for this projector. It must be set separately for 
+		//! RGB and Alpha channels of the images to project. By default it is enabled for all channels.
+		//! Disabling the angle fade will still use the AngleFadeEnd below to clamp the projection,
+		//! but it will ignore the start angle, so it will not apply a fading gradient.
+		void SetAngleFadeChannels( bool bFadeRGB, bool bFadeA );
+
+        //! Set the node generating the fading start angle. Only relevant if fading is enabled with 
 		void SetAngleFadeStart( NodeScalarPtr );
 
-        //! Get the node generating the threshold angle of the projection. In degrees.
-		NodeScalarPtr GetAngleFadeEnd() const;
+		//! Set the node generating the fading end angle
 		void SetAngleFadeEnd( NodeScalarPtr );
 
         //! Get the node generating the image to project.
