[Mutable] Fix incorrect range check in FindIntParameterValueName

diff --git a/Engine/Plugins/Experimental/Mutable/Source/CustomizableObject/Private/MuCO/CustomizableObject.cpp b/Engine/Plugins/Experimental/Mutable/Source/CustomizableObject/Private/MuCO/CustomizableObject.cpp
--- a/Engine/Plugins/Experimental/Mutable/Source/CustomizableObject/Private/MuCO/CustomizableObject.cpp
+++ b/Engine/Plugins/Experimental/Mutable/Source/CustomizableObject/Private/MuCO/CustomizableObject.cpp
@@ -1361,14 +1361,12 @@ FString UCustomizableObject::FindIntParameterValueName(int32 ParamIndex, int32 P
 	{
 		const TArray<FMutableModelParameterValue> & PossibleValues = ParameterProperties[ParamIndex].PossibleValues;
 
-		if (ParamValue >= 0 && ParamValue < PossibleValues.Num())
-		{
-			int32 MinValueIndex = PossibleValues[0].Value;
-			int32 ActualIndex = ParamValue - MinValueIndex;
-
-			check(ActualIndex >= 0 && ActualIndex < PossibleValues.Num());
+		const int32 MinValueIndex = !PossibleValues.IsEmpty() ? PossibleValues[0].Value : 0;
+		ParamValue = ParamValue - MinValueIndex;
 
-			return PossibleValues[ActualIndex].Name;
+		if (PossibleValues.IsValidIndex(ParamValue))
+		{
+			return PossibleValues[ParamValue].Name;
 		}
 	}
 	else
