[Mutable] Expose image saturate node in editor.

#rnx

diff --git a/Engine/Plugins/Experimental/Mutable/Source/CustomizableObjectEditor/Classes/MuCOE/Nodes/CustomizableObjectNodeTextureSaturate.h b/Engine/Plugins/Experimental/Mutable/Source/CustomizableObjectEditor/Classes/MuCOE/Nodes/CustomizableObjectNodeTextureSaturate.h
new file mode 100644
--- /dev/null
+++ b/Engine/Plugins/Experimental/Mutable/Source/CustomizableObjectEditor/Classes/MuCOE/Nodes/CustomizableObjectNodeTextureSaturate.h
@@ -0,0 +1,42 @@
+// Copyright Epic Games, Inc. All Rights Reserved.
+#pragma once
+
+#include "EdGraph/EdGraphNode.h"
+#include "EdGraph/EdGraphPin.h"
+#include "HAL/Platform.h"
+#include "Internationalization/Text.h"
+#include "Math/Color.h"
+#include "MuCOE/Nodes/CustomizableObjectNode.h"
+#include "UObject/UObjectGlobals.h"
+
+#include "CustomizableObjectNodeTextureSaturate.generated.h"
+
+class FArchive;
+class UCustomizableObjectNodeRemapPins;
+class UObject;
+
+UCLASS()
+class CUSTOMIZABLEOBJECTEDITOR_API UCustomizableObjectNodeTextureSaturate : public UCustomizableObjectNode
+{
+public:
+	GENERATED_BODY()
+
+	// Begin EdGraphNode interface
+	FText GetNodeTitle(ENodeTitleType::Type TittleType) const override;
+	FLinearColor GetNodeTitleColor() const override;
+	FText GetTooltipText() const override;
+	
+	// UCustomizableObjectNode interface
+	void AllocateDefaultPins(UCustomizableObjectNodeRemapPins* RemapPins) override;
+
+	UEdGraphPin* GetBaseImagePin() const;
+	UEdGraphPin* GetFactorPin() const;
+
+private:
+	UPROPERTY()
+	FEdGraphPinReference BaseImagePinReference;
+
+
+	UPROPERTY()
+	FEdGraphPinReference FactorPinReference;
+};
diff --git a/Engine/Plugins/Experimental/Mutable/Source/CustomizableObjectEditor/Private/MuCOE/EdGraphSchema_CustomizableObject.cpp b/Engine/Plugins/Experimental/Mutable/Source/CustomizableObjectEditor/Private/MuCOE/EdGraphSchema_CustomizableObject.cpp
--- a/Engine/Plugins/Experimental/Mutable/Source/CustomizableObjectEditor/Private/MuCOE/EdGraphSchema_CustomizableObject.cpp
+++ b/Engine/Plugins/Experimental/Mutable/Source/CustomizableObjectEditor/Private/MuCOE/EdGraphSchema_CustomizableObject.cpp
@@ -64,6 +64,7 @@
 #include "MuCOE/Nodes/CustomizableObjectNodeTextureSwitch.h"
 #include "MuCOE/Nodes/CustomizableObjectNodeTextureToChannels.h"
 #include "MuCOE/Nodes/CustomizableObjectNodeTextureTransform.h"
+#include "MuCOE/Nodes/CustomizableObjectNodeTextureSaturate.h"
 #include "MuCOE/Nodes/CustomizableObjectNodeTextureVariation.h"
 #include "ScopedTransaction.h"
 #include "Settings/EditorStyleSettings.h"
@@ -425,6 +426,7 @@ void UEdGraphSchema_CustomizableObject::GetGraphContextActions(FGraphContextMenu
 			ContextMenuBuilder.CreateTemplateNode<UCustomizableObjectNodeTextureInvert>(),
 			ContextMenuBuilder.CreateTemplateNode<UCustomizableObjectNodeTextureColourMap>(),
 			ContextMenuBuilder.CreateTemplateNode<UCustomizableObjectNodeTextureTransform>(),
+			ContextMenuBuilder.CreateTemplateNode<UCustomizableObjectNodeTextureSaturate>(),
 		};
 
 		AddNewNodeCategoryActionsFiltered(TextureTemplateNodes, ContextMenuBuilder, TEXT("Texture"), GeneralGrouping, Filter);
diff --git a/Engine/Plugins/Experimental/Mutable/Source/CustomizableObjectEditor/Private/MuCOE/GenerateMutableSource/GenerateMutableSource.cpp b/Engine/Plugins/Experimental/Mutable/Source/CustomizableObjectEditor/Private/MuCOE/GenerateMutableSource/GenerateMutableSource.cpp
--- a/Engine/Plugins/Experimental/Mutable/Source/CustomizableObjectEditor/Private/MuCOE/GenerateMutableSource/GenerateMutableSource.cpp
+++ b/Engine/Plugins/Experimental/Mutable/Source/CustomizableObjectEditor/Private/MuCOE/GenerateMutableSource/GenerateMutableSource.cpp
@@ -32,6 +32,7 @@
 #include "MuCOE/Nodes/CustomizableObjectNodeTextureSwitch.h"
 #include "MuCOE/Nodes/CustomizableObjectNodeTextureToChannels.h"
 #include "MuCOE/Nodes/CustomizableObjectNodeTextureTransform.h"
+#include "MuCOE/Nodes/CustomizableObjectNodeTextureSaturate.h"
 #include "MuCOE/Nodes/CustomizableObjectNodeTextureVariation.h"
 #include "MuCOE/UnrealEditorPortabilityHelpers.h"
 #include "MuT/NodeComponentEdit.h"
@@ -524,6 +525,14 @@ UTexture2D* FindReferenceImage(const UEdGraphPin* Pin, FMutableGraphGenerationCo
 		}
 	}
 
+	else if (const UCustomizableObjectNodeTextureSaturate* TypedNodeSaturate = Cast<UCustomizableObjectNodeTextureSaturate>(Node))
+	{
+		if ( UEdGraphPin* BaseImagePin = FollowInputPin(*TypedNodeSaturate->GetBaseImagePin()) )
+		{
+			Result = FindReferenceImage(BaseImagePin, GenerationContext);
+		}
+	}
+
 	else if (const UCustomizableObjectNodeTable* TypedNodeTable = Cast<UCustomizableObjectNodeTable>(Node))
 	{
 		if (Pin->PinType.PinCategory == Helper_GetPinCategory(Schema->PC_MaterialAsset))
diff --git a/Engine/Plugins/Experimental/Mutable/Source/CustomizableObjectEditor/Private/MuCOE/GenerateMutableSource/GenerateMutableSourceImage.cpp b/Engine/Plugins/Experimental/Mutable/Source/CustomizableObjectEditor/Private/MuCOE/GenerateMutableSource/GenerateMutableSourceImage.cpp
--- a/Engine/Plugins/Experimental/Mutable/Source/CustomizableObjectEditor/Private/MuCOE/GenerateMutableSource/GenerateMutableSourceImage.cpp
+++ b/Engine/Plugins/Experimental/Mutable/Source/CustomizableObjectEditor/Private/MuCOE/GenerateMutableSource/GenerateMutableSourceImage.cpp
@@ -25,6 +25,7 @@
 #include "MuCOE/Nodes/CustomizableObjectNodeTextureSwitch.h"
 #include "MuCOE/Nodes/CustomizableObjectNodeTextureToChannels.h"
 #include "MuCOE/Nodes/CustomizableObjectNodeTextureTransform.h"
+#include "MuCOE/Nodes/CustomizableObjectNodeTextureSaturate.h"
 #include "MuCOE/Nodes/CustomizableObjectNodeTextureVariation.h"
 #include "MuCOE/UnrealEditorPortabilityHelpers.h"
 #include "MuCOE/UnrealToMutableTextureConversionUtils.h"
@@ -39,6 +40,7 @@
 #include "MuT/NodeImagePlainColour.h"
 #include "MuT/NodeImageProject.h"
 #include "MuT/NodeImageResize.h"
+#include "MuT/NodeImageSaturate.h"
 #include "MuT/NodeImageSwitch.h"
 #include "MuT/NodeImageSwizzle.h"
 #include "MuT/NodeImageTable.h"
@@ -766,6 +768,25 @@ mu::NodeImagePtr GenerateMutableSourceImage(const UEdGraphPin* Pin, FMutableGrap
 			TransformNode->SetRotation( RotationNode ); 
 		}
 	}
+
+	else if (const UCustomizableObjectNodeTextureSaturate* TypedNodeSaturate = Cast<UCustomizableObjectNodeTextureSaturate>(Node))
+	{
+		mu::Ptr<mu::NodeImageSaturate> SaturateNode = new mu::NodeImageSaturate();
+		Result = SaturateNode;
+	
+		if ( UEdGraphPin* BaseImagePin = FollowInputPin(*TypedNodeSaturate->GetBaseImagePin()) )
+		{
+			mu::NodeImagePtr ImageNode = GenerateMutableSourceImage(BaseImagePin, GenerationContext, MaxTextureSize);
+			SaturateNode->SetSource(ImageNode); 
+		}
+
+		if ( UEdGraphPin* FactorPin = FollowInputPin(*TypedNodeSaturate->GetFactorPin()))
+		{
+			mu::NodeScalarPtr FactorNode = GenerateMutableSourceFloat(FactorPin, GenerationContext);
+			SaturateNode->SetFactor(FactorNode); 
+		}
+	}
+	
 	// If the node is a plain colour node, generate an image out of it
 	else if (Pin->PinType.PinCategory == Helper_GetPinCategory(Schema->PC_Color))
 	{
diff --git a/Engine/Plugins/Experimental/Mutable/Source/CustomizableObjectEditor/Private/MuCOE/Nodes/CustomizableObjectNodeTextureSaturate.cpp b/Engine/Plugins/Experimental/Mutable/Source/CustomizableObjectEditor/Private/MuCOE/Nodes/CustomizableObjectNodeTextureSaturate.cpp
new file mode 100644
--- /dev/null
+++ b/Engine/Plugins/Experimental/Mutable/Source/CustomizableObjectEditor/Private/MuCOE/Nodes/CustomizableObjectNodeTextureSaturate.cpp
@@ -0,0 +1,57 @@
+// Copyright Epic Games, Inc. All Rights Reserved.
+
+#include "MuCOE/Nodes/CustomizableObjectNodeTextureSaturate.h"
+
+#include "MuCOE/EdGraphSchema_CustomizableObject.h"
+
+class UCustomizableObjectNodeRemapPins;
+
+#define LOCTEXT_NAMESPACE "CustomizableObjectEditor"
+
+void UCustomizableObjectNodeTextureSaturate::AllocateDefaultPins(UCustomizableObjectNodeRemapPins* RemapPins)
+{
+	const UEdGraphSchema_CustomizableObject* Schema = GetDefault<UEdGraphSchema_CustomizableObject>();
+
+	FString PinName = TEXT("Texture");
+	UEdGraphPin* OutputPin = CustomCreatePin(EGPD_Output, Schema->PC_Image, FName(*PinName));
+	OutputPin->bDefaultValueIsIgnored = true;
+
+	PinName = TEXT("Base Texture");
+	UEdGraphPin* ImagePin = CustomCreatePin(EGPD_Input, Schema->PC_Image, FName(*PinName));
+
+	BaseImagePinReference = FEdGraphPinReference(ImagePin);
+	
+	PinName = TEXT("Factor");
+	UEdGraphPin* FactorPin = CustomCreatePin(EGPD_Input, Schema->PC_Float, FName(*PinName));
+	FactorPin->bDefaultValueIsIgnored = true;
+
+	FactorPinReference = FEdGraphPinReference(FactorPin);
+}
+
+UEdGraphPin* UCustomizableObjectNodeTextureSaturate::GetBaseImagePin() const
+{
+	return BaseImagePinReference.Get();
+}
+
+UEdGraphPin* UCustomizableObjectNodeTextureSaturate::GetFactorPin() const
+{
+	return FactorPinReference.Get();
+}
+
+FText UCustomizableObjectNodeTextureSaturate::GetNodeTitle(ENodeTitleType::Type TitleType) const
+{
+	return LOCTEXT("Texture_Saturate", "Texture Saturate");
+}
+
+FLinearColor UCustomizableObjectNodeTextureSaturate::GetNodeTitleColor() const
+{
+	const UEdGraphSchema_CustomizableObject* Schema = GetDefault<UEdGraphSchema_CustomizableObject>();
+	return Schema->GetPinTypeColor(Schema->PC_Image);
+}
+
+FText UCustomizableObjectNodeTextureSaturate::GetTooltipText() const
+{
+	return LOCTEXT("Texture_Saturate_Tooltip", "Saturate Image");
+}
+
+#undef LOCTEXT_NAMESPACE
