[mutable] Fix clamping of colours when image layers is used with a color larger than white.
#rnx

diff --git a/Engine/Plugins/Experimental/Mutable/Source/MutableRuntime/Private/MuR/CodeRunnerTasked.cpp b/Engine/Plugins/Experimental/Mutable/Source/MutableRuntime/Private/MuR/CodeRunnerTasked.cpp
--- a/Engine/Plugins/Experimental/Mutable/Source/MutableRuntime/Private/MuR/CodeRunnerTasked.cpp
+++ b/Engine/Plugins/Experimental/Mutable/Source/MutableRuntime/Private/MuR/CodeRunnerTasked.cpp
@@ -749,16 +749,16 @@ namespace mu
 				// \todo: precalculated tables for softlight
 				switch (EBlendType(Args.blendType))
 				{
-				case EBlendType::BT_NORMAL_COMBINE: ImageNormalCombine(pNew.get(), m_base.get(), m_mask.get(), vec3f(m_col)); break;
-				case EBlendType::BT_SOFTLIGHT: BufferLayerColour<SoftLightChannelMasked, SoftLightChannel, false>(pNew->GetData(), m_base.get(), m_mask.get(), vec3f(m_col)); break;
-				case EBlendType::BT_HARDLIGHT: BufferLayerColour<HardLightChannelMasked, HardLightChannel, false>(pNew->GetData(), m_base.get(), m_mask.get(), vec3f(m_col)); break;
-				case EBlendType::BT_BURN: BufferLayerColour<BurnChannelMasked, BurnChannel, false>(pNew->GetData(), m_base.get(), m_mask.get(), vec3f(m_col)); break;
-				case EBlendType::BT_DODGE: BufferLayerColour<DodgeChannelMasked, DodgeChannel, false>(pNew->GetData(), m_base.get(), m_mask.get(), vec3f(m_col)); break;
-				case EBlendType::BT_SCREEN: BufferLayerColour<ScreenChannelMasked, ScreenChannel, false>(pNew->GetData(), m_base.get(), m_mask.get(), vec3f(m_col)); break;
-				case EBlendType::BT_OVERLAY: BufferLayerColour<OverlayChannelMasked, OverlayChannel, false>(pNew->GetData(), m_base.get(), m_mask.get(), vec3f(m_col)); break;
-				case EBlendType::BT_LIGHTEN: BufferLayerColour<LightenChannelMasked, LightenChannel, false>(pNew->GetData(), m_base.get(), m_mask.get(), vec3f(m_col)); break;
-				case EBlendType::BT_MULTIPLY: BufferLayerColour<MultiplyChannelMasked, MultiplyChannel, false>(pNew->GetData(), m_base.get(), m_mask.get(), vec3f(m_col)); break;
-				case EBlendType::BT_BLEND: BufferLayerColour<BlendChannelMasked, BlendChannel, false>(pNew->GetData(), m_base.get(), m_mask.get(), vec3f(m_col)); break;
+				case EBlendType::BT_NORMAL_COMBINE: ImageNormalCombine(pNew.get(), m_base.get(), m_mask.get(), m_col); break;
+				case EBlendType::BT_SOFTLIGHT: BufferLayerColour<SoftLightChannelMasked, SoftLightChannel>(pNew->GetData(), m_base.get(), m_mask.get(), m_col); break;
+				case EBlendType::BT_HARDLIGHT: BufferLayerColour<HardLightChannelMasked, HardLightChannel>(pNew->GetData(), m_base.get(), m_mask.get(), m_col); break;
+				case EBlendType::BT_BURN: BufferLayerColour<BurnChannelMasked, BurnChannel>(pNew->GetData(), m_base.get(), m_mask.get(), m_col); break;
+				case EBlendType::BT_DODGE: BufferLayerColour<DodgeChannelMasked, DodgeChannel>(pNew->GetData(), m_base.get(), m_mask.get(), m_col); break;
+				case EBlendType::BT_SCREEN: BufferLayerColour<ScreenChannelMasked, ScreenChannel>(pNew->GetData(), m_base.get(), m_mask.get(), m_col); break;
+				case EBlendType::BT_OVERLAY: BufferLayerColour<OverlayChannelMasked, OverlayChannel>(pNew->GetData(), m_base.get(), m_mask.get(), m_col); break;
+				case EBlendType::BT_LIGHTEN: BufferLayerColour<LightenChannelMasked, LightenChannel>(pNew->GetData(), m_base.get(), m_mask.get(), m_col); break;
+				case EBlendType::BT_MULTIPLY: BufferLayerColour<MultiplyChannelMasked, MultiplyChannel>(pNew->GetData(), m_base.get(), m_mask.get(), m_col); break;
+				case EBlendType::BT_BLEND: BufferLayerColour<BlendChannelMasked, BlendChannel>(pNew->GetData(), m_base.get(), m_mask.get(), m_col); break;
 				default: check(false);
 				}
 
@@ -771,14 +771,14 @@ namespace mu
 					switch (EBlendType(Args.blendType))
 					{
 					case EBlendType::BT_NORMAL_COMBINE: check(false); break;
-					case EBlendType::BT_SOFTLIGHT: BufferLayerColourFromAlpha<SoftLightChannel, false>(pNew.get(), m_base.get(), vec3f(m_col)); break;
-					case EBlendType::BT_HARDLIGHT: BufferLayerColourFromAlpha<HardLightChannel, false>(pNew.get(), m_base.get(), vec3f(m_col)); break;
-					case EBlendType::BT_BURN: BufferLayerColourFromAlpha<BurnChannel, false>(pNew.get(), m_base.get(), vec3f(m_col)); break;
-					case EBlendType::BT_DODGE: BufferLayerColourFromAlpha<DodgeChannel, false>(pNew.get(), m_base.get(), vec3f(m_col)); break;
-					case EBlendType::BT_SCREEN: BufferLayerColourFromAlpha<ScreenChannel, false>(pNew.get(), m_base.get(), vec3f(m_col)); break;
-					case EBlendType::BT_OVERLAY: BufferLayerColourFromAlpha<OverlayChannel, false>(pNew.get(), m_base.get(), vec3f(m_col)); break;
-					case EBlendType::BT_LIGHTEN: BufferLayerColourFromAlpha<LightenChannel, false>(pNew.get(), m_base.get(), vec3f(m_col)); break;
-					case EBlendType::BT_MULTIPLY: BufferLayerColourFromAlpha<MultiplyChannel, false>(pNew.get(), m_base.get(), vec3f(m_col)); break;
+					case EBlendType::BT_SOFTLIGHT: BufferLayerColourFromAlpha<SoftLightChannel>(pNew.get(), m_base.get(), m_col); break;
+					case EBlendType::BT_HARDLIGHT: BufferLayerColourFromAlpha<HardLightChannel>(pNew.get(), m_base.get(), m_col); break;
+					case EBlendType::BT_BURN: BufferLayerColourFromAlpha<BurnChannel>(pNew.get(), m_base.get(), m_col); break;
+					case EBlendType::BT_DODGE: BufferLayerColourFromAlpha<DodgeChannel>(pNew.get(), m_base.get(), m_col); break;
+					case EBlendType::BT_SCREEN: BufferLayerColourFromAlpha<ScreenChannel>(pNew.get(), m_base.get(), m_col); break;
+					case EBlendType::BT_OVERLAY: BufferLayerColourFromAlpha<OverlayChannel>(pNew.get(), m_base.get(), m_col); break;
+					case EBlendType::BT_LIGHTEN: BufferLayerColourFromAlpha<LightenChannel>(pNew.get(), m_base.get(), m_col); break;
+					case EBlendType::BT_MULTIPLY: BufferLayerColourFromAlpha<MultiplyChannel>(pNew.get(), m_base.get(), m_col); break;
 					case EBlendType::BT_BLEND: check(false); break;
 					default: check(false);
 					}
@@ -787,15 +787,15 @@ namespace mu
 				{
 					switch (EBlendType(Args.blendType))
 					{
-					case EBlendType::BT_NORMAL_COMBINE: ImageNormalCombine(pNew.get(), m_base.get(), vec3f(m_col)); break;
-					case EBlendType::BT_SOFTLIGHT: BufferLayerColour<SoftLightChannel, false>(pNew.get(), m_base.get(), vec3f(m_col)); break;
-					case EBlendType::BT_HARDLIGHT: BufferLayerColour<HardLightChannel, false>(pNew.get(), m_base.get(), vec3f(m_col)); break;
-					case EBlendType::BT_BURN: BufferLayerColour<BurnChannel, false>(pNew.get(), m_base.get(), vec3f(m_col)); break;
-					case EBlendType::BT_DODGE: BufferLayerColour<DodgeChannel, false>(pNew.get(), m_base.get(), vec3f(m_col)); break;
-					case EBlendType::BT_SCREEN: BufferLayerColour<ScreenChannel, false>(pNew.get(), m_base.get(), vec3f(m_col)); break;
-					case EBlendType::BT_OVERLAY: BufferLayerColour<OverlayChannel, false>(pNew.get(), m_base.get(), vec3f(m_col)); break;
-					case EBlendType::BT_LIGHTEN: BufferLayerColour<LightenChannel, false>(pNew.get(), m_base.get(), vec3f(m_col)); break;
-					case EBlendType::BT_MULTIPLY: BufferLayerColour<MultiplyChannel, false>(pNew.get(), m_base.get(), vec3f(m_col)); break;
+					case EBlendType::BT_NORMAL_COMBINE: ImageNormalCombine(pNew.get(), m_base.get(), m_col); break;
+					case EBlendType::BT_SOFTLIGHT: BufferLayerColour<SoftLightChannel>(pNew.get(), m_base.get(), m_col); break;
+					case EBlendType::BT_HARDLIGHT: BufferLayerColour<HardLightChannel>(pNew.get(), m_base.get(), m_col); break;
+					case EBlendType::BT_BURN: BufferLayerColour<BurnChannel>(pNew.get(), m_base.get(), m_col); break;
+					case EBlendType::BT_DODGE: BufferLayerColour<DodgeChannel>(pNew.get(), m_base.get(), m_col); break;
+					case EBlendType::BT_SCREEN: BufferLayerColour<ScreenChannel>(pNew.get(), m_base.get(), m_col); break;
+					case EBlendType::BT_OVERLAY: BufferLayerColour<OverlayChannel>(pNew.get(), m_base.get(), m_col); break;
+					case EBlendType::BT_LIGHTEN: BufferLayerColour<LightenChannel>(pNew.get(), m_base.get(), m_col); break;
+					case EBlendType::BT_MULTIPLY: BufferLayerColour<MultiplyChannel>(pNew.get(), m_base.get(), m_col); break;
 					case EBlendType::BT_BLEND: FillPlainColourImage(pNew.get(), m_col); break;
 					default: check(false);
 					}
@@ -823,15 +823,15 @@ namespace mu
 				switch (EBlendType(Args.blendTypeAlpha))
 				{
 				case EBlendType::BT_NORMAL_COMBINE: check(false); break;
-				case EBlendType::BT_SOFTLIGHT: BufferLayerColourInPlace<SoftLightChannelMasked, SoftLightChannel, false, 1>(pNew.get(), m_mask.get(), vec3f(m_col), bOnlyOneMip, 3); break;
-				case EBlendType::BT_HARDLIGHT: BufferLayerColourInPlace<HardLightChannelMasked, HardLightChannel, false, 1>(pNew.get(), m_mask.get(), vec3f(m_col), bOnlyOneMip, 3); break;
-				case EBlendType::BT_BURN: BufferLayerColourInPlace<BurnChannelMasked, BurnChannel, false, 1>(pNew.get(), m_mask.get(), vec3f(m_col), bOnlyOneMip, 3); break;
-				case EBlendType::BT_DODGE: BufferLayerColourInPlace<DodgeChannelMasked, DodgeChannel, false, 1>(pNew.get(), m_mask.get(), vec3f(m_col), bOnlyOneMip, 3); break;
-				case EBlendType::BT_SCREEN: BufferLayerColourInPlace<ScreenChannelMasked, ScreenChannel, false, 1>(pNew.get(), m_mask.get(), vec3f(m_col), bOnlyOneMip, 3); break;
-				case EBlendType::BT_OVERLAY: BufferLayerColourInPlace<OverlayChannelMasked, OverlayChannel, false, 1>(pNew.get(), m_mask.get(), vec3f(m_col), bOnlyOneMip, 3); break;
-				case EBlendType::BT_LIGHTEN: BufferLayerColourInPlace<LightenChannelMasked, LightenChannel, false, 1>(pNew.get(), m_mask.get(), vec3f(m_col), bOnlyOneMip, 3); break;
-				case EBlendType::BT_MULTIPLY: BufferLayerColourInPlace<MultiplyChannelMasked, MultiplyChannel, false, 1>(pNew.get(), m_mask.get(), vec3f(m_col), bOnlyOneMip, 3); break;
-				case EBlendType::BT_BLEND: BufferLayerColourInPlace<BlendChannelMasked, BlendChannel, false, 1>(pNew.get(), m_mask.get(), vec3f(m_col), bOnlyOneMip, 3); break;
+				case EBlendType::BT_SOFTLIGHT: BufferLayerColourInPlace<SoftLightChannelMasked, SoftLightChannel, 1>(pNew.get(), m_mask.get(), m_col, bOnlyOneMip, 3); break;
+				case EBlendType::BT_HARDLIGHT: BufferLayerColourInPlace<HardLightChannelMasked, HardLightChannel, 1>(pNew.get(), m_mask.get(), m_col, bOnlyOneMip, 3); break;
+				case EBlendType::BT_BURN: BufferLayerColourInPlace<BurnChannelMasked, BurnChannel, 1>(pNew.get(), m_mask.get(), m_col, bOnlyOneMip, 3); break;
+				case EBlendType::BT_DODGE: BufferLayerColourInPlace<DodgeChannelMasked, DodgeChannel, 1>(pNew.get(), m_mask.get(), m_col, bOnlyOneMip, 3); break;
+				case EBlendType::BT_SCREEN: BufferLayerColourInPlace<ScreenChannelMasked, ScreenChannel, 1>(pNew.get(), m_mask.get(), m_col, bOnlyOneMip, 3); break;
+				case EBlendType::BT_OVERLAY: BufferLayerColourInPlace<OverlayChannelMasked, OverlayChannel, 1>(pNew.get(), m_mask.get(), m_col, bOnlyOneMip, 3); break;
+				case EBlendType::BT_LIGHTEN: BufferLayerColourInPlace<LightenChannelMasked, LightenChannel, 1>(pNew.get(), m_mask.get(), m_col, bOnlyOneMip, 3); break;
+				case EBlendType::BT_MULTIPLY: BufferLayerColourInPlace<MultiplyChannelMasked, MultiplyChannel, 1>(pNew.get(), m_mask.get(), m_col, bOnlyOneMip, 3); break;
+				case EBlendType::BT_BLEND: BufferLayerColourInPlace<BlendChannelMasked, BlendChannel, 1>(pNew.get(), m_mask.get(), m_col, bOnlyOneMip, 3); break;
 				default: check(false);
 				}
 
@@ -841,15 +841,15 @@ namespace mu
 				switch (EBlendType(Args.blendTypeAlpha))
 				{
 				case EBlendType::BT_NORMAL_COMBINE: check(false); break;
-				case EBlendType::BT_SOFTLIGHT: BufferLayerColourInPlace<SoftLightChannel, false, 1>(pNew.get(), vec3f(m_col), bOnlyOneMip, 3); break;
-				case EBlendType::BT_HARDLIGHT: BufferLayerColourInPlace<HardLightChannel, false, 1>(pNew.get(), vec3f(m_col), bOnlyOneMip, 3); break;
-				case EBlendType::BT_BURN: BufferLayerColourInPlace<BurnChannel, false, 1>(pNew.get(), vec3f(m_col), bOnlyOneMip, 3); break;
-				case EBlendType::BT_DODGE: BufferLayerColourInPlace<DodgeChannel, false, 1>(pNew.get(), vec3f(m_col), bOnlyOneMip, 3); break;
-				case EBlendType::BT_SCREEN: BufferLayerColourInPlace<ScreenChannel, false, 1>(pNew.get(), vec3f(m_col), bOnlyOneMip, 3); break;
-				case EBlendType::BT_OVERLAY: BufferLayerColourInPlace<OverlayChannel, false, 1>(pNew.get(), vec3f(m_col), bOnlyOneMip, 3); break;
-				case EBlendType::BT_LIGHTEN: BufferLayerColourInPlace<LightenChannel, false, 1>(pNew.get(), vec3f(m_col), bOnlyOneMip, 3); break;
-				case EBlendType::BT_MULTIPLY: BufferLayerColourInPlace<MultiplyChannel, false, 1>(pNew.get(), vec3f(m_col), bOnlyOneMip, 3); break;
-				case EBlendType::BT_BLEND: BufferLayerColourInPlace<BlendChannel, false, 1>(pNew.get(), vec3f(m_col), bOnlyOneMip, 3); break;
+				case EBlendType::BT_SOFTLIGHT: BufferLayerColourInPlace<SoftLightChannel, 1>(pNew.get(), m_col, bOnlyOneMip, 3); break;
+				case EBlendType::BT_HARDLIGHT: BufferLayerColourInPlace<HardLightChannel, 1>(pNew.get(), m_col, bOnlyOneMip, 3); break;
+				case EBlendType::BT_BURN: BufferLayerColourInPlace<BurnChannel, 1>(pNew.get(), m_col, bOnlyOneMip, 3); break;
+				case EBlendType::BT_DODGE: BufferLayerColourInPlace<DodgeChannel, 1>(pNew.get(), m_col, bOnlyOneMip, 3); break;
+				case EBlendType::BT_SCREEN: BufferLayerColourInPlace<ScreenChannel, 1>(pNew.get(), m_col, bOnlyOneMip, 3); break;
+				case EBlendType::BT_OVERLAY: BufferLayerColourInPlace<OverlayChannel, 1>(pNew.get(), m_col, bOnlyOneMip, 3); break;
+				case EBlendType::BT_LIGHTEN: BufferLayerColourInPlace<LightenChannel, 1>(pNew.get(), m_col, bOnlyOneMip, 3); break;
+				case EBlendType::BT_MULTIPLY: BufferLayerColourInPlace<MultiplyChannel, 1>(pNew.get(), m_col, bOnlyOneMip, 3); break;
+				case EBlendType::BT_BLEND: BufferLayerColourInPlace<BlendChannel, 1>(pNew.get(), m_col, bOnlyOneMip, 3); break;
 				default: check(false);
 				}
 			}
diff --git a/Engine/Plugins/Experimental/Mutable/Source/MutableRuntime/Private/MuR/OpImageLayer.h b/Engine/Plugins/Experimental/Mutable/Source/MutableRuntime/Private/MuR/OpImageLayer.h
--- a/Engine/Plugins/Experimental/Mutable/Source/MutableRuntime/Private/MuR/OpImageLayer.h
+++ b/Engine/Plugins/Experimental/Mutable/Source/MutableRuntime/Private/MuR/OpImageLayer.h
@@ -9,7 +9,11 @@
 
 namespace mu
 {
-
+	inline bool IsAnyComponentLargerThan1( FVector4f v )
+	{
+		return (v[0] > 1) || (v[1] > 1) || (v[2] > 1) || (v[3] > 1);
+	}
+	
 
 	//---------------------------------------------------------------------------------------------
 	//! Apply a blending function to an image with a colour source
@@ -183,27 +187,42 @@ namespace mu
     }
 
 
+	template< unsigned (*BLEND_FUNC)(unsigned, unsigned) >
+	inline void BufferLayerColour(Image* Result, const Image* Base, FVector4f Col)
+	{
+		bool bIsClampNeeded = IsAnyComponentLargerThan1(Col);
+		if (bIsClampNeeded)
+		{
+			BufferLayerColour<BLEND_FUNC, true>(Result, Base, Col);
+		}
+		else
+		{
+			BufferLayerColour<BLEND_FUNC, false>(Result, Base, Col);
+		}
+	}
+
+
 	//---------------------------------------------------------------------------------------------
 	template< unsigned (*BLEND_FUNC)(unsigned, unsigned), bool CLAMP >
-	inline void BufferLayerColourFromAlpha(Image* pResult, const Image* pBase, vec3<float> col)
+	inline void BufferLayerColourFromAlpha(Image* Result, const Image* Base, FVector4f Col)
 	{
-		check(pResult->GetFormat() == pBase->GetFormat());
-		check(pResult->GetSizeX() == pBase->GetSizeX());
-		check(pResult->GetSizeY() == pBase->GetSizeY());
-		check(pResult->GetLODCount() == pBase->GetLODCount());
+		check(Result->GetFormat() == Base->GetFormat());
+		check(Result->GetSizeX() == Base->GetSizeX());
+		check(Result->GetSizeY() == Base->GetSizeY());
+		check(Result->GetLODCount() == Base->GetLODCount());
 
-		EImageFormat baseFormat = pBase->GetFormat();
+		EImageFormat baseFormat = Base->GetFormat();
 
-		uint8* pDestBuf = pResult->GetData();
-		const uint8* pBaseBuf = pBase->GetData();
+		uint8* pDestBuf = Result->GetData();
+		const uint8* pBaseBuf = Base->GetData();
 
 		unsigned top[3];
-		top[0] = (unsigned)(255 * col[0]);
-		top[1] = (unsigned)(255 * col[1]);
-		top[2] = (unsigned)(255 * col[2]);
+		top[0] = (unsigned)(255 * Col[0]);
+		top[1] = (unsigned)(255 * Col[1]);
+		top[2] = (unsigned)(255 * Col[2]);
 
 		// Generic implementation
-		int32 pixelCount = pBase->CalculatePixelCount();
+		int32 pixelCount = Base->CalculatePixelCount();
 
 		constexpr int PixelCountConcurrencyThreshold = 0xff;
 
@@ -347,6 +366,22 @@ namespace mu
 	}
 
 
+	//---------------------------------------------------------------------------------------------
+	template< unsigned (*BLEND_FUNC)(unsigned, unsigned)>
+	inline void BufferLayerColourFromAlpha(Image* Result, const Image* Base, FVector4f Col)
+	{
+		bool bIsClampNeeded = IsAnyComponentLargerThan1(Col);
+		if (bIsClampNeeded)
+		{
+			BufferLayerColourFromAlpha<BLEND_FUNC,true>(Result, Base, Col);
+		}
+		else
+		{
+			BufferLayerColourFromAlpha<BLEND_FUNC,false>(Result, Base, Col);
+		}
+	}
+
+
 	//---------------------------------------------------------------------------------------------
 	template< unsigned (*BLEND_FUNC_MASKED)(unsigned,unsigned,unsigned),
 		unsigned (*BLEND_FUNC)(unsigned,unsigned),
@@ -355,26 +390,21 @@ namespace mu
 		uint32 CHANNELS_TO_BLEND,
 		// Number of total channels in the base image
 		int32 BASE_CHANNEL_STRIDE >
-    inline void BufferLayerColourFormat( uint8* pStartDestBuf,
-		const Image* pBase,
-		const Image* pMask,
-		vec3<float> col,
-		uint32 BaseOffset,
-		bool bOnlyOneMip )
+    inline void BufferLayerColourFormat( uint8* StartDestBuf, const Image* Base, const Image* Mask, FVector4f Col, uint32 BaseOffset, bool bOnlyOneMip )
 	{
 		check(CHANNELS_TO_BLEND+BaseOffset<=BASE_CHANNEL_STRIDE);
 
 		unsigned top[3];
-		top[0] = (unsigned)(255 * col[0]);
-		top[1] = (unsigned)(255 * col[1]);
-		top[2] = (unsigned)(255 * col[2]);
+		top[0] = (unsigned)(255 * Col[0]);
+		top[1] = (unsigned)(255 * Col[1]);
+		top[2] = (unsigned)(255 * Col[2]);
 
-		uint8* pDestBuf = pStartDestBuf+ BaseOffset;
-        const uint8* pMaskBuf = pMask->GetData();
-        const uint8* pBaseBuf = pBase->GetData()+ BaseOffset;
-		EImageFormat maskFormat = pMask->GetFormat();
+		uint8* pDestBuf = StartDestBuf+ BaseOffset;
+        const uint8* pMaskBuf = Mask->GetData();
+        const uint8* pBaseBuf = Base->GetData()+ BaseOffset;
+		EImageFormat maskFormat = Mask->GetFormat();
 
-		int32 LODCount = pBase->GetLODCount();
+		int32 LODCount = Base->GetLODCount();
 
 		if (bOnlyOneMip)
 		{
@@ -386,7 +416,7 @@ namespace mu
 		constexpr uint32 NumColorChannels = FMath::Min(CHANNELS_TO_BLEND,3u);
         if ( isUncompressed )
         {
-            int32 pixelCount = pBase->CalculatePixelCount();
+            int32 pixelCount = Base->CalculatePixelCount();
             for ( int i=0; i<pixelCount; ++i )
             {
                 unsigned mask = pMaskBuf[i];
@@ -413,12 +443,12 @@ namespace mu
         }
         else if ( maskFormat==EImageFormat::IF_L_UBYTE_RLE )
         {
-            int rows = pBase->GetSizeY();
-            int width = pBase->GetSizeX();
+            int rows = Base->GetSizeY();
+            int width = Base->GetSizeX();
 
 			// debug test
-			ImagePtr TempDecodedMask = new Image(pMask->GetSizeX(), pMask->GetSizeY(), pMask->GetLODCount(), EImageFormat::IF_L_UBYTE);
-			UncompressRLE_L(pMask, TempDecodedMask.get());
+			ImagePtr TempDecodedMask = new Image(Mask->GetSizeX(), Mask->GetSizeY(), Mask->GetLODCount(), EImageFormat::IF_L_UBYTE);
+			UncompressRLE_L(Mask, TempDecodedMask.get());
 
             for (int32 lod=0;lod<LODCount; ++lod)
             {
@@ -442,7 +472,7 @@ namespace mu
                         ++pMaskBuf;
 
                         // Equal pixels
-						check(pDestBuf + BASE_CHANNEL_STRIDE * equal <= pStartDestBuf + pBase->GetDataSize());
+						check(pDestBuf + BASE_CHANNEL_STRIDE * equal <= StartDestBuf + Base->GetDataSize());
                         if ( equalPixel==255 )
                         {
                             for ( int i=0; i<equal; ++i )
@@ -505,7 +535,7 @@ namespace mu
                         pBaseBuf += BASE_CHANNEL_STRIDE *equal;
 
                         // Different pixels
-						check(pDestBuf + BASE_CHANNEL_STRIDE * different <= pStartDestBuf + pBase->GetDataSize());
+						check(pDestBuf + BASE_CHANNEL_STRIDE * different <= StartDestBuf + Base->GetDataSize());
                         for ( int i=0; i<different; ++i )
                         {
                             for ( int32 c=0; c<NumColorChannels; ++c )
@@ -554,49 +584,46 @@ namespace mu
 			  unsigned (*BLEND_FUNC)(unsigned,unsigned),
 			  bool CLAMP
 			  >
-    inline void BufferLayerColour( uint8* pDestBuf,
-								   const Image* pBase,
-								   const Image* pMask,
-								   vec3<float> col )
+    inline void BufferLayerColour( uint8* DestBuf, const Image* Base, const Image* Mask, FVector4f Col )
 	{
-		check(pBase->GetSizeX() == pMask->GetSizeX());
-		check(pBase->GetSizeY() == pMask->GetSizeY());
-		check(pMask->GetFormat() == EImageFormat::IF_L_UBYTE
+		check(Base->GetSizeX() == Mask->GetSizeX());
+		check(Base->GetSizeY() == Mask->GetSizeY());
+		check(Mask->GetFormat() == EImageFormat::IF_L_UBYTE
 			||
-			pMask->GetFormat() == EImageFormat::IF_L_UBYTE_RLE);
+			Mask->GetFormat() == EImageFormat::IF_L_UBYTE_RLE);
 
 		bool bValid =
-			(pBase->GetSizeX() == pMask->GetSizeX())
+			(Base->GetSizeX() == Mask->GetSizeX())
 			&&
-			(pBase->GetSizeY() == pMask->GetSizeY())
+			(Base->GetSizeY() == Mask->GetSizeY())
 			&&
-			(pMask->GetFormat() == EImageFormat::IF_L_UBYTE
+			(Mask->GetFormat() == EImageFormat::IF_L_UBYTE
 				||
-				pMask->GetFormat() == EImageFormat::IF_L_UBYTE_RLE);
+				Mask->GetFormat() == EImageFormat::IF_L_UBYTE_RLE);
 		if (!bValid)
 		{
 			return;
 		}
 
-		EImageFormat baseFormat = pBase->GetFormat();
+		EImageFormat baseFormat = Base->GetFormat();
 		if ( baseFormat==EImageFormat::IF_RGB_UBYTE )
 		{
-			BufferLayerColourFormat<BLEND_FUNC_MASKED, BLEND_FUNC, CLAMP, 3, 3>( pDestBuf, pBase, pMask, col, 0, false);
+			BufferLayerColourFormat<BLEND_FUNC_MASKED, BLEND_FUNC, CLAMP, 3, 3>( DestBuf, Base, Mask, Col, 0, false);
 		}
         else if ( baseFormat==EImageFormat::IF_RGBA_UBYTE )
         {
-            BufferLayerColourFormat<BLEND_FUNC_MASKED, BLEND_FUNC, CLAMP, 3, 4>( pDestBuf, pBase, pMask, col, 0, false);
+            BufferLayerColourFormat<BLEND_FUNC_MASKED, BLEND_FUNC, CLAMP, 3, 4>( DestBuf, Base, Mask, Col, 0, false);
         }
         else if ( baseFormat==EImageFormat::IF_BGRA_UBYTE )
         {
-            float temp = col[0];
-            col[0] = col[2];
-            col[2] = temp;
-            BufferLayerColourFormat<BLEND_FUNC_MASKED, BLEND_FUNC, CLAMP, 3, 4>( pDestBuf, pBase, pMask, col, 0, false);
+            float temp = Col[0];
+            Col[0] = Col[2];
+            Col[2] = temp;
+            BufferLayerColourFormat<BLEND_FUNC_MASKED, BLEND_FUNC, CLAMP, 3, 4>( DestBuf, Base, Mask, Col, 0, false);
         }
         else if ( baseFormat==EImageFormat::IF_L_UBYTE )
 		{
-			BufferLayerColourFormat<BLEND_FUNC_MASKED, BLEND_FUNC, CLAMP, 1, 1>( pDestBuf, pBase, pMask, col, 0, false);
+			BufferLayerColourFormat<BLEND_FUNC_MASKED, BLEND_FUNC, CLAMP, 1, 1>( DestBuf, Base, Mask, Col, 0, false);
 		}
 		else
 		{
@@ -605,6 +632,24 @@ namespace mu
 	}
 
 
+	//---------------------------------------------------------------------------------------------
+	template< unsigned (*BLEND_FUNC_MASKED)(unsigned, unsigned, unsigned),
+		unsigned (*BLEND_FUNC)(unsigned, unsigned)
+	>
+	inline void BufferLayerColour(uint8* DestBuf, const Image* Base, const Image* Mask, FVector4f Col)
+	{
+		bool bIsClampNeeded = IsAnyComponentLargerThan1(Col);
+		if (bIsClampNeeded)
+		{
+			BufferLayerColour<BLEND_FUNC_MASKED, BLEND_FUNC, true>(DestBuf, Base, Mask, Col );
+		}
+		else
+		{
+			BufferLayerColour<BLEND_FUNC_MASKED, BLEND_FUNC, false>(DestBuf, Base, Mask, Col );
+		}
+	}
+
+
 	//---------------------------------------------------------------------------------------------
 	//! Apply a blending function to an image with another image as blending layer
 	//---------------------------------------------------------------------------------------------
@@ -614,27 +659,25 @@ namespace mu
 		// Number of total channels in the base image
 		int32 BASE_CHANNEL_STRIDE>
 
-	inline void BufferLayerColourFormatInPlace(Image* pBase, vec3<float> col,
-		uint32 BaseChannelOffset,
-		bool bOnlyFirstLOD)
+	inline void BufferLayerColourFormatInPlace(Image* Base, FVector4f Col, uint32 BaseChannelOffset, bool bOnlyFirstLOD)
 	{
-		uint8* pBaseBuf = pBase->GetData() + BaseChannelOffset;
+		uint8* pBaseBuf = Base->GetData() + BaseChannelOffset;
 
 		uint32 top[4];
-		top[0] = (uint32)(255 * col[0]);
-		top[1] = (uint32)(255 * col[1]);
-		top[2] = (uint32)(255 * col[2]);
-		top[3] = (uint32)(255 * col[0]); // used for 1 channel blends to alpha
+		top[0] = (uint32)(255 * Col[0]);
+		top[1] = (uint32)(255 * Col[1]);
+		top[2] = (uint32)(255 * Col[2]);
+		top[3] = (uint32)(255 * Col[0]); // used for 1 channel blends to alpha
 
 		// Generic implementation
 		int32 PixelCount = 0;
 		if (bOnlyFirstLOD)
 		{
-			PixelCount = pBase->GetSizeX() * pBase->GetSizeY();
+			PixelCount = Base->GetSizeX() * Base->GetSizeY();
 		}
 		else
 		{
-			PixelCount = pBase->CalculatePixelCount();
+			PixelCount = Base->CalculatePixelCount();
 		}
 
 		ParallelFor(PixelCount,
@@ -662,32 +705,32 @@ namespace mu
 
 	//---------------------------------------------------------------------------------------------
 	template< unsigned (*BLEND_FUNC)(unsigned, unsigned), bool CLAMP, uint32 CHANNEL_COUNT >
-	inline void BufferLayerColourInPlace(Image* pBase, vec3<float> col, bool bOnlyOneMip, uint32 BaseOffset)
+	inline void BufferLayerColourInPlace(Image* Base, FVector4f Col, bool bOnlyOneMip, uint32 BaseOffset)
 	{
-		EImageFormat baseFormat = pBase->GetFormat();
+		EImageFormat baseFormat = Base->GetFormat();
 
 		if (baseFormat == EImageFormat::IF_RGB_UBYTE)
 		{
 			check(BaseOffset + CHANNEL_COUNT <= 3);
-			BufferLayerColourFormatInPlace< BLEND_FUNC, CLAMP, CHANNEL_COUNT, 3 >(pBase, col, BaseOffset, bOnlyOneMip);
+			BufferLayerColourFormatInPlace< BLEND_FUNC, CLAMP, CHANNEL_COUNT, 3 >(Base, Col, BaseOffset, bOnlyOneMip);
 		}
 		else if (baseFormat == EImageFormat::IF_RGBA_UBYTE)
 		{
 			check(BaseOffset + CHANNEL_COUNT <= 4);
-			BufferLayerColourFormatInPlace< BLEND_FUNC, CLAMP, CHANNEL_COUNT, 4 >(pBase, col, BaseOffset, bOnlyOneMip);
+			BufferLayerColourFormatInPlace< BLEND_FUNC, CLAMP, CHANNEL_COUNT, 4 >(Base, Col, BaseOffset, bOnlyOneMip);
 		}
 		else if (baseFormat == EImageFormat::IF_BGRA_UBYTE)
 		{
-			float temp = col[0];
-			col[0] = col[2];
-			col[2] = temp;
-			BufferLayerColourFormatInPlace< BLEND_FUNC, CLAMP, CHANNEL_COUNT, 4>(pBase, col, BaseOffset, bOnlyOneMip);
+			float temp = Col[0];
+			Col[0] = Col[2];
+			Col[2] = temp;
+			BufferLayerColourFormatInPlace< BLEND_FUNC, CLAMP, CHANNEL_COUNT, 4>(Base, Col, BaseOffset, bOnlyOneMip);
 		}
 
 		else if (baseFormat == EImageFormat::IF_L_UBYTE)
 		{
 			check(BaseOffset + CHANNEL_COUNT <= 1);
-			BufferLayerColourFormatInPlace< BLEND_FUNC, CLAMP, CHANNEL_COUNT, 1 >(pBase, col, BaseOffset, bOnlyOneMip);
+			BufferLayerColourFormatInPlace< BLEND_FUNC, CLAMP, CHANNEL_COUNT, 1 >(Base, Col, BaseOffset, bOnlyOneMip);
 		}
 		else
 		{
@@ -696,41 +739,57 @@ namespace mu
 	}
 
 
+	//---------------------------------------------------------------------------------------------
+	template< unsigned (*BLEND_FUNC)(unsigned, unsigned), uint32 CHANNEL_COUNT >
+	inline void BufferLayerColourInPlace(Image* Base, FVector4f Col, bool bOnlyOneMip, uint32 BaseOffset)
+	{
+		bool bIsClampNeeded = IsAnyComponentLargerThan1(Col);
+		if (bIsClampNeeded)
+		{
+			BufferLayerColourInPlace<BLEND_FUNC, true, CHANNEL_COUNT>(Base, Col, bOnlyOneMip, BaseOffset);
+		}
+		else
+		{
+			BufferLayerColourInPlace<BLEND_FUNC, false, CHANNEL_COUNT>(Base, Col, bOnlyOneMip, BaseOffset);
+		}
+	}
+
+
 	//---------------------------------------------------------------------------------------------
 	template< 
 		unsigned (*BLEND_FUNC_MASKED)(unsigned, unsigned, unsigned),
 		unsigned (*BLEND_FUNC)(unsigned, unsigned), 
 		bool CLAMP, 
 		int32 CHANNEL_COUNT >
-	inline void BufferLayerColourInPlace(Image* pBase, const Image* Mask, vec3<float> col, bool bOnlyOneMip, uint32 BaseOffset)
+	inline void BufferLayerColourInPlace(Image* Base, const Image* Mask, FVector4f Col, bool bOnlyOneMip, uint32 BaseOffset)
 	{
-		check(pBase->GetSizeX() == Mask->GetSizeX());
-		check(pBase->GetSizeY() == Mask->GetSizeY());
+		check(Base->GetSizeX() == Mask->GetSizeX());
+		check(Base->GetSizeY() == Mask->GetSizeY());
 
-		EImageFormat baseFormat = pBase->GetFormat();
+		EImageFormat baseFormat = Base->GetFormat();
 
 		if (baseFormat == EImageFormat::IF_RGB_UBYTE)
 		{
 			check(BaseOffset + CHANNEL_COUNT <= 3);
-			BufferLayerColourFormat< BLEND_FUNC_MASKED, BLEND_FUNC, CLAMP, CHANNEL_COUNT, 3 >(pBase->GetData(), pBase, Mask, col, BaseOffset, bOnlyOneMip);
+			BufferLayerColourFormat< BLEND_FUNC_MASKED, BLEND_FUNC, CLAMP, CHANNEL_COUNT, 3 >(Base->GetData(), Base, Mask, Col, BaseOffset, bOnlyOneMip);
 		}
 		else if (baseFormat == EImageFormat::IF_RGBA_UBYTE)
 		{
 			check(BaseOffset + CHANNEL_COUNT <= 4);
-			BufferLayerColourFormat< BLEND_FUNC_MASKED, BLEND_FUNC, CLAMP, CHANNEL_COUNT, 4 >(pBase->GetData(), pBase, Mask, col, BaseOffset, bOnlyOneMip);
+			BufferLayerColourFormat< BLEND_FUNC_MASKED, BLEND_FUNC, CLAMP, CHANNEL_COUNT, 4 >(Base->GetData(), Base, Mask, Col, BaseOffset, bOnlyOneMip);
 		}
 		else if (baseFormat == EImageFormat::IF_BGRA_UBYTE)
 		{
 			check(BaseOffset + CHANNEL_COUNT <= 4);
-			float temp = col[0];
-			col[0] = col[2];
-			col[2] = temp;
-			BufferLayerColourFormat< BLEND_FUNC_MASKED, BLEND_FUNC, CLAMP, CHANNEL_COUNT, 4 >(pBase->GetData(), pBase, Mask, col, BaseOffset, bOnlyOneMip);
+			float temp = Col[0];
+			Col[0] = Col[2];
+			Col[2] = temp;
+			BufferLayerColourFormat< BLEND_FUNC_MASKED, BLEND_FUNC, CLAMP, CHANNEL_COUNT, 4 >(Base->GetData(), Base, Mask, Col, BaseOffset, bOnlyOneMip);
 		}
 		else if (baseFormat == EImageFormat::IF_L_UBYTE)
 		{
 			check(BaseOffset + CHANNEL_COUNT <= 1);
-			BufferLayerColourFormat< BLEND_FUNC_MASKED, BLEND_FUNC, CLAMP, CHANNEL_COUNT, 1 >(pBase->GetData(), pBase, Mask, col, BaseOffset, bOnlyOneMip);
+			BufferLayerColourFormat< BLEND_FUNC_MASKED, BLEND_FUNC, CLAMP, CHANNEL_COUNT, 1 >(Base->GetData(), Base, Mask, Col, BaseOffset, bOnlyOneMip);
 		}
 		else
 		{
@@ -739,6 +798,22 @@ namespace mu
 	}
 
 
+	//---------------------------------------------------------------------------------------------
+	template< unsigned (*BLEND_FUNC_MASKED)(unsigned, unsigned, unsigned), unsigned (*BLEND_FUNC)(unsigned, unsigned), uint32 CHANNEL_COUNT >
+	inline void BufferLayerColourInPlace(Image* Base, const Image* Mask, FVector4f Col, bool bOnlyOneMip, uint32 BaseOffset)
+	{
+		bool bIsClampNeeded = IsAnyComponentLargerThan1(Col);
+		if (bIsClampNeeded)
+		{
+			BufferLayerColourInPlace<BLEND_FUNC_MASKED, BLEND_FUNC, true, CHANNEL_COUNT>(Base, Mask, Col, bOnlyOneMip, BaseOffset);
+		}
+		else
+		{
+			BufferLayerColourInPlace<BLEND_FUNC_MASKED, BLEND_FUNC, false, CHANNEL_COUNT>(Base, Mask, Col, bOnlyOneMip, BaseOffset);
+		}
+	}
+
+
 	//---------------------------------------------------------------------------------------------
 	//! Apply a blending function to an image with another image as blending layer
 	//---------------------------------------------------------------------------------------------
