[Mutable] Fix RGBA textures reformat to RGB when applying saturation operations.

#rnx

diff --git a/Engine/Plugins/Experimental/Mutable/Source/MutableRuntime/Private/MuR/ImagePrivate.h b/Engine/Plugins/Experimental/Mutable/Source/MutableRuntime/Private/MuR/ImagePrivate.h
--- a/Engine/Plugins/Experimental/Mutable/Source/MutableRuntime/Private/MuR/ImagePrivate.h
+++ b/Engine/Plugins/Experimental/Mutable/Source/MutableRuntime/Private/MuR/ImagePrivate.h
@@ -125,6 +125,39 @@ namespace mu
     }
 
 
+    //---------------------------------------------------------------------------------------------
+    //!
+    //---------------------------------------------------------------------------------------------
+	inline EImageFormat GetRGBOrRGBAFormat(EImageFormat InFormat)
+    {
+		InFormat = GetUncompressedFormat(InFormat);
+
+		if (InFormat == EImageFormat::IF_NONE)
+		{
+			return InFormat;
+		}
+
+		switch (InFormat)
+		{
+		case EImageFormat::IF_L_UBYTE: 
+		{
+			return EImageFormat::IF_RGB_UBYTE;
+		}
+		case EImageFormat::IF_RGB_UBYTE:
+		case EImageFormat::IF_RGBA_UBYTE:
+		case EImageFormat::IF_BGRA_UBYTE:
+		{
+			return InFormat;
+		}
+		default:
+		{
+			unimplemented();
+		}
+		}
+
+		return EImageFormat::IF_NONE;
+    }
+
     //---------------------------------------------------------------------------------------------
     //!
     //---------------------------------------------------------------------------------------------
diff --git a/Engine/Plugins/Experimental/Mutable/Source/MutableTools/Private/MuT/CodeGenerator_Image.cpp b/Engine/Plugins/Experimental/Mutable/Source/MutableTools/Private/MuT/CodeGenerator_Image.cpp
--- a/Engine/Plugins/Experimental/Mutable/Source/MutableTools/Private/MuT/CodeGenerator_Image.cpp
+++ b/Engine/Plugins/Experimental/Mutable/Source/MutableTools/Private/MuT/CodeGenerator_Image.cpp
@@ -905,7 +905,8 @@ namespace mu
             // This argument is required
             base = GenerateMissingImageCode( "Saturate image", EImageFormat::IF_RGB_UBYTE, node.m_errorContext );
         }
-        base = GenerateImageFormat( base, EImageFormat::IF_RGB_UBYTE );
+		
+        base = GenerateImageFormat(base, GetRGBOrRGBAFormat(base->GetImageDesc().m_format));
         base = GenerateImageSize
                 ( base, FImageSize((uint16)m_imageState.Last().m_imageRect.size[0],(uint16)m_imageState.Last().m_imageRect.size[1]) );
         op->SetChild( op->op.args.ImageSaturate.base, base);
