[mutable] Optimization of pixel format operations and pixel invert.
#rnx

diff --git a/Engine/Plugins/Experimental/Mutable/Source/MutableTools/Private/MuT/ASTOpImagePixelFormat.cpp b/Engine/Plugins/Experimental/Mutable/Source/MutableTools/Private/MuT/ASTOpImagePixelFormat.cpp
--- a/Engine/Plugins/Experimental/Mutable/Source/MutableTools/Private/MuT/ASTOpImagePixelFormat.cpp
+++ b/Engine/Plugins/Experimental/Mutable/Source/MutableTools/Private/MuT/ASTOpImagePixelFormat.cpp
@@ -140,6 +140,22 @@ namespace mu
 			break;
 		}
 
+		case OP_TYPE::IM_INVERT:
+		{
+			// This op doesn't support compressed formats
+			if (!bIsCompressedFormat)
+			{
+				mu::Ptr<ASTOpFixed> newOp = mu::Clone<ASTOpFixed>(sourceAt);
+
+				mu::Ptr<ASTOpImagePixelFormat> fop = mu::Clone<ASTOpImagePixelFormat>(this);
+				fop->Source = newOp->children[newOp->op.args.ImageInvert.base].child();
+				newOp->SetChild(newOp->op.args.ImageInvert.base, fop);
+
+				at = newOp;
+			}
+			break;
+		}
+
 		case OP_TYPE::IM_RASTERMESH:
 		{
 			// This op doesn't support compressed formats
