[mutable] Fix for static analysis warning.
#rnx

diff --git a/Engine/Plugins/Experimental/Mutable/Source/MutableRuntime/Private/MuR/OpImageProject.cpp b/Engine/Plugins/Experimental/Mutable/Source/MutableRuntime/Private/MuR/OpImageProject.cpp
--- a/Engine/Plugins/Experimental/Mutable/Source/MutableRuntime/Private/MuR/OpImageProject.cpp
+++ b/Engine/Plugins/Experimental/Mutable/Source/MutableRuntime/Private/MuR/OpImageProject.cpp
@@ -540,14 +540,14 @@ namespace
 					{
 						if (m_bIsRGBFadingEnabled)
 						{
-							for (int i = 0; i < PIXEL_SIZE && i < 3; ++i)
+							for (int i = 0; i < FMath::Min(PIXEL_SIZE, 3); ++i)
 							{
 								pBufferPos[i] = uint8((pPixel[i] * factor_8) >> 8);
 							}
 						}
 						else
 						{
-							for (int i = 0; i < PIXEL_SIZE && i < 3; ++i)
+							for (int i = 0; i < FMath::Min(PIXEL_SIZE, 3); ++i)
 							{
 								pBufferPos[i] = uint8((pPixel[i] * factor_unfaded_8) >> 8);
 							}
