[mutable] Mutable Code Viewer : Fixed operation property not being set correctly (constant and dynamic)
- The culprint was an incorrect asumption made in order to improve overall performance.
- Now, each time an element gets expanded it will check if their children could be visible after the expansion. If so their state will get updated which, at the same time, will update the dynamic or constant bool values.

diff --git a/Engine/Plugins/Experimental/Mutable/Source/CustomizableObjectEditor/Private/MuCOE/SMutableCodeViewer.cpp b/Engine/Plugins/Experimental/Mutable/Source/CustomizableObjectEditor/Private/MuCOE/SMutableCodeViewer.cpp
--- a/Engine/Plugins/Experimental/Mutable/Source/CustomizableObjectEditor/Private/MuCOE/SMutableCodeViewer.cpp
+++ b/Engine/Plugins/Experimental/Mutable/Source/CustomizableObjectEditor/Private/MuCOE/SMutableCodeViewer.cpp
@@ -72,6 +72,41 @@ public:
 		);
 	}
 
+	FLinearColor OnGetExtraDataBoxColor() const
+	{
+		if (RowItem->bIsDynamicResource)
+		{
+			return DynamicResourceBoxColor;
+		}
+		else if (RowItem->bIsStateConstant)
+		{
+			return StateConstantBoxColor;
+		}
+		else
+		{
+			return ExtraDataBackgroundBoxDefaultColor;
+		}
+	}
+
+	FText OnGetExtraDataText() const
+	{
+		// DEBUG :Uncomment the next line in order to debug the current state being used by the element
+		// return FText::FromString(FString::FromInt(RowItem->GetStateIndex()));
+		
+		if (RowItem->bIsDynamicResource)
+		{
+			return  DynamicResourceText;
+		}
+		else if (RowItem->bIsStateConstant)
+		{
+			return StateConstantText;
+		}
+		else
+		{
+			return FText::FromString(FString(""));
+		}
+	}
+
 	/** Method intended with the generation of the wanted objects for each column*/
 	virtual TSharedRef<SWidget> GenerateWidgetForColumn(const FName& ColumnName) override
 	{
@@ -180,11 +215,15 @@ public:
 			// Prepare the text shown on the UI side of the operation tree
 			FString MainLabel = FString::Printf(TEXT("%s %d : %s"), *RowItem->Caption, int32(RowItem->MutableOperation), *OpName);
 
-#if UE_BUILD_DEBUG
-			FString IndexOnTree = FString::FromInt(RowItem->IndexOnTree);
-			IndexOnTree.Append(TEXT("- "));
-			MainLabel.InsertAt(0,IndexOnTree);
-#endif
+			// DEBUG : 
+			// FString IndexOnTree = FString::FromInt(RowItem->IndexOnTree);
+			// IndexOnTree.Append(TEXT("- "));
+			// MainLabel.InsertAt(0,IndexOnTree);
+
+			// DEBUG : 
+			// FString RowStateIndex = FString::FromInt(RowItem->GetStateIndex());
+			// RowStateIndex.Append(TEXT(" st "));
+			// MainLabel.InsertAt(0,RowStateIndex);
 
 			
 			if (RowItem->DuplicatedOf)
@@ -249,20 +288,6 @@ public:
 		// Second column showing some extra data related with the operation being displayed
 		if (ColumnName == MutableCodeTreeViewColumns::AdditionalDataColumnID)
 		{
-			// Determine what text will be shown on the extra data column alongside with the background color used for it
-			FText OperationExtraData = FText::FromString(FString(""));
-			const FSlateColor* OperationDataBoxColor = &ExtraDataBackgroundBoxDefaultColor;
-			if (RowItem->bIsDynamicResource)
-			{
-				OperationExtraData = DynamicResourceText;
-				OperationDataBoxColor = &DynamicResourceBoxColor;
-			}
-			else if (RowItem->bIsStateConstant)
-			{
-				OperationExtraData = StateConstantText;
-				OperationDataBoxColor = &StateConstantBoxColor;
-			}
-			
 			TSharedRef<SHorizontalBox> RowContainer =  SNew(SHorizontalBox)
 			
 			+ SHorizontalBox::Slot()
@@ -270,31 +295,19 @@ public:
 			[
 				SNew(SHorizontalBox)
 
-				// Cheap way of doing a fixed space with a color
 				+SHorizontalBox::Slot()
-				.AutoWidth()
+				.MaxWidth(4.0f)
 				[
-					 SNew(SOverlay)
-					 +SOverlay::Slot()
-					 [
-						SNew(SMutableColorPreviewBox)
-						.BoxColor(*OperationDataBoxColor)
-					 ]
-
-					 // Hacky way of setting a fixed size for the color box
-					 + SOverlay::Slot()
-					 [
-						SNew(STextBlock)
-						.Text(EmptyText)
-					 ]
+					SNew(SColorBlock)
+					.Color(this,&SMutableCodeTreeRow::OnGetExtraDataBoxColor)
 				]
 
 				+ SHorizontalBox::Slot()
 				.Padding(4,1)
-				.FillWidth(0.9f)
+				.AutoWidth()
 				[
 					SNew(STextBlock)
-					.Text(OperationExtraData)
+					.Text(this,&SMutableCodeTreeRow::OnGetExtraDataText)
 				]
 			];
 			
@@ -368,13 +381,13 @@ private:
 	const FText StateConstantText = FText::FromString(FString("const"));
 	
 	/** Color used on the extra data column when no extra data is shown */
-	const FSlateColor ExtraDataBackgroundBoxDefaultColor =  FSlateColor(TransparentColor);
+	const FLinearColor ExtraDataBackgroundBoxDefaultColor =  TransparentColor;
 
 	/** Color shown on the extra data column when the resource is found to be Dynamic */
-	const FSlateColor DynamicResourceBoxColor = FSlateColor(FLinearColor(0,0,1,0.8));
+	const FLinearColor DynamicResourceBoxColor = FLinearColor(0,0,1,0.8);
 
 	/** Color shown on the extra data column when the resource is found to be State Constant */
-	const FSlateColor StateConstantBoxColor = FSlateColor(FLinearColor(1,0,0,0.8));
+	const FLinearColor StateConstantBoxColor = FLinearColor(1,0,0,0.8);
 };
 
 
@@ -703,7 +716,14 @@ void SMutableCodeViewer::Construct(const FArguments& InArgs, const TSharedPtr<mu
 	];
 	
 	// Set the tree expanded by default
+	// It does not recalculate states since the expansion of the instance will NOT expand duplicates witch means the widget position
+	// of the children of duplicated (or the original of an operation with duplicates) will not change.
 	TreeExpandInstance();
+
+	// Enable the recalculation of states once the tree has already been initially expanded since now we do not control
+	// how the user is point to interact with the view.
+	bShouldRecalculateStates = true;
+	// Now, on expansion or contraction the states will get recalculated
 } 
 
 
@@ -1136,8 +1156,8 @@ void SMutableCodeViewer::GenerateAllTreeElements()
 
 		const FSlateColor LabelColor = ColorPerComputationalCost[StaticCast<uint8>(GetOperationTypeComputationalCost(
 			ModelPrivate->m_program.GetOpType(State.m_root)))];
-		
-		RootNodes.Add(MakeShareable(new FMutableCodeTreeElement(ItemCache.Num(), MutableModel, State.m_root, Caption,LabelColor)));
+
+		RootNodes.Add(MakeShareable(new FMutableCodeTreeElement(ItemCache.Num(),StateIndex, MutableModel, State.m_root, Caption,LabelColor)));
 
 		// Iterate over each root node and generate all the elements in a human readable pattern (Z Pattern)
 
@@ -1146,16 +1166,16 @@ void SMutableCodeViewer::GenerateAllTreeElements()
 
 		// Add the element to the cache so we keep the indices straight.
 		ItemCache.Add(Key, RootNodes.Last());
-		GenerateElementRecursive(State.m_root,Program);
+		GenerateElementRecursive(StateIndex,State.m_root,Program);
 	}
 	
 }
 
-void SMutableCodeViewer::GenerateElementRecursive(mu::OP::ADDRESS InParentAddress,  const mu::FProgram& InProgram)
+void SMutableCodeViewer::GenerateElementRecursive(const int32& InStateIndex, mu::OP::ADDRESS InParentAddress,  const mu::FProgram& InProgram)
 {
 	// This will be used to add operations
 	uint32 ChildIndex = 0;
-	auto AddOpFunc = [this, InParentAddress, &InProgram, &ChildIndex](mu::OP::ADDRESS ChildAddress, const FString& Caption)
+	auto AddOpFunc = [this, InParentAddress, &InProgram, &ChildIndex, &InStateIndex](mu::OP::ADDRESS ChildAddress, const FString& Caption)
 		{
 			if (ChildAddress)
 			{
@@ -1173,7 +1193,7 @@ void SMutableCodeViewer::GenerateElementRecursive(mu::OP::ADDRESS InParentAddres
 					const FSlateColor LabelColor = ColorPerComputationalCost[StaticCast<uint8>(GetOperationTypeComputationalCost(InProgram.GetOpType(ChildAddress)))];
 
 					// No caption for the generic tree
-					const TSharedPtr<FMutableCodeTreeElement> Item = MakeShareable(new FMutableCodeTreeElement(ItemCache.Num(), MutableModel, ChildAddress, Caption, LabelColor, MainItemPtr));
+					const TSharedPtr<FMutableCodeTreeElement> Item = MakeShareable(new FMutableCodeTreeElement(ItemCache.Num(),InStateIndex, MutableModel, ChildAddress, Caption, LabelColor, MainItemPtr));
 
 					// Cache this element for later access
 					ItemCache.Add(Key, Item);
@@ -1183,7 +1203,7 @@ void SMutableCodeViewer::GenerateElementRecursive(mu::OP::ADDRESS InParentAddres
 					{
 						MainItemPerOp.Add(ChildAddress, Item);
 
-						GenerateElementRecursive(ChildAddress, InProgram);
+						GenerateElementRecursive(InStateIndex,ChildAddress, InProgram);
 					}
 				}
 			}
@@ -1362,12 +1382,7 @@ void SMutableCodeViewer::GetChildrenForInfo(TSharedPtr<FMutableCodeTreeElement>
 		return;
 	}
 
-	// If this is a duplicated of another row, don't provide its children.
-	//if (InInfo->DuplicatedOf)
-	//{
-	//	return;
-	//}
-
+	check(MutableModel);
 	const mu::FProgram& Program = MutableModel->GetPrivate()->m_program;
 
 	mu::OP::ADDRESS ParentAddress = InInfo->MutableOperation;
@@ -1397,14 +1412,30 @@ void SMutableCodeViewer::GetChildrenForInfo(TSharedPtr<FMutableCodeTreeElement>
 
 void SMutableCodeViewer::OnExpansionChanged(TSharedPtr<FMutableCodeTreeElement> InItem, bool bInExpanded)
 {
+	// Update expanded state of the provided element
+	InItem->bIsExpanded = bInExpanded;
+	
 	// If an element gets expanded then contract (if found) the other element that uses the same address
 	if (bInExpanded)
 	{
 		const mu::OP::ADDRESS MutableOperation = InItem->MutableOperation;
-		const TSharedPtr<FMutableCodeTreeElement>* ExpandedElement = ExpandedElements.Find(MutableOperation);
-		if (ExpandedElement)
+		const TSharedPtr<FMutableCodeTreeElement>* PreviouslyExpandedElement = ExpandedElements.Find(MutableOperation);
+		if (PreviouslyExpandedElement)
+		{
+			TreeView->SetItemExpansion(*PreviouslyExpandedElement, false);
+		}
+
+		// Only do this if in a situation where it may be required (do not do it if the tree has not been interacted with yet)
+		if (bShouldRecalculateStates)
 		{
-			TreeView->SetItemExpansion(*ExpandedElement, false);
+			// Find all the children (recursive) of this item.
+			TSet<TSharedPtr<FMutableCodeTreeElement>> FoundChildren;
+			GetVisibleChildren(InItem, FoundChildren);
+			for (const TSharedPtr<FMutableCodeTreeElement>& Child : FoundChildren)
+			{
+				// For each of the children found set it's state to be the one found on the expanded element
+				Child->SetElementCurrentState(InItem->GetStateIndex());
+			}
 		}
 		
 		// Cache this element as one currently expanded
@@ -1417,6 +1448,53 @@ void SMutableCodeViewer::OnExpansionChanged(TSharedPtr<FMutableCodeTreeElement>
 	}
 }
 
+void SMutableCodeViewer::GetVisibleChildren(TSharedPtr<FMutableCodeTreeElement> InInfo, TSet<TSharedPtr<FMutableCodeTreeElement>>& OutChildren)
+{
+	check(MutableModel);
+	const mu::FProgram MutableProgram = MutableModel->GetPrivate()->m_program;
+	
+	TArray<TSharedPtr<FMutableCodeTreeElement>> ToSearchForChildren;
+	ToSearchForChildren.Add(InInfo);
+	while (!ToSearchForChildren.IsEmpty())
+	{
+		// Grab the first element in order to check for it's children
+		const TSharedPtr<FMutableCodeTreeElement> ToCheck = ToSearchForChildren[0];
+		ToSearchForChildren.RemoveAt(0);
+		
+		const mu::OP::ADDRESS ParentAddress = ToCheck->MutableOperation;
+	
+		// Generic case for unnamed children traversal.
+		uint32 ChildIndex = 0;
+		mu::ForEachReference(MutableProgram, ParentAddress, [this, ParentAddress, &ChildIndex, &OutChildren, &ToSearchForChildren ](mu::OP::ADDRESS ChildAddress)
+		{
+			if (ChildAddress)
+			{
+				const FItemCacheKey Key = { ParentAddress, ChildAddress, ChildIndex };
+				const TSharedPtr<FMutableCodeTreeElement> CachedItem = *ItemCache.Find(Key);
+
+				// Since we have already generated all elements CachedItem should be therefore always a valid pointer
+				check (CachedItem);
+				
+				// If the address has not been yet found then save it as one of the children affected
+				if (!OutChildren.Contains(CachedItem))
+				{
+					OutChildren.Add(CachedItem);
+
+					// And if the children is found to be expanded then also process it later to later return only the
+					// elements that are expanded in the tree view (using data manually set on each tree element)
+					if (CachedItem->bIsExpanded )
+					{
+						// Add for processing
+						ToSearchForChildren.Add(CachedItem);
+					}
+				}
+			}
+			++ChildIndex;
+		});
+	}
+	
+	UE_LOG(LogTemp,Warning,TEXT("Found a total of %i children elements "),OutChildren.Num());
+}
 
 
 void SMutableCodeViewer::OnSelectionChanged(TSharedPtr<FMutableCodeTreeElement> InNode, ESelectInfo::Type InSelectInfo)
diff --git a/Engine/Plugins/Experimental/Mutable/Source/CustomizableObjectEditor/Private/MuCOE/SMutableCodeViewer.h b/Engine/Plugins/Experimental/Mutable/Source/CustomizableObjectEditor/Private/MuCOE/SMutableCodeViewer.h
--- a/Engine/Plugins/Experimental/Mutable/Source/CustomizableObjectEditor/Private/MuCOE/SMutableCodeViewer.h
+++ b/Engine/Plugins/Experimental/Mutable/Source/CustomizableObjectEditor/Private/MuCOE/SMutableCodeViewer.h
@@ -180,7 +180,7 @@ private:
 	void GenerateAllTreeElements();
 
 	/** Support GenerateAllTreeElements by generating the elements recursively */
-	void GenerateElementRecursive(mu::OP::ADDRESS InParentAddress, const mu::FProgram& InProgram);
+	void GenerateElementRecursive( const int32& InStateIndex, mu::OP::ADDRESS InParentAddress, const mu::FProgram& InProgram);
 
 	
 	/** The addresses of the root operations. Cached when this object gets loaded on the Construct method of this slate */
@@ -639,6 +639,25 @@ private:
 	/** Called when any parameter value has changed, with the parameter index as argument.  */
 	void OnPreviewParameterValueChanged( int32 ParamIndex );
 
+	/*
+	* Element state recalculation 
+	*/
+	
+	/** FLag that determines if the expansion of elements should require the computation of the state of it's child elements.
+	 * Useful to avoid unnecessary calls for state resting of movable branches like the children of operations that are repeated
+	 * in the tree slate */
+	bool bShouldRecalculateStates = false;
+	
+	/** Provided an element of the tree view return all the children of it (direct and indirect). The search will stop (depth)
+	 * if a children is found to not be expanded. This way we avoid infinite searches (since we do not stop at duplicates).
+	 * @note It DOES NOT check for visibility by asking the TreeView. It only gets the elements that could be seen in theory.
+	 * @note This method will return expanded and not expanded children but only the expanded children will get it's children searched and therefore
+	 * provided by this method. In conclusion, only the elements that could be seen by the view will get provide.
+	 * @param InInfo The element whose children we want to grab. It can be expanded or not.
+	 * @param OutChildren A collection of children elements. Only the elements that could be seen in the tree view (expanded and the root element of an
+	 * unexpanded branch) wll get added to this collection.
+	 */
+	void GetVisibleChildren(TSharedPtr<FMutableCodeTreeElement> InInfo, TSet<TSharedPtr<FMutableCodeTreeElement>>& OutChildren);
 	
 
 	/*
@@ -715,7 +734,7 @@ public:
 class FMutableCodeTreeElement : public TSharedFromThis<FMutableCodeTreeElement>
 {
 public:
-	FMutableCodeTreeElement(int32 InIndexOnTree ,const TSharedPtr<mu::Model, ESPMode::ThreadSafe>& InModel, mu::OP::ADDRESS InOperation, const FString& InCaption,const FSlateColor InLabelColor, const TSharedPtr<FMutableCodeTreeElement>* InDuplicatedOf = nullptr)
+	FMutableCodeTreeElement(int32 InIndexOnTree  , const int32& InMutableStateIndex ,const TSharedPtr<mu::Model, ESPMode::ThreadSafe>& InModel, mu::OP::ADDRESS InOperation, const FString& InCaption,const FSlateColor InLabelColor, const TSharedPtr<FMutableCodeTreeElement>* InDuplicatedOf = nullptr)
 	{
 		MutableModel = InModel;
 		MutableOperation = InOperation;
@@ -726,61 +745,60 @@ public:
 		{
 			DuplicatedOf = *InDuplicatedOf;
 		}
-
-		// Check what type of operation is (state constant or dynamic resource)
+		
+		// Process the data that can be extracted from the current state
+		SetElementCurrentState(InMutableStateIndex);
+	}
+	
+	void SetElementCurrentState(const int32& InStateIndex)
+	{
+		// Skip operation if state is the same
+		if (InStateIndex == CurrentMutableStateIndex)
 		{
-			// If duplicated then grab the already processed data on the original operation
-			if (InDuplicatedOf)
-			{
-				bIsDynamicResource = DuplicatedOf->bIsDynamicResource;
-				bIsStateConstant = DuplicatedOf->bIsStateConstant;
+			return;
+		}
 
-				// All required data has been processed so an early exit is required
-				return;
-			}
-			
-			// Iterate over all states and try to locate the operation
-			const mu::FProgram& MutableProgram = InModel->GetPrivate()->m_program;
-			for (const mu::FProgram::FState& CurrentState : MutableProgram.m_states)
+		// Check for an out of bounds value
+		check(MutableModel)
+		mu::FProgram& MutableProgram = MutableModel->GetPrivate()->m_program;
+		check (InStateIndex >= 0 && InStateIndex < MutableProgram.m_states.Num());
+		
+		CurrentMutableStateIndex = InStateIndex;
+		const mu::FProgram::FState& CurrentState = MutableProgram.m_states[CurrentMutableStateIndex];
+		
+		// Check if it is a dynamic resource
+		for (auto& DynamicResource : CurrentState.m_dynamicResources)
+		{
+			// If the operation gets located then mark it as dynamic resource
+			if (DynamicResource.Key == MutableOperation)
 			{
-				// Check if it is a dynamic resource
-				for (auto& DynamicResource : CurrentState.m_dynamicResources)
-				{
-					// If the operation gets located then mark it as dynamic resource
-					if (DynamicResource.Key == MutableOperation)
-					{
-						bIsDynamicResource = true;
-						break;
-					}
-				}
-				
-				// Early exit: A dynamic resource can not be at the same time a state constant
-				if (bIsDynamicResource)
-				{
-					return;
-				}
-				
-				// Check if it is a state constant
-				bIsStateConstant = CurrentState.m_updateCache.Contains(MutableOperation);
+				bIsDynamicResource = true;
+				break;
 			}
 		}
+		// Early exit: A dynamic resource can not be at the same time a state constant
+		if (bIsDynamicResource)
+		{
+			return;
+		}
 		
+		// Check if it is a state constant
+		bIsStateConstant = CurrentState.m_updateCache.Contains(MutableOperation);
 	}
 
+	int32 GetStateIndex() const
+	{
+		return CurrentMutableStateIndex;
+	}
+	
 public:
 
 	/** */
 	TSharedPtr<mu::Model, ESPMode::ThreadSafe> MutableModel;
-
+	
 	/** Mutable Graph Node represented in this tree row. */
 	mu::OP::ADDRESS MutableOperation;
 
-	/** If true means that it will not update when a runtime parameter on the state gets updated */
-	bool bIsStateConstant = false;
-
-	/** If true then the mesh or image of this operation may change during the state update */
-	bool bIsDynamicResource = false;
-
 	/** Label representing this operation. */
 	FString Caption;
 
@@ -796,4 +814,29 @@ public:
 	
 	/** The current position of this element on the tree view. Used for navigation */
 	int32 IndexOnTree;
+	
+	/*
+	 * Dynamic data : Can and will change during the standard operation of the tree view object.
+	 */
+
+	/** If true means that it will not update when a runtime parameter on the state gets updated */
+	bool bIsStateConstant = false;
+
+	/** If true then the mesh or image of this operation may change during the state update */
+	bool bIsDynamicResource = false;
+
+	/** Flag that reflects the expanded state of the element in the tree view. It is used in order to know when an element
+	 * should or should not has it's set state index updated (and the data that comes with it).
+	 * @note It does not mean the element can be seen at the moment in the view, just that, in case of expansion of a parent
+	 * element then this could be part of the view (if inside view space).
+	 */
+	bool bIsExpanded = false;
+
+private:
+	/** Represents as part of what state this element is currently part of. Will change for elements that are children
+	 * of duplicated elements (duplicated) since children of duplicated elements are still unique and therefore they can
+	 * be present in more than one state (but can just appear in one state at any given time)
+	 */
+	int32 CurrentMutableStateIndex = -1;	
+
 };
