[mutable] Fix incorrect error message showing up when compiling COs

#rnx

diff --git a/Engine/Plugins/Experimental/Mutable/Source/CustomizableObjectEditor/Private/MuCOE/CustomizableObjectCompileRunnable.cpp b/Engine/Plugins/Experimental/Mutable/Source/CustomizableObjectEditor/Private/MuCOE/CustomizableObjectCompileRunnable.cpp
--- a/Engine/Plugins/Experimental/Mutable/Source/CustomizableObjectEditor/Private/MuCOE/CustomizableObjectCompileRunnable.cpp
+++ b/Engine/Plugins/Experimental/Mutable/Source/CustomizableObjectEditor/Private/MuCOE/CustomizableObjectCompileRunnable.cpp
@@ -213,13 +213,15 @@ uint32 FCustomizableObjectSaveDDRunnable::Run()
 
 		// Delete files...
 		bool bFilesDeleted = true;
-		if (!FileManager.Delete(*CompildeDataFullFileName, true, false, true))
+		if (FileManager.FileExists(*CompildeDataFullFileName)
+			&& !FileManager.Delete(*CompildeDataFullFileName, true, false, true))
 		{
 			UE_LOG(LogMutable, Error, TEXT("Failed to delete compiled data in file [%s]."), *CompildeDataFullFileName);
 			bFilesDeleted = false;
 		}
 
-		if (!FileManager.Delete(*StreamableDataFullFileName, true, false, true))
+		if (FileManager.FileExists(*StreamableDataFullFileName)
+			&& !FileManager.Delete(*StreamableDataFullFileName, true, false, true))
 		{
 			UE_LOG(LogMutable, Error, TEXT("Failed to delete streamed data in file [%s]."), *StreamableDataFullFileName);
 			bFilesDeleted = false;
