[mutable] Additional optimisation of the image operations recently used.
- image invert will try to reuse the image buffer.
- image layer will rebuild the mips at the end instead of before applying.
#rnx

diff --git a/Engine/Plugins/Experimental/Mutable/Source/MutableRuntime/Private/MuR/CodeRunner.cpp b/Engine/Plugins/Experimental/Mutable/Source/MutableRuntime/Private/MuR/CodeRunner.cpp
--- a/Engine/Plugins/Experimental/Mutable/Source/MutableRuntime/Private/MuR/CodeRunner.cpp
+++ b/Engine/Plugins/Experimental/Mutable/Source/MutableRuntime/Private/MuR/CodeRunner.cpp
@@ -3614,7 +3614,16 @@ namespace mu
 			{
 				Ptr<const Image> pA = GetMemory().GetImage(FCacheAddress(args.base, item));
 
-				ImagePtr pResult = ImageInvert(pA.get());
+				ImagePtr pResult;
+				if (pA->IsUnique())
+				{
+					pResult = mu::CloneOrTakeOver<>(pA.get());
+					ImageInvertInPlace(pResult.get());
+				}
+				else
+				{
+					pResult = ImageInvert(pA.get());
+				}
 
 				GetMemory().SetImage(item, pResult);
 				break;
diff --git a/Engine/Plugins/Experimental/Mutable/Source/MutableRuntime/Private/MuR/CodeRunnerTasked.cpp b/Engine/Plugins/Experimental/Mutable/Source/MutableRuntime/Private/MuR/CodeRunnerTasked.cpp
--- a/Engine/Plugins/Experimental/Mutable/Source/MutableRuntime/Private/MuR/CodeRunnerTasked.cpp
+++ b/Engine/Plugins/Experimental/Mutable/Source/MutableRuntime/Private/MuR/CodeRunnerTasked.cpp
@@ -494,27 +494,14 @@ namespace mu
 
 			check(pNew->GetDataSize() == m_base->GetDataSize());
 
+			bool bOnlyOneMip = false;
 			if (m_blended->GetLODCount() < m_base->GetLODCount())
 			{
-				MUTABLE_CPUPROFILER_SCOPE(ImageLayer_MipBlendedEmergencyFix);
-
-				int32 levelCount = m_base->GetLODCount();
-				ImagePtr pDest = new Image(m_blended->GetSizeX(), m_blended->GetSizeY(), levelCount, m_blended->GetFormat());
-
-				SCRATCH_IMAGE_MIPMAP scratch;
-				FMipmapGenerationSettings settings{};
-
-				ImageMipmap_PrepareScratch(pDest.get(), m_blended.get(), levelCount, &scratch);
-				ImageMipmap(m_imageCompressionQuality, pDest.get(), m_blended.get(), levelCount,
-					&scratch, settings);
-
-				m_blended = pDest;
+				bOnlyOneMip = true;
 			}
 
 			bool bDone = false;
 
-			bool bOnlyOneMip = false;
-
 			bool bUseMaskFromBlendAlpha = (Args.flags & OP::ImageLayerArgs::F_USE_MASK_FROM_BLENDED);
 
 			if (!m_mask && bUseMaskFromBlendAlpha
@@ -629,6 +616,13 @@ namespace mu
 				default: check(false);
 				}
 			}
+
+			if (bOnlyOneMip)
+			{
+				MUTABLE_CPUPROFILER_SCOPE(ImageLayer_MipFix);
+				FMipmapGenerationSettings DummyMipSettings{};
+				ImageMipmapInPlace(m_imageCompressionQuality, pNew.get(), DummyMipSettings);
+			}
 		}
 
 		if (InitialFormat != pNew->GetFormat())
diff --git a/Engine/Plugins/Experimental/Mutable/Source/MutableRuntime/Private/MuR/OpImageInvert.h b/Engine/Plugins/Experimental/Mutable/Source/MutableRuntime/Private/MuR/OpImageInvert.h
--- a/Engine/Plugins/Experimental/Mutable/Source/MutableRuntime/Private/MuR/OpImageInvert.h
+++ b/Engine/Plugins/Experimental/Mutable/Source/MutableRuntime/Private/MuR/OpImageInvert.h
@@ -15,13 +15,13 @@ namespace mu
 	{
 		MUTABLE_CPUPROFILER_SCOPE(ImageInvert)
 
-		ImagePtr pDest = new Image(pA->GetSizeX(), pA->GetSizeY(), pA->GetLODCount(), pA->GetFormat());
+			ImagePtr pDest = new Image(pA->GetSizeX(), pA->GetSizeY(), pA->GetLODCount(), pA->GetFormat());
 
 		uint8_t* pDestBuf = pDest->GetData();
 		const uint8_t* pABuf = pA->GetData();
 
 		//Generic implementation
-		int pixelCount = (int)pA->CalculatePixelCount();
+		int32 pixelCount = pA->CalculatePixelCount();
 
 		switch (pA->GetFormat())
 		{
@@ -33,7 +33,7 @@ namespace mu
 			}
 			break;
 		}
-		
+
 		case EImageFormat::IF_RGB_UBYTE:
 		{
 			for (int i = 0; i < pixelCount; ++i)
@@ -60,9 +60,60 @@ namespace mu
 			break;
 		}
 		default:
-			check(false);			
+			check(false);
 		}
 
 		return pDest;
 	}
+
+	//---------------------------------------------------------------------------------------------
+	//!
+	//---------------------------------------------------------------------------------------------
+	inline void ImageInvertInPlace(Image* pA)
+	{
+		MUTABLE_CPUPROFILER_SCOPE(ImageInvertInPlace)
+
+		uint8_t* pABuf = pA->GetData();
+
+		//Generic implementation
+		int32 pixelCount = pA->CalculatePixelCount();
+
+		switch (pA->GetFormat())
+		{
+		case EImageFormat::IF_L_UBYTE:
+		{
+			for (int i = 0; i < pixelCount; ++i)
+			{
+				pABuf[i] = 255 - pABuf[i];
+			}
+			break;
+		}
+
+		case EImageFormat::IF_RGB_UBYTE:
+		{
+			for (int i = 0; i < pixelCount; ++i)
+			{
+				for (int c = 0; c < 3; ++c)
+				{
+					pABuf[i * 3 + c] = 255 - pABuf[i * 3 + c];
+				}
+			}
+			break;
+		}
+		case EImageFormat::IF_RGBA_UBYTE:
+		case EImageFormat::IF_BGRA_UBYTE:
+		{
+			for (int i = 0; i < pixelCount; ++i)
+			{
+				for (int c = 0; c < 3; ++c)
+				{
+					pABuf[i * 4 + c] = 255 - pABuf[i * 4 + c];
+				}
+			}
+			break;
+		}
+		default:
+			check(false);
+		}
+	}
 }
