[mutable] Fix ensure failure when LoadingExternalImages in the editor

#rnx

diff --git a/Engine/Plugins/Experimental/Mutable/Source/CustomizableObject/Private/MuCO/UnrealMutableImageProvider.cpp b/Engine/Plugins/Experimental/Mutable/Source/CustomizableObject/Private/MuCO/UnrealMutableImageProvider.cpp
--- a/Engine/Plugins/Experimental/Mutable/Source/CustomizableObject/Private/MuCO/UnrealMutableImageProvider.cpp
+++ b/Engine/Plugins/Experimental/Mutable/Source/CustomizableObject/Private/MuCO/UnrealMutableImageProvider.cpp
@@ -167,7 +167,7 @@ TTuple<FGraphEventRef, TFunction<void()>> FUnrealMutableImageProvider::GetImageA
 			// not being modified by the game thread at the moment and the texture cannot be GCed because of the AddReferencedObjects
 			// in the FUnrealMutableImageProvider
 
-			uint8 MipIndex = MipmapsToSkip < TextureToLoad->GetNumMips() ? MipmapsToSkip : TextureToLoad->GetNumMips() - 1;
+			uint8 MipIndex = MipmapsToSkip < TextureToLoad->GetPlatformData()->Mips.Num() ? MipmapsToSkip : TextureToLoad->GetPlatformData()->Mips.Num() - 1;
 
 			// Mips in the mip tail are inlined and can't be streamed, find the smallest mip available.
 			for (; MipIndex > 0; --MipIndex)
@@ -379,7 +379,7 @@ mu::FImageDesc FUnrealMutableImageProvider::GetImageDesc(mu::EXTERNAL_IMAGE_ID i
 			// not being modified by the game thread at the moment and the texture cannot be GCed because of the AddReferencedObjects
 			// in the FUnrealMutableImageProvider
 
-			uint8 MipIndex = MipmapsToSkip < TextureToLoad->GetNumMips() ? MipmapsToSkip : TextureToLoad->GetNumMips() - 1;
+			uint8 MipIndex = MipmapsToSkip < TextureToLoad->GetPlatformData()->Mips.Num() ? MipmapsToSkip : TextureToLoad->GetPlatformData()->Mips.Num() - 1;
 
 			// Mips in the mip tail are inlined and can't be streamed, find the smallest mip available.
 			for (; MipIndex > 0; --MipIndex)
