#rnx

diff --git a/Engine/Plugins/Experimental/Mutable/Source/CustomizableObject/Classes/MuCO/CustomizableObject.h b/Engine/Plugins/Experimental/Mutable/Source/CustomizableObject/Classes/MuCO/CustomizableObject.h
--- a/Engine/Plugins/Experimental/Mutable/Source/CustomizableObject/Classes/MuCO/CustomizableObject.h
+++ b/Engine/Plugins/Experimental/Mutable/Source/CustomizableObject/Classes/MuCO/CustomizableObject.h
@@ -1203,7 +1203,7 @@ private:
 	// This is a manual version number for the binary blobs in this asset.
 	// Increasing it invalidates all the previously compiled models.
 	// Warning: If while merging code both versions have changed, take the highest+1.
-	static const int32 CurrentSupportedVersion = 364;
+	static const int32 CurrentSupportedVersion = 365;
 
 public:
 
diff --git a/Engine/Plugins/Experimental/Mutable/Source/CustomizableObject/Private/MuCO/CustomizableObjectInstance.cpp b/Engine/Plugins/Experimental/Mutable/Source/CustomizableObject/Private/MuCO/CustomizableObjectInstance.cpp
--- a/Engine/Plugins/Experimental/Mutable/Source/CustomizableObject/Private/MuCO/CustomizableObjectInstance.cpp
+++ b/Engine/Plugins/Experimental/Mutable/Source/CustomizableObject/Private/MuCO/CustomizableObjectInstance.cpp
@@ -4589,7 +4589,7 @@ void UCustomizableInstancePrivateData::AdditionalAssetsAsyncLoaded(UCustomizable
 }
 
 
-void UpdateTextureRegionsMutable(UTexture2D* Texture, int32 MipIndex, const FUpdateTextureRegion2D& Region, uint32 SrcPitch, FByteBulkData* BulkData)
+void UpdateTextureRegionsMutable(UTexture2D* Texture, int32 MipIndex, uint32 NumMips, const FUpdateTextureRegion2D& Region, uint32 SrcPitch, const FByteBulkData* BulkData)
 {
 	if (Texture->GetResource())
 	{
@@ -4599,6 +4599,7 @@ void UpdateTextureRegionsMutable(UTexture2D* Texture, int32 MipIndex, const FUpd
 			int32 MipIndex;
 			FUpdateTextureRegion2D Region;
 			uint32 SrcPitch;
+			uint32 NumMips;
 		};
 
 		FUpdateTextureRegionsData* RegionData = new FUpdateTextureRegionsData;
@@ -4607,19 +4608,30 @@ void UpdateTextureRegionsMutable(UTexture2D* Texture, int32 MipIndex, const FUpd
 		RegionData->MipIndex = MipIndex;
 		RegionData->Region = Region;
 		RegionData->SrcPitch = SrcPitch;
-
+		RegionData->NumMips = NumMips;
 
 		ENQUEUE_RENDER_COMMAND(UpdateTextureRegionsMutable)(
 			[RegionData, BulkData](FRHICommandList& CmdList)
 			{
+				check(int32(RegionData->NumMips) >= RegionData->Texture2DResource->GetCurrentMipCount());
+				int32 MipDifference = RegionData->NumMips - RegionData->Texture2DResource->GetCurrentMipCount();
+				check(MipDifference >= 0);
 				int32 CurrentFirstMip = RegionData->Texture2DResource->GetCurrentFirstMip();
 				uint8* SrcData = (uint8*)BulkData->LockReadOnly();
 
-				if (RegionData->MipIndex >= CurrentFirstMip)
+				//uint32 Size = RegionData->SrcPitch / (sizeof(uint8) * 4);
+				//UE_LOG(LogMutable, Warning, TEXT("UpdateTextureRegionsMutable MipIndex = %d, FirstMip = %d, size = %d"),
+				//	RegionData->MipIndex, CurrentFirstMip, Size);
+
+				//checkf(Size <= RegionData->Texture2DResource->GetSizeX(),
+				//	TEXT("UpdateTextureRegionsMutable incorrect size. %d, %d. NumMips=%d"), 
+				//	Size, RegionData->Texture2DResource->GetSizeX(), RegionData->Texture2DResource->GetCurrentMipCount());
+
+				if (RegionData->MipIndex >= CurrentFirstMip + MipDifference)
 				{
 					RHIUpdateTexture2D(
 						RegionData->Texture2DResource->GetTexture2DRHI(),
-						RegionData->MipIndex - CurrentFirstMip,
+						RegionData->MipIndex - CurrentFirstMip - MipDifference,
 						RegionData->Region,
 						RegionData->SrcPitch,
 						SrcData);
@@ -4631,11 +4643,14 @@ void UpdateTextureRegionsMutable(UTexture2D* Texture, int32 MipIndex, const FUpd
 	}
 }
 
+
 void UCustomizableInstancePrivateData::ReuseTexture(UTexture2D* Texture)
 {
 	if (Texture->GetPlatformData())
 	{
-		for (uint32 i = 0; i < (uint32)Texture->GetPlatformData()->Mips.Num(); i++)
+		uint32 NumMips = Texture->GetPlatformData()->Mips.Num();
+
+		for (uint32 i = 0; i < NumMips; i++)
 		{
 			FTexture2DMipMap& Mip = Texture->GetPlatformData()->Mips[i];
 
@@ -4650,7 +4665,10 @@ void UCustomizableInstancePrivateData::ReuseTexture(UTexture2D* Texture)
 				Region.Width = Mip.SizeX;
 				Region.Height = Mip.SizeY;
 
-				UpdateTextureRegionsMutable(Texture, i, Region, Mip.SizeX * sizeof(uint8) * 4, &Mip.BulkData);
+				check(int32(Region.Width) <= Texture->GetSizeX());
+				check(int32(Region.Height) <= Texture->GetSizeY());
+
+				UpdateTextureRegionsMutable(Texture, i, NumMips, Region, Mip.SizeX * sizeof(uint8) * 4, &Mip.BulkData);
 			}
 		}
 	}
@@ -4677,7 +4695,7 @@ void UCustomizableInstancePrivateData::BuildMaterials(const TSharedPtr<FMutableO
 	UpdateContext->State = OperationData->State;
 
 	const bool bUseCurrentMinLODAsBaseLOD = HasCOInstanceFlags(ReduceLODs);
-	const bool bReuseTextures = HasCOInstanceFlags(ReuseTextures);
+	const bool bReuseTextures = OperationData->bReuseInstanceTextures;
 
 	const FInstanceUpdateData::FLOD& FirstLOD = OperationData->InstanceUpdateData.LODs[OperationData->CurrentMinLOD];
 	const int32 NumComponents = FirstLOD.ComponentCount;
@@ -4847,8 +4865,27 @@ void UCustomizableInstancePrivateData::BuildMaterials(const TSharedPtr<FMutableO
 
 									if (ReusedTexture && (*ReusedTexture).IsValid() && !(*ReusedTexture)->HasAnyFlags(RF_BeginDestroyed))
 									{
-										Texture = (*ReusedTexture).Get();
-										check(Texture != nullptr);
+										// Only uncompressed textures can be reused. This also fixes an issue in the editor where textures supposedly 
+										// uncompressed by their state, are still compressed because the CO has not been compiled at maximum settings
+										// and the uncompressed setting cannot be applied to them.
+										EPixelFormat PixelFormat = (*ReusedTexture)->GetPixelFormat();
+
+										if (PixelFormat == EPixelFormat::PF_R8G8B8A8)
+										{
+											Texture = (*ReusedTexture).Get();
+											check(Texture != nullptr);
+										}
+										else
+										{
+											ReusedTexture = nullptr;
+											Texture = CreateTexture();
+
+#if WITH_EDITOR
+											UE_LOG(LogMutable, Warning,
+												TEXT("Tried to reuse an uncompressed texture with name %s. Make sure the selected Mutable state disables texture compression/streaming, that one of the state's runtime parameters affects the texture and that the CO is compiled with max. optimization settings."),
+												*Texture->GetName());
+#endif
+										}
 									}
 									else
 									{
diff --git a/Engine/Plugins/Experimental/Mutable/Source/CustomizableObject/Private/MuCO/CustomizableObjectSystem.cpp b/Engine/Plugins/Experimental/Mutable/Source/CustomizableObject/Private/MuCO/CustomizableObjectSystem.cpp
--- a/Engine/Plugins/Experimental/Mutable/Source/CustomizableObject/Private/MuCO/CustomizableObjectSystem.cpp
+++ b/Engine/Plugins/Experimental/Mutable/Source/CustomizableObject/Private/MuCO/CustomizableObjectSystem.cpp
@@ -682,6 +682,14 @@ static FAutoConsoleVariableRef CVarEnableMutableLiveUpdate(
 	ECVF_Default);
 
 
+int32 FCustomizableObjectSystemPrivate::EnableReuseInstanceTextures = 0;
+
+static FAutoConsoleVariableRef CVarEnableMutableReuseInstanceTextures(
+	TEXT("mutable.EnableReuseInstanceTextures"), FCustomizableObjectSystemPrivate::EnableReuseInstanceTextures,
+	TEXT("If set to 1 or greater and set in the corresponding setting in the current Mutable state, Mutable can reuse instance UTextures (only uncompressed and not streaming, so set the options in the state) and their resources between updates when they are modified. If geometry or state is changed they cannot be reused."),
+	ECVF_Default);
+
+
 /** Update the given Instance Skeletal Meshes and call its callbacks. */
 void UpdateSkeletalMesh(UCustomizableObjectInstance* CustomizableObjectInstance)
 {
@@ -2115,12 +2123,13 @@ namespace impl
 
 		CandidateInstance->CommitMinMaxLOD();
 
+		FString StateName = CandidateInstance->GetCustomizableObject()->GetStateName(CandidateInstance->GetState());
+		const FParameterUIData* StateData = CandidateInstance->GetCustomizableObject()->StateUIDataMap.Find(StateName);
+
 		bool bLiveUpdateMode = false;
 
 		if (SystemPrivateData->EnableMutableLiveUpdate)
 		{
-			FString StateName = CandidateInstance->GetCustomizableObject()->GetStateName(CandidateInstance->GetState());
-			const FParameterUIData* StateData = CandidateInstance->GetCustomizableObject()->StateUIDataMap.Find(StateName);
 			bLiveUpdateMode = StateData ? StateData->bLiveUpdateMode : false;
 		}
 
@@ -2130,6 +2139,22 @@ namespace impl
 			bLiveUpdateMode = false;
 		}
 
+		bool bReuseInstanceTextures = false;
+
+		if (SystemPrivateData->EnableReuseInstanceTextures)
+		{
+			if (Operation->bNeverStream)
+			{
+				bReuseInstanceTextures = StateData ? StateData->bReuseInstanceTextures : false;
+				bReuseInstanceTextures |= CandidateInstancePrivateData->HasCOInstanceFlags(ReuseTextures);
+			}
+			else
+			{
+				UE_LOG(LogMutable, Warning, TEXT("Instance texture reuse requires that the current Mutable state is in non-streaming mode. Change it in the Mutable graph base node in the state definition."));
+				bReuseInstanceTextures = false;
+			}
+		}
+
 		FCustomizableObjectSystemPrivate* CustomizableObjectSystemPrivateData = System->GetPrivate();
 		check(CustomizableObjectSystemPrivateData != nullptr);
 
@@ -2151,6 +2176,7 @@ namespace impl
 		CurrentOperationData->CurrentMaxLOD = Operation->InstanceDescriptorRuntimeHash.GetMaxLOD();
 		CurrentOperationData->bNeverStream = Operation->bNeverStream;
 		CurrentOperationData->bLiveUpdateMode = bLiveUpdateMode;
+		CurrentOperationData->bReuseInstanceTextures = bReuseInstanceTextures;
 		CurrentOperationData->InstanceID = bLiveUpdateMode ? CandidateInstancePrivateData->LiveUpdateModeInstanceID : 0;
 		CurrentOperationData->MipsToSkip = Operation->MipsToSkip;
 		CurrentOperationData->MutableParameters = Parameters;
diff --git a/Engine/Plugins/Experimental/Mutable/Source/CustomizableObject/Private/MuCO/CustomizableObjectSystemPrivate.h b/Engine/Plugins/Experimental/Mutable/Source/CustomizableObject/Private/MuCO/CustomizableObjectSystemPrivate.h
--- a/Engine/Plugins/Experimental/Mutable/Source/CustomizableObject/Private/MuCO/CustomizableObjectSystemPrivate.h
+++ b/Engine/Plugins/Experimental/Mutable/Source/CustomizableObject/Private/MuCO/CustomizableObjectSystemPrivate.h
@@ -423,6 +423,7 @@ struct FMutableOperationData
 	bool bNeverStream = false;
 	/** When this option is enabled it will reuse the Mutable core instance and its temp data between updates.  */
 	bool bLiveUpdateMode = false;
+	bool bReuseInstanceTextures = false;
 	/** This option comes from the operation request. It is used to reduce the number of mipmaps that mutable must generate for images.  */
 	int32 MipsToSkip = 0;
 
@@ -495,6 +496,7 @@ public:
 
 	static int32 EnableMutableProgressiveMipStreaming;
 	static int32 EnableMutableLiveUpdate;
+	static int32 EnableReuseInstanceTextures;
 	static int32 EnableMutableAnimInfoDebugging;
 
 	/** */
diff --git a/Engine/Plugins/Experimental/Mutable/Source/CustomizableObject/Public/MuCO/CustomizableObjectUIData.h b/Engine/Plugins/Experimental/Mutable/Source/CustomizableObject/Public/MuCO/CustomizableObjectUIData.h
--- a/Engine/Plugins/Experimental/Mutable/Source/CustomizableObject/Public/MuCO/CustomizableObjectUIData.h
+++ b/Engine/Plugins/Experimental/Mutable/Source/CustomizableObject/Public/MuCO/CustomizableObjectUIData.h
@@ -264,6 +264,9 @@ struct FParameterUIData
 	UPROPERTY(BlueprintReadWrite, Category = CustomizableObject)
 	bool bLiveUpdateMode = false;
 
+	UPROPERTY(BlueprintReadWrite, Category = CustomizableObject)
+	bool bReuseInstanceTextures = false;
+
 	UPROPERTY(BlueprintReadWrite, Category = CustomizableObject)
 	TMap<FString, FString> ForcedParameterValues;
 
@@ -271,7 +274,8 @@ struct FParameterUIData
 	{
 		if (Name != Other.Name || ParamUIMetadata != Other.ParamUIMetadata || Type != Other.Type 
 			|| ArrayIntegerParameterOption != Other.ArrayIntegerParameterOption || IntegerParameterGroupType != Other.IntegerParameterGroupType 
-			|| bLiveUpdateMode != Other.bLiveUpdateMode || !ForcedParameterValues.OrderIndependentCompareEqual(Other.ForcedParameterValues)
+			|| bLiveUpdateMode != Other.bLiveUpdateMode || bReuseInstanceTextures != Other.bReuseInstanceTextures
+			|| !ForcedParameterValues.OrderIndependentCompareEqual(Other.ForcedParameterValues)
 			)
 		{
 			return false;
@@ -289,6 +293,7 @@ struct FParameterUIData
 		Ar << UIData.IntegerParameterGroupType;
 		Ar << UIData.bDontCompressRuntimeTextures;
 		Ar << UIData.bLiveUpdateMode;
+		Ar << UIData.bReuseInstanceTextures;
 		Ar << UIData.ForcedParameterValues;
 
 		return Ar;
diff --git a/Engine/Plugins/Experimental/Mutable/Source/CustomizableObjectEditor/Classes/MuCOE/Nodes/CustomizableObjectNodeObject.h b/Engine/Plugins/Experimental/Mutable/Source/CustomizableObjectEditor/Classes/MuCOE/Nodes/CustomizableObjectNodeObject.h
--- a/Engine/Plugins/Experimental/Mutable/Source/CustomizableObjectEditor/Classes/MuCOE/Nodes/CustomizableObjectNodeObject.h
+++ b/Engine/Plugins/Experimental/Mutable/Source/CustomizableObjectEditor/Classes/MuCOE/Nodes/CustomizableObjectNodeObject.h
@@ -34,6 +34,12 @@ struct CUSTOMIZABLEOBJECTEDITOR_API FCustomizableObjectState
 	UPROPERTY(EditAnywhere, Category = CustomizableObject)
 	bool bLiveUpdateMode = false;
 
+	// Enables the reuse of all possible textures when the instance is updated without any changes in geometry or state (the first update after creation doesn't reuse any)
+	// It will only work if the textures aren't compressed, so set the instance to a Mutable state with texture compression disabled
+	// WARNING! If texture reuse is enabled, do NOT keep external references to the textures of the instance. The instance owns the textures.
+	UPROPERTY(EditAnywhere, Category = CustomizableObject)
+	bool bReuseInstanceTextures = false;
+
 	UPROPERTY(EditAnywhere, Category = CustomizableObject)
 	bool bBuildOnlyFirstLOD = false;
 
diff --git a/Engine/Plugins/Experimental/Mutable/Source/CustomizableObjectEditor/Private/MuCOE/GenerateMutableSource/GenerateMutableSource.cpp b/Engine/Plugins/Experimental/Mutable/Source/CustomizableObjectEditor/Private/MuCOE/GenerateMutableSource/GenerateMutableSource.cpp
--- a/Engine/Plugins/Experimental/Mutable/Source/CustomizableObjectEditor/Private/MuCOE/GenerateMutableSource/GenerateMutableSource.cpp
+++ b/Engine/Plugins/Experimental/Mutable/Source/CustomizableObjectEditor/Private/MuCOE/GenerateMutableSource/GenerateMutableSource.cpp
@@ -710,6 +710,7 @@ mu::NodeObjectPtr GenerateMutableSource(const UEdGraphPin * Pin, FMutableGraphGe
 			FParameterUIData ParameterUIData(State.Name, State.StateUIMetadata, EMutableParameterType::None);
 			ParameterUIData.bDontCompressRuntimeTextures = State.bDontCompressRuntimeTextures;
 			ParameterUIData.bLiveUpdateMode = State.bLiveUpdateMode;
+			ParameterUIData.bReuseInstanceTextures = State.bReuseInstanceTextures;
 			ParameterUIData.ForcedParameterValues = State.ForcedParameterValues;
 
 			GenerationContext.StateUIDataMap.Add(State.Name, ParameterUIData);
