[mutable] Fix incorrect size expression return in image raster mesh.
Cleanup coding style in "image invert" op.
#rnx

diff --git a/Engine/Plugins/Experimental/Mutable/Source/MutableRuntime/Private/MuR/OpImageInvert.h b/Engine/Plugins/Experimental/Mutable/Source/MutableRuntime/Private/MuR/OpImageInvert.h
--- a/Engine/Plugins/Experimental/Mutable/Source/MutableRuntime/Private/MuR/OpImageInvert.h
+++ b/Engine/Plugins/Experimental/Mutable/Source/MutableRuntime/Private/MuR/OpImageInvert.h
@@ -8,17 +8,15 @@
 namespace mu
 {
 
-	//---------------------------------------------------------------------------------------------
-	//!
-	//---------------------------------------------------------------------------------------------
+	/** Create a new image inverting the colour (RGB or L) components of an image. Leave alpha untouched. */
 	inline ImagePtr ImageInvert(const Image* pA)
 	{
-		MUTABLE_CPUPROFILER_SCOPE(ImageInvert)
+		MUTABLE_CPUPROFILER_SCOPE(ImageInvert);
 
-			ImagePtr pDest = new Image(pA->GetSizeX(), pA->GetSizeY(), pA->GetLODCount(), pA->GetFormat());
+		ImagePtr pDest = new Image(pA->GetSizeX(), pA->GetSizeY(), pA->GetLODCount(), pA->GetFormat());
 
-		uint8_t* pDestBuf = pDest->GetData();
-		const uint8_t* pABuf = pA->GetData();
+		uint8* pDestBuf = pDest->GetData();
+		const uint8* pABuf = pA->GetData();
 
 		//Generic implementation
 		int32 pixelCount = pA->CalculatePixelCount();
@@ -27,7 +25,7 @@ namespace mu
 		{
 		case EImageFormat::IF_L_UBYTE:
 		{
-			for (int i = 0; i < pixelCount; ++i)
+			for (int32 i = 0; i < pixelCount; ++i)
 			{
 				pDestBuf[i] = 255 - pABuf[i];
 			}
@@ -36,7 +34,7 @@ namespace mu
 
 		case EImageFormat::IF_RGB_UBYTE:
 		{
-			for (int i = 0; i < pixelCount; ++i)
+			for (int32 i = 0; i < pixelCount; ++i)
 			{
 				for (int c = 0; c < 3; ++c)
 				{
@@ -48,9 +46,9 @@ namespace mu
 		case EImageFormat::IF_RGBA_UBYTE:
 		case EImageFormat::IF_BGRA_UBYTE:
 		{
-			for (int i = 0; i < pixelCount; ++i)
+			for (int32 i = 0; i < pixelCount; ++i)
 			{
-				for (int c = 0; c < 3; ++c)
+				for (int32 c = 0; c < 3; ++c)
 				{
 					pDestBuf[i * 4 + c] = 255 - pABuf[i * 4 + c];
 				}
@@ -66,14 +64,17 @@ namespace mu
 		return pDest;
 	}
 
-	//---------------------------------------------------------------------------------------------
-	//!
-	//---------------------------------------------------------------------------------------------
+	/** Invert the colour (RGB or L) components of an image. Leave alpha untouched. */
 	inline void ImageInvertInPlace(Image* pA)
 	{
-		MUTABLE_CPUPROFILER_SCOPE(ImageInvertInPlace)
+		MUTABLE_CPUPROFILER_SCOPE(ImageInvertInPlace);
+			
+		if (!pA)
+		{
+			return;
+		}
 
-		uint8_t* pABuf = pA->GetData();
+		uint8* pABuf = pA->GetData();
 
 		//Generic implementation
 		int32 pixelCount = pA->CalculatePixelCount();
@@ -82,7 +83,7 @@ namespace mu
 		{
 		case EImageFormat::IF_L_UBYTE:
 		{
-			for (int i = 0; i < pixelCount; ++i)
+			for (int32 i = 0; i < pixelCount; ++i)
 			{
 				pABuf[i] = 255 - pABuf[i];
 			}
@@ -91,9 +92,9 @@ namespace mu
 
 		case EImageFormat::IF_RGB_UBYTE:
 		{
-			for (int i = 0; i < pixelCount; ++i)
+			for (int32 i = 0; i < pixelCount; ++i)
 			{
-				for (int c = 0; c < 3; ++c)
+				for (int32 c = 0; c < 3; ++c)
 				{
 					pABuf[i * 3 + c] = 255 - pABuf[i * 3 + c];
 				}
@@ -103,9 +104,9 @@ namespace mu
 		case EImageFormat::IF_RGBA_UBYTE:
 		case EImageFormat::IF_BGRA_UBYTE:
 		{
-			for (int i = 0; i < pixelCount; ++i)
+			for (int32 i = 0; i < pixelCount; ++i)
 			{
-				for (int c = 0; c < 3; ++c)
+				for (int32 c = 0; c < 3; ++c)
 				{
 					pABuf[i * 4 + c] = 255 - pABuf[i * 4 + c];
 				}
diff --git a/Engine/Plugins/Experimental/Mutable/Source/MutableTools/Private/MuT/ASTOpImageRasterMesh.cpp b/Engine/Plugins/Experimental/Mutable/Source/MutableTools/Private/MuT/ASTOpImageRasterMesh.cpp
--- a/Engine/Plugins/Experimental/Mutable/Source/MutableTools/Private/MuT/ASTOpImageRasterMesh.cpp
+++ b/Engine/Plugins/Experimental/Mutable/Source/MutableTools/Private/MuT/ASTOpImageRasterMesh.cpp
@@ -181,8 +181,7 @@ namespace mu
 		pRes->type = ImageSizeExpression::ISET_CONSTANT;
 		pRes->size[0] = sizeX ? sizeX : 256;
 		pRes->size[1] = sizeY ? sizeY : 256;
-
-		return nullptr;
+		return pRes;
 	}
 
 
