[mutable] Fix crash when incorrect layouts are specified in UI.
#rnx

diff --git a/Engine/Plugins/Experimental/Mutable/Source/MutableTools/Private/MuT/CodeGenerator_Image.cpp b/Engine/Plugins/Experimental/Mutable/Source/MutableTools/Private/MuT/CodeGenerator_Image.cpp
--- a/Engine/Plugins/Experimental/Mutable/Source/MutableTools/Private/MuT/CodeGenerator_Image.cpp
+++ b/Engine/Plugins/Experimental/Mutable/Source/MutableTools/Private/MuT/CodeGenerator_Image.cpp
@@ -1459,10 +1459,16 @@ namespace mu
 
 			// Match the block id of the block we are generating with the id that resulted in the generated mesh
 			GeneratedLayoutBlockId = -1;
-			if (LayoutBlockIndex >= 0)
+			if (node.m_layout<MeshResult.GeneratedLayouts.Num() 
+				&& 
+				LayoutBlockIndex >= 0 && LayoutBlockIndex< MeshResult.GeneratedLayouts[node.m_layout]->m_blocks.Num() )
 			{
 				GeneratedLayoutBlockId = MeshResult.GeneratedLayouts[node.m_layout]->m_blocks[LayoutBlockIndex].m_id;
 			}
+			else
+			{
+				m_pErrorLog->GetPrivate()->Add("Layout or block index error.", ELMT_ERROR, node.m_errorContext);
+			}
 
             pop->SetChild( pop->op.args.MeshProject.mesh, MeshResult.meshOp );
 
