[mutable] Improvements at Content Validation procedure for COs
- The same root CO will no longer be compiled more than once if more than one checked out CO has the same CO as root.
- Before standard compilation a series of compilations with an altered LodBias value will be ran in order to find issues related with the LOD configuration.

diff --git a/Engine/Plugins/Experimental/Mutable/Source/CustomizableObject/Classes/MuCO/CustomizableObject.h b/Engine/Plugins/Experimental/Mutable/Source/CustomizableObject/Classes/MuCO/CustomizableObject.h
--- a/Engine/Plugins/Experimental/Mutable/Source/CustomizableObject/Classes/MuCO/CustomizableObject.h
+++ b/Engine/Plugins/Experimental/Mutable/Source/CustomizableObject/Classes/MuCO/CustomizableObject.h
@@ -1414,6 +1414,18 @@ public:
 	virtual EDataValidationResult IsDataValid(class FDataValidationContext& Context) override;
 	// End of UObject Interface
 
+private:
+
+	/** Cached handle to be able later to remove the bound method from the FEditorDelegates::OnPostAssetValidation delegate */
+	inline static FDelegateHandle OnPostCOValidationHandle;
+	
+	/** Collection with all root objects tested during this IsDataValidRun. Shared with all COs */
+	inline static TArray<UCustomizableObject*> AlreadyValidatedRootObjects;
+	
+	/** Method invoked once the validation of all assets has been completed. */
+	static void OnPostCOsValidation();
+
+public:
 	// UObject Interface -> Asset saving
 	virtual void PreSaveRoot(FObjectPreSaveRootContext ObjectSaveContext) override;
 	// End of UObject Interface
diff --git a/Engine/Plugins/Experimental/Mutable/Source/CustomizableObject/Private/MuCO/CustomizableObject.cpp b/Engine/Plugins/Experimental/Mutable/Source/CustomizableObject/Private/MuCO/CustomizableObject.cpp
--- a/Engine/Plugins/Experimental/Mutable/Source/CustomizableObject/Private/MuCO/CustomizableObject.cpp
+++ b/Engine/Plugins/Experimental/Mutable/Source/CustomizableObject/Private/MuCO/CustomizableObject.cpp
@@ -5,6 +5,7 @@
 #include "Algo/Copy.h"
 #include "AssetRegistry/AssetRegistryModule.h"
 #include "Async/AsyncFileHandle.h"
+#include "Delegates/IDelegateInstance.h"
 #include "EdGraph/EdGraph.h"
 #include "Engine/SkeletalMesh.h"
 #include "HAL/FileManager.h"
@@ -25,6 +26,7 @@
 #include "Serialization/MemoryReader.h"
 #include "Serialization/MemoryWriter.h"
 #include "UObject/ObjectSaveContext.h"
+#include "Templates/UniquePtr.h"
 
 #include UE_INLINE_GENERATED_CPP_BY_NAME(CustomizableObject)
 
@@ -1498,68 +1500,140 @@ EDataValidationResult UCustomizableObject::IsDataValid(FDataValidationContext& C
 		return Result;
 	}
 
-	UE_LOG(LogMutable,Display,TEXT("Running data validation checks for %s CO"),*this->GetName());
+	UE_LOG(LogMutable,Display,TEXT("Running data validation checks for %s CO."),*this->GetName());
+
+	// Bind the post validation method to the post validation delegate if not bound already to be able to know when the validation
+	// operation (for all assets) concludes
+	if (!UCustomizableObject::OnPostCOValidationHandle.IsValid())
+	{
+		UCustomizableObject::OnPostCOValidationHandle = FEditorDelegates::OnPostAssetValidation.AddStatic(OnPostCOsValidation);	
+	}
 	
 	// Request a compiler to be able to locate the root and to compile it
-	FCustomizableObjectCompilerBase* Compiler = UCustomizableObjectSystem::GetInstance()->GetNewCompiler();
+	const TUniquePtr<FCustomizableObjectCompilerBase> Compiler =
+		TUniquePtr<FCustomizableObjectCompilerBase>(UCustomizableObjectSystem::GetInstance()->GetNewCompiler());
+	
+	// Find out witch is the root for this CO (it may be itself but that is OK)
+	UCustomizableObject* RootObject = Compiler->GetRootObject(this);
+	check (RootObject);
+	
+	// Check that the object to be compiled has not already been compiled
+	if (UCustomizableObject::AlreadyValidatedRootObjects.Contains(RootObject))
 	{
-		// Find out witch is the root for this CO (it may be itself but that is OK)
-		UCustomizableObject* RootObject = Compiler->GetRootObject(this);
-		check (RootObject);
-		
-		// Run Sync compilation -> Warning : Potentially long operation -------------
-		Compiler->Compile(*RootObject, this->CompileOptions, false);
-		// --------------------------------------------------------------------------
-
-		// Request the compilation warnings and errors and provide them to the caller
-		TArray<FText> ValidationErrors;
-		TArray<FText> ValidationWarnings;
-		Compiler->GetCompilationMessages(ValidationWarnings, ValidationErrors);
+		return Result;
+	}
 
-		for (const FText& ValidationError : ValidationErrors)
-		{
-			Context.AddError(ValidationError);
-		}
+	// Root Object not yet tested -> Proceed with the testing
+	
+	// Collection of configurations to be tested with the located root object
+	constexpr int32 MaxBias = 15;
+	TArray<FCompilationOptions> CompilationOptionsToTest;
+	for (int32 LodBias = 0; LodBias < MaxBias; LodBias++)
+	{
+		FCompilationOptions ModifiedCompilationOptions = this->CompileOptions;
+		ModifiedCompilationOptions.bForceLargeLODBias = true;	
+		ModifiedCompilationOptions.DebugBias = LodBias;	
 
-		for (const FText& ValidationWarning : ValidationWarnings)
-		{
-			Context.AddWarning(ValidationWarning);
-		}
+		// Add one configuration object for each bias setting
+		CompilationOptionsToTest.Add(ModifiedCompilationOptions);
+	}
+	
+	// Add current configuration to be tested as well.
+	CompilationOptionsToTest.Add(this->CompileOptions);
 
-		// Check if the compilation was able to be performed (not if there were errors but if a critical issue appeared
-		// before starting the compilation as we know.
-		switch (Compiler->GetCompilationState())
+	
+	// Caches with all the data produced by the subsequent compilations of the root of this CO
+	TArray<FText> CachedValidationErrors;
+	TArray<FText> CachedValidationWarnings;
+	TArray<ECustomizableObjectCompilationState> CachedCompilationEndStates;
+	
+	// Iterate over the compilation options that we want to test and perform the compilation
+	for	(const FCompilationOptions& Options : CompilationOptionsToTest)
+	{
+		// Run Sync compilation -> Warning : Potentially long operation -------------
+		Compiler->Compile(*RootObject, Options, false);
+		// --------------------------------------------------------------------------
+		
+		// Get compilation errors and warnings
+		TArray<FText> CompilationErrors;
+		TArray<FText> CompilationWarnings;
+		Compiler->GetCompilationMessages(CompilationWarnings, CompilationErrors);
+		
+		// Cache the messages returned by the compiler
+		for ( const FText& FoundError : CompilationErrors)
 		{
-		case ECustomizableObjectCompilationState::Completed:
+			// Add message if not already present
+			if (!CachedValidationErrors.ContainsByPredicate([&FoundError](const FText& ArrayEntry)
+				{ return FoundError.EqualTo(ArrayEntry);}))
 			{
-				// If a warning or error was found then this object failed the validation process
-				Result = (ValidationWarnings.IsEmpty() && ValidationErrors.IsEmpty()) ? EDataValidationResult::Valid : EDataValidationResult::Invalid;
+				CachedValidationErrors.Add(FoundError);
 			}
-			break;
-			
-		case ECustomizableObjectCompilationState::Failed:
+		}
+		for ( const FText& FoundWarning : CompilationWarnings)
+		{
+			if (!CachedValidationWarnings.ContainsByPredicate([&FoundWarning](const FText& ArrayEntry)
+				{ return FoundWarning.EqualTo(ArrayEntry);}))
 			{
-				// Early CO compilation error (before starting mutable compilation) -> Output is invalid
-				Result = EDataValidationResult::Invalid;
-				UE_LOG(LogMutable, Error,
-					   TEXT("Compilation of %s failed : Check previous log messages to get more information."),
-					   *this->GetName())
+				CachedValidationWarnings.Add(FoundWarning);
 			}
-			break;
-		
-			// Invalid options for this context:
-			// ECustomizableObjectCompilationState::None: would mean we are working with a locked asset, with should not be possible with a sync compilation
-		case ECustomizableObjectCompilationState::None:
-		case ECustomizableObjectCompilationState::InProgress:
-		default:
-			checkNoEntry();
 		}
+	
+		CachedCompilationEndStates.Add(Compiler->GetCompilationState());
 	}
-	delete Compiler;
+	
+	// Cache root object to avoid processing it again when processing another CO related with the same root CO
+	AlreadyValidatedRootObjects.Add(RootObject);
+	
+	// Wrapping up : Fill message output caches and determine if the compilation was successful or not
 
+	// Provide the warning and log messages to the context object (so it can later notify the user using the UI)
+	for (const FText& ValidationError : CachedValidationErrors)
+	{
+		Context.AddError(ValidationError);
+	}
+	for (const FText& ValidationWarning : CachedValidationWarnings)
+	{
+		Context.AddWarning(ValidationWarning);
+	}
+	
+	// Return informed guess about what the validation state of this object should be
+
+	// If one or more tests failed to ran then the result must be invalid
+	if (CachedCompilationEndStates.Contains(ECustomizableObjectCompilationState::Failed))
+	{
+		// Early CO compilation error (before starting mutable compilation) -> Output is invalid
+		Result = EDataValidationResult::Invalid;
+		UE_LOG(LogMutable, Error,
+			   TEXT("Compilation of %s failed : Check previous log messages to get more information."),
+			   *this->GetName())
+	}
+	// If it contains invalid states then notify about it too:
+	// ECustomizableObjectCompilationState::None would mean the resource is locked (and should not be)
+	// ECustomizableObjectCompilationState::InProgress should not be possible since we are compiling synchronously.
+	else if (CachedCompilationEndStates.Contains(ECustomizableObjectCompilationState::InProgress) ||
+		CachedCompilationEndStates.Contains(ECustomizableObjectCompilationState::None))
+	{
+		checkNoEntry();
+	}
+	// All compilations completed successfully
+	else 
+	{
+		// If a warning or error was found then this object failed the validation process
+		Result = (CachedValidationWarnings.IsEmpty() && CachedValidationErrors.IsEmpty()) ? EDataValidationResult::Valid : EDataValidationResult::Invalid;
+	}
 	
 	return Result;
 }
+
+void UCustomizableObject::OnPostCOsValidation()
+{
+	// Unbound this method from the validation end delegate
+	UCustomizableObject::OnPostCOValidationHandle.Reset();
+
+	// Clear collection with the already processed COs once the validation system has completed its operation
+	UCustomizableObject::AlreadyValidatedRootObjects.Empty();
+}
+
 #endif
 
 
