[Mutable] Implement new LOD system to allow generation of individual LODs depending on the PredictedLODLevels.
+ Base implementation of a new kind of update that allows reusing data from previous updates (currently disabled)
+ Refactor of the CustomizableObjectInstance update workflow and code clean up

#rnx

diff --git a/Engine/Plugins/Experimental/Mutable/Source/CustomizableObject/Classes/MuCO/CustomizableInstanceLODManagement.h b/Engine/Plugins/Experimental/Mutable/Source/CustomizableObject/Classes/MuCO/CustomizableInstanceLODManagement.h
--- a/Engine/Plugins/Experimental/Mutable/Source/CustomizableObject/Classes/MuCO/CustomizableInstanceLODManagement.h
+++ b/Engine/Plugins/Experimental/Mutable/Source/CustomizableObject/Classes/MuCO/CustomizableInstanceLODManagement.h
@@ -15,14 +15,14 @@ public:
 	UCustomizableInstanceLODManagementBase() : UObject() {};
 	virtual ~UCustomizableInstanceLODManagementBase() {};
 
-	virtual void UpdateInstanceDistsAndLODs() { check(0 && "You must override this"); };
-
-	virtual int32 GetNumGeneratedInstancesLimitFullLODs() { check(0 && "You must override this"); return 0; };
-	virtual int32 GetNumGeneratedInstancesLimitLOD1() { check(0 && "You must override this"); return 0; };
-	virtual int32 GetNumGeneratedInstancesLimitLOD2() { check(0 && "You must override this"); return 0; };
-
-	virtual bool IsOnlyUpdateCloseCustomizableObjectsEnabled() const { check(0 && "You must override this"); return false; };
-	virtual float GetOnlyUpdateCloseCustomizableObjectsDist() const { check(0 && "You must override this"); return 0.f; };
+	// WARNING! The following methods must be overriden in derived classes
+	virtual void UpdateInstanceDistsAndLODs() { check(0); };
+	virtual int32 GetNumGeneratedInstancesLimitFullLODs() const { check(0); return 0; };
+	virtual int32 GetNumGeneratedInstancesLimitLOD1() const { check(0); return 0; };
+	virtual int32 GetNumGeneratedInstancesLimitLOD2() const { check(0); return 0; };
+	virtual float GetOnlyUpdateCloseCustomizableObjectsDist() const { check(0); return 0.f; };
+	virtual bool IsOnlyUpdateCloseCustomizableObjectsEnabled() const { check(0); return false; };
+	virtual bool IsOnlyGenerateRequestedLODLevelsEnabled() const { check(0); return false; };
 };
 
 
@@ -36,9 +36,9 @@ public:
 
 	virtual void UpdateInstanceDistsAndLODs() override;
 
-	virtual int32 GetNumGeneratedInstancesLimitFullLODs() override;
-	virtual int32 GetNumGeneratedInstancesLimitLOD1() override;
-	virtual int32 GetNumGeneratedInstancesLimitLOD2() override;
+	virtual int32 GetNumGeneratedInstancesLimitFullLODs() const override;
+	virtual int32 GetNumGeneratedInstancesLimitLOD1() const override;
+	virtual int32 GetNumGeneratedInstancesLimitLOD2() const override;
 
 	// Used to define the actors that will be considered the centers of LOD distance calculations. If left empty it will revert to the default, the first player controller actor
 	void AddViewCenter(const AActor* const InCentralActor) { ViewCenters.Add(TWeakObjectPtr<const AActor>(InCentralActor)); };
@@ -49,10 +49,11 @@ public:
 	void SetNumberOfPriorityUpdateInstances(int32 NumPriorityUpdateInstances);
 	int32 GetNumberOfPriorityUpdateInstances() const;
 
-	virtual bool IsOnlyUpdateCloseCustomizableObjectsEnabled() const override;
+	virtual void SetCustomizableObjectsUpdateDistance(float Distance);
 	virtual float GetOnlyUpdateCloseCustomizableObjectsDist() const override;
-	void EnableOnlyUpdateCloseCustomizableObjects(float CloseDist);
-	void DisableOnlyUpdateCloseCustomizableObjects();
+
+	virtual bool IsOnlyUpdateCloseCustomizableObjectsEnabled() const override;
+	virtual bool IsOnlyGenerateRequestedLODLevelsEnabled() const override { return true; };
 
 private:
 	// The list of actors that define a view radius (CloseDist) around them used for LOD management
@@ -61,7 +62,6 @@ private:
 	// Sets how many of the nearest instances to the player will have updates with priority over LOD changes
 	int32 NumPriorityUpdateInstances = 1;
 
-	bool bOnlyUpdateCloseCustomizableObjects;
 	float CloseCustomizableObjectsDist;	
 };
 
diff --git a/Engine/Plugins/Experimental/Mutable/Source/CustomizableObject/Classes/MuCO/CustomizableObjectInstance.h b/Engine/Plugins/Experimental/Mutable/Source/CustomizableObject/Classes/MuCO/CustomizableObjectInstance.h
--- a/Engine/Plugins/Experimental/Mutable/Source/CustomizableObject/Classes/MuCO/CustomizableObjectInstance.h
+++ b/Engine/Plugins/Experimental/Mutable/Source/CustomizableObject/Classes/MuCO/CustomizableObjectInstance.h
@@ -587,11 +587,17 @@ public:
 	void SetIsPlayerOrNearIt(bool NewValue);
 	float GetMinSquareDistToPlayer() const;
 	void SetMinSquareDistToPlayer(float NewValue);
-	void SetMinMaxLODToLoad(int32 NewMinLOD = 0, int32 NewMaxLOD = INT32_MAX, bool LimitLODUpgrades = true);
+
+	int32 GetNumComponents() const;
+
+
 	int32 GetMinLODToLoad() const;
 	int32 GetMaxLODToLoad() const;
 	int32 GetNumLODsAvailable() const;
 
+	UE_DEPRECATED(5.2, "Use SetRequestedLODs instead.")
+	void SetMinMaxLODToLoad(int32 NewMinLOD = 0, int32 NewMaxLOD = INT32_MAX, bool LimitLODUpgrades = true);
+
 	/** Return the Min LOD this Instance is using (from the beginning of an update. If an update fails this value will be incorrect). */
 	int32 GetCurrentMinLOD() const;
 
@@ -600,14 +606,18 @@ public:
 
 	/** Save the Min and Max LOD that will be used for the update. */
 	void CommitMinMaxLOD();
-	
-	/* If enabled, CurrentMinLOD will be the first LOD of the generated SkeletalMesh. Otherwise the number of LODs will remain constant and LODs [0 .. CurrentMinLOD] will share the same RenderData. 
-	 * Enabled by default */
-	void SetUseCurrentMinLODAsBaseLOD(bool bIsBaseLOD);
-	bool GetUseCurrentMinLODAsBaseLOD() const;
+
+	/** Sets an array of LODs to generate per component. Mutable will generate those plus the currently generated LODs (if any).
+	 * Requires mutable.EnableOnlyGenerateRequestedLODs and CurrentInstanceLODManagement->IsOnlyGenerateRequestedLODLevelsEnabled() to be true.
+	 * @param InMinLOD - MinLOD to generate.
+	 * @param InMaxLOD - MaxLOD to generate.
+	 * @param InRequestedLODsPerComponent - Array with bitmasks of requested LODs per component with range from [0 .. CO->GetComponentCount()]. */
+	void SetRequestedLODs(int32 InMinLOD, int32 InMaxLOD, const TArray<uint16>& InRequestedLODsPerComponent);
+
+	const TArray<uint16>& GetRequestedLODsPerComponent() const;
 
 	/** Instance updated. */
-	void Updated(EUpdateResult Result);
+	void Updated(EUpdateResult Result, const FDescriptorRuntimeHash& UpdatedHash);
 	
 	/** Return the UCustomizableObjectInstance::Descriptor hash on the last update request. */
 	FDescriptorRuntimeHash GetDescriptorRuntimeHash() const;
diff --git a/Engine/Plugins/Experimental/Mutable/Source/CustomizableObject/Classes/MuCO/CustomizableObjectSystem.h b/Engine/Plugins/Experimental/Mutable/Source/CustomizableObject/Classes/MuCO/CustomizableObjectSystem.h
--- a/Engine/Plugins/Experimental/Mutable/Source/CustomizableObject/Classes/MuCO/CustomizableObjectSystem.h
+++ b/Engine/Plugins/Experimental/Mutable/Source/CustomizableObject/Classes/MuCO/CustomizableObjectSystem.h
@@ -259,6 +259,8 @@ public:
 
 	bool IsProgressiveMipStreamingEnabled() const;
 
+	bool IsOnlyGenerateRequestedLODsEnabled() const;
+
 	void AddPendingReleaseSkeletalMesh( USkeletalMesh* SkeletalMesh );
 
 	void PurgePendingReleaseSkeletalMesh();
diff --git a/Engine/Plugins/Experimental/Mutable/Source/CustomizableObject/Private/MuCO/CustomizableInstanceLODManagement.cpp b/Engine/Plugins/Experimental/Mutable/Source/CustomizableObject/Private/MuCO/CustomizableInstanceLODManagement.cpp
--- a/Engine/Plugins/Experimental/Mutable/Source/CustomizableObject/Private/MuCO/CustomizableInstanceLODManagement.cpp
+++ b/Engine/Plugins/Experimental/Mutable/Source/CustomizableObject/Private/MuCO/CustomizableInstanceLODManagement.cpp
@@ -47,10 +47,14 @@ static TAutoConsoleVariable<int32> CVarDistanceForFixedLOD2(
 	TEXT("If NumGeneratedInstancesLimit is different than 0, sets the distance at which the system will fix the LOD of an instance to the lowest res one (LOD2) to prevent unnecessary LOD changes and memory consumption"),
 	ECVF_Scalability);
 
+static TAutoConsoleVariable<bool> CVarOnlyUpdateCloseCustomizableObjects(
+	TEXT("b.OnlyUpdateCloseCustomizableObjects"),
+	true,
+	TEXT("If true, only CustomizableObjects within a predefined distance to the view centers will be generated"),
+	ECVF_Scalability);
 
 UCustomizableInstanceLODManagement::UCustomizableInstanceLODManagement() : UCustomizableInstanceLODManagementBase()
 {
-	bOnlyUpdateCloseCustomizableObjects = false;
 	CloseCustomizableObjectsDist = 2000.f;
 }
 
@@ -90,73 +94,75 @@ void UpdateCameraToInstancesDistance(const FVector CameraPosition)
 
 void UCustomizableInstanceLODManagement::UpdateInstanceDistsAndLODs()
 {
-	UCustomizableObjectSystem* CustomizableObjectSystem = UCustomizableObjectSystem::GetInstance();
-
-	if (ViewCenters.Num() == 0) // Just use the first pawn
+	int32 NumGeneratedInstancesLimitLODs = GetNumGeneratedInstancesLimitFullLODs();
+	if (NumGeneratedInstancesLimitLODs > 0 || (IsOnlyUpdateCloseCustomizableObjectsEnabled() && IsOnlyGenerateRequestedLODLevelsEnabled()))
 	{
-		UCustomizableSkeletalComponent* FirstCustomizableSkeletalComponent = nullptr;
-
-#if WITH_EDITOR
-		bool bLevelEditorInstancesUpdated = false;
-#endif
+		UCustomizableObjectSystem* CustomizableObjectSystem = UCustomizableObjectSystem::GetInstance();
 
-		for (TObjectIterator<UCustomizableSkeletalComponent> CustomizableSkeletalComponent; CustomizableSkeletalComponent; ++CustomizableSkeletalComponent)
+		if (ViewCenters.Num() == 0) // Just use the first pawn
 		{
-			if (CustomizableSkeletalComponent && !CustomizableSkeletalComponent->IsTemplate())
-			{
-				UWorld* LocalWorld = CustomizableSkeletalComponent->GetWorld();
-				APlayerController* Controller = LocalWorld ? LocalWorld->GetFirstPlayerController() : nullptr;
-				TWeakObjectPtr<const AActor> ViewCenter = Controller ? TWeakObjectPtr<const AActor>(Controller->GetPawn()) : nullptr;
+			UCustomizableSkeletalComponent* FirstCustomizableSkeletalComponent = nullptr;
 
-				if (ViewCenter.IsValid())
-				{
-					UpdatePawnToInstancesDistances(nullptr, ViewCenter);
-					break;
-				}
+	#if WITH_EDITOR
+			bool bLevelEditorInstancesUpdated = false;
+	#endif
 
-#if WITH_EDITOR
-				else if (CustomizableSkeletalComponent->GetWorld())
+			for (TObjectIterator<UCustomizableSkeletalComponent> CustomizableSkeletalComponent; CustomizableSkeletalComponent; ++CustomizableSkeletalComponent)
+			{
+				if (CustomizableSkeletalComponent && !CustomizableSkeletalComponent->IsTemplate())
 				{
-					EWorldType::Type worldType = CustomizableSkeletalComponent->GetWorld()->WorldType;
-					
-					// Level Editor Instances (non PIE)
-					if (!bLevelEditorInstancesUpdated && worldType == EWorldType::Editor)
+					UWorld* LocalWorld = CustomizableSkeletalComponent->GetWorld();
+					APlayerController* Controller = LocalWorld ? LocalWorld->GetFirstPlayerController() : nullptr;
+					TWeakObjectPtr<const AActor> ViewCenter = Controller ? TWeakObjectPtr<const AActor>(Controller->GetPawn()) : nullptr;
+
+					if (ViewCenter.IsValid())
 					{
-						for (FLevelEditorViewportClient* LevelVC : GEditor->GetLevelViewportClients())
+						UpdatePawnToInstancesDistances(nullptr, ViewCenter);
+						break;
+					}
+
+	#if WITH_EDITOR
+					else if (CustomizableSkeletalComponent->GetWorld())
+					{
+						EWorldType::Type worldType = CustomizableSkeletalComponent->GetWorld()->WorldType;
+					
+						// Level Editor Instances (non PIE)
+						if (!bLevelEditorInstancesUpdated && worldType == EWorldType::Editor)
 						{
-							if (LevelVC && LevelVC->IsPerspective())
+							for (FLevelEditorViewportClient* LevelVC : GEditor->GetLevelViewportClients())
 							{
-								UpdateCameraToInstancesDistance(LevelVC->GetViewLocation());
-								bLevelEditorInstancesUpdated = true;
-								break;
+								if (LevelVC && LevelVC->IsPerspective())
+								{
+									UpdateCameraToInstancesDistance(LevelVC->GetViewLocation());
+									bLevelEditorInstancesUpdated = true;
+									break;
+								}
 							}
 						}
-					}
 
-					// Blueprint instances
-					else if (worldType == EWorldType::EditorPreview)
-					{
-						CustomizableSkeletalComponent->EditorUpdateComponent();
+						// Blueprint instances
+						else if (worldType == EWorldType::EditorPreview)
+						{
+							CustomizableSkeletalComponent->EditorUpdateComponent();
+						}
 					}
-				}
-#endif // WITH_EDITOR
+	#endif // WITH_EDITOR
 
+				}
 			}
 		}
-	}
-	else
-	{
-		for (const TWeakObjectPtr<const AActor> ViewCenter : ViewCenters)
+		else
 		{
-			if (ViewCenter.IsValid())
+			for (const TWeakObjectPtr<const AActor> ViewCenter : ViewCenters)
 			{
-				UpdatePawnToInstancesDistances(nullptr, ViewCenter);
+				if (ViewCenter.IsValid())
+				{
+					UpdatePawnToInstancesDistances(nullptr, ViewCenter);
+				}
 			}
 		}
 	}
 
-	int32 NumGeneratedInstancesLimitLODs = GetNumGeneratedInstancesLimitFullLODs();
-
 	if (NumGeneratedInstancesLimitLODs > 0)
 	{
 		int32 NumGeneratedInstancesLimitLOD1 = GetNumGeneratedInstancesLimitLOD1();
@@ -186,6 +192,8 @@ void UCustomizableInstanceLODManagement::UpdateInstanceDistsAndLODs()
 			SortedInstances[i]->SetIsPlayerOrNearIt(true);
 		}
 
+		TArray<uint16> RequestedLODs;
+
 		if (SortedInstances.Num() > NumGeneratedInstancesLimitLODs)
 		{
 			SortedInstances.Sort([](UCustomizableObjectInstance& A, UCustomizableObjectInstance& B)
@@ -197,15 +205,18 @@ void UCustomizableInstanceLODManagement::UpdateInstanceDistsAndLODs()
 			const float DistanceForFixedLOD = float(CVarDistanceForFixedLOD2.GetValueOnGameThread());
 			const float DistanceForFixedLODSquared = FMath::Square(DistanceForFixedLOD);
 
+
 			for (int32 i = 0; i < NumGeneratedInstancesLimitLODs && i < SortedInstances.Num(); ++i)
 			{
 				if (!bAlreadyReachedFixedLOD && SortedInstances[i]->GetMinSquareDistToPlayer() < DistanceForFixedLODSquared)
 				{
-					SortedInstances[i]->SetMinMaxLODToLoad();
+					RequestedLODs.Init(MAX_uint16, SortedInstances[i]->GetNumComponents());
+					SortedInstances[i]->SetRequestedLODs(0, MAX_int32, RequestedLODs);
 				}
 				else
 				{
-					SortedInstances[i]->SetMinMaxLODToLoad(2, 2);
+					RequestedLODs.Init(MAX_uint16, SortedInstances[i]->GetNumComponents());
+					SortedInstances[i]->SetRequestedLODs(2, 2, RequestedLODs);
 					bAlreadyReachedFixedLOD = true;
 				}
 			}
@@ -214,18 +225,21 @@ void UCustomizableInstanceLODManagement::UpdateInstanceDistsAndLODs()
 			{
 				if (!bAlreadyReachedFixedLOD && SortedInstances[i]->GetMinSquareDistToPlayer() < DistanceForFixedLODSquared)
 				{
-					SortedInstances[i]->SetMinMaxLODToLoad(1, 2);
+					RequestedLODs.Init(MAX_uint16, SortedInstances[i]->GetNumComponents());
+					SortedInstances[i]->SetRequestedLODs(1, MAX_int32, RequestedLODs);
 				}
 				else
 				{
-					SortedInstances[i]->SetMinMaxLODToLoad(2, 2);
+					RequestedLODs.Init(MAX_uint16, SortedInstances[i]->GetNumComponents());
+					SortedInstances[i]->SetRequestedLODs(2, MAX_int32, RequestedLODs);
 					bAlreadyReachedFixedLOD = true;
 				}
 			}
 
 			for (int32 i = NumGeneratedInstancesLimitLODs + NumGeneratedInstancesLimitLOD1; i < NumGeneratedInstancesLimitLODs + NumGeneratedInstancesLimitLOD1 + NumGeneratedInstancesLimitLOD2 && i < SortedInstances.Num(); ++i)
 			{
-				SortedInstances[i]->SetMinMaxLODToLoad(2, 2);
+				RequestedLODs.Init(MAX_uint16, SortedInstances[i]->GetNumComponents());
+				SortedInstances[i]->SetRequestedLODs(2, MAX_int32, RequestedLODs);
 			}
 
 			for (int32 i = NumGeneratedInstancesLimitLODs + NumGeneratedInstancesLimitLOD1 + NumGeneratedInstancesLimitLOD2; i < SortedInstances.Num(); ++i)
@@ -238,176 +252,140 @@ void UCustomizableInstanceLODManagement::UpdateInstanceDistsAndLODs()
 			// No limit surpassed, set all instances to have all LODs, there will be an UpdateSkeletalMesh only if there's a change in LOD state
 			for (int32 i = 0; i < SortedInstances.Num(); ++i)
 			{
-				SortedInstances[i]->SetMinMaxLODToLoad();
+				RequestedLODs.Init(MAX_uint16, SortedInstances[i]->GetNumComponents());
+				SortedInstances[i]->SetRequestedLODs(2, MAX_int32, RequestedLODs);
 			}
 		}
 	}
-	else if (bOnlyUpdateCloseCustomizableObjects)
+	else if (IsOnlyGenerateRequestedLODLevelsEnabled())
 	{
-		TArray<FVector> Locations;
-		TArray<FMatrix> ProjectionMatrices;
-		TMap <const AActor*, int32> ViewCentersToIndexMap;
+		struct FLODTracker
+		{
+			int32 MinLOD = MAX_int32;
+			int32 MaxLOD = MAX_int32;
+
+			bool bInitialized = false;
+			TArray<uint16> RequestedLODsPerComponent;
+		};
 
-		UCameraComponent* PlayerCamera = nullptr;
+		TMap<TObjectPtr<UCustomizableObjectInstance>, FLODTracker> InstancesMinLOD;
+		InstancesMinLOD.Reserve(100);
 
-		if (ViewCenters.Num() == 0)
+		for (TObjectIterator<UCustomizableSkeletalComponent> CustomizableSkeletalComponent; CustomizableSkeletalComponent; ++CustomizableSkeletalComponent)
 		{
-			for (TObjectIterator<UCustomizableSkeletalComponent> CustomizableSkeletalComponent; CustomizableSkeletalComponent; ++CustomizableSkeletalComponent)
+			if (CustomizableSkeletalComponent && !CustomizableSkeletalComponent->IsTemplate())
 			{
-				UWorld* LocalWorld = CustomizableSkeletalComponent->GetWorld();
-				APlayerController* Controller = LocalWorld ? LocalWorld->GetFirstPlayerController() : nullptr;
-
-				if (Controller && Controller->GetViewTarget())
+				UCustomizableObjectInstance* COI = CustomizableSkeletalComponent->CustomizableObjectInstance;
+				if (!COI || !COI->GetCustomizableObject())
 				{
-					PlayerCamera = Cast<UCameraComponent>(Controller->GetViewTarget()->GetComponentByClass(UCameraComponent::StaticClass()));
-					ViewCentersToIndexMap.Add(Controller, 0);
-					break;
+					continue;
 				}
-			}
 
-			if (PlayerCamera && PlayerCamera->IsActive())
-			{
-				FMinimalViewInfo MinimalViewInfo;
-				PlayerCamera->GetCameraView(16.6f, MinimalViewInfo); // The Delta Time parameter is irrelevant
+#if WITH_EDITOR
 
-				FMatrix ViewMatrix;
-				FMatrix ProjectionMatrix;
-				FMatrix ViewProjectionMatrix;
-				UGameplayStatics::GetViewProjectionMatrix(MinimalViewInfo, ViewMatrix, ProjectionMatrix, ViewProjectionMatrix);
+				EWorldType::Type WorldType = EWorldType::Type::None;
 
-				Locations.Add(MinimalViewInfo.Location);
-				ProjectionMatrices.Add(ProjectionMatrix);
-			}
-		}
-		else
-		{
-			int32 index = 0;
-			for (const TWeakObjectPtr<const AActor> ViewCenter : ViewCenters)
-			{
-				if (!ViewCenter.IsValid())
+				UWorld* World = CustomizableSkeletalComponent->GetWorld();
+				if (World)
 				{
-					continue;
+					WorldType = World->WorldType;
 				}
 
-				UCameraComponent* Camera = nullptr;
-				if (const ACameraActor* CameraActor = Cast<ACameraActor>(ViewCenter.Get()))
+				// Blueprint instances and CO Editors
+				if (WorldType == EWorldType::EditorPreview || (!World && !CustomizableSkeletalComponent->GetAttachParentActor()))
 				{
-					Camera = CameraActor->GetCameraComponent();
+					CustomizableSkeletalComponent->EditorUpdateComponent();
+					continue;
 				}
-				else
+
+				// Skip if world type is EditorPreview, GamePreview, etc...
+				if (WorldType == EWorldType::GamePreview || WorldType == EWorldType::GameRPC || WorldType == EWorldType::Inactive)
 				{
-					Camera = Cast<UCameraComponent>(ViewCenter->GetComponentByClass(UCameraComponent::StaticClass()));
+					continue;
 				}
-				
-				if (Camera && Camera->IsActive())
+#endif // WITH_EDITOR
+
+				FLODTracker& LODTracker = InstancesMinLOD.FindOrAdd(COI);
+
+				if (!LODTracker.bInitialized)
 				{
-					FMinimalViewInfo MinimalViewInfo;
-					Camera->GetCameraView(16.6f, MinimalViewInfo); // The Delta Time parameter is irrelevant
-
-					FMatrix ViewMatrix;
-					FMatrix ProjectionMatrix;
-					FMatrix ViewProjectionMatrix;
-					UGameplayStatics::GetViewProjectionMatrix(MinimalViewInfo, ViewMatrix, ProjectionMatrix, ViewProjectionMatrix);
-
-					Locations.Add(MinimalViewInfo.Location);
-					ProjectionMatrices.Add(ProjectionMatrix);
-					ViewCentersToIndexMap.Add(ViewCenter.Get(), index);
-					index++;
+					LODTracker.RequestedLODsPerComponent.AddZeroed(COI->GetCustomizableObject()->GetComponentCount());
+					LODTracker.bInitialized = true;
 				}
-			}
-		}
-
-		if (Locations.Num())
-		{
-			float SquaredDist = FMath::Square(CloseCustomizableObjectsDist);
-			float ScreenSizeLOD1 = 0.04f;
-			float ScreenSizeLOD2 = 0.02f;
 
-			for (TObjectIterator<UCustomizableObjectInstance> CustomizableObjectInstance; CustomizableObjectInstance; ++CustomizableObjectInstance)
-			{
-				if ( IsValidChecked(*CustomizableObjectInstance) && CustomizableObjectInstance->GetPrivate() )
+				USkeletalMeshComponent* Parent = Cast<USkeletalMeshComponent>(CustomizableSkeletalComponent->GetAttachParent());
+				if (Parent)
 				{
-					UCustomizableObject* CustomizableObject = CustomizableObjectInstance->GetCustomizableObject(); 
-					if (!CustomizableObject)
-					{
-						continue;
-					}
-														
-					if (CustomizableObjectInstance->GetMinSquareDistToPlayer() > SquaredDist || !CustomizableObjectInstance->NearestToActor.IsValid())
-					{
-						continue;
-					}
+					COI->SetIsBeingUsedByComponentInPlay(true);
 
-					USkeletalMesh* RefSkeletalMesh = CustomizableObject->GetRefSkeletalMesh();
-					if (!RefSkeletalMesh)
-					{
-						continue;
-					}
-					
-					if (RefSkeletalMesh->GetLODInfoArray().Num() > 1)
+					// If it's the local player set max priority
+					const APawn* Pawn = Cast<APawn>(Parent->GetAttachParentActor());
+					if (Pawn && Pawn->IsLocallyControlled())
 					{
-						ScreenSizeLOD1 = FMath::Square(RefSkeletalMesh->GetLODInfo(1)->ScreenSize.Default);
+						COI->SetMinSquareDistToPlayer(-1.f);
 					}
-					
-					if (RefSkeletalMesh->GetLODInfoArray().Num() > 2)
+					else
 					{
-						ScreenSizeLOD2 = FMath::Square(RefSkeletalMesh->GetLODInfo(2)->ScreenSize.Default);
+						LODTracker.MaxLOD = COI->GetNumLODsAvailable() - 1;
 					}
-					
-					float BoundSize = RefSkeletalMesh->GetBounds().SphereRadius;
-					FVector BoundOrigin;
 
-					AActor* ParentActor = CustomizableObjectInstance->NearestToActor->GetAttachmentRootActor();
-					if (ParentActor)
-					{
-						BoundSize *= ParentActor->GetActorScale().Size() * 1.55f;
-						BoundOrigin = ParentActor->GetActorLocation();
-					}
-					else
+
+					// Use the component minLOD to set the minimum LOD to generate
+					LODTracker.MinLOD = FMath::Min(LODTracker.MinLOD, Parent->bOverrideMinLod ? Parent->MinLodModel : 0);
+
+					// If the parent component have a SkeletalMesh use the RequestedLODLevel of the component as reference to know which LODs mutable should generate.
+					if (Parent->GetSkeletalMeshAsset() && LODTracker.RequestedLODsPerComponent.IsValidIndex(CustomizableSkeletalComponent->ComponentIndex))
 					{
-						BoundOrigin = CustomizableObjectInstance->NearestToActor->GetComponentLocation();
+						LODTracker.RequestedLODsPerComponent[CustomizableSkeletalComponent->ComponentIndex] |= 1 << Parent->GetPredictedLODLevel();
 					}
+				}
+			}
+		}
 
+		for (TPair<TObjectPtr<UCustomizableObjectInstance>, FLODTracker>& It : InstancesMinLOD)
+		{
+			if (IsValidChecked(It.Key) && It.Key->GetPrivate())
+			{
+				const UCustomizableObject* CustomizableObject = It.Key->GetCustomizableObject();
+				if (!CustomizableObject || CustomizableObject->IsLocked())
+				{
+					continue;
+				}
 
-					int32* ViewIndex = ViewCentersToIndexMap.Find(CustomizableObjectInstance->NearestToViewCenter.Get());
-					float InstanceScreenSize = ComputeBoundsScreenRadiusSquared(BoundOrigin, BoundSize, Locations[ViewIndex ? *ViewIndex : 0 ], ProjectionMatrices[ ViewIndex ? *ViewIndex : 0]);
-
-					int32 InstanceMaxLOD = FMath::Min(RefSkeletalMesh->GetLODNum(), CustomizableObjectInstance->GetNumLODsAvailable()) - 1;
+				if (IsOnlyUpdateCloseCustomizableObjectsEnabled() && (It.Key->GetMinSquareDistToPlayer() > FMath::Square(CloseCustomizableObjectsDist) || !It.Key->NearestToActor.IsValid()))
+				{
+					continue;
+				}
 
-					bool bIsPawn = CustomizableObjectInstance->GetMinSquareDistToPlayer() <= 0;
+				// Limit MinLOD
+				It.Value.MinLOD = FMath::Min(It.Value.MinLOD, CustomizableObject->GetNumLODs() - 1);
 
-					if (InstanceMaxLOD == 0 || InstanceScreenSize > ScreenSizeLOD1 || bIsPawn)
-					{
-						CustomizableObjectInstance->SetMinMaxLODToLoad(0, InstanceMaxLOD, false);
-					}
-					else if (InstanceMaxLOD == 1 || InstanceScreenSize > ScreenSizeLOD2 )
-					{
-						CustomizableObjectInstance->SetMinMaxLODToLoad(1, InstanceMaxLOD, false);
-					}
-					else
-					{
-						CustomizableObjectInstance->SetMinMaxLODToLoad(2, 2);
-					}
+				// If it's the player generate only the first LOD
+				if (It.Key->GetMinSquareDistToPlayer() == -1.f && It.Value.MaxLOD == MAX_int32)
+				{
+					It.Value.MaxLOD = It.Value.MinLOD;
 				}
+
+				It.Key->SetRequestedLODs(It.Value.MinLOD, It.Value.MaxLOD, It.Value.RequestedLODsPerComponent);
 			}
 		}
 	}
 }
 
 
-int32 UCustomizableInstanceLODManagement::GetNumGeneratedInstancesLimitFullLODs()
+int32 UCustomizableInstanceLODManagement::GetNumGeneratedInstancesLimitFullLODs() const
 {
 	return CVarNumGeneratedInstancesLimit.GetValueOnGameThread();
 }
 
 
-int32 UCustomizableInstanceLODManagement::GetNumGeneratedInstancesLimitLOD1()
+int32 UCustomizableInstanceLODManagement::GetNumGeneratedInstancesLimitLOD1() const
 {
 	return CVarNumGeneratedInstancesLimitLOD1.GetValueOnGameThread();
 }
 
 
-int32 UCustomizableInstanceLODManagement::GetNumGeneratedInstancesLimitLOD2()
+int32 UCustomizableInstanceLODManagement::GetNumGeneratedInstancesLimitLOD2() const
 {
 	return CVarNumGeneratedInstancesLimitLOD2.GetValueOnGameThread();
 }
@@ -425,26 +403,20 @@ int32 UCustomizableInstanceLODManagement::GetNumberOfPriorityUpdateInstances() c
 }
 
 
-bool UCustomizableInstanceLODManagement::IsOnlyUpdateCloseCustomizableObjectsEnabled() const
+void UCustomizableInstanceLODManagement::SetCustomizableObjectsUpdateDistance(float Distance)
 {
-	return bOnlyUpdateCloseCustomizableObjects;
+	CloseCustomizableObjectsDist = Distance;
 }
 
 
-void UCustomizableInstanceLODManagement::EnableOnlyUpdateCloseCustomizableObjects(float CloseDist)
+float UCustomizableInstanceLODManagement::GetOnlyUpdateCloseCustomizableObjectsDist() const
 {
-	bOnlyUpdateCloseCustomizableObjects = true;
-	CloseCustomizableObjectsDist = CloseDist;
+	return CloseCustomizableObjectsDist;
 }
 
 
-void UCustomizableInstanceLODManagement::DisableOnlyUpdateCloseCustomizableObjects()
+bool UCustomizableInstanceLODManagement::IsOnlyUpdateCloseCustomizableObjectsEnabled() const
 {
-	bOnlyUpdateCloseCustomizableObjects = false;
+	return CVarOnlyUpdateCloseCustomizableObjects.GetValueOnGameThread();
 }
 
-
-float UCustomizableInstanceLODManagement::GetOnlyUpdateCloseCustomizableObjectsDist() const
-{
-	return CloseCustomizableObjectsDist;
-}
diff --git a/Engine/Plugins/Experimental/Mutable/Source/CustomizableObject/Private/MuCO/CustomizableInstancePrivateData.h b/Engine/Plugins/Experimental/Mutable/Source/CustomizableObject/Private/MuCO/CustomizableInstancePrivateData.h
--- a/Engine/Plugins/Experimental/Mutable/Source/CustomizableObject/Private/MuCO/CustomizableInstancePrivateData.h
+++ b/Engine/Plugins/Experimental/Mutable/Source/CustomizableObject/Private/MuCO/CustomizableInstancePrivateData.h
@@ -26,25 +26,24 @@ enum ECOInstanceFlags
 	Updating						= 1 << 0,	//
 	CreatingSkeletalMesh			= 1 << 1,	//
 	Generated						= 1 << 2,	//
-	AssetsLoaded					= 1 << 3,	// Assets required to generate a mesh. Mainly Materials, PhysicsAssets and AnimBps 
-	ReduceLODs						= 1 << 4,	// CurrentMinLOD will be mapped to the LOD 0 of the generated SkeletalMesh. 
-	ReuseTextures					= 1 << 5, 	// 
-	ReplacePhysicsAssets			= 1 << 6,	// Merge active PhysicsAssets and replace the base physics asset
-	PendingMeshUpdate				= 1 << 7,	// Internaly used to know if we should regenerate the SkeletalMeshes when updating the instance
+	ReuseTextures					= 1 << 3, 	// 
+	ReplacePhysicsAssets			= 1 << 4,	// Merge active PhysicsAssets and replace the base physics asset
 
 	// Update priorities
-	UsedByComponent					= 1 << 8,	// If any components are using this instance, they will set flag every frame
-	UsedByComponentInPlay			= 1 << 9,	// If any components are using this instance in play, they will set flag every frame
-	UsedByPlayerOrNearIt			= 1 << 10,	// The instance is used by the player or is near the player, used to give more priority to its updates
-	DiscardedByNumInstancesLimit	= 1 << 11,	// The instance is descarded because we exceeded the limit of instances generated 
+	UsedByComponent					= 1 << 5,	// If any components are using this instance, they will set flag every frame
+	UsedByComponentInPlay			= 1 << 6,	// If any components are using this instance in play, they will set flag every frame
+	UsedByPlayerOrNearIt			= 1 << 7,	// The instance is used by the player or is near the player, used to give more priority to its updates
+	DiscardedByNumInstancesLimit	= 1 << 8,	// The instance is descarded because we exceeded the limit of instances generated 
 
 	// Types of updates
-	PendingLODsUpdate				= 1 << 12,	// Used to queue an update due to a change in LODs required by the instance
-	PendingLODsUpdateSecondStage	= 1 << 13,	// Second stage of the previous flag, Internaly used to queue an update due to a change in LODs required by the instance
-	PendingLODsDowngrade			= 1 << 14,	// Used to queue a downgrade update to reduce the number of LODs. LOD update goes from a high res level to a low res one, ex: 0 to 1 or 1 to 2
+	PendingLODsUpdate				= 1 << 9,	// Used to queue an update due to a change in LODs required by the instance
+	PendingLODsDowngrade			= 1 << 10,	// Used to queue a downgrade update to reduce the number of LODs. LOD update goes from a high res level to a low res one, ex: 0 to 1 or 1 to 2
 
 	// Streaming
-	LODsStreamingEnabled			= 1 << 15,	// Stream LODs instead of generating all LODs at once. Enables LODs update(upgrade)/downgrade.
+	LODsStreamingEnabled			= 1 << 11,	// Stream LODs instead of generating all LODs at once. Enables LODs update(upgrade)/downgrade.
+	
+	// Generation
+	ForceGenerateAllLODs			= 1 << 12,	// If set, Requested LOD Levels will be ignored and all LODs in between the current min/max lod will be generated
 };
 
 
@@ -230,8 +229,6 @@ public:
 	// Copy data generated in the mutable thread over to the instance and initializes additional data required during the update
 	void PrepareForUpdate(const TSharedPtr<FMutableOperationData>& OperationData);
 
-	void SetMinMaxLODToLoad(UCustomizableObjectInstance* Public, int32 NewMinLOD, int32 NewMaxLOD, bool bLimitLODUpgrades = true);
-	
 	int32 GetNumLODsAvailable() const { return NumLODsAvailable; }
 	
 	// The following method is basically copied from PostEditChangeProperty and/or SkeletalMesh.cpp to be able to replicate PostEditChangeProperty without the editor
@@ -252,11 +249,19 @@ public:
 
 private:
 
-	bool BuildSkeletalMeshSkeletonData(const TSharedPtr<FMutableOperationData>& OperationData, USkeletalMesh* SkeletalMesh, const FMutableRefSkeletalMeshData* RefSkeletalMeshData, UCustomizableObjectInstance* CustomizableObjectIntance, int32 ComponentIndex);
-	void BuildMorphTargetsData(const TSharedPtr<FMutableOperationData>& OperationData, USkeletalMesh* SkeletalMesh, UCustomizableObjectInstance* CustomizableObjectInstance, int32 ComponentIndex);
-	void BuildClothingData(const TSharedPtr<FMutableOperationData>& OperationData, USkeletalMesh* SkeletalMesh, UCustomizableObjectInstance* CustomizableObjectInstance, int32 ComponentIndex);
-	void BuildSkeletalMeshElementData(const TSharedPtr<FMutableOperationData>& OperationData, USkeletalMesh* SkeletalMesh, UCustomizableObjectInstance* CustomizableObjectInstance, int32 ComponentIndex);
-	bool BuildSkeletalMeshRenderData(const TSharedPtr<FMutableOperationData>& OperationData, USkeletalMesh* SkeletalMesh, UCustomizableObjectInstance* CustomizableObjectInstance, int32 ComponentIndex);
+	void InitLastUpdateData(const TSharedPtr<FMutableOperationData>& OperationData);
+
+	void InitSkeletalMeshData(const TSharedPtr<FMutableOperationData>& OperationData, USkeletalMesh* SkeletalMesh, const FMutableRefSkeletalMeshData* RefSkeletalMeshData, int32 ComponentIndex);
+
+	bool BuildSkeletonData(const TSharedPtr<FMutableOperationData>& OperationData, USkeletalMesh* SkeletalMesh, const FMutableRefSkeletalMeshData* RefSkeletalMeshData, UCustomizableObjectInstance* CustomizableObjectIntance, int32 ComponentIndex);
+	void BuildMeshSockets(const TSharedPtr<FMutableOperationData>& OperationData, USkeletalMesh* SkeletalMesh, const FMutableRefSkeletalMeshData* RefSkeletalMeshData, UCustomizableObjectInstance* CustomizableObjectInstance, mu::MeshPtrConst MutableMesh);
+	void BuildOrCopyElementData(const TSharedPtr<FMutableOperationData>& OperationData, USkeletalMesh* SkeletalMesh, UCustomizableObjectInstance* CustomizableObjectInstance, int32 ComponentIndex);
+	void BuildOrCopyMorphTargetsData(const TSharedPtr<FMutableOperationData>& OperationData, USkeletalMesh* SkeletalMesh, const USkeletalMesh* SrcSkeletalMesh, UCustomizableObjectInstance* CustomizableObjectInstance, int32 ComponentIndex);
+	bool BuildOrCopyRenderData(const TSharedPtr<FMutableOperationData>& OperationData, USkeletalMesh* SkeletalMesh, const USkeletalMesh* SrcSkeletalMesh, UCustomizableObjectInstance* CustomizableObjectInstance, int32 ComponentIndex);
+	void BuildOrCopyClothingData(const TSharedPtr<FMutableOperationData>& OperationData, USkeletalMesh* SkeletalMesh, const USkeletalMesh* SrcSkeletalMesh, UCustomizableObjectInstance* CustomizableObjectInstance, int32 ComponentIndex);
+
+	bool CopySkeletonData(const TSharedPtr<FMutableOperationData>& OperationData, USkeletalMesh* SrcSkeletalMesh, USkeletalMesh* DestSkeletalMesh, int32 ComponentIndex);
+	void CopyMeshSockets(USkeletalMesh* SrcSkeletalMesh, USkeletalMesh* DestSkeletalMesh);
 
 	//
 	USkeleton* MergeSkeletons(UCustomizableObjectInstance* Public, const FMutableRefSkeletalMeshData* RefSkeletalMeshData, int32 ComponentIndex);
@@ -264,26 +269,23 @@ private:
 	//
 	UPhysicsAsset* GetOrBuildPhysicsAsset(TObjectPtr<class UPhysicsAsset> TamplateAsset, const mu::PhysicsBody* PhysicsBody, int32 ComponentIndex, bool bDisableCollisionBetweenAssets);
 	
+	// Create a transient texture and add it to the TextureTrackerArray
+	UTexture2D* CreateTexture();
+
 	
+	void InvalidateGeneratedData();
+
+	bool DoComponentsNeedUpdate(UCustomizableObjectInstance* CustomizableObjectInstance, const TSharedPtr<FMutableOperationData>& OperationData, TArray<bool>& OutComponentNeedsUpdate, bool& bOutEmptyMesh);
+
 	int32 GetLastMeshId(int32 ComponentIndex, int32 LODIndex) const;
 	void SetLastMeshId(int32 ComponentIndex, int32 LODIndex, int32 MeshId);
-	void ClearAllLastMeshIds();
-
-	bool MeshNeedsUpdate(UCustomizableObjectInstance* Public, const TSharedPtr<FMutableOperationData>& OperationData, bool& bOutEmptyMesh);
 
-	UTexture2D* CreateTexture();
-	
 public:
 
 	// If any components are using this instance, they will store the min of their distances to the player here every frame for LOD purposes
 	float MinSquareDistFromComponentToPlayer;
 	float LastMinSquareDistFromComponentToPlayer; // The same as the previous dist for last frame
 												
-	double LastUpdateTime;
-
-	/** Saves the LODs requested on the update. */
-	void SaveMinMaxLODToLoad(const UCustomizableObjectInstance* Public);
-	
 	// This is the LODs that the Customizable Object has
 	int32 NumLODsAvailable;
 
@@ -316,6 +318,9 @@ public:
 	UPROPERTY(Transient, Category = CustomizableObjectInstance, editfixedsize, VisibleAnywhere)
 	FGameplayTagContainer AnimBPGameplayTags;
 
+	// Struct used during an update to avoid generating resources that can be reused.
+	FInstanceGeneratedData LastUpdateData;
+
 private:
 	
 	ECOInstanceFlags InstanceFlagsPrivate = ECOInstanceFlags::ECONone;
diff --git a/Engine/Plugins/Experimental/Mutable/Source/CustomizableObject/Private/MuCO/CustomizableObjectInstance.cpp b/Engine/Plugins/Experimental/Mutable/Source/CustomizableObject/Private/MuCO/CustomizableObjectInstance.cpp
--- a/Engine/Plugins/Experimental/Mutable/Source/CustomizableObject/Private/MuCO/CustomizableObjectInstance.cpp
+++ b/Engine/Plugins/Experimental/Mutable/Source/CustomizableObject/Private/MuCO/CustomizableObjectInstance.cpp
@@ -74,16 +74,48 @@ void UCustomizableInstancePrivateData::SetLastMeshId(int32 ComponentIndex, int32
 }
 
 
-void UCustomizableInstancePrivateData::ClearAllLastMeshIds()
+void UCustomizableInstancePrivateData::InitLastUpdateData(const TSharedPtr<FMutableOperationData>& OperationData)
 {
-	if (NumLODsAvailable != INT32_MAX)
+	LastUpdateData.LODs.Init(FInstanceGeneratedData::FLOD(), OperationData->NumLODsAvailable);
+	for (int32 LODIndex = OperationData->CurrentMinLOD; LODIndex <= OperationData->CurrentMaxLOD; ++LODIndex)
+	{
+		LastUpdateData.LODs[LODIndex].FirstComponent = OperationData->InstanceUpdateData.LODs[LODIndex].FirstComponent;
+		LastUpdateData.LODs[LODIndex].ComponentCount = OperationData->InstanceUpdateData.LODs[LODIndex].ComponentCount;
+	}
+
+	const int32 NumComponents = OperationData->InstanceUpdateData.Components.Num();
+	LastUpdateData.Components.Init(FInstanceGeneratedData::FComponent(), NumComponents);
+	for (int32 ComponentIndex = 0; ComponentIndex < NumComponents; ++ComponentIndex)
+	{
+		FInstanceGeneratedData::FComponent& LastUpdateDataComponent = LastUpdateData.Components[ComponentIndex];
+		const FInstanceUpdateData::FComponent& UpdateComponent = OperationData->InstanceUpdateData.Components[ComponentIndex];
+		
+		LastUpdateDataComponent.ComponentId = UpdateComponent.Id;
+		LastUpdateDataComponent.bGenerated = UpdateComponent.bGenerated;
+
+		LastUpdateDataComponent.MeshID = UpdateComponent.MeshID;
+
+		LastUpdateDataComponent.FirstSurface = UpdateComponent.FirstSurface;
+		LastUpdateDataComponent.SurfaceCount = UpdateComponent.SurfaceCount;
+	}
+
+	const int32 NumSurfaces = OperationData->InstanceUpdateData.Surfaces.Num();
+	LastUpdateData.SurfaceIds.Init(0, NumSurfaces);
+	for (int32 SurfaceIndex = 0; SurfaceIndex < NumSurfaces; ++SurfaceIndex)
+{
+		LastUpdateData.SurfaceIds[SurfaceIndex] = OperationData->InstanceUpdateData.Surfaces[SurfaceIndex].SurfaceId;
+	}
+}
+
+
+void UCustomizableInstancePrivateData::InvalidateGeneratedData()
 	{
 		for (FCustomizableInstanceComponentData& ComponentData : ComponentsData)
 		{
-			ComponentData.LastMeshIdPerLOD.Reset(NumLODsAvailable);
 			ComponentData.LastMeshIdPerLOD.Init(INDEX_NONE, NumLODsAvailable);
 		}
-	}
+	
+	LastUpdateData.Clear();
 }
 
 
@@ -124,67 +156,10 @@ void UCustomizableObjectInstance::SetDescriptor(const FCustomizableObjectInstanc
 }
 
 
-void UCustomizableInstancePrivateData::SetMinMaxLODToLoad(UCustomizableObjectInstance* Public, int32 NewMinLOD, int32 NewMaxLOD, bool bLimitLODUpgrades)
-{
-	if (!HasCOInstanceFlags(LODsStreamingEnabled))
-	{
-		return;
-	}
-
-	// Min LOD
-	NewMinLOD = FMath::Min(FMath::Max(NewMinLOD, static_cast<int32>(FirstLODAvailable)), FirstLODAvailable + NumMaxLODsToStream);
-
-	const bool bIsDowngradeLODUpdate = Public->GetCurrentMinLOD() >= 0 && NewMinLOD > Public->GetCurrentMinLOD();
-	SetCOInstanceFlags(bIsDowngradeLODUpdate ? PendingLODsDowngrade : ECONone);
-
-	if (!bIsDowngradeLODUpdate && bLimitLODUpgrades && HasCOInstanceFlags(Generated))
-	{
-		if (FPlatformTime::Seconds() - LastUpdateTime < 3.f) // This is a sort of hysteresis/cooldown to prevent too many instances being discarded/updated per second
-		{
-			return;
-		}
-	}
-
-	// Max LOD
-	if (NumLODsAvailable != INT32_MAX)
-	{
-		if (Public->Descriptor.MaxLOD == INT32_MAX)
-		{
-			Public->Descriptor.MaxLOD = NumLODsAvailable - 1;
-		}
-
-		NewMaxLOD = FMath::Min(NewMaxLOD, NumLODsAvailable - 1);
-	}
-	
-	NewMaxLOD = FMath::Max(NewMaxLOD, NewMinLOD);
-
-	// Save the new LODs
-	const bool bLODChanged = Public->Descriptor.MinLOD != NewMinLOD || Public->Descriptor.MaxLOD != NewMaxLOD;
-	
-	Public->Descriptor.MinLOD = NewMinLOD;
-	Public->Descriptor.MaxLOD = NewMaxLOD;
-
-	SetCOInstanceFlags(bLODChanged ? PendingLODsUpdate : ECONone);
-
-	if (bLODChanged)
-	{
-		Public->UpdateDescriptorRuntimeHash.UpdateMinMaxLOD(NewMinLOD, NewMaxLOD);
-
-		if (const FMutableQueueElem* QueueElem = UCustomizableObjectSystem::GetInstance()->GetPrivate()->MutableOperationQueue.Get(Public))
-		{
-			QueueElem->Operation->InstanceDescriptorRuntimeHash.UpdateMinMaxLOD(NewMinLOD, NewMaxLOD);
-		}
-	}
-}
-
-
 void UCustomizableInstancePrivateData::PrepareForUpdate(const TSharedPtr<FMutableOperationData>& OperationData)
 {
-	const bool bNumLODsUpdated = NumLODsAvailable != OperationData->NumLODsAvailable;
-	
-	SetCOInstanceFlags(bNumLODsUpdated ? PendingMeshUpdate : ECONone);
-
-	NumLODsAvailable = OperationData->NumLODsAvailable;
+	// Prepare LastUpdateData to allow reuse of meshes and surfaces
+	InitLastUpdateData(OperationData);
 
 	TArray<FInstanceUpdateData::FLOD>& LODs = OperationData->InstanceUpdateData.LODs;
 	TArray<FInstanceUpdateData::FComponent>& Components = OperationData->InstanceUpdateData.Components;
@@ -202,6 +177,11 @@ void UCustomizableInstancePrivateData::PrepareForUpdate(const TSharedPtr<FMutabl
 
 			if (FCustomizableInstanceComponentData* ComponentData = GetComponentData(Component.Id))
 			{
+				if (!Component.bGenerated || Component.bReuseMesh)
+				{
+					continue;
+				}
+
 				ComponentData->AnimSlotToBP.Empty();
 
 #if WITH_EDITORONLY_DATA
@@ -214,16 +194,12 @@ void UCustomizableInstancePrivateData::PrepareForUpdate(const TSharedPtr<FMutabl
 				ComponentData->PhysicsAssets.PhysicsAssetsToMerge.Empty();
 				ComponentData->ClothingPhysicsAssetsToStream.Empty();
 
-				if (bNumLODsUpdated)
-				{
-					ComponentData->LastMeshIdPerLOD.Init(-1, NumLODsAvailable);
-				}
 				continue;
 			}
 
 			FCustomizableInstanceComponentData& NewComponentData = ComponentsData.AddDefaulted_GetRef();
 			NewComponentData.ComponentIndex = Component.Id;
-			NewComponentData.LastMeshIdPerLOD.Init(-1, NumLODsAvailable);
+			NewComponentData.LastMeshIdPerLOD.Init(INDEX_NONE, NumLODsAvailable);
 		}
 	}
 
@@ -478,7 +454,7 @@ void UCustomizableObjectInstance::PostLoad()
 	Super::PostLoad();
 
 	Descriptor.ReloadParameters();
-	SetMinMaxLODToLoad();
+	SetRequestedLODs(Descriptor.MinLOD, Descriptor.MaxLOD, Descriptor.RequestedLODLevels);
 }
 
 
@@ -548,14 +524,17 @@ void UCustomizableInstancePrivateData::InstanceUpdateFlags(const UCustomizableOb
 	const UCustomizableObject* CustomizableObject = Public.GetCustomizableObject();
 	
 	FirstLODAvailable = CustomizableObject->LODSettings.FirstLODAvailable;
+	NumLODsAvailable = CustomizableObject->GetNumLODs();
 
 	if (CustomizableObject->LODSettings.bLODStreamingEnabled || !UCustomizableObjectSystem::GetInstance()->IsProgressiveMipStreamingEnabled())
 	{
 		SetCOInstanceFlags(LODsStreamingEnabled);
+		ClearCOInstanceFlags(ForceGenerateAllLODs);
 		NumMaxLODsToStream = CustomizableObject->LODSettings.NumLODsToStream;
 	}
 	else
 	{
+		SetCOInstanceFlags(ForceGenerateAllLODs);
 		ClearCOInstanceFlags(LODsStreamingEnabled);
 		NumMaxLODsToStream = 0;
 	}
@@ -605,7 +584,7 @@ void UCustomizableInstancePrivateData::ReloadParameters(UCustomizableObjectInsta
 		ProjectorStates.Reset();
 	}
 	
-	ClearAllLastMeshIds();
+	InvalidateGeneratedData();
 	
 	Public->GetDescriptor().ReloadParameters();
 }
@@ -615,7 +594,7 @@ void UCustomizableObjectInstance::SetObject(UCustomizableObject* InObject)
 {
 	Descriptor.SetCustomizableObject(*InObject);
 	PrivateData->ReloadParameters(this);
-	SetMinMaxLODToLoad();
+	SetRequestedLODs(Descriptor.MinLOD, Descriptor.MaxLOD, Descriptor.RequestedLODLevels);
 }
 
 
@@ -830,7 +809,6 @@ void UCustomizableInstancePrivateData::DoUpdateSkeletalMesh(UCustomizableObjectI
 	{
 		if (bShouldUpdateLODs)
 		{
-			SetCOInstanceFlags(PendingLODsUpdateSecondStage);
 			UE_LOG(LogMutable, Verbose, TEXT("LOD change: %d, %d -> %d, %d"), CustomizableInstance.GetCurrentMinLOD(), CustomizableInstance.GetCurrentMaxLOD(), CustomizableInstance.GetMinLODToLoad(), CustomizableInstance.GetMinLODToLoad());
 		}
 
@@ -881,15 +859,10 @@ void UCustomizableInstancePrivateData::DoUpdateSkeletalMesh(UCustomizableObjectI
 				{
 					CustomizableInstance.PreUpdateSkeletalMeshStatus = CustomizableInstance.SkeletalMeshStatus;
 				}
-
-				CustomizableInstance.SkeletalMeshStatus = ESkeletalMeshState::AsyncUpdatePending;
 			}
-#else
-			CustomizableInstance.SkeletalMeshStatus = ESkeletalMeshState::AsyncUpdatePending;
 #endif
 
 			SetCOInstanceFlags(Generated); // Will be done in UpdateSkeletalMesh_PostBeginUpdate
-			LastUpdateTime = FPlatformTime::Seconds();
 			
 			// Do not do work after calling InitUpdateSkeletalMesh. This function can optimize an update and fully complete it before even exiting its scope.
 			System->GetPrivate()->InitUpdateSkeletalMesh(CustomizableInstance, Priority);
@@ -902,7 +875,7 @@ void UCustomizableInstancePrivateData::DoUpdateSkeletalMesh(UCustomizableObjectI
 		if (HasCOInstanceFlags(Generated) && !HasCOInstanceFlags(Updating))
 		{
 			System->GetPrivate()->InitDiscardResourcesSkeletalMesh(&CustomizableInstance);
-			ClearCOInstanceFlags((ECOInstanceFlags)(PendingLODsUpdate | PendingLODsUpdateSecondStage));
+			ClearCOInstanceFlags(PendingLODsUpdate);
 		}
 	}
 }
@@ -1406,335 +1379,331 @@ UPhysicsAsset* UCustomizableInstancePrivateData::GetOrBuildPhysicsAsset(
 }
 
 
-bool UCustomizableInstancePrivateData::UpdateSkeletalMesh_PostBeginUpdate0(UCustomizableObjectInstance* Public, const TSharedPtr<FMutableOperationData>& OperationData )
+void SetMeshUVChannelDensity(FMeshUVChannelInfo& UVChannelInfo, float Density = 0.f)
 {
-	MUTABLE_CPUPROFILER_SCOPE(UCustomizableInstancePrivateData::UpdateSkeletalMesh_PostBeginUpdate0)
-	
-	UCustomizableObject* CustomizableObject = Public->GetCustomizableObject();
-
-	Public->SkeletalMeshStatus = ESkeletalMeshState::Correct;
+	Density = Density > 0.f ? Density : 150.f;
 
-	bool bEmptyMesh;
-	bool bMeshNeedsUpdate = MeshNeedsUpdate(Public, OperationData, bEmptyMesh);
-	SetCOInstanceFlags(CreatingSkeletalMesh);
+	UVChannelInfo.bInitialized = true;
+	UVChannelInfo.bOverrideDensities = false;
 
-	if (bMeshNeedsUpdate)
-	{
-		TextureReuseCache.Empty(); // Sections may have changed, so invalidate the texture reuse cache because it's indexed by section
-	}
-#if WITH_EDITOR
-	else
+	for (int32 i = 0; i < TEXSTREAM_MAX_NUM_UVCHANNELS; ++i)
 	{
-		Public->SkeletalMeshStatus = Public->PreUpdateSkeletalMeshStatus; // The mesh won't be updated, set the previous SkeletalMeshStatus to avoid losing error notifications
+		UVChannelInfo.LocalUVDensities[i] = Density;
 	}
-#endif
+}
 
-	if (bEmptyMesh)
-	{
-		Public->SkeletalMeshStatus = ESkeletalMeshState::UpdateError;
-		UE_LOG(LogMutable, Warning, TEXT("Updating skeletal mesh error for mesh %s"), *Public->GetName());
-	}
 
-	const int32 FirstComponent = OperationData->InstanceUpdateData.LODs[OperationData->CurrentMinLOD].FirstComponent;
-	const int32 NumComponents = OperationData->InstanceUpdateData.LODs[OperationData->CurrentMinLOD].ComponentCount;
+bool UCustomizableInstancePrivateData::DoComponentsNeedUpdate(UCustomizableObjectInstance* Public, const TSharedPtr<FMutableOperationData>& OperationData, TArray<bool>& OutComponentNeedsUpdate, bool& bOutEmptyMesh)
+{
+	UCustomizableObject* CustomizableObject = Public->GetCustomizableObject();
+	check(CustomizableObject);
+
+	const int32 MaxNumLODs = OperationData->InstanceUpdateData.LODs.Num();
+	const int32 MaxNumComponents = CustomizableObject->GetComponentCount();
 
-	bool bHasSkeletalMesh = false;
+	OutComponentNeedsUpdate.AddDefaulted(MaxNumComponents);
 
-	if (bMeshNeedsUpdate && !bEmptyMesh)
+	// Return true if at least one mesh needs to be updated
+	bool bUpdateMeshes = false;
+
+	for (int32 ComponentIndex = 0; ComponentIndex < MaxNumComponents; ++ComponentIndex)
 	{
-		MUTABLE_CPUPROFILER_SCOPE(ConstructMesh);
+		USkeletalMesh* SkeletalMesh = Public->SkeletalMeshes.IsValidIndex(ComponentIndex) ? Public->SkeletalMeshes[ComponentIndex] : nullptr;
+		FCustomizableInstanceComponentData* ComponentData = GetComponentData(ComponentIndex);
 
-		bool bSuccess = true;
+		// Check if the mesh is/isn't generated and it should/shouldn't
+		bool bMeshNeedsUpdate = (SkeletalMesh && !ComponentData) || (!SkeletalMesh && ComponentData);
 
-		const int32 MaxNumComponents = CustomizableObject->ReferenceSkeletalMeshesData.Num();
-		
-		// Initialize the maximum number of SkeletalMeshes we could possibly have. 
-		Public->SkeletalMeshes.Init(nullptr, MaxNumComponents);
+		if (SkeletalMesh)
+		{
+			bMeshNeedsUpdate = bMeshNeedsUpdate || (SkeletalMesh->GetLODNum() - 1) != OperationData->CurrentMaxLOD;
+		}
 
-		for (int32 ComponentIndex = 0; ComponentIndex < NumComponents; ++ComponentIndex)
+		if (ComponentData)
 		{
-			const FInstanceUpdateData::FComponent& Component = OperationData->InstanceUpdateData.Components[FirstComponent+ComponentIndex];
-			
-			if (!Component.Mesh)
+			for (int32 LODIndex = 0; !bMeshNeedsUpdate && LODIndex < MaxNumLODs; ++LODIndex)
 			{
-				continue;
+				// Check if an LOD is generated and it shouldn't
+				const bool LODCanBeGenerated = LODIndex >= OperationData->CurrentMinLOD && LODIndex <= OperationData->CurrentMaxLOD;
+				bMeshNeedsUpdate = !LODCanBeGenerated && ComponentData->LastMeshIdPerLOD[LODIndex] != INDEX_NONE;
 			}
+		}
 
-			// ComponentIndex could be different than the final SkeletalMesh index if a component is not generated.
-			const uint16 SkeletalMeshIndex = Component.Id;
-
-			if (Component.SurfaceCount > 0)
-			{
-				Public->SkeletalMeshes[SkeletalMeshIndex] = NewObject<USkeletalMesh>(GetTransientPackage(), NAME_None, RF_Transient);
-				bHasSkeletalMesh = true;
-			}
+		OutComponentNeedsUpdate[ComponentIndex] = bMeshNeedsUpdate;
 
-			USkeletalMesh* SkeletalMesh = Public->SkeletalMeshes[SkeletalMeshIndex];
+		bUpdateMeshes = bUpdateMeshes || bMeshNeedsUpdate;
+	}
 
-			const FMutableRefSkeletalMeshData* RefSkeletalMeshData = CustomizableObject->GetRefSkeletalMeshData(SkeletalMeshIndex);
-			
-			if (SkeletalMesh && !RefSkeletalMeshData)
-			{
-				bSuccess = false;
-				break;
-			}
+	bool bFoundEmptyMesh = false;
+	bool bFoundEmptyLOD = false;
 
-			if (SkeletalMesh)
-			{
-				INC_DWORD_STAT(STAT_MutableNumSkeletalMeshes);
+	bool bFoundNonEmptyMesh = false;
 
-				// Mutable skeletal meshes are generated dynamically in-game and cannot be streamed from disk
-				SkeletalMesh->NeverStream = 1;
+	for (int32 LODIndex = OperationData->CurrentMinLOD; LODIndex <= OperationData->CurrentMaxLOD && LODIndex < MaxNumLODs; ++LODIndex)
+	{
 
-				bSuccess = BuildSkeletalMeshSkeletonData(OperationData, SkeletalMesh, RefSkeletalMeshData, Public, ComponentIndex);
-				if (!bSuccess)
-				{
-					break;
-				}
+		const FInstanceUpdateData::FLOD& LOD = OperationData->InstanceUpdateData.LODs[LODIndex];
 
-				SkeletalMesh->SetImportedBounds(RefSkeletalMeshData->Bounds);
-				SkeletalMesh->SetPhysicsAsset(RefSkeletalMeshData->PhysicsAsset.Get());
-				SkeletalMesh->SetEnablePerPolyCollision(RefSkeletalMeshData->Settings.bEnablePerPolyCollision);
-				SkeletalMesh->SetPostProcessAnimBlueprint(RefSkeletalMeshData->PostProcessAnimInst.Get());
-				SkeletalMesh->SetShadowPhysicsAsset(RefSkeletalMeshData->ShadowPhysicsAsset.Get());
+		for (int32 ComponentIndex = 0; ComponentIndex < LOD.ComponentCount; ++ComponentIndex)
+		{
+			const FInstanceUpdateData::FComponent& Component = OperationData->InstanceUpdateData.Components[LOD.FirstComponent + ComponentIndex];
 
-				if (!HasCOInstanceFlags(ReduceLODs))
-				{
-					SkeletalMesh->SetMinLod(FirstLODAvailable);
-					SkeletalMesh->SetQualityLevelMinLod(FirstLODAvailable);
-				}
+			const int32 LastMeshID = GetLastMeshId(Component.Id, LODIndex);
 
-				// Merge Physics Assets coming from SubMeshes of the newly generated Mesh
-				if ( OperationData->InstanceUpdateData.LODs.Num() )
+			if (Component.bGenerated)
+			{
+				if (Component.Mesh)
 				{
-					mu::MeshPtrConst MutableMesh = Component.Mesh;
-					mu::Ptr<const mu::PhysicsBody> MutablePhysics = MutableMesh ? MutableMesh->GetPhysicsBody() : nullptr;
+					OutComponentNeedsUpdate[Component.Id] = OutComponentNeedsUpdate[Component.Id] || LastMeshID != Component.MeshID;
 
-					if (MutablePhysics)
+					// Don't build a degenerated mesh if something went wrong
+					if (Component.Mesh->GetVertexCount() > 0)
 					{
-						constexpr bool bDisallowCollisionBetweenAssets = true;
-						UPhysicsAsset* PhysicsAssetResult = GetOrBuildPhysicsAsset( 
-								RefSkeletalMeshData->PhysicsAsset.Get(), MutablePhysics.get(), SkeletalMeshIndex, bDisallowCollisionBetweenAssets);
-					
-						SkeletalMesh->SetPhysicsAsset(PhysicsAssetResult);
-
-						if (PhysicsAssetResult)
-						{
-							// This is only called in editor, no need to add editor guards.	
-							PhysicsAssetResult->SetPreviewMesh(SkeletalMesh);
-						}
+						bFoundNonEmptyMesh = true;
+					}
+					else if (LOD.ComponentCount == 1)
+					{
+						bFoundEmptyMesh = true;
 					}
 				}
-
-				// Update mesh sockets.
+				else if (Component.bReuseMesh)
 				{
-					MUTABLE_CPUPROFILER_SCOPE(ConstructMesh_UpdateSockets);
-	
-					const uint32 SocketCount = RefSkeletalMeshData->Sockets.Num();
+					bFoundNonEmptyMesh = true;
+					check(LastMeshID != INDEX_NONE);
+				}
+				else
+				{
+					OutComponentNeedsUpdate[Component.Id] = true;
+				}
+			}
+			else
+			{
+				OutComponentNeedsUpdate[Component.Id] = OutComponentNeedsUpdate[Component.Id] || LastMeshID != INDEX_NONE;
+			}
 
-					TArray<USkeletalMeshSocket*>& Sockets = SkeletalMesh->GetMeshOnlySocketList();
-					Sockets.Empty(SocketCount);
-					TMap<FName, TTuple<int32, int32>> SocketMap; // Maps Socket name to Sockets Array index and priority
-					
-					for (uint32 SocketIndex = 0; SocketIndex < SocketCount; ++SocketIndex)
-					{
-						const FMutableRefSocket& RefSocket = RefSkeletalMeshData->Sockets[SocketIndex];
+			bUpdateMeshes = bUpdateMeshes || OutComponentNeedsUpdate[Component.Id];
+		}
+	}
 
-						USkeletalMeshSocket* Socket = NewObject<USkeletalMeshSocket>(SkeletalMesh);
-						
-						Socket->SocketName = RefSocket.SocketName;
-						Socket->BoneName = RefSocket.BoneName;
 
-						Socket->RelativeLocation = RefSocket.RelativeLocation;
-						Socket->RelativeRotation = RefSocket.RelativeRotation;
-						Socket->RelativeScale = RefSocket.RelativeScale;
+	if (!bFoundNonEmptyMesh)
+	{
+		bFoundEmptyLOD = true;
+		bUpdateMeshes = true;
+		UE_LOG(LogMutable, Warning, TEXT("Building instance: an LOD has no mesh geometry. This cannot be handled by Unreal."));
+	}
 
-						Socket->bForceAlwaysAnimated = RefSocket.bForceAlwaysAnimated;
-						const int32 LastIndex = Sockets.Add(Socket);
+	bOutEmptyMesh = bFoundEmptyLOD || bFoundEmptyMesh;
+	return bUpdateMeshes;
+}
 
-						SocketMap.Add(Socket->SocketName, TTuple<int32, int32>(LastIndex, RefSocket.Priority));
-					}
 
-					mu::MeshPtrConst MutableMesh = Component.Mesh;
 
-					for (int32 TagIndex = 0; TagIndex < MutableMesh->GetTagCount(); ++TagIndex)
-					{
-						FString Tag = MutableMesh->GetTag(TagIndex);
+bool UCustomizableInstancePrivateData::UpdateSkeletalMesh_PostBeginUpdate0(UCustomizableObjectInstance* Public, const TSharedPtr<FMutableOperationData>& OperationData)
+{
+	MUTABLE_CPUPROFILER_SCOPE(UCustomizableInstancePrivateData::UpdateSkeletalMesh_PostBeginUpdate0)
 
-						if (Tag.RemoveFromStart("__Socket:"))
-						{
-							check(Tag.IsNumeric());
-							const int32 MutableSocketIndex = FCString::Atoi(*Tag);
-							//check(CustomizableObject->SocketArray.IsValidIndex(MutableSocketIndex));
+		bool bEmptyMesh = false;
 
-							if (CustomizableObject->SocketArray.IsValidIndex(MutableSocketIndex))
-							{
-								const FMutableRefSocket& MutableSocket = CustomizableObject->SocketArray[MutableSocketIndex];
-								int32 IndexToWriteSocket = -1;
+	TArray<bool> ComponentNeedsUpdate;
+	bool bUpdateMeshes = DoComponentsNeedUpdate(Public, OperationData, ComponentNeedsUpdate, bEmptyMesh);
 
-								if (TTuple<int32, int32>* FoundSocket = SocketMap.Find(MutableSocket.SocketName))
-								{
-									if (FoundSocket->Value < MutableSocket.Priority)
-									{
-										// Overwrite the existing socket because the new mesh part one is higher priority
-										IndexToWriteSocket = FoundSocket->Key;
-										FoundSocket->Value = MutableSocket.Priority;
-									}
-								}
-								else
-								{
-									// New Socket
-									USkeletalMeshSocket* Socket = NewObject<USkeletalMeshSocket>(SkeletalMesh);
-									IndexToWriteSocket = Sockets.Add(Socket);
-									SocketMap.Add(MutableSocket.SocketName, TTuple<int32, int32>(IndexToWriteSocket, MutableSocket.Priority));
-								}
+	// We can not handle empty meshes, clear any generated mesh and return
+	if (bEmptyMesh)
+	{
+		Public->SkeletalMeshStatus = ESkeletalMeshState::UpdateError;
+		UE_LOG(LogMutable, Warning, TEXT("Updating skeletal mesh error for mesh %s"), *Public->GetName());
 
-								if (IndexToWriteSocket >= 0)
-								{
-									check(Sockets.IsValidIndex(IndexToWriteSocket));
-
-									USkeletalMeshSocket* SocketToWrite = Sockets[IndexToWriteSocket];
-
-									SocketToWrite->SocketName = MutableSocket.SocketName;
-									SocketToWrite->BoneName = MutableSocket.BoneName;
-									
-									SocketToWrite->RelativeLocation = MutableSocket.RelativeLocation;
-									SocketToWrite->RelativeRotation = MutableSocket.RelativeRotation;
-									SocketToWrite->RelativeScale = MutableSocket.RelativeScale;
-									
-									SocketToWrite->bForceAlwaysAnimated = MutableSocket.bForceAlwaysAnimated;
-								}
-							}
-						}
-					}
-				}
-			}
-		}
+		Public->SkeletalMeshes.Reset(); // What about all the references
 
-		if (!bSuccess)
-		{
-			Public->SkeletalMeshes.Reset();
-			return bSuccess;
-		}
-	}
-	else if (bMeshNeedsUpdate)
-	{
-		Public->SkeletalMeshes.Reset();
-		
 		SetCOInstanceFlags(Generated);
-		ClearCOInstanceFlags((ECOInstanceFlags)(CreatingSkeletalMesh | PendingMeshUpdate));
+
+		InvalidateGeneratedData();
 
 		return false;
 	}
 
-	if (bMeshNeedsUpdate && bHasSkeletalMesh)
+	Public->SkeletalMeshStatus = ESkeletalMeshState::Correct;
+
+	// None of the current meshes requires a mesh update. Continue to BuildMaterials
+	if (!bUpdateMeshes)
+	{
+		SetCOInstanceFlags(Generated);
+		return true;
+	}
+
+	SetCOInstanceFlags(CreatingSkeletalMesh);
+
+	TextureReuseCache.Empty(); // Sections may have changed, so invalidate the texture reuse cache because it's indexed by section
+
+	TArray<TObjectPtr<USkeletalMesh>> OldSkeletalMeshes = Public->SkeletalMeshes;
+
+	check(NumLODsAvailable == OperationData->InstanceUpdateData.LODs.Num());
+
+	bool bSuccess = true;
+	bool bHasSkeletalMesh = false;
+
+	UCustomizableObject* CustomizableObject = Public->GetCustomizableObject();
+	check(CustomizableObject);
+
+	// Initialize the maximum number of SkeletalMeshes we could possibly have. 
+	Public->SkeletalMeshes.Init(nullptr, CustomizableObject->GetComponentCount());
+
+	const int32 NumMutableComponents = OperationData->InstanceUpdateData.Components.Num();
+	for (int32 ComponentIndex = 0; ComponentIndex < NumMutableComponents; ++ComponentIndex)
 	{
-		MUTABLE_CPUPROFILER_SCOPE(ReleaseAndBuildMesh);
-		bool bSuccess = true;
+		const FInstanceUpdateData::FComponent& Component = OperationData->InstanceUpdateData.Components[ComponentIndex];
 
-		for (int32 ComponentIndex = 0; ComponentIndex < NumComponents; ++ComponentIndex)
+		if (!ComponentNeedsUpdate[Component.Id])
 		{
-			const FInstanceUpdateData::FComponent& Component = OperationData->InstanceUpdateData.Components[FirstComponent + ComponentIndex];
+			Public->SkeletalMeshes[Component.Id] = OldSkeletalMeshes[Component.Id];
+			bHasSkeletalMesh = true;
+			continue;
+		}
 
-			if (!Component.Mesh)
-			{
-				continue;
-			}
+		if (!Component.bGenerated || (!Component.Mesh && !Component.bReuseMesh) || Component.SurfaceCount == 0)
+		{
+			continue;
+		}
+
+		const FMutableRefSkeletalMeshData* RefSkeletalMeshData = CustomizableObject->GetRefSkeletalMeshData(Component.Id);
+		if (!RefSkeletalMeshData)
+		{
+			bSuccess = false;
+			break;
+		}
+
+		// Check if we have initialized the component
+		if (Public->SkeletalMeshes[Component.Id])
+		{
+			continue;
+		}
+
+		// Create and initialize the SkeletalMesh for this component
+		MUTABLE_CPUPROFILER_SCOPE(ConstructMesh);
+
+		Public->SkeletalMeshes[Component.Id] = NewObject<USkeletalMesh>(GetTransientPackage(), NAME_None, RF_Transient);
 
-			const uint16 SkeletalMeshIndex = Component.Id;
+		USkeletalMesh* SkeletalMesh = Public->SkeletalMeshes[Component.Id];
+		check(SkeletalMesh);
 
-			check(Public->SkeletalMeshes.IsValidIndex(SkeletalMeshIndex));
-			USkeletalMesh* SkeletalMesh = Public->SkeletalMeshes[SkeletalMeshIndex];
+		INC_DWORD_STAT(STAT_MutableNumSkeletalMeshes);
 
-			const FMutableRefSkeletalMeshData* RefSkeletalMeshData = CustomizableObject->GetRefSkeletalMeshData(SkeletalMeshIndex);
+		// Set up the default information any mesh from this component will have (LODArrayInfos, RenderData, Mesh settings, etc). 
+		InitSkeletalMeshData(OperationData, SkeletalMesh, RefSkeletalMeshData, Component.Id);
+
+		if (Component.Mesh)
+		{
+			// Construct a new skeleton, fix up ActiveBones and Bonemap arrays and recompute the RefInvMatrices
+			bSuccess = BuildSkeletonData(OperationData, SkeletalMesh, RefSkeletalMeshData, Public, Component.Id);
 
-			if (SkeletalMesh && !RefSkeletalMeshData)
+			if (!bSuccess)
 			{
-				bSuccess = false;
 				break;
 			}
-			
-			if (SkeletalMesh)
-			{
-				{
-					MUTABLE_CPUPROFILER_SCOPE(ReleaseAndBuildMesh_Empty);
-
-					if (FSkeletalMeshRenderData* Resource = Helper_GetResourceForRendering(SkeletalMesh))
-					{
-						Helper_GetLODData(Resource).Empty();
-					}
-					else
-					{
-						SkeletalMesh->AllocateResourceForRendering();
-					}
-					Helper_GetLODInfoArray(SkeletalMesh).Empty();
-				}
-
-				// Add LODs
-				{
-					MUTABLE_CPUPROFILER_SCOPE(ReleaseAndBuildMesh_AddLODs);
 
-					for (int32 LODIndex = 0; LODIndex <= OperationData->CurrentMaxLOD && LODIndex < OperationData->InstanceUpdateData.LODs.Num(); ++LODIndex)
-					{
-						if (HasCOInstanceFlags(ReduceLODs) && LODIndex < OperationData->CurrentMinLOD)
-						{
-							continue;
-						}
-						
-						Helper_GetLODData(SkeletalMesh).Add((LODIndex <= OperationData->CurrentMinLOD && Helper_GetLODData(SkeletalMesh).Num()) ? &Helper_GetLODData(SkeletalMesh)[0] : new Helper_LODDataType());
+			// Build PhysicsAsset Merge Physics Assets coming from SubMeshes of the newly generated Mesh
+			if (mu::Ptr<const mu::PhysicsBody> MutablePhysics = Component.Mesh->GetPhysicsBody())
+			{
+				constexpr bool bDisallowCollisionBetweenAssets = true;
+				UPhysicsAsset* PhysicsAssetResult = GetOrBuildPhysicsAsset(
+					RefSkeletalMeshData->PhysicsAsset.Get(), MutablePhysics.get(), Component.Id, bDisallowCollisionBetweenAssets);
 
-						const FMutableRefLODData& LODData = RefSkeletalMeshData->LODData[LODIndex];
-						
-						FSkeletalMeshLODInfo& LODInfo = SkeletalMesh->AddLODInfo();
-						LODInfo.ScreenSize = LODData.LODInfo.ScreenSize;
-						LODInfo.LODHysteresis =  LODData.LODInfo.LODHysteresis;
-						LODInfo.bSupportUniformlyDistributedSampling =  LODData.LODInfo.bSupportUniformlyDistributedSampling;
-						LODInfo.bAllowCPUAccess = LODData.LODInfo.bAllowCPUAccess;
-						
-						// Disable LOD simplification when baking instances
-						LODInfo.ReductionSettings.NumOfTrianglesPercentage = 1.f;
-						LODInfo.ReductionSettings.NumOfVertPercentage = 1.f;
-						
-#if WITH_EDITORONLY_DATA
-						LODInfo.BuildSettings.bUseFullPrecisionUVs = true;
-#endif
-					}
-				}
+				SkeletalMesh->SetPhysicsAsset(PhysicsAssetResult);
 
+				if (PhysicsAssetResult)
 				{
-					BuildSkeletalMeshElementData(OperationData, SkeletalMesh, Public, ComponentIndex);
-					bSuccess = BuildSkeletalMeshRenderData(OperationData, SkeletalMesh, Public, ComponentIndex);
+					// This is only called in editor, no need to add editor guards.	
+					PhysicsAssetResult->SetPreviewMesh(SkeletalMesh);
 				}
+			}
 
-				// TODO: update the model resources cache		
-				//mu::MeshPtrConst MutableMesh = PrivateData->InstanceUpdateData.LODs[0].Components[0].Mesh;
-				//FMutableResourceCache& Cache = UCustomizableObjectSystem::GetInstance()->GetPrivate()->GetObjectCache(CustomizableObject);
-				//Cache.Meshes.Add(MutableMesh->GetId(), SkeletalMesh);
+			// Add sockets from the SkeletalMesh of reference and from the MutableMesh
+			BuildMeshSockets(OperationData, SkeletalMesh, RefSkeletalMeshData, Public, Component.Mesh);
+		}
+		else if (Component.bReuseMesh)
+		{
+			// Source mesh to reuse data from
+			USkeletalMesh* SrcSkeletalMesh = OldSkeletalMeshes[Component.Id];
+			check(SrcSkeletalMesh);
 
-				if (bSuccess)
-				{
-					BuildMorphTargetsData(OperationData, SkeletalMesh, Public, ComponentIndex);
-					BuildClothingData(OperationData, SkeletalMesh, Public, ComponentIndex);
+			// Set the Skeleton from the previously generated mesh and fix up ActiveBones and Bonemap arrays from new components
+			bSuccess = CopySkeletonData(OperationData, SrcSkeletalMesh, SkeletalMesh, Component.Id);
 
-					ensure(Helper_GetLODData(SkeletalMesh).Num() > 0);
-					ensure(Helper_GetLODInfoArray(SkeletalMesh).Num() > 0);
-				}
-			}
-			
 			if (!bSuccess)
 			{
-				Public->SkeletalMeshes.Reset();
 				break;
 			}
+
+			// Set the PhysicsAsset from the previously generated mesh
+			SkeletalMesh->SetPhysicsAsset(SrcSkeletalMesh->GetPhysicsAsset());
+
+			// Copy Sockets from the previously generated mesh
+			CopyMeshSockets(SrcSkeletalMesh, SkeletalMesh);
+		}
+		else
+		{
+			bSuccess = false;
+			break;
 		}
 
-		return bSuccess;
+		bHasSkeletalMesh = true;
 	}
 
-	SetCOInstanceFlags(Generated);
-	ClearCOInstanceFlags((ECOInstanceFlags)(CreatingSkeletalMesh | PendingMeshUpdate));
 
-	return !bMeshNeedsUpdate;
+	if (bSuccess && bHasSkeletalMesh)
+	{
+		MUTABLE_CPUPROFILER_SCOPE(BuildSkeletalMeshData);
+
+		const int32 NumComponents = Public->SkeletalMeshes.Num();
+		for (int32 ComponentIndex = 0; bSuccess && ComponentIndex < NumComponents; ++ComponentIndex)
+		{
+			if (!ComponentNeedsUpdate[ComponentIndex])
+			{
+				continue;
+			}
+
+			USkeletalMesh* SkeletalMesh = Public->SkeletalMeshes[ComponentIndex];
+
+			if (!SkeletalMesh)
+			{
+				continue;
+			}
+
+			// Mesh to copy data from if possible. 
+			USkeletalMesh* OldSkeletalMesh = OldSkeletalMeshes.IsValidIndex(ComponentIndex) ? OldSkeletalMeshes[ComponentIndex] : nullptr;
+
+			{
+				BuildOrCopyElementData(OperationData, SkeletalMesh, Public, ComponentIndex);
+				bSuccess = BuildOrCopyRenderData(OperationData, SkeletalMesh, OldSkeletalMesh, Public, ComponentIndex);
+			}
+
+			if (bSuccess)
+			{
+				BuildOrCopyMorphTargetsData(OperationData, SkeletalMesh, OldSkeletalMesh, Public, ComponentIndex);
+				BuildOrCopyClothingData(OperationData, SkeletalMesh, OldSkeletalMesh, Public, ComponentIndex);
+
+				ensure(Helper_GetLODData(SkeletalMesh).Num() > 0);
+				ensure(Helper_GetLODInfoArray(SkeletalMesh).Num() > 0);
+			}
+		}
+	}
+
+	SetCOInstanceFlags(Generated);
+	ClearCOInstanceFlags(CreatingSkeletalMesh); // TODO MTBL-391: Review
+
+	if (!bSuccess)
+	{
+#if WITH_EDITOR
+		Public->SkeletalMeshStatus = Public->PreUpdateSkeletalMeshStatus; // The mesh won't be updated, set the previous SkeletalMeshStatus to avoid losing error notifications
+#endif
+
+		Public->SkeletalMeshes.Reset();
+
+		InvalidateGeneratedData();
+	}
+
+	return bSuccess;
 }
 
 
@@ -1753,8 +1722,6 @@ UCustomizableInstancePrivateData::UCustomizableInstancePrivateData()
 	MinSquareDistFromComponentToPlayer = FLT_MAX;
 	LastMinSquareDistFromComponentToPlayer = FLT_MAX;
 	
-	LastUpdateTime = 0.f;
-
 	NumLODsAvailable = INT32_MAX;
 	FirstLODAvailable = 0;
 
@@ -2124,75 +2091,6 @@ int32 UCustomizableObjectInstance::FindProjectorParameterNameIndex(const FString
 }
 
 
-bool UCustomizableInstancePrivateData::MeshNeedsUpdate( UCustomizableObjectInstance* Public, const TSharedPtr<FMutableOperationData>& OperationData, bool& bOutEmptyMesh)
-{
-	mu::MeshPtrConst MutableMesh;
-	bool bMeshNeedsUpdate = HasCOInstanceFlags(PendingMeshUpdate);
-	bool bFoundEmptyMesh = false;
-	bool bFoundEmptyLOD = false;
-
-	// Check if the current meshes have the number of LODs the operation requires or they must be updated.
-	const int32 SkeletalMeshCount = Public->SkeletalMeshes.Num();
-
-	// All meshes have the same number of LODs, check the first valid SkeletalMesh
-	const int32 CurrentMinLODIndex = HasCOInstanceFlags(ReduceLODs) ? 0 : OperationData->CurrentMinLOD;
-
-	for(int32 SkeletalMeshIndex = 0; SkeletalMeshIndex < SkeletalMeshCount; ++SkeletalMeshIndex)
-	{
-		if(const USkeletalMesh* SkeletalMesh = Public->SkeletalMeshes[SkeletalMeshIndex])
-		{
-			bMeshNeedsUpdate = (SkeletalMesh->GetLODNum() - 1) != (OperationData->CurrentMaxLOD - CurrentMinLODIndex);
-			break;
-		}
-	}
-	
-	for (int32 LODIndex = OperationData->CurrentMinLOD; LODIndex <= OperationData->CurrentMaxLOD && LODIndex < OperationData->InstanceUpdateData.LODs.Num(); ++LODIndex)
-	{
-		bool bFoundNonEmptyMesh = false;
-
-		const FInstanceUpdateData::FLOD& LOD = OperationData->InstanceUpdateData.LODs[LODIndex];
-
-		bMeshNeedsUpdate = bMeshNeedsUpdate || (LOD.ComponentCount > SkeletalMeshCount);
-
-		for (int32 ComponentIndex = 0; ComponentIndex < LOD.ComponentCount; ++ComponentIndex)
-		{
-			const FInstanceUpdateData::FComponent& Component = OperationData->InstanceUpdateData.Components[LOD.FirstComponent + ComponentIndex];
-			if (Component.Mesh)
-			{
-				bMeshNeedsUpdate = bMeshNeedsUpdate || (Component.Mesh->GetId() != GetLastMeshId(Component.Id, LODIndex)) || !Public->SkeletalMeshes[ComponentIndex];
-
-				// Don't build a degenerated mesh if something went wrong
-				if (Component.Mesh->GetVertexCount() > 0)
-				{
-					bFoundNonEmptyMesh = true;
-				}
-				else
-				{
-					if (LOD.ComponentCount == 1)
-					{
-						bFoundEmptyMesh = true;
-					}
-				}
-			}
-			else
-			{
-				bMeshNeedsUpdate = bMeshNeedsUpdate || Public->SkeletalMeshes[ComponentIndex];
-			}
-		}
-
-		if (!bFoundNonEmptyMesh)
-		{
-			bFoundEmptyLOD = true;
-			UE_LOG(LogMutable, Warning, TEXT("Building instance: an LOD has no mesh geometry. This cannot be handled by Unreal."));
-		}
-	}
-
-	bOutEmptyMesh = bFoundEmptyLOD || bFoundEmptyMesh;
-
-	return bMeshNeedsUpdate;
-}
-
-
 void UCustomizableObjectInstance::SetRandomValues()
 {
 	Descriptor.SetRandomValues();
@@ -2452,7 +2350,8 @@ void UCustomizableInstancePrivateData::DiscardResourcesAndSetReferenceSkeletalMe
 	}
 
 	ClearCOInstanceFlags(Generated);
-	ClearAllLastMeshIds();
+	
+	InvalidateGeneratedData();
 	
 	Public->SkeletalMeshes.Reset();
 	Public->DescriptorRuntimeHash = FDescriptorRuntimeHash();
@@ -2518,11 +2417,11 @@ void SetTexturePropertiesFromMutableImageProps(UTexture2D* Texture, const FMutab
 }
 
 
-void UCustomizableObjectInstance::Updated(EUpdateResult Result)
+void UCustomizableObjectInstance::Updated(EUpdateResult Result, const FDescriptorRuntimeHash& InUpdatedHash)
 {
 	if (Result == EUpdateResult::Success)
 	{
-		DescriptorRuntimeHash = UpdateDescriptorRuntimeHash;
+		DescriptorRuntimeHash = InUpdatedHash;
 	}
 
 	// Call Customizable Skeletal Components updated callbacks.
@@ -2897,47 +2796,79 @@ void UCustomizableInstancePrivateData::ConvertImage(UTexture2D* Texture, mu::Ima
 }
 
 
-void SetMeshUVChannelDensity(FMeshUVChannelInfo& UVChannelInfo, float Density = 0.f)
+void UCustomizableInstancePrivateData::InitSkeletalMeshData(const TSharedPtr<FMutableOperationData>& OperationData, USkeletalMesh* SkeletalMesh, const FMutableRefSkeletalMeshData* RefSkeletalMeshData, int32 ComponentIndex)
 {
-	Density = Density > 0.f ? Density : 150.f;
+	MUTABLE_CPUPROFILER_SCOPE(UCustomizableInstancePrivateData::InitSkeletalMesh);
 
-	UVChannelInfo.bInitialized = true;
-	UVChannelInfo.bOverrideDensities = false;
+	check(SkeletalMesh);
+	check(RefSkeletalMeshData);
 
-	for (int32 i = 0; i < TEXSTREAM_MAX_NUM_UVCHANNELS; ++i)
+	// Mutable skeletal meshes are generated dynamically in-game and cannot be streamed from disk
+	SkeletalMesh->NeverStream = 1;
+
+	SkeletalMesh->SetImportedBounds(RefSkeletalMeshData->Bounds);
+	SkeletalMesh->SetPostProcessAnimBlueprint(RefSkeletalMeshData->PostProcessAnimInst.Get());
+	SkeletalMesh->SetShadowPhysicsAsset(RefSkeletalMeshData->ShadowPhysicsAsset.Get());
+
+	// Set Min LOD
+	SkeletalMesh->SetMinLod(FirstLODAvailable);
+	SkeletalMesh->SetQualityLevelMinLod(FirstLODAvailable);
+
+	SkeletalMesh->SetHasBeenSimplified(false);
+	SkeletalMesh->SetHasVertexColors(false);
+
+	// Set the default Physics Assets
+	SkeletalMesh->SetPhysicsAsset(RefSkeletalMeshData->PhysicsAsset.Get());
+	SkeletalMesh->SetEnablePerPolyCollision(RefSkeletalMeshData->Settings.bEnablePerPolyCollision);
+
+
+
+	// Allocate resources for rendering and add LOD Info
 	{
-		UVChannelInfo.LocalUVDensities[i] = Density;
-	}
-}
+		MUTABLE_CPUPROFILER_SCOPE(InitSkeletalMesh_AddLODData);
+		SkeletalMesh->AllocateResourceForRendering();
 
+		FSkeletalMeshRenderData* RenderData = SkeletalMesh->GetResourceForRendering();
+		for (int32 LODIndex = 0; LODIndex <= OperationData->CurrentMaxLOD; ++LODIndex)
+		{
+			RenderData->LODRenderData.Add(new Helper_LODDataType());
 
-void UCustomizableInstancePrivateData::BuildSkeletalMeshElementData(const TSharedPtr<FMutableOperationData>& OperationData, USkeletalMesh* SkeletalMesh, UCustomizableObjectInstance* CustomizableObjectInstance, int32 ComponentIndex)
-{
-	MUTABLE_CPUPROFILER_SCOPE(UCustomizableInstancePrivateData::BuildSkeletalMeshElementData);
+			const FMutableRefLODData& LODData = RefSkeletalMeshData->LODData[LODIndex];
 
-	// Set up unreal's default material
-	UMaterialInterface* UnrealMaterial = UMaterial::GetDefaultMaterial(MD_Surface);
-	SkeletalMesh->GetMaterials().SetNum(1);
-	SkeletalMesh->GetMaterials()[0] = UnrealMaterial;
+			FSkeletalMeshLODInfo& LODInfo = SkeletalMesh->AddLODInfo();
+			LODInfo.ScreenSize = LODData.LODInfo.ScreenSize;
+			LODInfo.LODHysteresis = LODData.LODInfo.LODHysteresis;
+			LODInfo.bSupportUniformlyDistributedSampling = LODData.LODInfo.bSupportUniformlyDistributedSampling;
+			LODInfo.bAllowCPUAccess = LODData.LODInfo.bAllowCPUAccess;
 
-	// Default density
-	SetMeshUVChannelDensity(SkeletalMesh->GetMaterials()[0].UVChannelData);
+			// Disable LOD simplification when baking instances
+			LODInfo.ReductionSettings.NumOfTrianglesPercentage = 1.f;
+			LODInfo.ReductionSettings.NumOfVertPercentage = 1.f;
 
-	int32 MeshLODIndex = HasCOInstanceFlags(ReduceLODs) ? 0 : OperationData->CurrentMinLOD;
+#if WITH_EDITORONLY_DATA
+			LODInfo.BuildSettings.bUseFullPrecisionUVs = true;
+#endif
+			LODInfo.LODMaterialMap.SetNumZeroed(1);
+		}
+	}
 
-	int32 LODCount = OperationData->InstanceUpdateData.LODs.Num();
-	for (int32 LODIndex = OperationData->CurrentMinLOD; LODIndex <= OperationData->CurrentMaxLOD && LODIndex < LODCount; ++LODIndex, ++MeshLODIndex)
+	// Set up unreal's default material, will be replaced when building materials
 	{
-		const FInstanceUpdateData::FLOD& LOD = OperationData->InstanceUpdateData.LODs[LODIndex];
-		const mu::MeshPtrConst MutableMesh = OperationData->InstanceUpdateData.Components[LOD.FirstComponent + ComponentIndex].Mesh;
-		UnrealConversionUtils::BuildSkeletalMeshElementDataAtLOD(MeshLODIndex,MutableMesh,SkeletalMesh);
+		MUTABLE_CPUPROFILER_SCOPE(InitSkeletalMesh_AddDefaultMaterial);
+		UMaterialInterface* UnrealMaterial = UMaterial::GetDefaultMaterial(MD_Surface);
+		SkeletalMesh->GetMaterials().SetNum(1);
+		SkeletalMesh->GetMaterials()[0] = UnrealMaterial;
+
+		// Default density
+		SetMeshUVChannelDensity(SkeletalMesh->GetMaterials()[0].UVChannelData);
 	}
+
 }
 
 
-bool UCustomizableInstancePrivateData::BuildSkeletalMeshSkeletonData(const TSharedPtr<FMutableOperationData>& OperationData, USkeletalMesh* SkeletalMesh, const FMutableRefSkeletalMeshData* RefSkeletalMeshData, UCustomizableObjectInstance* Public, int32 ComponentIndex)
+bool UCustomizableInstancePrivateData::BuildSkeletonData(const TSharedPtr<FMutableOperationData>& OperationData, USkeletalMesh* SkeletalMesh, const FMutableRefSkeletalMeshData* RefSkeletalMeshData, UCustomizableObjectInstance* Public, int32 ComponentIndex)
 {
-	MUTABLE_CPUPROFILER_SCOPE(BuildSkeletonData);
+	MUTABLE_CPUPROFILER_SCOPE(UCustomizableInstancePrivateData::BuildSkeletonData);
 
 	check(SkeletalMesh);
 	check(RefSkeletalMeshData);
@@ -2961,7 +2892,7 @@ bool UCustomizableInstancePrivateData::BuildSkeletalMeshSkeletonData(const TShar
 	BoneToFinalBoneIndexMap.Reserve(MutBoneCount);
 
 	{
-		MUTABLE_CPUPROFILER_SCOPE(BuildSkeletonData_CheckForMissingBones);
+		MUTABLE_CPUPROFILER_SCOPE(BuildSkeletonData_EnsureBonesExist);
 
 		// Ensure all the required bones are present in the skeleton
 		for (int32 BoneIndex = 0; BoneIndex < MutBoneCount; ++BoneIndex)
@@ -3070,126 +3001,407 @@ bool UCustomizableInstancePrivateData::BuildSkeletalMeshSkeletonData(const TShar
 }
 
 
-void UCustomizableInstancePrivateData::BuildMorphTargetsData(const TSharedPtr<FMutableOperationData>& OperationData, USkeletalMesh* SkeletalMesh, UCustomizableObjectInstance* CustomizableObjectInstance, int32 ComponentIndex)
+bool UCustomizableInstancePrivateData::CopySkeletonData(const TSharedPtr<FMutableOperationData>& OperationData, USkeletalMesh* SrcSkeletalMesh, USkeletalMesh* DestSkeletalMesh, int32 ComponentIndex)
 {
-	MUTABLE_CPUPROFILER_SCOPE(UCustomizableInstancePrivateData::BuildMorphTargetsData);
+	MUTABLE_CPUPROFILER_SCOPE(UCustomizableInstancePrivateData::CopySkeletonData);
+
+	check(SrcSkeletalMesh);
+	check(DestSkeletalMesh);
 
-	if (!(SkeletalMesh && OperationData->InstanceUpdateData.LODs.Num()))
+	const TObjectPtr<USkeleton> Skeleton = SrcSkeletalMesh->GetSkeleton();
+	if (!Skeleton)
 	{
-		return;
+		return false;
 	}
 
-	const UCustomizableObject* CustomizableObject = CustomizableObjectInstance->GetCustomizableObject();
-	const TArray<FMorphTargetInfo>& ContributingMorphTargetsInfo = CustomizableObject->ContributingMorphTargetsInfo;
-	const TArray<FMorphTargetVertexData>& MorphTargetReconstructionData = CustomizableObject->MorphTargetReconstructionData;
+	DestSkeletalMesh->SetSkeleton(Skeleton);
+	DestSkeletalMesh->SetRefSkeleton(Skeleton->GetReferenceSkeleton());
+	DestSkeletalMesh->SetRefBasesInvMatrix(SrcSkeletalMesh->GetRefBasesInvMatrix());
+
+	const FReferenceSkeleton& RefSkeleton = DestSkeletalMesh->GetRefSkeleton();
+
+	// Check that the bones we need are present in the current Skeleton
+	FInstanceUpdateData::FSkeletonData& MutSkeletonData = OperationData->InstanceUpdateData.Skeletons[ComponentIndex];
+	const int32 MutBoneCount = MutSkeletonData.BoneNames.Num();
+	
+	TMap<FName, uint16> BoneToFinalBoneIndexMap;
+	BoneToFinalBoneIndexMap.Reserve(MutBoneCount);
 
-	if (!(MorphTargetReconstructionData.Num() && ContributingMorphTargetsInfo.Num()))
 	{
-		return;
-	}
+		MUTABLE_CPUPROFILER_SCOPE(CopySkeletonData_CheckForMissingBones);
 
-	TArray<int32> SectionMorphTargetVertices;
-	SectionMorphTargetVertices.SetNumZeroed(ContributingMorphTargetsInfo.Num());
+		// Ensure all the required bones are present in the skeleton
+		for (int32 BoneIndex = 0; BoneIndex < MutBoneCount; ++BoneIndex)
+		{
+			FName BoneName = MutSkeletonData.BoneNames[BoneIndex];
+			check(BoneName != NAME_None);
+
+			int32 SourceBoneIndex = RefSkeleton.FindRawBoneIndex(BoneName);
+			
+			if (SourceBoneIndex == INDEX_NONE)
+			{
+				if (MutSkeletonData.BoneMatricesWithScale.Find(BoneName))
+				{
+					// Merged skeleton is missing some bones! This happens if one of the skeletons involved in the merge is discarded due to being incompatible with the rest
+					// or if the source mesh is not in sync with the skeleton. 
+					UE_LOG(LogMutable, Warning, TEXT("Building instance: generated mesh has a bone [%s] not present in the reference mesh [%s]. Failing to generate mesh. "),
+						*BoneName.ToString(), *DestSkeletalMesh->GetName());
+					return false;
+				}
+
+				// The bone is missing but it's not required. Fix the SourceBoneIndex and continue.
+				SourceBoneIndex = 0;
+			}
+
+			BoneToFinalBoneIndexMap.Add(BoneName, SourceBoneIndex);
+		}
+	}
 
-	SkeletalMesh->GetMorphTargets().Empty(ContributingMorphTargetsInfo.Num());
-	for (const FMorphTargetInfo& MorphTargetInfo : ContributingMorphTargetsInfo)
 	{
-		UMorphTarget* NewMorphTarget = NewObject<UMorphTarget>(SkeletalMesh, MorphTargetInfo.Name);
+		MUTABLE_CPUPROFILER_SCOPE(CopySkeletonData_FixBoneIndices);
 
-		NewMorphTarget->BaseSkelMesh = SkeletalMesh;
-		NewMorphTarget->GetMorphLODModels().SetNum(MorphTargetInfo.LodNum);
-		SkeletalMesh->GetMorphTargets().Add(NewMorphTarget);
+		// Fix up BoneMaps and ActiveBones indices
+		for (FInstanceUpdateData::FComponent& Component : OperationData->InstanceUpdateData.Components)
+		{
+			if (Component.Id != ComponentIndex)
+			{
+				continue;
+			}
+
+			for (uint16& BoneIndex : Component.BoneMap)
+			{
+				BoneIndex = BoneToFinalBoneIndexMap[MutSkeletonData.BoneNames[BoneIndex]];
+			}
+
+			TArray<uint16> UniqueActiveBones;
+			for (const uint16 BoneIndex : Component.ActiveBones)
+			{
+				UniqueActiveBones.AddUnique(BoneToFinalBoneIndexMap[MutSkeletonData.BoneNames[BoneIndex]]);
+			}
+
+			Exchange(Component.ActiveBones, UniqueActiveBones);
+			Component.ActiveBones.Sort();
+		}
 	}
 
-	const int32 SkeletalMeshMorphTargetsNum = SkeletalMesh->GetMorphTargets().Num();
+	return true;
+}
+
+
+void UCustomizableInstancePrivateData::BuildMeshSockets(const TSharedPtr<FMutableOperationData>& OperationData, USkeletalMesh* SkeletalMesh, const FMutableRefSkeletalMeshData* RefSkeletalMeshData, UCustomizableObjectInstance* Public, mu::MeshPtrConst MutableMesh)
+{
+	// Build mesh sockets.
+	MUTABLE_CPUPROFILER_SCOPE(UCustomizableInstancePrivateData::BuildMeshSockets);
+
+	check(SkeletalMesh);
+	check(RefSkeletalMeshData);
+	check(Public);
 
-	int32 LODCount = OperationData->InstanceUpdateData.LODs.Num();
-	for (int32 LODIndex = OperationData->CurrentMinLOD; LODIndex <= OperationData->CurrentMaxLOD && LODIndex < LODCount; ++LODIndex)
+	const uint32 SocketCount = RefSkeletalMeshData->Sockets.Num();
+
+	TArray<USkeletalMeshSocket*>& Sockets = SkeletalMesh->GetMeshOnlySocketList();
+	Sockets.Empty(SocketCount);
+	TMap<FName, TTuple<int32, int32>> SocketMap; // Maps Socket name to Sockets Array index and priority
+	
+	// Add sockets used by the SkeletalMesh of reference.
 	{
-		const FInstanceUpdateData::FLOD& LOD = OperationData->InstanceUpdateData.LODs[LODIndex];
-		if (mu::MeshPtrConst MutableMesh = OperationData->InstanceUpdateData.Components[LOD.FirstComponent+ComponentIndex].Mesh)
+		MUTABLE_CPUPROFILER_SCOPE(BuildMeshSockets_RefMeshSockets);
+	
+		for (uint32 SocketIndex = 0; SocketIndex < SocketCount; ++SocketIndex)
 		{
-			const mu::FMeshBufferSet& MeshSet = MutableMesh->GetVertexBuffers();
+			const FMutableRefSocket& RefSocket = RefSkeletalMeshData->Sockets[SocketIndex];
 
-			int32 VertexMorphsInfoIndexBufferIndex, VertexMorphsInfoIndexBufferChannel;
-			MeshSet.FindChannel(mu::MBS_OTHER, 0, &VertexMorphsInfoIndexBufferIndex, &VertexMorphsInfoIndexBufferChannel);
+			USkeletalMeshSocket* Socket = NewObject<USkeletalMeshSocket>(SkeletalMesh);
 
-			int32 VertexMorphsCountBufferIndex, VertexMorphsCountBufferChannel;
-			MeshSet.FindChannel(mu::MBS_OTHER, 1, &VertexMorphsCountBufferIndex, &VertexMorphsCountBufferChannel);
+			Socket->SocketName = RefSocket.SocketName;
+			Socket->BoneName = RefSocket.BoneName;
 
-			if (VertexMorphsInfoIndexBufferIndex < 0 || VertexMorphsCountBufferIndex < 0)
-			{
-				continue;
-			}
+			Socket->RelativeLocation = RefSocket.RelativeLocation;
+			Socket->RelativeRotation = RefSocket.RelativeRotation;
+			Socket->RelativeScale = RefSocket.RelativeScale;
 
-			const int32* const VertexMorphsInfoIndexBuffer = reinterpret_cast<const int32*>(MeshSet.GetBufferData(VertexMorphsInfoIndexBufferIndex));
-			TArrayView<const int32> VertexMorphsInfoIndexView(VertexMorphsInfoIndexBuffer, MeshSet.GetElementCount());
+			Socket->bForceAlwaysAnimated = RefSocket.bForceAlwaysAnimated;
+			const int32 LastIndex = Sockets.Add(Socket);
 
-			const int32* const VertexMorphsCountBuffer = reinterpret_cast<const int32*>(MeshSet.GetBufferData(VertexMorphsCountBufferIndex));
-			TArrayView<const int32> VertexMorphsCountView(VertexMorphsCountBuffer, MeshSet.GetElementCount());
+			SocketMap.Add(Socket->SocketName, TTuple<int32, int32>(LastIndex, RefSocket.Priority));
+		}
+	}
 
-			const int32 SurfaceCount = MutableMesh->GetSurfaceCount();
-			for (int32 Section = 0; Section < SurfaceCount; ++Section)
-			{
-				// Reset SectionMorphTargets.
-				for (auto& Elem : SectionMorphTargetVertices)
-				{
-					Elem = 0;
-				}
+	// Add or update sockets modified by Mutable.
+	if (MutableMesh)
+	{
+		MUTABLE_CPUPROFILER_SCOPE(BuildMeshSockets_MutableSockets);
 
-				int FirstVertex, VerticesCount, FirstIndex, IndiciesCount;
-				MutableMesh->GetSurface(Section, &FirstVertex, &VerticesCount, &FirstIndex, &IndiciesCount);
+		const UCustomizableObject* CustomizableObject = Public->GetCustomizableObject();
+		check(CustomizableObject);
 
-				for (int32 VertexIdx = FirstVertex; VertexIdx < FirstVertex + VerticesCount; ++VertexIdx)
+		for (int32 TagIndex = 0; TagIndex < MutableMesh->GetTagCount(); ++TagIndex)
+		{
+			FString Tag = MutableMesh->GetTag(TagIndex);
+
+			if (Tag.RemoveFromStart("__Socket:"))
+			{
+				check(Tag.IsNumeric());
+				const int32 MutableSocketIndex = FCString::Atoi(*Tag);
+
+				if (CustomizableObject->SocketArray.IsValidIndex(MutableSocketIndex))
 				{
-					const int32 MorphCount = VertexMorphsCountView[VertexIdx];
-					if (MorphCount <= 0)
-					{
-						continue;
-					}
+					const FMutableRefSocket& MutableSocket = CustomizableObject->SocketArray[MutableSocketIndex];
+					int32 IndexToWriteSocket = -1;
 
-					TArrayView<const FMorphTargetVertexData> MorphsVertexDataView(&(MorphTargetReconstructionData[VertexMorphsInfoIndexView[VertexIdx]]), MorphCount);
-					for (const FMorphTargetVertexData& SourceVertex : MorphsVertexDataView)
+					if (TTuple<int32, int32>* FoundSocket = SocketMap.Find(MutableSocket.SocketName))
 					{
-						// check(SkeletalMeshMorphTargetsNum <= SourceVertex.MorphIndex);
-						if (SkeletalMeshMorphTargetsNum <= SourceVertex.MorphIndex)
+						if (FoundSocket->Value < MutableSocket.Priority)
 						{
-							continue;
+							// Overwrite the existing socket because the new mesh part one is higher priority
+							IndexToWriteSocket = FoundSocket->Key;
+							FoundSocket->Value = MutableSocket.Priority;
 						}
+					}
+					else
+					{
+						// New Socket
+						USkeletalMeshSocket* Socket = NewObject<USkeletalMeshSocket>(SkeletalMesh);
+						IndexToWriteSocket = Sockets.Add(Socket);
+						SocketMap.Add(MutableSocket.SocketName, TTuple<int32, int32>(IndexToWriteSocket, MutableSocket.Priority));
+					}
 
-						FMorphTargetLODModel& DestMorphLODModel = SkeletalMesh->GetMorphTargets()[SourceVertex.MorphIndex]->GetMorphLODModels()[LODIndex];
-
-						DestMorphLODModel.Vertices.Emplace(
-							FMorphTargetDelta{ SourceVertex.PositionDelta, SourceVertex.TangentZDelta, static_cast<uint32>(VertexIdx) });
+					if (IndexToWriteSocket >= 0)
+					{
+						check(Sockets.IsValidIndex(IndexToWriteSocket));
 
-						++SectionMorphTargetVertices[SourceVertex.MorphIndex];
+						USkeletalMeshSocket* SocketToWrite = Sockets[IndexToWriteSocket];
 
-					}
-				}
+						SocketToWrite->SocketName = MutableSocket.SocketName;
+						SocketToWrite->BoneName = MutableSocket.BoneName;
 
-				const int32 SectionMorphTargetsNum = SectionMorphTargetVertices.Num();
-				for (int32 MorphIdx = 0; MorphIdx < SectionMorphTargetsNum; ++MorphIdx)
-				{
-					if (SectionMorphTargetVertices[MorphIdx] > 0)
-					{
-						FMorphTargetLODModel& MorphTargetLodModel = SkeletalMesh->GetMorphTargets()[MorphIdx]->GetMorphLODModels()[LODIndex];
+						SocketToWrite->RelativeLocation = MutableSocket.RelativeLocation;
+						SocketToWrite->RelativeRotation = MutableSocket.RelativeRotation;
+						SocketToWrite->RelativeScale = MutableSocket.RelativeScale;
 
-						MorphTargetLodModel.SectionIndices.Add(Section);
-						MorphTargetLodModel.NumVertices = SectionMorphTargetVertices[MorphIdx];
+						SocketToWrite->bForceAlwaysAnimated = MutableSocket.bForceAlwaysAnimated;
 					}
 				}
 			}
 		}
 	}
-
-	SkeletalMesh->InitMorphTargets();
 }
 
-namespace 
+
+void UCustomizableInstancePrivateData::CopyMeshSockets(USkeletalMesh* SrcSkeletalMesh, USkeletalMesh* DestSkeletalMesh)
 {
-	// Only used to be able to create new clothing assets and assign a new guid to them without the factory.
-	class UCustomizableObjectClothingAsset : public UClothingAssetCommon
-	{
+	// Copy mesh sockets.
+	MUTABLE_CPUPROFILER_SCOPE(UCustomizableInstancePrivateData::CopyMeshSockets);
+
+	check(SrcSkeletalMesh);
+	check(DestSkeletalMesh);
+
+	const TArray<USkeletalMeshSocket*>& SrcSockets = SrcSkeletalMesh->GetMeshOnlySocketList();
+	TArray<USkeletalMeshSocket*>& DestSockets = DestSkeletalMesh->GetMeshOnlySocketList();
+	DestSockets.Empty(SrcSockets.Num());
+
+	for (const USkeletalMeshSocket* SrcSocket : SrcSockets)
+	{
+		if (SrcSocket)
+		{
+			USkeletalMeshSocket* NewSocket = NewObject<USkeletalMeshSocket>(DestSkeletalMesh);
+
+
+			NewSocket->SocketName = SrcSocket->SocketName;
+			NewSocket->BoneName = SrcSocket->BoneName;
+
+			NewSocket->RelativeLocation = SrcSocket->RelativeLocation;
+			NewSocket->RelativeRotation = SrcSocket->RelativeRotation;
+			NewSocket->RelativeScale = SrcSocket->RelativeScale;
+
+			NewSocket->bForceAlwaysAnimated = SrcSocket->bForceAlwaysAnimated;
+			DestSockets.Add(NewSocket);
+
+		}
+	}
+}
+
+
+void UCustomizableInstancePrivateData::BuildOrCopyElementData(const TSharedPtr<FMutableOperationData>& OperationData, USkeletalMesh* SkeletalMesh, UCustomizableObjectInstance* CustomizableObjectInstance, int32 ComponentIndex)
+{
+	MUTABLE_CPUPROFILER_SCOPE(UCustomizableInstancePrivateData::BuildOrCopyElementData);
+
+	for (int32 LODIndex = OperationData->CurrentMaxLOD; LODIndex >= FirstLODAvailable; --LODIndex)
+	{
+		const FInstanceUpdateData::FLOD& LOD = OperationData->InstanceUpdateData.LODs[LODIndex];
+
+		if (!LOD.ComponentCount)
+		{
+			continue;
+		}
+
+		const FInstanceUpdateData::FComponent& Component = OperationData->InstanceUpdateData.Components[LOD.FirstComponent + ComponentIndex];
+
+		if (!Component.bGenerated || Component.bReuseMesh)
+		{
+			continue;
+		}
+
+		for (int32 SurfaceIndex = 0; SurfaceIndex < Component.SurfaceCount; ++SurfaceIndex)
+		{
+			new(Helper_GetLODRenderSections(SkeletalMesh, LODIndex)) Helper_SkelMeshRenderSection();
+		}
+	}
+}
+
+
+void UCustomizableInstancePrivateData::BuildOrCopyMorphTargetsData(const TSharedPtr<FMutableOperationData>& OperationData, USkeletalMesh* SkeletalMesh, const USkeletalMesh* LastUpdateSkeletalMesh, UCustomizableObjectInstance* CustomizableObjectInstance, int32 ComponentIndex)
+{
+	MUTABLE_CPUPROFILER_SCOPE(UCustomizableInstancePrivateData::BuildOrCopyMorphTargetsData);
+
+	if (!SkeletalMesh || !CustomizableObjectInstance)
+	{
+		return;
+	}
+
+	const UCustomizableObject* CustomizableObject = CustomizableObjectInstance->GetCustomizableObject();
+	const TArray<FMorphTargetInfo>& ContributingMorphTargetsInfo = CustomizableObject->ContributingMorphTargetsInfo;
+	const TArray<FMorphTargetVertexData>& MorphTargetReconstructionData = CustomizableObject->MorphTargetReconstructionData;
+
+	if (!(MorphTargetReconstructionData.Num() && ContributingMorphTargetsInfo.Num()))
+	{
+		return;
+	}
+
+	TArray<int32> SectionMorphTargetVertices;
+	SectionMorphTargetVertices.SetNumZeroed(ContributingMorphTargetsInfo.Num());
+
+	SkeletalMesh->GetMorphTargets().Empty(ContributingMorphTargetsInfo.Num());
+	for (const FMorphTargetInfo& MorphTargetInfo : ContributingMorphTargetsInfo)
+	{
+		UMorphTarget* NewMorphTarget = NewObject<UMorphTarget>(SkeletalMesh, MorphTargetInfo.Name);
+
+		NewMorphTarget->BaseSkelMesh = SkeletalMesh;
+		NewMorphTarget->GetMorphLODModels().SetNum(MorphTargetInfo.LodNum);
+		SkeletalMesh->GetMorphTargets().Add(NewMorphTarget);
+	}
+
+	const int32 SkeletalMeshMorphTargetsNum = SkeletalMesh->GetMorphTargets().Num();
+
+	int32 LastValidLODIndex = OperationData->CurrentMaxLOD;
+	for (int32 LODIndex = OperationData->CurrentMaxLOD; LODIndex >= FirstLODAvailable; --LODIndex)
+	{
+		const FInstanceUpdateData::FLOD& LOD = OperationData->InstanceUpdateData.LODs[LODIndex];
+
+		if (LODIndex >= OperationData->CurrentMinLOD && OperationData->InstanceUpdateData.Components[OperationData->InstanceUpdateData.LODs[LODIndex].FirstComponent + ComponentIndex].bGenerated)
+		{
+			LastValidLODIndex = LODIndex;
+		}
+
+		const FInstanceUpdateData::FComponent& Component = OperationData->InstanceUpdateData.Components[OperationData->InstanceUpdateData.LODs[LastValidLODIndex].FirstComponent + ComponentIndex];
+
+		if (Component.Mesh)
+		{
+			MUTABLE_CPUPROFILER_SCOPE(BuildOrCopyMorphTargetsData_Build);
+
+			const mu::FMeshBufferSet& MeshSet = Component.Mesh->GetVertexBuffers();
+
+			int32 VertexMorphsInfoIndexBufferIndex, VertexMorphsInfoIndexBufferChannel;
+			MeshSet.FindChannel(mu::MBS_OTHER, 0, &VertexMorphsInfoIndexBufferIndex, &VertexMorphsInfoIndexBufferChannel);
+
+			int32 VertexMorphsCountBufferIndex, VertexMorphsCountBufferChannel;
+			MeshSet.FindChannel(mu::MBS_OTHER, 1, &VertexMorphsCountBufferIndex, &VertexMorphsCountBufferChannel);
+
+			if (VertexMorphsInfoIndexBufferIndex < 0 || VertexMorphsCountBufferIndex < 0)
+			{
+				continue;
+			}
+
+			const int32* const VertexMorphsInfoIndexBuffer = reinterpret_cast<const int32*>(MeshSet.GetBufferData(VertexMorphsInfoIndexBufferIndex));
+			TArrayView<const int32> VertexMorphsInfoIndexView(VertexMorphsInfoIndexBuffer, MeshSet.GetElementCount());
+
+			const int32* const VertexMorphsCountBuffer = reinterpret_cast<const int32*>(MeshSet.GetBufferData(VertexMorphsCountBufferIndex));
+			TArrayView<const int32> VertexMorphsCountView(VertexMorphsCountBuffer, MeshSet.GetElementCount());
+
+			const int32 SurfaceCount = Component.Mesh->GetSurfaceCount();
+			for (int32 Section = 0; Section < SurfaceCount; ++Section)
+			{
+				// Reset SectionMorphTargets.
+				for (auto& Elem : SectionMorphTargetVertices)
+				{
+					Elem = 0;
+				}
+
+				int FirstVertex, VerticesCount, FirstIndex, IndiciesCount;
+				Component.Mesh->GetSurface(Section, &FirstVertex, &VerticesCount, &FirstIndex, &IndiciesCount);
+
+				for (int32 VertexIdx = FirstVertex; VertexIdx < FirstVertex + VerticesCount; ++VertexIdx)
+				{
+					const int32 MorphCount = VertexMorphsCountView[VertexIdx];
+					if (MorphCount <= 0)
+					{
+						continue;
+					}
+
+					TArrayView<const FMorphTargetVertexData> MorphsVertexDataView(&(MorphTargetReconstructionData[VertexMorphsInfoIndexView[VertexIdx]]), MorphCount);
+					for (const FMorphTargetVertexData& SourceVertex : MorphsVertexDataView)
+					{
+						// check(SkeletalMeshMorphTargetsNum <= SourceVertex.MorphIndex);
+						if (SkeletalMeshMorphTargetsNum <= SourceVertex.MorphIndex)
+						{
+							continue;
+						}
+
+						FMorphTargetLODModel& DestMorphLODModel = SkeletalMesh->GetMorphTargets()[SourceVertex.MorphIndex]->GetMorphLODModels()[LODIndex];
+
+						DestMorphLODModel.Vertices.Emplace(
+							FMorphTargetDelta{ SourceVertex.PositionDelta, SourceVertex.TangentZDelta, static_cast<uint32>(VertexIdx) });
+
+						++SectionMorphTargetVertices[SourceVertex.MorphIndex];
+
+					}
+				}
+
+				const int32 SectionMorphTargetsNum = SectionMorphTargetVertices.Num();
+				for (int32 MorphIdx = 0; MorphIdx < SectionMorphTargetsNum; ++MorphIdx)
+				{
+					if (SectionMorphTargetVertices[MorphIdx] > 0)
+					{
+						FMorphTargetLODModel& MorphTargetLodModel = SkeletalMesh->GetMorphTargets()[MorphIdx]->GetMorphLODModels()[LODIndex];
+
+						MorphTargetLodModel.SectionIndices.Add(Section);
+						MorphTargetLodModel.NumVertices = SectionMorphTargetVertices[MorphIdx];
+					}
+				}
+			}
+		}
+		else
+		{
+			MUTABLE_CPUPROFILER_SCOPE(BuildOrCopyMorphTargetsData_Copy);
+
+			// Copy data from the last valid LOD generated from either the current mesh or the previous mesh
+			const FInstanceUpdateData::FLOD& LastLOD = OperationData->InstanceUpdateData.LODs[LastValidLODIndex];
+			const FInstanceUpdateData::FComponent& LastComponent = OperationData->InstanceUpdateData.Components[LastLOD.FirstComponent + ComponentIndex];
+
+			const USkeletalMesh* SrcSkeletalMesh = LastComponent.bReuseMesh ? LastUpdateSkeletalMesh : SkeletalMesh;
+			if (!SrcSkeletalMesh)
+			{
+				check(false);
+				return;
+			}
+
+			const int32 NumMorphTargets = SrcSkeletalMesh->GetMorphTargets().Num();
+			for (int32 MorphTargetIndex = 0; MorphTargetIndex < NumMorphTargets; ++MorphTargetIndex)
+			{
+				SkeletalMesh->GetMorphTargets()[MorphTargetIndex]->GetMorphLODModels()[LODIndex] = SrcSkeletalMesh->GetMorphTargets()[MorphTargetIndex]->GetMorphLODModels()[LastValidLODIndex];
+			}
+		}
+	}
+
+	SkeletalMesh->InitMorphTargets();
+
+}
+
+namespace 
+{
+	// Only used to be able to create new clothing assets and assign a new guid to them without the factory.
+	class UCustomizableObjectClothingAsset : public UClothingAssetCommon
+	{
 	public:
 		void AssignNewGuid()
 		{
@@ -3199,11 +3411,12 @@ namespace
 
 }
 
-void UCustomizableInstancePrivateData::BuildClothingData(const TSharedPtr<FMutableOperationData>&OperationData, USkeletalMesh * SkeletalMesh, UCustomizableObjectInstance * CustomizableObjectInstance, int32 ComponentIndex)
+
+void UCustomizableInstancePrivateData::BuildOrCopyClothingData(const TSharedPtr<FMutableOperationData>&OperationData, USkeletalMesh * SkeletalMesh, const USkeletalMesh* LastUpdateSkeletalMesh, UCustomizableObjectInstance * CustomizableObjectInstance, int32 ComponentIndex)
 {
-	MUTABLE_CPUPROFILER_SCOPE(UCustomizableInstancePrivateData::BuildClothingData);
+	MUTABLE_CPUPROFILER_SCOPE(UCustomizableInstancePrivateData::BuildOrCopyClothingData);
 
-	if (!(SkeletalMesh && OperationData->InstanceUpdateData.LODs.Num()))
+	if (!SkeletalMesh)
 	{
 		return;
 	}
@@ -3236,23 +3449,35 @@ void UCustomizableInstancePrivateData::BuildClothingData(const TSharedPtr<FMutab
 	TArray<FSectionWithClothData> SectionsWithCloth;
 	SectionsWithCloth.Reserve(32);
 
-	int32 MeshNumLods = 0;
+	const int32 LODCount = OperationData->InstanceUpdateData.LODs.Num();
 
-	int32 LODCount = OperationData->InstanceUpdateData.LODs.Num();
 	{
 		MUTABLE_CPUPROFILER_SCOPE(DiscoverSectionsWithCloth);
 
-		int32 MeshLODIndex = HasCOInstanceFlags(ReduceLODs) ? 0 : OperationData->CurrentMinLOD;
-		for (int32 LODIndex = OperationData->CurrentMinLOD; LODIndex <= OperationData->CurrentMaxLOD && LODIndex < LODCount; ++LODIndex, ++MeshLODIndex)
+		for (int32 LODIndex = OperationData->CurrentMinLOD; LODIndex <= OperationData->CurrentMaxLOD; ++LODIndex)
 		{
 			const FInstanceUpdateData::FLOD& LOD = OperationData->InstanceUpdateData.LODs[LODIndex];
-			const int32 ComponentCount = LOD.ComponentCount;
-			if (!ComponentCount)
+			if (ComponentIndex >= LOD.ComponentCount)
 			{
-				return;
+				continue;
 			}
 
-			if (mu::MeshPtrConst MutableMesh = OperationData->InstanceUpdateData.Components[LOD.FirstComponent+ComponentIndex].Mesh)
+			FInstanceUpdateData::FComponent& Component = OperationData->InstanceUpdateData.Components[LOD.FirstComponent + ComponentIndex];
+
+			if (!Component.bGenerated)
+			{
+				continue;
+			}
+
+
+			if (!Component.Mesh && Component.bReuseMesh)
+			{
+				// TODO PERE: Implement copying LODRenderDatas with clothing
+				unimplemented()
+			}
+
+
+			if (mu::MeshPtrConst MutableMesh = Component.Mesh)
 			{
 				const mu::FMeshBufferSet& MeshSet = MutableMesh->GetVertexBuffers();
 
@@ -3339,12 +3564,13 @@ void UCustomizableInstancePrivateData::BuildClothingData(const TSharedPtr<FMutab
 					}
 
 					SectionsWithCloth.Add
-							(FSectionWithClothData{ ClothAssetIndex, ClothAssetLodIndex, Section, MeshLODIndex, FirstVertex, IndicesView16Bits, IndicesView32Bits, ClothingDataView, TArray<FMeshToMeshVertData>() });
+							(FSectionWithClothData{ ClothAssetIndex, ClothAssetLodIndex, Section, LODIndex, FirstVertex, IndicesView16Bits, IndicesView32Bits, ClothingDataView, TArray<FMeshToMeshVertData>() });
+
+					// TODO Pere: Remove once we support the copy of LODRenderDatas with clothing
+					LastUpdateData.Components[LOD.FirstComponent + ComponentIndex].bGenerated = false;
 				}
 			}
 		}
-
-		MeshNumLods = MeshLODIndex;
 	}
 
 	if (!SectionsWithCloth.Num())
@@ -3365,7 +3591,7 @@ void UCustomizableInstancePrivateData::BuildClothingData(const TSharedPtr<FMutab
 			// Only initilialize once, multiple sections with cloth could point to the same cloth asset.
 			if (!DstAssetData.LodMap.Num())
 			{
-				DstAssetData.LodMap.Init( -1, MeshNumLods );
+				DstAssetData.LodMap.Init( INDEX_NONE, OperationData->CurrentMaxLOD + 1 );
 
 				DstAssetData.LodData.SetNum( SrcAssetData.LodData.Num() );
 				DstAssetData.UsedBoneNames = SrcAssetData.UsedBoneNames;
@@ -3902,21 +4128,19 @@ void UCustomizableInstancePrivateData::BuildClothingData(const TSharedPtr<FMutab
 
 	// Based on FSkeletalMeshLODModel::GetClothMappingData().
 	TArray<TArray<FMeshToMeshVertData>> LodMappingData;
-	LodMappingData.SetNum(FMath::Max<int32>(OperationData->CurrentMaxLOD + 1, LODCount));
+	LodMappingData.SetNum(LODCount);
 
 	TArray<TArray<FClothBufferIndexMapping>> LodClothingIndexMapping;
-	LodClothingIndexMapping.SetNum(FMath::Max<int32>(OperationData->CurrentMaxLOD + 1, LODCount));
+	LodClothingIndexMapping.SetNum(LODCount);
 	{
 		int32 NumSectionsWithClothProcessed = 0;
-		int32 MeshLODIndex = HasCOInstanceFlags(ReduceLODs) ? 0 : OperationData->CurrentMinLOD;
 
-		for (int32 LODIndex = OperationData->CurrentMinLOD; LODIndex <= OperationData->CurrentMaxLOD && LODIndex < LODCount; ++LODIndex, ++MeshLODIndex)
+		for (int32 LODIndex = OperationData->CurrentMinLOD; LODIndex <= OperationData->CurrentMaxLOD; ++LODIndex)
 		{
 			const FInstanceUpdateData::FLOD& LOD = OperationData->InstanceUpdateData.LODs[LODIndex];
-			const int32 ComponentCount = LOD.ComponentCount;
-			if (!ComponentCount)
+			if (ComponentIndex >= LOD.ComponentCount)
 			{
-				return;
+				continue;
 			}
 
 			if (mu::MeshPtrConst MutableMesh = OperationData->InstanceUpdateData.Components[LOD.FirstComponent+ComponentIndex].Mesh)
@@ -3964,10 +4188,9 @@ void UCustomizableInstancePrivateData::BuildClothingData(const TSharedPtr<FMutab
 	{
 		MUTABLE_CPUPROFILER_SCOPE(InitClothRenderData)
 		// Based on FSkeletalMeshLODModel::GetClothMappingData().
-		int32 MeshLODIndex = HasCOInstanceFlags(ReduceLODs) ? 0 : OperationData->CurrentMinLOD;
-		for (int32 LODIndex = OperationData->CurrentMinLOD; LODIndex <= OperationData->CurrentMaxLOD && LODIndex < LODCount; ++LODIndex, ++MeshLODIndex)
+		for (int32 LODIndex = OperationData->CurrentMinLOD; LODIndex <= OperationData->CurrentMaxLOD; ++LODIndex)
 		{
-			FSkeletalMeshLODRenderData& LODModel = RenderResource->LODRenderData[MeshLODIndex];
+			FSkeletalMeshLODRenderData& LODModel = RenderResource->LODRenderData[LODIndex];
 	
 			if (LodMappingData[LODIndex].Num() > 0)
 			{
@@ -4074,36 +4297,76 @@ void UCustomizableInstancePrivateData::BuildClothingData(const TSharedPtr<FMutab
 	SkeletalMesh->SetHasActiveClothingAssets( static_cast<bool>( SectionsWithCloth.Num() ) );
 }
 
-bool UCustomizableInstancePrivateData::BuildSkeletalMeshRenderData(const TSharedPtr<FMutableOperationData>& OperationData, USkeletalMesh* SkeletalMesh, UCustomizableObjectInstance* Public, int32 ComponentIndex)
-{
-	MUTABLE_CPUPROFILER_SCOPE(UCustomizableInstancePrivateData::BuildSkeletalMeshRenderData)
 
-	UCustomizableObject* CustomizableObject = Public->GetCustomizableObject();
-
-	SkeletalMesh->SetHasBeenSimplified(false);
-	SkeletalMesh->SetHasVertexColors(false);
+bool UCustomizableInstancePrivateData::BuildOrCopyRenderData(const TSharedPtr<FMutableOperationData>& OperationData, USkeletalMesh* SkeletalMesh, const USkeletalMesh* LastUpdateSkeletalMesh, UCustomizableObjectInstance* Public, int32 ComponentIndex)
+{
+	MUTABLE_CPUPROFILER_SCOPE(UCustomizableInstancePrivateData::BuildOrCopyRenderData)
 
-	int32 MeshLODIndex = HasCOInstanceFlags(ReduceLODs) ? 0 : OperationData->CurrentMinLOD;
+		UCustomizableObject* CustomizableObject = Public->GetCustomizableObject();
+	check(CustomizableObject);
 
-	int32 LODCount = OperationData->InstanceUpdateData.LODs.Num();
-	for (int32 LODIndex = OperationData->CurrentMinLOD; LODIndex <= OperationData->CurrentMaxLOD && LODIndex < LODCount; ++LODIndex, ++MeshLODIndex)
+	int32 LastValidLODIndex = OperationData->CurrentMaxLOD;
+	for (int32 LODIndex = OperationData->CurrentMaxLOD; LODIndex >= FirstLODAvailable; --LODIndex)
 	{
-		MUTABLE_CPUPROFILER_SCOPE(LODLoop);
+		MUTABLE_CPUPROFILER_SCOPE(BuildOrCopyRenderData_LODLoop);
 
 		const FInstanceUpdateData::FLOD& LOD = OperationData->InstanceUpdateData.LODs[LODIndex];
 
-		// There could be components without a mesh in LODs
+		bool bCopyComponent = false;
 		if (!LOD.ComponentCount)
 		{
-			UE_LOG(LogMutable, Warning, TEXT("Building instance: generated mesh [%s] has LOD [%d] with no component.")
-				, *SkeletalMesh->GetName()
-				, LODIndex);
+			// Interrupt the generation if the LOD is empty and it should have been generated.
+			if (LODIndex >= OperationData->CurrentMinLOD && LODIndex <= OperationData->CurrentMaxLOD)
+			{
+				UE_LOG(LogMutable, Warning, TEXT("Building instance: generated mesh [%s] has LOD [%d] with no component.")
+					, *SkeletalMesh->GetName()
+					, LODIndex);
 
-			// End with failure
-			return false;
+				// End with failure
+				return false;
+			}
+
+			// LODIndex is below the CurrentMinLOD copy the data from the LastValidLODIndex
+			bCopyComponent = true;
 		}
 
-		FInstanceUpdateData::FComponent& Component = OperationData->InstanceUpdateData.Components[LOD.FirstComponent+ComponentIndex];
+		else
+		{
+			check(ComponentIndex < LOD.ComponentCount);
+			FInstanceUpdateData::FComponent& Component = OperationData->InstanceUpdateData.Components[LOD.FirstComponent + ComponentIndex];
+
+			if (!Component.bGenerated || Component.bReuseMesh)
+			{
+				bCopyComponent = true;
+			}
+			else
+			{
+				LastValidLODIndex = LODIndex;
+			}
+		}
+
+
+		if (bCopyComponent)
+		{
+			// Copy data from the last valid LOD generated from either the current mesh or the previous mesh
+			const FInstanceUpdateData::FLOD& LastValidLOD = OperationData->InstanceUpdateData.LODs[LastValidLODIndex];
+			const FInstanceUpdateData::FComponent& LastValidComponent = OperationData->InstanceUpdateData.Components[LastValidLOD.FirstComponent + ComponentIndex];
+
+			const USkeletalMesh* SrcSkeletalMesh = LastValidComponent.bReuseMesh ? LastUpdateSkeletalMesh : SkeletalMesh;
+			if (!SrcSkeletalMesh)
+			{
+				check(false);
+				return false;
+			}
+
+			TRACE_CPUPROFILER_EVENT_SCOPE_TEXT(*FString::Printf(TEXT("BuildOrCopyRenderData_CopyData: From LOD %d to LOD %d"), LastValidLODIndex, LODIndex));
+
+			// Render Data will be reused from the previously generated component
+			UnrealConversionUtils::CopySkeletalMeshLODRenderData(SrcSkeletalMesh, SkeletalMesh, LastValidLODIndex, LODIndex);
+			continue;
+		}
+
+		FInstanceUpdateData::FComponent& Component = OperationData->InstanceUpdateData.Components[OperationData->InstanceUpdateData.LODs[LastValidLODIndex].FirstComponent + ComponentIndex];
 
 		// There could be components without a mesh in LODs
 		if (!Component.Mesh)
@@ -4117,20 +4380,22 @@ bool UCustomizableInstancePrivateData::BuildSkeletalMeshRenderData(const TShared
 			return false;
 		}
 
+		TRACE_CPUPROFILER_EVENT_SCOPE_TEXT(*FString::Printf(TEXT("BuildOrCopyRenderData_BuildData: LOD %d"), LODIndex));
+
 		SetLastMeshId(Component.Id, LODIndex, Component.Mesh->GetId());
 
 		UnrealConversionUtils::SetupRenderSections(
 			SkeletalMesh,
 			Component.Mesh,
-			MeshLODIndex,
+			LODIndex,
 			Component.BoneMap);
 
 		UnrealConversionUtils::CopyMutableVertexBuffers(
 			SkeletalMesh,
 			Component.Mesh,
-			MeshLODIndex);
+			LODIndex);
 
-		FSkeletalMeshLODRenderData& LODModel = Helper_GetLODData(SkeletalMesh)[MeshLODIndex];
+		FSkeletalMeshLODRenderData& LODModel = Helper_GetLODData(SkeletalMesh)[LODIndex];
 
 		if (LODModel.DoesVertexBufferUse16BitBoneIndex() && !UCustomizableObjectSystem::GetInstance()->IsSupport16BitBoneIndexEnabled())
 		{
@@ -4153,7 +4418,7 @@ bool UCustomizableInstancePrivateData::BuildSkeletalMeshRenderData(const TShared
 		LODModel.ActiveBoneIndices.Append(Component.ActiveBones);
 		LODModel.RequiredBones.Append(Component.ActiveBones);
 
-		const TArray<FMutableSkinWeightProfileInfo>& SkinWeightProfilesInfo = Public->GetCustomizableObject()->SkinWeightProfilesInfo;
+		const TArray<FMutableSkinWeightProfileInfo>& SkinWeightProfilesInfo = CustomizableObject->SkinWeightProfilesInfo;
 		if (SkinWeightProfilesInfo.Num())
 		{
 			bool bHasSkinWeightProfiles = false;
@@ -4184,7 +4449,7 @@ bool UCustomizableInstancePrivateData::BuildSkeletalMeshRenderData(const TShared
 
 				const FSkinWeightProfileInfo* ExistingProfile = SkeletalMesh->GetSkinWeightProfiles().FindByPredicate(
 					[&Profile](const FSkinWeightProfileInfo& P) { return P.Name == Profile.Name; });
-				
+
 				if (!ExistingProfile)
 				{
 					SkeletalMesh->AddSkinWeightProfile({ Profile.Name, Profile.DefaultProfile, Profile.DefaultProfileFromLODIndex });
@@ -4216,14 +4481,16 @@ bool UCustomizableInstancePrivateData::BuildSkeletalMeshRenderData(const TShared
 		LODModel.BuffersSize = 1;
 
 	}
-	
+
 
 	return true;
 }
 
 
-FGraphEventRef UCustomizableInstancePrivateData::LoadAdditionalAssetsAsync(const TSharedPtr<FMutableOperationData>& OperationData, UCustomizableObjectInstance* Public, FStreamableManager &StreamableManager)
+FGraphEventRef UCustomizableInstancePrivateData::LoadAdditionalAssetsAsync(const TSharedPtr<FMutableOperationData>& OperationData, UCustomizableObjectInstance* Public, FStreamableManager& StreamableManager)
 {
+	MUTABLE_CPUPROFILER_SCOPE(UCustomizableInstancePrivateData::LoadAdditionalAssetsAsync);
+
 	UCustomizableObject* CustomizableObject = Public->GetCustomizableObject();
 
 	FGraphEventRef Result = nullptr;
@@ -4242,7 +4509,7 @@ FGraphEventRef UCustomizableInstancePrivateData::LoadAdditionalAssetsAsync(const
 
 	GatheredAnimBPs.Empty();
 	AnimBPGameplayTags.Reset();
-	
+
 	for (const FInstanceUpdateData::FSurface& Surface : OperationData->InstanceUpdateData.Surfaces)
 	{
 		uint32 MaterialIndex = Surface.MaterialIndex;
@@ -4320,15 +4587,15 @@ FGraphEventRef UCustomizableInstancePrivateData::LoadAdditionalAssetsAsync(const
 		const FInstanceUpdateData::FLOD& LOD = OperationData->InstanceUpdateData.LODs[OperationData->CurrentMinLOD];
 		for (int32 ComponentIndex = 0; ComponentIndex < LOD.ComponentCount; ++ComponentIndex)
 		{
-			const FInstanceUpdateData::FComponent& Component = Components[LOD.FirstComponent+ComponentIndex];
+			const FInstanceUpdateData::FComponent& Component = Components[LOD.FirstComponent + ComponentIndex];
 			mu::MeshPtrConst MutableMesh = Component.Mesh;
 
 			if (!MutableMesh)
 			{
 				continue;
 			}
-				
-			FCustomizableInstanceComponentData* ComponentData =	GetComponentData(Component.Id);
+
+			FCustomizableInstanceComponentData* ComponentData = GetComponentData(Component.Id);
 
 			const bool bReplacePhysicsAssets = HasCOInstanceFlags(ReplacePhysicsAssets);
 
@@ -4358,7 +4625,7 @@ FGraphEventRef UCustomizableInstancePrivateData::LoadAdditionalAssetsAsync(const
 
 					if (Tag.Split(TEXT("_AssetIdx_"), &AssetIndexString, &AssetPath) && AssetIndexString.IsNumeric())
 					{
-						
+
 						int32 AssetIndex = FCString::Atoi(*AssetIndexString);
 
 						TSoftObjectPtr<UPhysicsAsset>* PhysicsAssetPtr = CustomizableObject->PhysicsAssetsMap.Find(AssetPath);
@@ -4369,7 +4636,7 @@ FGraphEventRef UCustomizableInstancePrivateData::LoadAdditionalAssetsAsync(const
 						{
 							if (PhysicsAssetPtr->Get())
 							{
-								if( ClothingPhysicsAssets.IsValidIndex(AssetIndex) )
+								if (ClothingPhysicsAssets.IsValidIndex(AssetIndex))
 								{
 									ClothingPhysicsAssets[AssetIndex] = PhysicsAssetPtr->Get();
 								}
@@ -4442,17 +4709,11 @@ FGraphEventRef UCustomizableInstancePrivateData::LoadAdditionalAssetsAsync(const
 
 	if (AssetsToStream.Num() > 0)
 	{
-		ClearCOInstanceFlags(AssetsLoaded);
-
 		check(!StreamingHandle);
 
 		Result = FGraphEvent::CreateGraphEvent();
 		StreamingHandle = StreamableManager.RequestAsyncLoad(AssetsToStream, FStreamableDelegate::CreateUObject(Public, &UCustomizableObjectInstance::AdditionalAssetsAsyncLoaded, Result));
 	}
-	else
-	{
-		SetCOInstanceFlags(AssetsLoaded);
-	}
 
 	return Result;
 }
@@ -4469,9 +4730,9 @@ void UCustomizableObjectInstance::AdditionalAssetsAsyncLoaded( FGraphEventRef Co
 
 void UCustomizableInstancePrivateData::AdditionalAssetsAsyncLoaded(UCustomizableObjectInstance* Public)
 {
-	UCustomizableObject* CustomizableObject = Public->GetCustomizableObject();
+	MUTABLE_CPUPROFILER_SCOPE(UCustomizableInstancePrivateData::AdditionalAssetsAsyncLoaded);
 
-	SetCOInstanceFlags(AssetsLoaded);
+	UCustomizableObject* CustomizableObject = Public->GetCustomizableObject();
 
 	// Loaded Materials
 	check(ObjectToInstanceIndexMap.Num() == ReferencedMaterials.Num());
@@ -4675,11 +4936,11 @@ void UCustomizableInstancePrivateData::ReuseTexture(UTexture2D* Texture)
 }
 
 
-void UCustomizableInstancePrivateData::BuildMaterials(const TSharedPtr<FMutableOperationData>& OperationData, UCustomizableObjectInstance* Public )
+void UCustomizableInstancePrivateData::BuildMaterials(const TSharedPtr<FMutableOperationData>& OperationData, UCustomizableObjectInstance* Public)
 {
 	MUTABLE_CPUPROFILER_SCOPE(UCustomizableInstancePrivateData::BuildMaterials)
 
-	UCustomizableObject* CustomizableObject = Public->GetCustomizableObject();
+		UCustomizableObject* CustomizableObject = Public->GetCustomizableObject();
 
 	const int32 LODCount = OperationData->InstanceUpdateData.LODs.Num();
 
@@ -4694,7 +4955,6 @@ void UCustomizableInstancePrivateData::BuildMaterials(const TSharedPtr<FMutableO
 	UpdateContext->Parameters = OperationData->MutableParameters;
 	UpdateContext->State = OperationData->State;
 
-	const bool bUseCurrentMinLODAsBaseLOD = HasCOInstanceFlags(ReduceLODs);
 	const bool bReuseTextures = OperationData->bReuseInstanceTextures;
 
 	const FInstanceUpdateData::FLOD& FirstLOD = OperationData->InstanceUpdateData.LODs[OperationData->CurrentMinLOD];
@@ -4703,447 +4963,451 @@ void UCustomizableInstancePrivateData::BuildMaterials(const TSharedPtr<FMutableO
 	for (int32 ComponentIndex = 0; ComponentIndex < NumComponents; ++ComponentIndex)
 	{
 		USkeletalMesh* SkeletalMesh = Public->SkeletalMeshes[ComponentIndex];
-		if (SkeletalMesh)
+
+		if (!SkeletalMesh)
 		{
-			SkeletalMesh->GetMaterials().Reset();
+			continue;
 		}
 
+		SkeletalMesh->GetMaterials().Reset();
+
 		// Maps serializations of FMutableMaterialPlaceholder to Created Dynamic Material instances, used to reuse materials across LODs
 		TMap<FString, FMutableMaterialPlaceholder*> ReuseMaterialCache;
 		// Stores the FMutableMaterialPlaceholders and keeps them in memory while the component's LODs are being processed
 		TArray<TSharedPtr<FMutableMaterialPlaceholder>> MutableMaterialPlaceholderArray;
 
-		MUTABLE_CPUPROFILER_SCOPE(LODLoop);
-
-		int32 MeshLODIndex = bUseCurrentMinLODAsBaseLOD ? 0 : OperationData->CurrentMinLOD;
+		MUTABLE_CPUPROFILER_SCOPE(BuildMaterials_LODLoop);
 
-		for (int32 LODIndex = OperationData->CurrentMinLOD; LODIndex <= OperationData->CurrentMaxLOD && LODIndex < LODCount; ++LODIndex, ++MeshLODIndex)
+		int32 LastValidLODIndex = OperationData->CurrentMaxLOD;
+		for (int32 LODIndex = OperationData->CurrentMaxLOD; LODIndex >= FirstLODAvailable; --LODIndex)
 		{
-			const FInstanceUpdateData::FLOD& LOD = OperationData->InstanceUpdateData.LODs[LODIndex];
+			if (Helper_GetLODInfoArray(SkeletalMesh).Num() != 0)
+			{
+				Helper_GetLODInfoArray(SkeletalMesh)[LODIndex].LODMaterialMap.Reset();
+			}
 
-			if (SkeletalMesh && Helper_GetLODInfoArray(SkeletalMesh).Num() != 0)
+			if (LODIndex >= OperationData->CurrentMinLOD && OperationData->InstanceUpdateData.Components[OperationData->InstanceUpdateData.LODs[LODIndex].FirstComponent + ComponentIndex].bGenerated)
 			{
-				Helper_GetLODInfoArray(SkeletalMesh)[MeshLODIndex].LODMaterialMap.Reset();
+				LastValidLODIndex = LODIndex;
 			}
 
-			if (ComponentIndex<LOD.ComponentCount)
+			const FInstanceUpdateData::FLOD& LOD = OperationData->InstanceUpdateData.LODs[LastValidLODIndex];
+			const FInstanceUpdateData::FComponent& Component = OperationData->InstanceUpdateData.Components[LOD.FirstComponent + ComponentIndex];
+			check(Component.bGenerated);
+
+			const FMutableRefSkeletalMeshData* RefSkeletalMeshData = CustomizableObject->GetRefSkeletalMeshData(Component.Id);
+			check(RefSkeletalMeshData);
+
+			for (int32 SurfaceIndex = 0; SurfaceIndex < Component.SurfaceCount; ++SurfaceIndex)
 			{
-				const FInstanceUpdateData::FComponent& Component = OperationData->InstanceUpdateData.Components[LOD.FirstComponent+ComponentIndex];
-				auto MutableMesh = Component.Mesh;
+				const FInstanceUpdateData::FSurface& Surface = OperationData->InstanceUpdateData.Surfaces[Component.FirstSurface + SurfaceIndex];
 
-				if (!MutableMesh)
+				MutableMaterialPlaceholderArray.Add(TSharedPtr<FMutableMaterialPlaceholder>(new FMutableMaterialPlaceholder));
+				FMutableMaterialPlaceholder& MutableMaterialPlaceholder = *MutableMaterialPlaceholderArray.Last();
+
+				if (LastValidLODIndex != LODIndex)
 				{
-					// This case is dangerous. If the LOD has no mesh at all it will crash unreal.
-					// In theory this is checked earlier, and the LOD is removed.
+					// Copy information from the LastValidLODIndex
+					Helper_GetLODInfoArray(SkeletalMesh)[LODIndex].LODMaterialMap = Helper_GetLODInfoArray(SkeletalMesh)[LastValidLODIndex].LODMaterialMap;
+					Helper_GetLODRenderSections(SkeletalMesh, LODIndex)[SurfaceIndex].MaterialIndex = Helper_GetLODRenderSections(SkeletalMesh, LastValidLODIndex)[SurfaceIndex].MaterialIndex;
 					continue;
 				}
 
-				const FMutableRefSkeletalMeshData* RefSkeletalMeshData = CustomizableObject->GetRefSkeletalMeshData(Component.Id);
-				check(RefSkeletalMeshData);
-
-				for (int32 SurfaceIndex = 0; SurfaceIndex < Component.SurfaceCount; ++SurfaceIndex)
 				{
-					const FInstanceUpdateData::FSurface& Surface = OperationData->InstanceUpdateData.Surfaces[Component.FirstSurface + SurfaceIndex];
-
-					MutableMaterialPlaceholderArray.Add(TSharedPtr<FMutableMaterialPlaceholder>(new FMutableMaterialPlaceholder));
-					FMutableMaterialPlaceholder& MutableMaterialPlaceholder = *MutableMaterialPlaceholderArray.Last();
+					MUTABLE_CPUPROFILER_SCOPE(ParamLoop);
 
+					for (int32 VectorIndex = 0; VectorIndex < Surface.VectorCount; ++VectorIndex)
 					{
-						MUTABLE_CPUPROFILER_SCOPE(ParamLoop);
-
-						for (int32 VectorIndex = 0; VectorIndex < Surface.VectorCount; ++VectorIndex)
-						{
-							const FInstanceUpdateData::FVector& Vector = OperationData->InstanceUpdateData.Vectors[Surface.FirstVector + VectorIndex];
+						const FInstanceUpdateData::FVector& Vector = OperationData->InstanceUpdateData.Vectors[Surface.FirstVector + VectorIndex];
 
-							// Decoding Material Layer from Mutable parameter name
-							FString EncodingString = "-MutableLayerParam:";
+						// Decoding Material Layer from Mutable parameter name
+						FString EncodingString = "-MutableLayerParam:";
 
-							int32 EncodingPosition = Vector.Name.Find(EncodingString);
-							int32 LayerIndex = -1;
-							
-							if (EncodingPosition == INDEX_NONE)
-							{
-								MutableMaterialPlaceholder.AddParam(
-									FMutableMaterialPlaceholder::FMutableMaterialPlaceHolderParam(FName(Vector.Name), -1, Vector.Vector));
-							}
-							else
-							{
-								//Getting layer index
-								int32 LayerPosition = Vector.Name.Len() - (EncodingPosition + EncodingString.Len());
-								FString IndexString = Vector.Name.RightChop(Vector.Name.Len() - LayerPosition);
-								LayerIndex = FCString::Atof(*IndexString);
+						int32 EncodingPosition = Vector.Name.Find(EncodingString);
+						int32 LayerIndex = -1;
 
-								//Getting parameter name
-								FString Sufix = EncodingString + FString::FromInt(LayerIndex);
-								FString NameCopy = Vector.Name;
-								NameCopy.RemoveFromEnd(Sufix);
-									
-								MutableMaterialPlaceholder.AddParam(
-									FMutableMaterialPlaceholder::FMutableMaterialPlaceHolderParam(FName(NameCopy), LayerIndex, Vector.Vector));
-							}
+						if (EncodingPosition == INDEX_NONE)
+						{
+							MutableMaterialPlaceholder.AddParam(
+								FMutableMaterialPlaceholder::FMutableMaterialPlaceHolderParam(FName(Vector.Name), -1, Vector.Vector));
 						}
-
-						for (int32 ScalarIndex = 0; ScalarIndex < Surface.ScalarCount; ++ScalarIndex)
+						else
 						{
-							const FInstanceUpdateData::FScalar& Scalar = OperationData->InstanceUpdateData.Scalars[Surface.FirstScalar + ScalarIndex];
-							
-							// Decoding Material Layer from Mutable parameter name
-							FString EncodingString = "-MutableLayerParam:";
+							//Getting layer index
+							int32 LayerPosition = Vector.Name.Len() - (EncodingPosition + EncodingString.Len());
+							FString IndexString = Vector.Name.RightChop(Vector.Name.Len() - LayerPosition);
+							LayerIndex = FCString::Atof(*IndexString);
+
+							//Getting parameter name
+							FString Sufix = EncodingString + FString::FromInt(LayerIndex);
+							FString NameCopy = Vector.Name;
+							NameCopy.RemoveFromEnd(Sufix);
+
+							MutableMaterialPlaceholder.AddParam(
+								FMutableMaterialPlaceholder::FMutableMaterialPlaceHolderParam(FName(NameCopy), LayerIndex, Vector.Vector));
+						}
+					}
 
-							int32 EncodingPosition = Scalar.Name.Find(EncodingString);
-							int32 LayerIndex = -1;
+					for (int32 ScalarIndex = 0; ScalarIndex < Surface.ScalarCount; ++ScalarIndex)
+					{
+						const FInstanceUpdateData::FScalar& Scalar = OperationData->InstanceUpdateData.Scalars[Surface.FirstScalar + ScalarIndex];
 
-							if (EncodingPosition == INDEX_NONE)
-							{
-								MutableMaterialPlaceholder.AddParam(
-									FMutableMaterialPlaceholder::FMutableMaterialPlaceHolderParam(FName(Scalar.Name), -1, Scalar.Scalar));
-							}
-							else
-							{
-								//Getting layer index
-								int32 LayerPosition = Scalar.Name.Len() - (EncodingPosition + EncodingString.Len());
-								FString IndexString = Scalar.Name.RightChop(Scalar.Name.Len() - LayerPosition);
-								LayerIndex = FCString::Atof(*IndexString);
+						// Decoding Material Layer from Mutable parameter name
+						FString EncodingString = "-MutableLayerParam:";
 
-								//Getting parameter name
-								FString Sufix = EncodingString + FString::FromInt(LayerIndex);
-								FString NameCopy = Scalar.Name;
-								NameCopy.RemoveFromEnd(Sufix);
+						int32 EncodingPosition = Scalar.Name.Find(EncodingString);
+						int32 LayerIndex = -1;
 
-								MutableMaterialPlaceholder.AddParam(
-									FMutableMaterialPlaceholder::FMutableMaterialPlaceHolderParam(FName(NameCopy), LayerIndex, Scalar.Scalar));
-							}
+						if (EncodingPosition == INDEX_NONE)
+						{
+							MutableMaterialPlaceholder.AddParam(
+								FMutableMaterialPlaceholder::FMutableMaterialPlaceHolderParam(FName(Scalar.Name), -1, Scalar.Scalar));
+						}
+						else
+						{
+							//Getting layer index
+							int32 LayerPosition = Scalar.Name.Len() - (EncodingPosition + EncodingString.Len());
+							FString IndexString = Scalar.Name.RightChop(Scalar.Name.Len() - LayerPosition);
+							LayerIndex = FCString::Atof(*IndexString);
+
+							//Getting parameter name
+							FString Sufix = EncodingString + FString::FromInt(LayerIndex);
+							FString NameCopy = Scalar.Name;
+							NameCopy.RemoveFromEnd(Sufix);
+
+							MutableMaterialPlaceholder.AddParam(
+								FMutableMaterialPlaceholder::FMutableMaterialPlaceHolderParam(FName(NameCopy), LayerIndex, Scalar.Scalar));
 						}
 					}
+				}
 
-					{
-						MUTABLE_CPUPROFILER_SCOPE(ImageLoop);
+				{
+					MUTABLE_CPUPROFILER_SCOPE(BuildMaterials_ImageLoop);
 
-						// Get the cache of resources of all live instances of this object
-						FMutableResourceCache& Cache = UCustomizableObjectSystem::GetInstance()->GetPrivate()->GetObjectCache(CustomizableObject);
+					// Get the cache of resources of all live instances of this object
+					FMutableResourceCache& Cache = UCustomizableObjectSystem::GetInstance()->GetPrivate()->GetObjectCache(CustomizableObject);
 
-						FString CurrentState = Public->GetCurrentState();
-						bool bNeverStream = OperationData->bNeverStream;
+					FString CurrentState = Public->GetCurrentState();
+					bool bNeverStream = OperationData->bNeverStream;
 
-						check((bNeverStream && OperationData->MipsToSkip == 0) ||
-							 (!bNeverStream && OperationData->MipsToSkip >= 0));
+					check((bNeverStream && OperationData->MipsToSkip == 0) ||
+						(!bNeverStream && OperationData->MipsToSkip >= 0));
 
-						for (int32 ImageIndex = 0; ImageIndex < Surface.ImageCount; ++ImageIndex)
-						{
-							const FInstanceUpdateData::FImage& Image = OperationData->InstanceUpdateData.Images[Surface.FirstImage+ImageIndex];
-							FString KeyName = Image.Name;
-							mu::ImagePtrConst MutableImage = Image.Image;
+					for (int32 ImageIndex = 0; ImageIndex < Surface.ImageCount; ++ImageIndex)
+					{
+						const FInstanceUpdateData::FImage& Image = OperationData->InstanceUpdateData.Images[Surface.FirstImage + ImageIndex];
+						FString KeyName = Image.Name;
+						mu::ImagePtrConst MutableImage = Image.Image;
 
-							UTexture2D* Texture = nullptr;
+						UTexture2D* Texture = nullptr;
 
-							// \TODO: Change this key to a struct.
-							FString TextureReuseCacheRef = bReuseTextures ? FString::Printf(TEXT("%d-%d-%d-%d"), LODIndex, ComponentIndex, Surface.SurfaceId, ImageIndex) : FString();
+						// \TODO: Change this key to a struct.
+						FString TextureReuseCacheRef = bReuseTextures ? FString::Printf(TEXT("%d-%d-%d-%d"), LastValidLODIndex, ComponentIndex, Surface.SurfaceId, ImageIndex) : FString();
 
-							// If the mutable image is null, it must be in the cache
-							FMutableImageCacheKey ImageCacheKey = { Image.ImageID, OperationData->MipsToSkip };
-							if (!MutableImage)
+						// If the mutable image is null, it must be in the cache
+						FMutableImageCacheKey ImageCacheKey = { Image.ImageID, OperationData->MipsToSkip };
+						if (!MutableImage)
+						{
+							TWeakObjectPtr<UTexture2D>* CachedPointerPtr = Cache.Images.Find(ImageCacheKey);
+							if (CachedPointerPtr)
 							{
-								TWeakObjectPtr<UTexture2D>* CachedPointerPtr = Cache.Images.Find(ImageCacheKey);
-								if (CachedPointerPtr)
-								{
-									ensure(!CachedPointerPtr->IsStale());
-									Texture = CachedPointerPtr->Get();
-								}
-
-								check(Texture);
+								ensure(!CachedPointerPtr->IsStale());
+								Texture = CachedPointerPtr->Get();
 							}
 
-							// Find the additional information for this image
-							int32 ImageKey = FCString::Atoi(*KeyName);
-							if (ImageKey >= 0 && ImageKey < CustomizableObject->ImageProperties.Num())
+							check(Texture);
+						}
+
+						// Find the additional information for this image
+						int32 ImageKey = FCString::Atoi(*KeyName);
+						if (ImageKey >= 0 && ImageKey < CustomizableObject->ImageProperties.Num())
+						{
+							const FMutableModelImageProperties& Props = CustomizableObject->ImageProperties[ImageKey];
+
+							if (MutableImage)
 							{
-								const FMutableModelImageProperties& Props = CustomizableObject->ImageProperties[ImageKey];
+								TWeakObjectPtr<UTexture2D>* ReusedTexture = bReuseTextures ? TextureReuseCache.Find(TextureReuseCacheRef) : nullptr;
 
-								if (MutableImage)
+								if (ReusedTexture && (*ReusedTexture).IsValid() && !(*ReusedTexture)->HasAnyFlags(RF_BeginDestroyed))
 								{
-									TWeakObjectPtr<UTexture2D>* ReusedTexture = bReuseTextures ? TextureReuseCache.Find(TextureReuseCacheRef) : nullptr;
+									// Only uncompressed textures can be reused. This also fixes an issue in the editor where textures supposedly 
+									// uncompressed by their state, are still compressed because the CO has not been compiled at maximum settings
+									// and the uncompressed setting cannot be applied to them.
+									EPixelFormat PixelFormat = (*ReusedTexture)->GetPixelFormat();
 
-									if (ReusedTexture && (*ReusedTexture).IsValid() && !(*ReusedTexture)->HasAnyFlags(RF_BeginDestroyed))
+									if (PixelFormat == EPixelFormat::PF_R8G8B8A8)
 									{
-										// Only uncompressed textures can be reused. This also fixes an issue in the editor where textures supposedly 
-										// uncompressed by their state, are still compressed because the CO has not been compiled at maximum settings
-										// and the uncompressed setting cannot be applied to them.
-										EPixelFormat PixelFormat = (*ReusedTexture)->GetPixelFormat();
-
-										if (PixelFormat == EPixelFormat::PF_R8G8B8A8)
-										{
-											Texture = (*ReusedTexture).Get();
-											check(Texture != nullptr);
-										}
-										else
-										{
-											ReusedTexture = nullptr;
-											Texture = CreateTexture();
-
-#if WITH_EDITOR
-											UE_LOG(LogMutable, Warning,
-												TEXT("Tried to reuse an uncompressed texture with name %s. Make sure the selected Mutable state disables texture compression/streaming, that one of the state's runtime parameters affects the texture and that the CO is compiled with max. optimization settings."),
-												*Texture->GetName());
-#endif
-										}
+										Texture = (*ReusedTexture).Get();
+										check(Texture != nullptr);
 									}
 									else
 									{
 										ReusedTexture = nullptr;
 										Texture = CreateTexture();
+
+#if WITH_EDITOR
+										UE_LOG(LogMutable, Warning,
+											TEXT("Tried to reuse an uncompressed texture with name %s. Make sure the selected Mutable state disables texture compression/streaming, that one of the state's runtime parameters affects the texture and that the CO is compiled with max. optimization settings."),
+											*Texture->GetName());
+#endif
 									}
+								}
+								else
+								{
+									ReusedTexture = nullptr;
+									Texture = CreateTexture();
+								}
 
-									if (Texture)
+								if (Texture)
+								{
+									if (OperationData->ImageToPlatformDataMap.Contains(Image.ImageID))
 									{
-										if (OperationData->ImageToPlatformDataMap.Contains(Image.ImageID))
-										{
-											SetTexturePropertiesFromMutableImageProps(Texture, Props, bNeverStream);
-
-											FTexturePlatformData* PlatformData = OperationData->ImageToPlatformDataMap[Image.ImageID];
+										SetTexturePropertiesFromMutableImageProps(Texture, Props, bNeverStream);
 
-											if (ReusedTexture)
-											{
-												check(PlatformData->Mips.Num() == Texture->GetPlatformData()->Mips.Num());
-												check(PlatformData->Mips[0].SizeX == Texture->GetPlatformData()->Mips[0].SizeX);
-												check(PlatformData->Mips[0].SizeY == Texture->GetPlatformData()->Mips[0].SizeY);
-											}
-
-											if (Texture->GetPlatformData())
-											{
-												delete Texture->GetPlatformData();
-											}
+										FTexturePlatformData* PlatformData = OperationData->ImageToPlatformDataMap[Image.ImageID];
 
-											Texture->SetPlatformData( PlatformData );
-											OperationData->ImageToPlatformDataMap.Remove(Image.ImageID);
-										}
-										else
+										if (ReusedTexture)
 										{
-											UE_LOG(LogMutable, Error, TEXT("Required image was not generated in the mutable thread, and it is not cached."));
+											check(PlatformData->Mips.Num() == Texture->GetPlatformData()->Mips.Num());
+											check(PlatformData->Mips[0].SizeX == Texture->GetPlatformData()->Mips[0].SizeX);
+											check(PlatformData->Mips[0].SizeY == Texture->GetPlatformData()->Mips[0].SizeY);
 										}
 
-										if (bNeverStream)
+										if (Texture->GetPlatformData())
 										{
-											// To prevent LogTexture Error "Loading non-streamed mips from an external bulk file."
-											for (int32 i = 0; i < Texture->GetPlatformData()->Mips.Num(); ++i)
-											{
-												Texture->GetPlatformData()->Mips[i].BulkData.ClearBulkDataFlags(BULKDATA_PayloadInSeperateFile);
-											}
+											delete Texture->GetPlatformData();
 										}
 
+										Texture->SetPlatformData(PlatformData);
+										OperationData->ImageToPlatformDataMap.Remove(Image.ImageID);
+									}
+									else
+									{
+										UE_LOG(LogMutable, Error, TEXT("Required image was not generated in the mutable thread, and it is not cached."));
+									}
+
+									if (bNeverStream)
+									{
+										// To prevent LogTexture Error "Loading non-streamed mips from an external bulk file."
+										for (int32 i = 0; i < Texture->GetPlatformData()->Mips.Num(); ++i)
 										{
-											MUTABLE_CPUPROFILER_SCOPE(UpdateResource);
+											Texture->GetPlatformData()->Mips[i].BulkData.ClearBulkDataFlags(BULKDATA_PayloadInSeperateFile);
+										}
+									}
+
+									{
+										MUTABLE_CPUPROFILER_SCOPE(UpdateResource);
 #if !PLATFORM_DESKTOP && !PLATFORM_SWITCH // Switch does this in FTexture2DResource::InitRHI()
-											for (int32 i = 0; i < Texture->GetPlatformData()->Mips.Num(); ++i)
-											{
-												uint32 DataFlags = Texture->GetPlatformData()->Mips[i].BulkData.GetBulkDataFlags();
-												Texture->GetPlatformData()->Mips[i].BulkData.SetBulkDataFlags(DataFlags | BULKDATA_SingleUse);
-											}
+										for (int32 i = 0; i < Texture->GetPlatformData()->Mips.Num(); ++i)
+										{
+											uint32 DataFlags = Texture->GetPlatformData()->Mips[i].BulkData.GetBulkDataFlags();
+											Texture->GetPlatformData()->Mips[i].BulkData.SetBulkDataFlags(DataFlags | BULKDATA_SingleUse);
+										}
 #endif
 
-											if (ReusedTexture)
+										if (ReusedTexture)
+										{
+											// Must remove texture from cache since it will be reused with a different ImageID
+											for (TPair<FMutableImageCacheKey, TWeakObjectPtr<UTexture2D>>& CachedTexture : Cache.Images)
 											{
-												// Must remove texture from cache since it will be reused with a different ImageID
-												for (TPair<FMutableImageCacheKey, TWeakObjectPtr<UTexture2D>>& CachedTexture : Cache.Images)
+												if (CachedTexture.Value == Texture)
 												{
-													if (CachedTexture.Value == Texture)
-													{
-														Cache.Images.Remove(CachedTexture.Key);
-														break;
-													}
+													Cache.Images.Remove(CachedTexture.Key);
+													break;
 												}
-
-												ReuseTexture(Texture);
 											}
-											else if (Texture)
+
+											ReuseTexture(Texture);
+										}
+										else if (Texture)
+										{
+											//if (!bNeverStream) // No need to check bNeverStream. In that case, the texture won't use 
+											// the MutableMipDataProviderFactory anyway and it's needed for detecting Mutable textures elsewhere
 											{
-												//if (!bNeverStream) // No need to check bNeverStream. In that case, the texture won't use 
-												// the MutableMipDataProviderFactory anyway and it's needed for detecting Mutable textures elsewhere
+												UMutableTextureMipDataProviderFactory* MutableMipDataProviderFactory = Cast<UMutableTextureMipDataProviderFactory>(Texture->GetAssetUserDataOfClass(UMutableTextureMipDataProviderFactory::StaticClass()));
+												if (!MutableMipDataProviderFactory)
 												{
-													UMutableTextureMipDataProviderFactory* MutableMipDataProviderFactory = Cast<UMutableTextureMipDataProviderFactory>(Texture->GetAssetUserDataOfClass(UMutableTextureMipDataProviderFactory::StaticClass()));
-													if (!MutableMipDataProviderFactory)
+													MutableMipDataProviderFactory = NewObject<UMutableTextureMipDataProviderFactory>();
+
+													if (MutableMipDataProviderFactory)
 													{
-														MutableMipDataProviderFactory = NewObject<UMutableTextureMipDataProviderFactory>();
-
-														if (MutableMipDataProviderFactory)
-														{
-															MutableMipDataProviderFactory->CustomizableObjectInstance = Public;
-															check(LODIndex < 256 && ComponentIndex < 256 && ImageIndex < 256);
-															MutableMipDataProviderFactory->ImageRef.ImageID = Image.ImageID;
-															MutableMipDataProviderFactory->ImageRef.SurfaceId = Surface.SurfaceId;
-															MutableMipDataProviderFactory->ImageRef.LOD = uint8(LODIndex);
-															MutableMipDataProviderFactory->ImageRef.Component = uint8(ComponentIndex);
-															MutableMipDataProviderFactory->ImageRef.Image = uint8(ImageIndex);
-															MutableMipDataProviderFactory->UpdateContext = UpdateContext;
-															Texture->AddAssetUserData(MutableMipDataProviderFactory);
-														}
+														MutableMipDataProviderFactory->CustomizableObjectInstance = Public;
+														check(LastValidLODIndex < 256 && ComponentIndex < 256 && ImageIndex < 256);
+														MutableMipDataProviderFactory->ImageRef.ImageID = Image.ImageID;
+														MutableMipDataProviderFactory->ImageRef.SurfaceId = Surface.SurfaceId;
+														MutableMipDataProviderFactory->ImageRef.LOD = uint8(LastValidLODIndex);
+														MutableMipDataProviderFactory->ImageRef.Component = uint8(ComponentIndex);
+														MutableMipDataProviderFactory->ImageRef.Image = uint8(ImageIndex);
+														MutableMipDataProviderFactory->UpdateContext = UpdateContext;
+														Texture->AddAssetUserData(MutableMipDataProviderFactory);
 													}
 												}
-
-												Texture->UpdateResource();
 											}
-										}
 
-										// update the model resources cache
-										Cache.Images.Add(ImageCacheKey, Texture);
-									}
-									else
-									{
-										UE_LOG(LogMutable, Error, TEXT("Texture creation failed."));
+											Texture->UpdateResource();
+										}
 									}
-								}
-
-								FGeneratedTexture TextureData;
-								TextureData.Id = Image.ImageID;
-								TextureData.Name = Props.TextureParameterName;
-								TextureData.Texture = Texture;
-								NewGeneratedTextures.Add(TextureData);
-
-								// Decoding Material Layer from Mutable parameter name
-								FString ImageName = Image.Name;
-								FString EncodingString = "-MutableLayerParam:";
-
-								int32 EncodingPosition = ImageName.Find(EncodingString);
-								int32 LayerIndex = -1;
 
-								if (EncodingPosition == INDEX_NONE)
-								{
-									MutableMaterialPlaceholder.AddParam(
-										FMutableMaterialPlaceholder::FMutableMaterialPlaceHolderParam(*Props.TextureParameterName, -1, TextureData));
+									// update the model resources cache
+									Cache.Images.Add(ImageCacheKey, Texture);
 								}
 								else
 								{
-									//Getting layer index
-									int32 LayerPosition = ImageName.Len() - (EncodingPosition + EncodingString.Len());
-									FString IndexString = ImageName.RightChop(ImageName.Len() - LayerPosition);
-									LayerIndex = FCString::Atof(*IndexString);
-
-									MutableMaterialPlaceholder.AddParam(
-										FMutableMaterialPlaceholder::FMutableMaterialPlaceHolderParam(*Props.TextureParameterName, LayerIndex, TextureData));
+									UE_LOG(LogMutable, Error, TEXT("Texture creation failed."));
 								}
 							}
+
+							FGeneratedTexture TextureData;
+							TextureData.Id = Image.ImageID;
+							TextureData.Name = Props.TextureParameterName;
+							TextureData.Texture = Texture;
+							NewGeneratedTextures.Add(TextureData);
+
+							// Decoding Material Layer from Mutable parameter name
+							FString ImageName = Image.Name;
+							FString EncodingString = "-MutableLayerParam:";
+
+							int32 EncodingPosition = ImageName.Find(EncodingString);
+							int32 LayerIndex = -1;
+
+							if (EncodingPosition == INDEX_NONE)
+							{
+								MutableMaterialPlaceholder.AddParam(
+									FMutableMaterialPlaceholder::FMutableMaterialPlaceHolderParam(*Props.TextureParameterName, -1, TextureData));
+							}
 							else
 							{
-								// This means the compiled model (maybe coming from derived data) has images that the asset doesn't know about.
-								UE_LOG(LogMutable, Error, TEXT("CustomizableObject derived data out of sync with asset for [%s]. Try recompiling it."), *CustomizableObject->GetName());
+								//Getting layer index
+								int32 LayerPosition = ImageName.Len() - (EncodingPosition + EncodingString.Len());
+								FString IndexString = ImageName.RightChop(ImageName.Len() - LayerPosition);
+								LayerIndex = FCString::Atof(*IndexString);
+
+								MutableMaterialPlaceholder.AddParam(
+									FMutableMaterialPlaceholder::FMutableMaterialPlaceHolderParam(*Props.TextureParameterName, LayerIndex, TextureData));
 							}
+						}
+						else
+						{
+							// This means the compiled model (maybe coming from derived data) has images that the asset doesn't know about.
+							UE_LOG(LogMutable, Error, TEXT("CustomizableObject derived data out of sync with asset for [%s]. Try recompiling it."), *CustomizableObject->GetName());
+						}
 
-							if (bReuseTextures)
+						if (bReuseTextures)
+						{
+							if (Texture)
 							{
-								if (Texture)
-								{
-									TextureReuseCache.Add(TextureReuseCacheRef, Texture);
-								}
-								else
-								{
-									TextureReuseCache.Remove(TextureReuseCacheRef);
-								}
+								TextureReuseCache.Add(TextureReuseCacheRef, Texture);
+							}
+							else
+							{
+								TextureReuseCache.Remove(TextureReuseCacheRef);
 							}
 						}
 					}
+				}
 
-					const uint32 ReferencedMaterialIndex = ObjectToInstanceIndexMap[Surface.MaterialIndex];
-					UMaterialInterface* MaterialTemplate = ReferencedMaterials[ReferencedMaterialIndex];
-					MutableMaterialPlaceholder.ParentMaterial = MaterialTemplate;
+				const uint32 ReferencedMaterialIndex = ObjectToInstanceIndexMap[Surface.MaterialIndex];
+				UMaterialInterface* MaterialTemplate = ReferencedMaterials[ReferencedMaterialIndex];
+				MutableMaterialPlaceholder.ParentMaterial = MaterialTemplate;
 
-					FString MaterialPlaceholderSerialization = MutableMaterialPlaceholder.GetSerialization();
-					FMutableMaterialPlaceholder** FoundMaterialPlaceholder = ReuseMaterialCache.Find(MaterialPlaceholderSerialization);
+				FString MaterialPlaceholderSerialization = MutableMaterialPlaceholder.GetSerialization();
+				FMutableMaterialPlaceholder** FoundMaterialPlaceholder = ReuseMaterialCache.Find(MaterialPlaceholderSerialization);
 
-					if (FoundMaterialPlaceholder)
+				if (FoundMaterialPlaceholder)
+				{
+					int32 MatIndex = (*FoundMaterialPlaceholder)->MatIndex;
+					check(MatIndex >= 0);
+					int32 LODMaterialIndex = Helper_GetLODInfoArray(SkeletalMesh)[LODIndex].LODMaterialMap.Add(MatIndex);
+					Helper_GetLODRenderSections(SkeletalMesh, LODIndex)[SurfaceIndex].MaterialIndex = LODMaterialIndex;
+				}
+				else
+				{
+					if (MutableMaterialPlaceholder.ParentMaterial)
 					{
-						int32 MatIndex = (*FoundMaterialPlaceholder)->MatIndex;
-						check(MatIndex >= 0);
-						int32 LODMaterialIndex = Helper_GetLODInfoArray(SkeletalMesh)[MeshLODIndex].LODMaterialMap.Add(MatIndex);
-						Helper_GetLODRenderSections(SkeletalMesh, MeshLODIndex)[SurfaceIndex].MaterialIndex = LODMaterialIndex;
+						check(MaterialPlaceholderSerialization != FString("null"));
+						ReuseMaterialCache.Add(MaterialPlaceholderSerialization, &MutableMaterialPlaceholder);
 					}
-					else
-					{
-						if (MutableMaterialPlaceholder.ParentMaterial)
-						{
-							check(MaterialPlaceholderSerialization != FString("null"));
-							ReuseMaterialCache.Add(MaterialPlaceholderSerialization, &MutableMaterialPlaceholder);
-						}
 
-						FGeneratedMaterial Material;
-						UMaterialInstanceDynamic* MaterialInstance = nullptr;
+					FGeneratedMaterial Material;
+					UMaterialInstanceDynamic* MaterialInstance = nullptr;
 
-						if (MaterialTemplate)
-						{
-							MUTABLE_CPUPROFILER_SCOPE(CreateMaterial);
+					if (MaterialTemplate)
+					{
+						MUTABLE_CPUPROFILER_SCOPE(BuildMaterials_CreateMaterial);
 
-							UMaterialInterface* ActualMaterialInterface = MaterialTemplate;
+						UMaterialInterface* ActualMaterialInterface = MaterialTemplate;
 
-							if (MutableMaterialPlaceholder.Params.Num() != 0)
-							{
-								MaterialInstance = UMaterialInstanceDynamic::Create(MaterialTemplate, GetTransientPackage());
-								ActualMaterialInterface = MaterialInstance;
-							}
+						if (MutableMaterialPlaceholder.Params.Num() != 0)
+						{
+							MaterialInstance = UMaterialInstanceDynamic::Create(MaterialTemplate, GetTransientPackage());
+							ActualMaterialInterface = MaterialInstance;
+						}
 
-							if (SkeletalMesh)
-							{
-								const int32 MatIndex = SkeletalMesh->GetMaterials().Num();
-								MutableMaterialPlaceholder.MatIndex = MatIndex;
+						if (SkeletalMesh)
+						{
+							const int32 MatIndex = SkeletalMesh->GetMaterials().Num();
+							MutableMaterialPlaceholder.MatIndex = MatIndex;
 
-								// Set up SkeletalMaterial data
-								FSkeletalMaterial& SkeletalMaterial = SkeletalMesh->GetMaterials().Add_GetRef(ActualMaterialInterface);
-								SkeletalMaterial.MaterialSlotName = CustomizableObject->ReferencedMaterialSlotNames[Surface.MaterialIndex];
-								SetMeshUVChannelDensity(SkeletalMaterial.UVChannelData, RefSkeletalMeshData->Settings.DefaultUVChannelDensity);
+							// Set up SkeletalMaterial data
+							FSkeletalMaterial& SkeletalMaterial = SkeletalMesh->GetMaterials().Add_GetRef(ActualMaterialInterface);
+							SkeletalMaterial.MaterialSlotName = CustomizableObject->ReferencedMaterialSlotNames[Surface.MaterialIndex];
+							SetMeshUVChannelDensity(SkeletalMaterial.UVChannelData, RefSkeletalMeshData->Settings.DefaultUVChannelDensity);
 
-								const int32 LODMaterialIndex = Helper_GetLODInfoArray(SkeletalMesh)[MeshLODIndex].LODMaterialMap.Add(MatIndex);
-								Helper_GetLODRenderSections(SkeletalMesh, MeshLODIndex)[SurfaceIndex].MaterialIndex = LODMaterialIndex;
-							}
+							const int32 LODMaterialIndex = Helper_GetLODInfoArray(SkeletalMesh)[LODIndex].LODMaterialMap.Add(MatIndex);
+							Helper_GetLODRenderSections(SkeletalMesh, LODIndex)[SurfaceIndex].MaterialIndex = LODMaterialIndex;
+						}
 
-							if (MaterialInstance)
+						if (MaterialInstance)
+						{
+							for (const FMutableMaterialPlaceholder::FMutableMaterialPlaceHolderParam& Param : MutableMaterialPlaceholder.Params)
 							{
-								for (const FMutableMaterialPlaceholder::FMutableMaterialPlaceHolderParam& Param : MutableMaterialPlaceholder.Params)
+								switch (Param.Type)
 								{
-									switch (Param.Type)
+								case FMutableMaterialPlaceholder::EPlaceHolderParamType::Vector:
+									if (Param.LayerIndex < 0)
 									{
-									case FMutableMaterialPlaceholder::EPlaceHolderParamType::Vector:
-										if (Param.LayerIndex < 0)
-										{
-											MaterialInstance->SetVectorParameterValue(Param.ParamName, Param.Vector);
-										}
-										else
-										{
-											FMaterialParameterInfo ParameterInfo = FMaterialParameterInfo(Param.ParamName, EMaterialParameterAssociation::LayerParameter, Param.LayerIndex);
-											MaterialInstance->SetVectorParameterValueByInfo(ParameterInfo, Param.Vector);
-										}
+										MaterialInstance->SetVectorParameterValue(Param.ParamName, Param.Vector);
+									}
+									else
+									{
+										FMaterialParameterInfo ParameterInfo = FMaterialParameterInfo(Param.ParamName, EMaterialParameterAssociation::LayerParameter, Param.LayerIndex);
+										MaterialInstance->SetVectorParameterValueByInfo(ParameterInfo, Param.Vector);
+									}
 
-										break;
+									break;
 
-									case FMutableMaterialPlaceholder::EPlaceHolderParamType::Scalar:
-										if (Param.LayerIndex < 0)
-										{
-											MaterialInstance->SetScalarParameterValue(FName(Param.ParamName), Param.Scalar);
-										}
-										else
-										{
-											FMaterialParameterInfo ParameterInfo = FMaterialParameterInfo(Param.ParamName, EMaterialParameterAssociation::LayerParameter, Param.LayerIndex);
-											MaterialInstance->SetScalarParameterValueByInfo(ParameterInfo, Param.Scalar);
-										}
+								case FMutableMaterialPlaceholder::EPlaceHolderParamType::Scalar:
+									if (Param.LayerIndex < 0)
+									{
+										MaterialInstance->SetScalarParameterValue(FName(Param.ParamName), Param.Scalar);
+									}
+									else
+									{
+										FMaterialParameterInfo ParameterInfo = FMaterialParameterInfo(Param.ParamName, EMaterialParameterAssociation::LayerParameter, Param.LayerIndex);
+										MaterialInstance->SetScalarParameterValueByInfo(ParameterInfo, Param.Scalar);
+									}
 
-										break;
+									break;
 
-									case FMutableMaterialPlaceholder::EPlaceHolderParamType::Texture:
-										if (Param.LayerIndex < 0)
-										{
-											MaterialInstance->SetTextureParameterValue(Param.ParamName, Param.Texture.Texture);
-										}
-										else
-										{
-											FMaterialParameterInfo ParameterInfo = FMaterialParameterInfo(Param.ParamName, EMaterialParameterAssociation::LayerParameter, Param.LayerIndex);
-											MaterialInstance->SetTextureParameterValueByInfo(ParameterInfo, Param.Texture.Texture);
-										}
+								case FMutableMaterialPlaceholder::EPlaceHolderParamType::Texture:
+									if (Param.LayerIndex < 0)
+									{
+										MaterialInstance->SetTextureParameterValue(Param.ParamName, Param.Texture.Texture);
+									}
+									else
+									{
+										FMaterialParameterInfo ParameterInfo = FMaterialParameterInfo(Param.ParamName, EMaterialParameterAssociation::LayerParameter, Param.LayerIndex);
+										MaterialInstance->SetTextureParameterValueByInfo(ParameterInfo, Param.Texture.Texture);
+									}
 
-										Material.Textures.Add(Param.Texture);
+									Material.Textures.Add(Param.Texture);
 
-										break;
-									}
+									break;
 								}
 							}
-
-							GeneratedMaterials.Add(Material);
 						}
+
+						GeneratedMaterials.Add(Material);
 					}
 				}
 			}
@@ -5151,7 +5415,7 @@ void UCustomizableInstancePrivateData::BuildMaterials(const TSharedPtr<FMutableO
 	}
 
 	{
-		MUTABLE_CPUPROFILER_SCOPE(Exchange);
+		MUTABLE_CPUPROFILER_SCOPE(BuildMaterials_Exchange);
 
 		FCustomizableObjectSystemPrivate* CustomizableObjectSystem = UCustomizableObjectSystem::GetInstance()->GetPrivate();
 		TexturesToRelease.Empty();
@@ -5392,7 +5656,13 @@ void UCustomizableObjectInstance::SetMinSquareDistToPlayer(float NewValue)
 
 void UCustomizableObjectInstance::SetMinMaxLODToLoad(int32 NewMinLOD, int32 NewMaxLOD, bool bLimitLODUpgrades)
 {
-	GetPrivate()->SetMinMaxLODToLoad(this, NewMinLOD, NewMaxLOD, bLimitLODUpgrades);
+	SetRequestedLODs(NewMinLOD, NewMaxLOD, Descriptor.RequestedLODLevels);
+}
+
+
+int32 UCustomizableObjectInstance::GetNumComponents() const
+{
+	return GetCustomizableObject() ? GetCustomizableObject()->GetComponentCount() : 0;
 }
 
 
@@ -5408,33 +5678,121 @@ int32 UCustomizableObjectInstance::GetMaxLODToLoad() const
 }
 
 
-void UCustomizableObjectInstance::SetUseCurrentMinLODAsBaseLOD(bool bIsBaseLOD)
+int32 UCustomizableObjectInstance::GetNumLODsAvailable() const
 {
-	bIsBaseLOD ? GetPrivate()->SetCOInstanceFlags(ReduceLODs) : GetPrivate()->ClearCOInstanceFlags(ReduceLODs);
+	return GetPrivate()->GetNumLODsAvailable();
 }
 
 
-bool UCustomizableObjectInstance::GetUseCurrentMinLODAsBaseLOD() const
+int32 UCustomizableObjectInstance::GetCurrentMinLOD() const
 {
-	return GetPrivate()->HasCOInstanceFlags(ReduceLODs);
+	return CurrentMinLOD;
 }
 
 
-int32 UCustomizableObjectInstance::GetNumLODsAvailable() const
+int32 UCustomizableObjectInstance::GetCurrentMaxLOD() const
 {
-	return GetPrivate()->GetNumLODsAvailable();
+	return CurrentMaxLOD;
 }
 
 
-int32 UCustomizableObjectInstance::GetCurrentMinLOD() const
+void UCustomizableObjectInstance::SetRequestedLODs(int32 InMinLOD, int32 InMaxLOD, const TArray<uint16>& InRequestedLODsPerComponent)
 {
-	return CurrentMinLOD;
+	check(PrivateData);
+	if (!PrivateData->HasCOInstanceFlags(LODsStreamingEnabled))
+	{
+		return;
+	}
+
+	// Clamp Min LOD
+	InMinLOD = FMath::Min(FMath::Max(InMinLOD, static_cast<int32>(PrivateData->FirstLODAvailable)), PrivateData->FirstLODAvailable + PrivateData->NumMaxLODsToStream);
+
+	// Clamp Max LOD
+	InMaxLOD = FMath::Max(FMath::Min(InMaxLOD, PrivateData->NumLODsAvailable - 1), InMinLOD);
+
+	const bool bMinMaxLODChanged = Descriptor.MaxLOD != InMaxLOD || Descriptor.MinLOD != InMinLOD;
+
+	const bool bIsDowngradeLODUpdate = GetCurrentMinLOD() >= 0 && InMinLOD > GetCurrentMinLOD();
+	PrivateData->SetCOInstanceFlags(bIsDowngradeLODUpdate ? PendingLODsDowngrade : ECONone);
+
+	// Save the new LODs
+	Descriptor.MinLOD = InMinLOD;
+	Descriptor.MaxLOD = InMaxLOD;
+
+	bool bUpdateRequestedLODs = false;
+	if (UCustomizableObjectSystem::GetInstance()->IsOnlyGenerateRequestedLODsEnabled())
+	{
+		// TODO Pere: With Requested LODs We could skip updates by increasing the lower limit of LODs mutable can generate if only MinLOD have changed.
+
+		const int32 ComponentCount = GetNumComponents();
+		const TArray<uint16>& GeneratedLODsPerComponent = DescriptorRuntimeHash.GetRequestedLODs();
+
+		const bool bIgnoreGeneratedLODs = DescriptorRuntimeHash != UpdateDescriptorRuntimeHash || GeneratedLODsPerComponent.Num() != ComponentCount;
+
+		if (bMinMaxLODChanged || bIgnoreGeneratedLODs || Descriptor.GetRequestedLODLevels() != InRequestedLODsPerComponent)
+		{
+			check(InRequestedLODsPerComponent.Num() == ComponentCount);
+
+			bUpdateRequestedLODs = bIgnoreGeneratedLODs;
+			for (int32 ComponentIndex = 0; ComponentIndex < ComponentCount; ++ComponentIndex)
+			{
+				// Generate at least the MaxLOD
+				int32 RequestedLODs = 1 << Descriptor.MaxLOD;
+
+				for (int32 LODIndex = 0; LODIndex < Descriptor.MaxLOD; ++LODIndex)
+				{
+					// Add previously generated LODs if the descriptor have not changed
+					if (!bIgnoreGeneratedLODs && GeneratedLODsPerComponent[ComponentIndex] & (1 << LODIndex))
+					{
+						RequestedLODs |= (1 << LODIndex);
+					}
+
+					// Add new requested LODs that fall within the range
+					if (InRequestedLODsPerComponent[ComponentIndex] & (1 << LODIndex))
+					{
+						if (LODIndex >= Descriptor.MinLOD)
+						{
+							RequestedLODs |= (1 << LODIndex);
+						}
+						else
+						{
+							// RequestedLOD is out of range, generate the MinLOD instead;
+							RequestedLODs |= (1 << Descriptor.MinLOD);
+						}
+					}
+				}
+
+				bUpdateRequestedLODs |= RequestedLODs != Descriptor.RequestedLODLevels[ComponentIndex];
+
+				// Save new RequestedLODs
+				Descriptor.RequestedLODLevels[ComponentIndex] = RequestedLODs;
+			}
+		}
+	}
+
+	if (bMinMaxLODChanged || bUpdateRequestedLODs)
+	{
+		UpdateDescriptorRuntimeHash.UpdateMinMaxLOD(InMinLOD, InMaxLOD);
+		UpdateDescriptorRuntimeHash.UpdateRequestedLODs(Descriptor.RequestedLODLevels);
+
+		// Update the Queued operation data or request a new update
+		const FMutableQueueElem* QueueElem = UCustomizableObjectSystem::GetInstance()->GetPrivate()->MutableOperationQueue.Get(this);
+		if (QueueElem && QueueElem->Operation && QueueElem->Operation->Type == FMutableOperation::EOperationType::Update)
+		{
+			QueueElem->Operation->InstanceDescriptorRuntimeHash.UpdateMinMaxLOD(InMinLOD, InMaxLOD);
+			QueueElem->Operation->InstanceDescriptorRuntimeHash.UpdateRequestedLODs(Descriptor.RequestedLODLevels);
+		}
+		else
+		{
+			PrivateData->SetCOInstanceFlags(PendingLODsUpdate);
+		}
+	}
 }
 
 
-int32 UCustomizableObjectInstance::GetCurrentMaxLOD() const
+const TArray<uint16>& UCustomizableObjectInstance::GetRequestedLODsPerComponent() const
 {
-	return CurrentMaxLOD;
+	return Descriptor.RequestedLODLevels;
 }
 
 
diff --git a/Engine/Plugins/Experimental/Mutable/Source/CustomizableObject/Private/MuCO/CustomizableObjectInstanceDescriptor.cpp b/Engine/Plugins/Experimental/Mutable/Source/CustomizableObject/Private/MuCO/CustomizableObjectInstanceDescriptor.cpp
--- a/Engine/Plugins/Experimental/Mutable/Source/CustomizableObject/Private/MuCO/CustomizableObjectInstanceDescriptor.cpp
+++ b/Engine/Plugins/Experimental/Mutable/Source/CustomizableObject/Private/MuCO/CustomizableObjectInstanceDescriptor.cpp
@@ -677,6 +677,9 @@ void FCustomizableObjectInstanceDescriptor::ReloadParameters()
 	}
 
 	SetState(FMath::Clamp(GetState(), 0, CustomizableObject->GetStateCount() - 1));
+	
+	MaxLOD = FMath::Max(0, CustomizableObject->GetNumLODs() - 1);
+	RequestedLODLevels.Init(0, CustomizableObject->GetComponentCount());
 
 	TArray<FCustomizableObjectBoolParameterValue> OldBoolParameters = BoolParameters;
 	TArray<FCustomizableObjectIntParameterValue> OldIntParameters = IntParameters;
@@ -1037,6 +1040,18 @@ void FCustomizableObjectInstanceDescriptor::SetMaxLod(int32 InMaxLOD)
 }
 
 
+void FCustomizableObjectInstanceDescriptor::SetRequestedLODLevels(const TArray<uint16>& InRequestedLODLevels)
+{
+	RequestedLODLevels = InRequestedLODLevels;
+}
+
+
+const TArray<uint16>& FCustomizableObjectInstanceDescriptor::GetRequestedLODLevels() const
+{
+	return RequestedLODLevels;
+}
+
+
 TArray<FCustomizableObjectBoolParameterValue>& FCustomizableObjectInstanceDescriptor::GetBoolParameters()
 {
 	return BoolParameters;
@@ -2367,14 +2382,38 @@ FDescriptorRuntimeHash::FDescriptorRuntimeHash(const FCustomizableObjectInstance
 {
 	MinLOD = Descriptor.MinLOD;
 	MaxLOD = Descriptor.MaxLOD;	
+
+	RequestedLODsPerComponent = Descriptor.RequestedLODLevels;
 }
 
 
 bool FDescriptorRuntimeHash::IsSubset(const FDescriptorRuntimeHash& Other) const
 {
-	return FDescriptorHash::operator==(Other) &&
-		MinLOD >= Other.MinLOD &&
-		MaxLOD <= Other.MaxLOD;
+	if (FDescriptorHash::operator!=(Other) || MinLOD < Other.MinLOD || MaxLOD != Other.MaxLOD)
+	{
+		return false;
+	}
+
+	if (RequestedLODsPerComponent == Other.RequestedLODsPerComponent)
+	{
+		return true;
+	}
+
+	for (int32 ComponentIndex = 0; ComponentIndex < RequestedLODsPerComponent.Num(); ++ComponentIndex)
+	{
+		int32 RequestedLODs = RequestedLODsPerComponent[ComponentIndex];
+		int32 OtherRequestedLODs = Other.RequestedLODsPerComponent[ComponentIndex];
+		for (uint16 LODIndex = MinLOD; LODIndex <= MaxLOD; ++LODIndex)
+		{
+			// To be a subset all bits set in RequestedLODs must be set in OtherRequestedLODs. OtherRequestedLODs can have additional bits set
+			if (RequestedLODs & (1 << LODIndex) && !(OtherRequestedLODs & (1 << LODIndex)))
+			{
+				return false;
+			}
+		}
+	}
+
+	return true;
 }
 
 
@@ -2395,3 +2434,15 @@ int32 FDescriptorRuntimeHash::GetMaxLOD() const
 {
 	return MaxLOD;
 }
+
+
+void FDescriptorRuntimeHash::UpdateRequestedLODs(const TArray<uint16>& InRequestedLODs)
+{
+	RequestedLODsPerComponent = InRequestedLODs;
+}
+
+
+const TArray<uint16>& FDescriptorRuntimeHash::GetRequestedLODs() const
+{
+	return RequestedLODsPerComponent;
+}
diff --git a/Engine/Plugins/Experimental/Mutable/Source/CustomizableObject/Private/MuCO/CustomizableObjectSystem.cpp b/Engine/Plugins/Experimental/Mutable/Source/CustomizableObject/Private/MuCO/CustomizableObjectSystem.cpp
--- a/Engine/Plugins/Experimental/Mutable/Source/CustomizableObject/Private/MuCO/CustomizableObjectSystem.cpp
+++ b/Engine/Plugins/Experimental/Mutable/Source/CustomizableObject/Private/MuCO/CustomizableObjectSystem.cpp
@@ -690,8 +690,22 @@ static FAutoConsoleVariableRef CVarEnableMutableReuseInstanceTextures(
 	ECVF_Default);
 
 
+bool FCustomizableObjectSystemPrivate::bEnableMutableReusePreviousUpdateData = false;
+static FAutoConsoleVariableRef CVarEnableMutableReusePreviousUpdateData(
+	TEXT("mutable.EnableMutableReusePreviousUpdateData"), FCustomizableObjectSystemPrivate::bEnableMutableReusePreviousUpdateData,
+	TEXT("If true, Mutable will try to reuse render sections from previous SkeletalMeshes. If false, all SkeletalMeshes will be build from scratch."),
+	ECVF_Default);
+
+
+bool FCustomizableObjectSystemPrivate::bEnableOnlyGenerateRequestedLODs = true;
+static FAutoConsoleVariableRef CVarEnableOnlyGenerateRequestedLODs(
+	TEXT("mutable.EnableOnlyGenerateRequestedLODs"), FCustomizableObjectSystemPrivate::bEnableOnlyGenerateRequestedLODs,
+	TEXT("If true, Only the RequestedLODLevels will be generated. If false, all LODs will be build."),
+	ECVF_Default);
+
+
 /** Update the given Instance Skeletal Meshes and call its callbacks. */
-void UpdateSkeletalMesh(UCustomizableObjectInstance* CustomizableObjectInstance)
+void UpdateSkeletalMesh(UCustomizableObjectInstance* CustomizableObjectInstance, const FDescriptorRuntimeHash& UpdatedDescriptorRuntimeHash)
 {
 	// This used to be CustomizableObjectInstance::UpdateSkeletalMesh_PostBeginUpdate3
 	{
@@ -721,7 +735,7 @@ void UpdateSkeletalMesh(UCustomizableObjectInstance* CustomizableObjectInstance)
 			{
 				MUTABLE_CPUPROFILER_SCOPE(UpdateSkeletalMesh_SetSkeletalMesh);
 
-				const bool bIsCreatingSkeletalMesh = CustomizableObjectInstancePrivateData->HasCOInstanceFlags(CreatingSkeletalMesh);
+				const bool bIsCreatingSkeletalMesh = CustomizableObjectInstancePrivateData->HasCOInstanceFlags(CreatingSkeletalMesh); //TODO MTBL-391: Review
 				CustomizableSkeletalComponent->SetSkeletalMesh(CustomizableObjectInstance->SkeletalMeshes[CustomizableSkeletalComponent->ComponentIndex], false, bIsCreatingSkeletalMesh);
 
 				if (CustomizableObjectInstancePrivateData->HasCOInstanceFlags(ReplacePhysicsAssets))
@@ -741,7 +755,8 @@ void UpdateSkeletalMesh(UCustomizableObjectInstance* CustomizableObjectInstance)
 		{
 			MUTABLE_CPUPROFILER_SCOPE(UpdateSkeletalMesh_UpdatedDelegate);
 
-			CustomizableObjectInstance->Updated(CustomizableObjectInstance->SkeletalMeshStatus == ESkeletalMeshState::Correct ? EUpdateResult::Success : EUpdateResult::Error);
+			CustomizableObjectInstance->Updated(CustomizableObjectInstance->SkeletalMeshStatus == ESkeletalMeshState::Correct ? EUpdateResult::Success : EUpdateResult::Error, UpdatedDescriptorRuntimeHash);
+
 
 #if WITH_EDITOR
 			CustomizableObjectInstance->InstanceUpdated = true;
@@ -811,7 +826,7 @@ void FCustomizableObjectSystemPrivate::InitUpdateSkeletalMesh(UCustomizableObjec
 		CurrentMutableOperation->Type == FMutableOperation::EOperationType::Update &&
 		UpdateDescriptorHash.IsSubset(CurrentMutableOperation->InstanceDescriptorRuntimeHash))
 	{
-		return; // The the requested update is equal to the running update.
+		return; // The requested update is equal to the running update.
 	}
 
 	if (const TObjectPtr<UDefaultImageProvider> DefaultImageProvider = UCustomizableObjectSystem::GetInstance()->DefaultImageProvider)
@@ -826,10 +841,11 @@ void FCustomizableObjectSystemPrivate::InitUpdateSkeletalMesh(UCustomizableObjec
 	if (UpdateDescriptorHash.IsSubset(Instance.GetDescriptorRuntimeHash()))
 	{
 		Instance.SkeletalMeshStatus = ESkeletalMeshState::Correct; // TODO GMT MTBL-1033 should not be here. Move to UCustomizableObjectInstance::Updated
-		UpdateSkeletalMesh(&Instance);
+		UpdateSkeletalMesh(&Instance, Instance.GetDescriptorRuntimeHash());
 	}
 	else
 	{
+		Instance.SkeletalMeshStatus = ESkeletalMeshState::AsyncUpdatePending;
 		MutableOperationQueue.Enqueue(FMutableQueueElem::Create(Operation, Priority, Instance.GetPrivate()->MinSquareDistFromComponentToPlayer));
 	}
 }
@@ -1227,6 +1243,22 @@ namespace impl
 			OperationData->CurrentMaxLOD = OperationData->NumLODsAvailable - 1;
 		}
 
+		if(!OperationData->RequestedLODs.IsEmpty())
+		{
+			// Initialize RequestedLODs to zero if not set
+			const int32 ComponentCount = Instance->GetComponentCount(OperationData->CurrentMinLOD);
+			OperationData->RequestedLODs.SetNumZeroed(ComponentCount);
+			
+			for (int32 ComponentIndex = 0; ComponentIndex < ComponentCount; ++ComponentIndex)
+			{
+				// Ensure we're generating at least one LOD
+				if (OperationData->RequestedLODs[ComponentIndex] == 0)
+				{
+					OperationData->RequestedLODs[ComponentIndex] |= (1 << OperationData->CurrentMaxLOD);
+				}
+			}
+		}
+
 		// Generate the mesh and gather all the required resource Ids
 		OperationData->InstanceUpdateData.LODs.SetNum(Instance->GetLODCount());
 		for (int32 MutableLODIndex = 0; MutableLODIndex < Instance->GetLODCount(); ++MutableLODIndex)
@@ -1241,6 +1273,10 @@ namespace impl
 			LOD.FirstComponent = OperationData->InstanceUpdateData.Components.Num();
 			LOD.ComponentCount = Instance->GetComponentCount(MutableLODIndex);
 
+			const FInstanceGeneratedData::FLOD& PrevUpdateLODData = OperationData->LastUpdateData.LODs.IsValidIndex(MutableLODIndex) ?
+				OperationData->LastUpdateData.LODs[MutableLODIndex] : FInstanceGeneratedData::FLOD();
+
+
 			for (int32 ComponentIndex = 0; ComponentIndex < LOD.ComponentCount; ++ComponentIndex)
 			{
 				OperationData->InstanceUpdateData.Components.Push(FInstanceUpdateData::FComponent());
@@ -1249,12 +1285,17 @@ namespace impl
 				Component.FirstSurface = OperationData->InstanceUpdateData.Surfaces.Num();
 				Component.SurfaceCount = 0;
 
+				const FInstanceGeneratedData::FComponent& PrevUpdateComponent = OperationData->LastUpdateData.Components.IsValidIndex(PrevUpdateLODData.FirstComponent + ComponentIndex) ?
+					OperationData->LastUpdateData.Components[PrevUpdateLODData.FirstComponent +  ComponentIndex] : FInstanceGeneratedData::FComponent();
+
+				const bool bGenerateLOD = OperationData->RequestedLODs.IsValidIndex(Component.Id) ? (OperationData->RequestedLODs[Component.Id] & (1 << MutableLODIndex)) != 0 : true;
+
 				// Mesh
 				if (Instance->GetMeshCount(MutableLODIndex, ComponentIndex) > 0)
 				{
 					MUTABLE_CPUPROFILER_SCOPE(GetMesh);
 
-					mu::RESOURCE_ID MeshId = Instance->GetMeshId(MutableLODIndex, ComponentIndex, 0);
+					Component.MeshID = Instance->GetMeshId(MutableLODIndex, ComponentIndex, 0);
 
 					// Mesh cache is not enabled yet.
 					//auto CachedMesh = Cache.Meshes.Find(meshId);
@@ -1264,20 +1305,30 @@ namespace impl
 					//	INC_DWORD_STAT(STAT_MutableNumCachedSkeletalMeshes);
 					//}
 
-					Component.Mesh = System->GetMesh(OperationData->InstanceID, MeshId);
+					// Check if we can use data from the currently generated SkeletalMesh
+					Component.bReuseMesh = OperationData->bCanReuseGeneratedData && PrevUpdateComponent.bGenerated ? PrevUpdateComponent.MeshID == Component.MeshID : false;
+
+					if(bGenerateLOD && !Component.bReuseMesh)
+					{
+						Component.Mesh = System->GetMesh(OperationData->InstanceID, Component.MeshID);
+					}
 				}
 
-				if (!Component.Mesh)
+				if (!Component.Mesh && !Component.bReuseMesh)
 				{
 					continue;
 				}
 
+				Component.bGenerated = true;
+
+				const int32 SurfaceCount = Component.Mesh ? Component.Mesh->GetSurfaceCount() : PrevUpdateComponent.SurfaceCount;
+
 				// Materials and images
-				for (int32 MeshSurfaceIndex = 0; MeshSurfaceIndex < Component.Mesh->GetSurfaceCount(); ++MeshSurfaceIndex)
+				for (int32 MeshSurfaceIndex = 0; MeshSurfaceIndex < SurfaceCount; ++MeshSurfaceIndex)
 				{
-					uint32 SurfaceId = Component.Mesh->GetSurfaceId(MeshSurfaceIndex);
+					uint32 SurfaceId = Component.Mesh ? Component.Mesh->GetSurfaceId(MeshSurfaceIndex) : OperationData->LastUpdateData.SurfaceIds[PrevUpdateComponent.FirstSurface + MeshSurfaceIndex];
 					int32 InstanceSurfaceIndex = Instance->FindSurfaceById(MutableLODIndex, ComponentIndex, SurfaceId);
-					check(Component.Mesh->GetVertexCount() == 0 || InstanceSurfaceIndex >= 0);
+					check(Component.bReuseMesh || Component.Mesh->GetVertexCount() == 0 || InstanceSurfaceIndex >= 0);
 
 					if (InstanceSurfaceIndex >= 0)
 					{
@@ -1665,16 +1716,16 @@ namespace impl
 			return;
 		}
 
+		FCustomizableObjectSystemPrivate * CustomizableObjectSystemPrivateData = System->GetPrivate();
+		check(CustomizableObjectSystemPrivateData != nullptr);
+
 		// Actual work
-		UpdateSkeletalMesh(CustomizableObjectInstance);
+		// TODO MTBL-391: Review This hotfix
+		UpdateSkeletalMesh(CustomizableObjectInstance, CustomizableObjectSystemPrivateData->CurrentMutableOperation->InstanceDescriptorRuntimeHash);
 
 		// TODO: T2927
 		if (LogBenchmarkUtil::isLoggingActive())
 		{
-			// Log the time stat
-			FCustomizableObjectSystemPrivate* CustomizableObjectSystemPrivateData = System->GetPrivate();
-			check(CustomizableObjectSystemPrivateData != nullptr);
-
 			double deltaSeconds = FPlatformTime::Seconds() - CustomizableObjectSystemPrivateData->CurrentMutableOperation->StartUpdateTime;
 			int32 deltaMs = int32(deltaSeconds * 1000);
 			int64 streamingCache = CustomizableObjectSystemPrivateData->LastStreamingMemorySize / 1024;
@@ -2039,67 +2090,65 @@ namespace impl
 			return;
 		}
 
-		UCustomizableInstancePrivateData* ObjectInstancePrivateData = Operation->CustomizableObjectInstance->GetPrivate();
-		check(ObjectInstancePrivateData != nullptr);
-		if (ObjectInstancePrivateData->HasCOInstanceFlags(PendingLODsUpdate))
+		TObjectPtr<UCustomizableObjectInstance> CandidateInstance = Operation->CustomizableObjectInstance.Get();
+		
+		UCustomizableInstancePrivateData* CandidateInstancePrivateData = CandidateInstance->GetPrivate();
+		check(CandidateInstancePrivateData);
+
+		if (CandidateInstancePrivateData->HasCOInstanceFlags(PendingLODsUpdate))
 		{
-			ObjectInstancePrivateData->ClearCOInstanceFlags(PendingLODsUpdate);
+			CandidateInstancePrivateData->ClearCOInstanceFlags(PendingLODsUpdate);
 			// TODO: Is anything needed for this now?
 			//Operation->CustomizableObjectInstance->ReleaseMutableInstanceId(); // To make mutable regenerate the LODs even if the instance parameters have not changed
 		}
 
-		float LastMinSquareDistFromComponentToPlayer = ObjectInstancePrivateData->LastMinSquareDistFromComponentToPlayer;
-
-		UCustomizableObjectInstance* CandidateInstance = Operation->CustomizableObjectInstance.Get();
-
 		bool bCancel = false;
-		UCustomizableObject* CustomizableObject = nullptr;
-		mu::Ptr<const mu::Parameters> Parameters;
 
 		// If the object is locked (for instance, compiling) we skip any instance update.
-		if (!CandidateInstance) // TODO DRB: Shouldn't this be impossible?  This is Operation->CustomizableObjectInstance, and we returned above if it wasn't valid.
+		TObjectPtr<UCustomizableObject> CustomizableObject = CandidateInstance->GetCustomizableObject();
+		if (!CustomizableObject)
 		{
 			bCancel = true;
 		}
 		else
 		{
-			CustomizableObject = CandidateInstance->GetCustomizableObject();
-			if (!CustomizableObject)
+			check(CustomizableObject->GetPrivate());
+			if (CustomizableObject->GetPrivate()->bLocked)
 			{
 				bCancel = true;
 			}
-			else
-			{
-				check(CustomizableObject->GetPrivate() != nullptr);
-				if (CustomizableObject->GetPrivate()->bLocked)
-				{
-					bCancel = true;
-				}
-			}
+		}
 
-			// Only update resources if the instance is in range (it could have got far from the player since the task was queued)
-			check(System->CurrentInstanceLODManagement != nullptr);
-			if (System->CurrentInstanceLODManagement->IsOnlyUpdateCloseCustomizableObjectsEnabled()
-				&& LastMinSquareDistFromComponentToPlayer > FMath::Square(System->CurrentInstanceLODManagement->GetOnlyUpdateCloseCustomizableObjectsDist())
-				&& LastMinSquareDistFromComponentToPlayer != FLT_MAX // This means it is the first frame so it has to be updated
-			   )
-			{
-				bCancel = true;
-			}
+		// Only update resources if the instance is in range (it could have got far from the player since the task was queued)
+		check(System->CurrentInstanceLODManagement != nullptr);
+		if (System->CurrentInstanceLODManagement->IsOnlyUpdateCloseCustomizableObjectsEnabled()
+			&& CandidateInstancePrivateData->LastMinSquareDistFromComponentToPlayer > FMath::Square(System->CurrentInstanceLODManagement->GetOnlyUpdateCloseCustomizableObjectsDist())
+			&& CandidateInstancePrivateData->LastMinSquareDistFromComponentToPlayer != FLT_MAX // This means it is the first frame so it has to be updated
+		   )
+		{
+			bCancel = true;
+		}
 
-			Parameters = Operation->GetParameters();
-			if (!Parameters)
-			{
-				bCancel = true;
-			}
+		// Skip update, the requested update is equal to the running update.
+		if (Operation->InstanceDescriptorRuntimeHash.IsSubset(CandidateInstance->GetDescriptorRuntimeHash()))
+		{
+			UpdateSkeletalMesh(CandidateInstance, CandidateInstance->GetDescriptorRuntimeHash());
+			bCancel = true;
+		}
+
+		mu::Ptr<const mu::Parameters> Parameters = Operation->GetParameters();
+		if (!Parameters)
+		{
+			bCancel = true;
 		}
 
 		if (bCancel)
 		{
-			if (Operation->CustomizableObjectInstance.IsValid()) // TODO DRB: If this was true at the top of this functoin, we returned.  Can it have become invalid in between?  Any reason we're not using CandidateInstance here?
+			if (CandidateInstancePrivateData) 
 			{
-				ObjectInstancePrivateData->ClearCOInstanceFlags(Updating);
+				CandidateInstancePrivateData->ClearCOInstanceFlags(Updating);
 			}
+
 			System->ClearCurrentMutableOperation();
 			return;
 		}
@@ -2114,9 +2163,6 @@ namespace impl
 
 		SystemPrivateData->CurrentInstanceBeingUpdated = CandidateInstance;
 
-		UCustomizableInstancePrivateData* CandidateInstancePrivateData = CandidateInstance->GetPrivate();
-		check(CandidateInstancePrivateData != nullptr);
-
 		// Prepare streaming for the current customizable object
 		check(SystemPrivateData->Streamer != nullptr);
 		SystemPrivateData->Streamer->PrepareStreamingForObject(CustomizableObject);
@@ -2172,6 +2218,8 @@ namespace impl
 		check(CurrentOperationData);
 		CurrentOperationData->TextureCoverageQueries_MutableThreadParams = CandidateInstancePrivateData->TextureCoverageQueries;
 		CurrentOperationData->TextureCoverageQueries_MutableThreadResults.Empty();
+		CurrentOperationData->bCanReuseGeneratedData = SystemPrivateData->bEnableMutableReusePreviousUpdateData;
+		CurrentOperationData->LastUpdateData = SystemPrivateData->bEnableMutableReusePreviousUpdateData ? CandidateInstancePrivateData->LastUpdateData : FInstanceGeneratedData();
 		CurrentOperationData->CurrentMinLOD = Operation->InstanceDescriptorRuntimeHash.GetMinLOD();
 		CurrentOperationData->CurrentMaxLOD = Operation->InstanceDescriptorRuntimeHash.GetMaxLOD();
 		CurrentOperationData->bNeverStream = Operation->bNeverStream;
@@ -2180,8 +2228,12 @@ namespace impl
 		CurrentOperationData->InstanceID = bLiveUpdateMode ? CandidateInstancePrivateData->LiveUpdateModeInstanceID : 0;
 		CurrentOperationData->MipsToSkip = Operation->MipsToSkip;
 		CurrentOperationData->MutableParameters = Parameters;
-		check(Operation->CustomizableObjectInstance.IsValid());
-		CurrentOperationData->State = Operation->CustomizableObjectInstance->GetState();
+		CurrentOperationData->State = CandidateInstance->GetState();
+
+		if (SystemPrivateData->bEnableOnlyGenerateRequestedLODs && System->CurrentInstanceLODManagement->IsOnlyGenerateRequestedLODLevelsEnabled() && !CandidateInstancePrivateData->HasCOInstanceFlags(ForceGenerateAllLODs))
+		{
+			CurrentOperationData->RequestedLODs = Operation->InstanceDescriptorRuntimeHash.GetRequestedLODs();
+		}
 
 		FGraphEventRef Mutable_GetMeshTask;
 		{
@@ -2377,14 +2429,50 @@ bool UCustomizableObjectSystem::Tick(float DeltaTime)
 
 	MUTABLE_CPUPROFILER_SCOPE(TickCustomizableObjectSystem)
 
-	// Reset the instance relevancy
-	// \TODO: Review
-	// TODO: This should be done only when requiring a new job. And for the current operation instance, in case it needs to be cancelled.
-	CurrentInstanceLODManagement->UpdateInstanceDistsAndLODs();
-	
 	// Get a new operation if we aren't working on one
 	if (!Private->CurrentMutableOperation)
 	{
+		// Reset the instance relevancy
+		// \TODO: Review
+		// TODO: This should be done only when requiring a new job. And for the current operation instance, in case it needs to be cancelled.
+		CurrentInstanceLODManagement->UpdateInstanceDistsAndLODs();
+
+		for (TObjectIterator<UCustomizableObjectInstance> CustomizableObjectInstance; CustomizableObjectInstance; ++CustomizableObjectInstance)
+		{
+			if (IsValidChecked(*CustomizableObjectInstance) && CustomizableObjectInstance->GetPrivate())
+			{
+				UCustomizableInstancePrivateData* ObjectInstancePrivateData = CustomizableObjectInstance->GetPrivate();
+				check(ObjectInstancePrivateData != nullptr);
+
+				if (ObjectInstancePrivateData->HasCOInstanceFlags(UsedByComponentInPlay))
+				{
+					ObjectInstancePrivateData->TickUpdateCloseCustomizableObjects(**CustomizableObjectInstance);
+				}
+				else if (ObjectInstancePrivateData->HasCOInstanceFlags(UsedByComponent))
+				{
+					ObjectInstancePrivateData->UpdateInstanceIfNotGenerated(**CustomizableObjectInstance);
+				}
+
+				ObjectInstancePrivateData->ClearCOInstanceFlags((ECOInstanceFlags)(UsedByComponent | UsedByComponentInPlay | PendingLODsUpdate)); // TODO MTBL-391: Makes no sense to clear it here, what if an update is requested before we set it back to true
+			}
+		}
+
+		// Update the queue. TODO: This should be done only when requiring a new job. And for the current operation instance, in case it needs to be cancelled.
+		{
+			Private->MutableOperationQueue.ChangePriorities();
+
+			for (TObjectIterator<UCustomizableObjectInstance> CustomizableObjectInstance; CustomizableObjectInstance; ++CustomizableObjectInstance)
+			{
+				if (IsValidChecked(*CustomizableObjectInstance) && CustomizableObjectInstance->GetPrivate())
+				{
+					CustomizableObjectInstance->GetPrivate()->LastMinSquareDistFromComponentToPlayer = CustomizableObjectInstance->GetPrivate()->MinSquareDistFromComponentToPlayer;
+					CustomizableObjectInstance->GetPrivate()->MinSquareDistFromComponentToPlayer = FLT_MAX;
+				}
+			}
+
+			Private->MutableOperationQueue.Sort();
+		}
+
 		// Update the streaming limit if it has changed. It is safe to do this now.
 		Private->UpdateStreamingLimit();
 
@@ -2406,41 +2494,7 @@ bool UCustomizableObjectSystem::Tick(float DeltaTime)
 		AdvanceCurrentOperation();
 	}
 
-	for (TObjectIterator<UCustomizableObjectInstance> CustomizableObjectInstance; CustomizableObjectInstance; ++CustomizableObjectInstance)
-	{
-		if (IsValidChecked(*CustomizableObjectInstance) && CustomizableObjectInstance->GetPrivate())
-		{
-			UCustomizableInstancePrivateData* ObjectInstancePrivateData = CustomizableObjectInstance->GetPrivate();
-			check(ObjectInstancePrivateData != nullptr);
-
-			if (ObjectInstancePrivateData->HasCOInstanceFlags(UsedByComponentInPlay))
-			{
-				ObjectInstancePrivateData->TickUpdateCloseCustomizableObjects(**CustomizableObjectInstance);
-			}
-			else if (ObjectInstancePrivateData->HasCOInstanceFlags(UsedByComponent))
-			{
-				ObjectInstancePrivateData->UpdateInstanceIfNotGenerated(**CustomizableObjectInstance);
-			}
-
-			ObjectInstancePrivateData->ClearCOInstanceFlags((ECOInstanceFlags)(UsedByComponent | UsedByComponentInPlay));
-		}
-	}
-
-	// Update the queue. TODO: This should be done only when requiring a new job. And for the current operation instance, in case it needs to be cancelled.
-	{
-		Private->MutableOperationQueue.ChangePriorities();
-
-		for (TObjectIterator<UCustomizableObjectInstance> CustomizableObjectInstance; CustomizableObjectInstance; ++CustomizableObjectInstance)
-		{
-			if (IsValidChecked(*CustomizableObjectInstance) && CustomizableObjectInstance->GetPrivate())
-			{
-				CustomizableObjectInstance->GetPrivate()->LastMinSquareDistFromComponentToPlayer = CustomizableObjectInstance->GetPrivate()->MinSquareDistFromComponentToPlayer;
-				CustomizableObjectInstance->GetPrivate()->MinSquareDistFromComponentToPlayer = FLT_MAX;
-			}
-		}
 
-		Private->MutableOperationQueue.Sort();
-	}
 
 	// TODO: T2927
 	if (LogBenchmarkUtil::isLoggingActive())
@@ -2591,7 +2645,8 @@ void FMutableOperation::MutableIsDisabledCase()
 		CustomizableObjectInstance->bEditorPropertyChanged = false;
 	}
 
-	CustomizableObjectInstance->Updated(EUpdateResult::Success);
+	// We must invalidate the current DescriptorRuntimeHash, since we're discarding everything.
+	CustomizableObjectInstance->Updated(EUpdateResult::Success, FDescriptorRuntimeHash());
 
 #if WITH_EDITOR
 	CustomizableObjectInstance->InstanceUpdated = true;
@@ -2643,6 +2698,7 @@ bool UCustomizableObjectSystem::IsSupport16BitBoneIndexEnabled() const
 	return Private->bSupport16BitBoneIndex;
 }
 
+
 bool UCustomizableObjectSystem::IsProgressiveMipStreamingEnabled() const
 {
 	check(Private != nullptr);
@@ -2650,6 +2706,13 @@ bool UCustomizableObjectSystem::IsProgressiveMipStreamingEnabled() const
 }
 
 
+bool UCustomizableObjectSystem::IsOnlyGenerateRequestedLODsEnabled() const
+{
+	check(Private != nullptr);
+	return Private->bEnableOnlyGenerateRequestedLODs;
+}
+
+
 void UCustomizableObjectSystem::AddUncompiledCOWarning(UCustomizableObject* InObject)
 {
 	if (!InObject)
diff --git a/Engine/Plugins/Experimental/Mutable/Source/CustomizableObject/Private/MuCO/CustomizableObjectSystemPrivate.h b/Engine/Plugins/Experimental/Mutable/Source/CustomizableObject/Private/MuCO/CustomizableObjectSystemPrivate.h
--- a/Engine/Plugins/Experimental/Mutable/Source/CustomizableObject/Private/MuCO/CustomizableObjectSystemPrivate.h
+++ b/Engine/Plugins/Experimental/Mutable/Source/CustomizableObject/Private/MuCO/CustomizableObjectSystemPrivate.h
@@ -263,6 +263,44 @@ struct FMutableTask
 };
 
 
+// Mutable data generated during the last update of an instance.
+struct FInstanceGeneratedData
+{
+	struct FComponent
+	{
+		uint16 ComponentId = 0;
+
+		/** True if it can be reused */
+		bool bGenerated = false;
+
+		mu::RESOURCE_ID MeshID;
+
+		/** Range in the Surfaces array */
+		uint16 FirstSurface = 0;
+		uint16 SurfaceCount = 0;
+	};
+
+	struct FLOD
+	{
+		/** Range in the Components array */
+		uint16 FirstComponent = 0;
+		uint16 ComponentCount = 0;
+	};
+
+	TArray<FLOD> LODs;
+	TArray<FComponent> Components;
+	TArray<uint32> SurfaceIds;
+
+	/** Clear data, called upon failing to generate a mesh and after recompiling the CO */
+	void Clear()
+	{
+		LODs.Empty();
+		Components.Empty();
+		SurfaceIds.Empty();
+	}
+};
+
+
 // Mutable data generated during the update steps.
 // We keep it from begin to end update, and it is used in several steps.
 // TODO: Flatten this structure into 4 arrays and use indices in fixed-size structs instead of subarrays
@@ -313,6 +351,13 @@ struct FInstanceUpdateData
 	{
 		uint16 Id = 0;
 		
+		// True if the Mesh is valid or if we're reusing a component
+		bool bGenerated = false;
+
+		// Reuse component from a previously generated SkeletalMesh
+		bool bReuseMesh = false;
+
+		mu::RESOURCE_ID MeshID;
 		mu::MeshPtrConst Mesh;
 
 		/** Range in the Surfaces array */
@@ -415,6 +460,9 @@ struct FPendingTextureCoverageQuery
 /** Runtime data used during a mutable instance update */
 struct FMutableOperationData
 {
+	bool bCanReuseGeneratedData = false;
+	FInstanceGeneratedData LastUpdateData;
+
 	FInstanceUpdateData InstanceUpdateData;
 	FParameterDecorationsUpdateData ParametersUpdateData;
 	TArray<int> RelevantParametersInProgress;
@@ -432,6 +480,8 @@ struct FMutableOperationData
 	int32 CurrentMaxLOD = 0;
 	int32 NumLODsAvailable = 0;
 
+	TArray<uint16> RequestedLODs;
+
 	TMap<FString, FTextureCoverageQueryData> TextureCoverageQueries_MutableThreadParams;
 	TMap<FString, FTextureCoverageQueryData> TextureCoverageQueries_MutableThreadResults;
 
@@ -498,6 +548,8 @@ public:
 	static int32 EnableMutableLiveUpdate;
 	static int32 EnableReuseInstanceTextures;
 	static int32 EnableMutableAnimInfoDebugging;
+	static bool bEnableMutableReusePreviousUpdateData;
+	static bool bEnableOnlyGenerateRequestedLODs;
 
 	/** */
 	inline void AddGameThreadTask(const FMutableTask& Task)
diff --git a/Engine/Plugins/Experimental/Mutable/Source/CustomizableObject/Private/MuCO/CustomizableSkeletalComponent.cpp b/Engine/Plugins/Experimental/Mutable/Source/CustomizableObject/Private/MuCO/CustomizableSkeletalComponent.cpp
--- a/Engine/Plugins/Experimental/Mutable/Source/CustomizableObject/Private/MuCO/CustomizableSkeletalComponent.cpp
+++ b/Engine/Plugins/Experimental/Mutable/Source/CustomizableObject/Private/MuCO/CustomizableSkeletalComponent.cpp
@@ -184,7 +184,7 @@ void UCustomizableSkeletalComponent::EditorUpdateComponent()
 {
 	if (CustomizableObjectInstance)
 	{
-		CustomizableObjectInstance->GetPrivate()->SetCOInstanceFlags(UsedByComponent);
+		CustomizableObjectInstance->GetPrivate()->SetCOInstanceFlags((ECOInstanceFlags)(UsedByComponent | ForceGenerateAllLODs));
 
 		AActor* ParentActor = GetAttachmentRootActor();
 		if (ParentActor)
diff --git a/Engine/Plugins/Experimental/Mutable/Source/CustomizableObject/Private/MuCO/LogInformationUtil.cpp b/Engine/Plugins/Experimental/Mutable/Source/CustomizableObject/Private/MuCO/LogInformationUtil.cpp
--- a/Engine/Plugins/Experimental/Mutable/Source/CustomizableObject/Private/MuCO/LogInformationUtil.cpp
+++ b/Engine/Plugins/Experimental/Mutable/Source/CustomizableObject/Private/MuCO/LogInformationUtil.cpp
@@ -166,7 +166,6 @@ void LogInformationUtil::LogShowInstanceDataFull(const UCustomizableObjectInstan
 	MessageChunk += FString::Printf(TEXT("        bShouldUpdateLODs = %d\n"), InstanceFlags & ECOInstanceFlags::PendingLODsUpdate);
 	LogData += MessageChunk;
 
-	MessageChunk = FString::Printf(TEXT("        bShouldUpdateLODsSecondStage = %d\n"), InstanceFlags & ECOInstanceFlags::PendingLODsUpdateSecondStage);
 	MessageChunk += FString::Printf(TEXT("        bIsDowngradeLODUpdate = %d\n"), InstanceFlags & ECOInstanceFlags::PendingLODsDowngrade);
 	MessageChunk += FString::Printf(TEXT("        bIsUpdating = %d\n"), InstanceFlags & ECOInstanceFlags::Updating);
 	MessageChunk += FString::Printf(TEXT("        bIsCreatingSkeletalMesh = %d\n"), InstanceFlags & ECOInstanceFlags::CreatingSkeletalMesh);
diff --git a/Engine/Plugins/Experimental/Mutable/Source/CustomizableObject/Private/MuCO/UnrealConversionUtils.cpp b/Engine/Plugins/Experimental/Mutable/Source/CustomizableObject/Private/MuCO/UnrealConversionUtils.cpp
--- a/Engine/Plugins/Experimental/Mutable/Source/CustomizableObject/Private/MuCO/UnrealConversionUtils.cpp
+++ b/Engine/Plugins/Experimental/Mutable/Source/CustomizableObject/Private/MuCO/UnrealConversionUtils.cpp
@@ -105,7 +105,7 @@ namespace UnrealConversionUtils
 		const int32 SourceBoneCount = InSourceReferenceSkeleton.GetRawBoneNum();
 		const TArray<FTransform>& SourceRawMeshBonePose = InSourceReferenceSkeleton.GetRawRefBonePose();
 
-		MUTABLE_CPUPROFILER_SCOPE(BuildSkeletonData_BuildRefSkeleton);
+		MUTABLE_CPUPROFILER_SCOPE(BuildRefSkeleton);
 
 		// Build new RefSkeleton	
 		FReferenceSkeletonModifier RefSkeletonModifier(InRefSkeleton, InSkeleton);
@@ -138,25 +138,6 @@ namespace UnrealConversionUtils
 	}
 
 
-	void BuildSkeletalMeshElementDataAtLOD(const int32 MeshLODIndex,
-		mu::MeshPtrConst InMutableMesh, USkeletalMesh* OutSkeletalMesh)
-	{
-		int32 SurfaceCount = 0;
-		if (InMutableMesh)
-		{
-			SurfaceCount = InMutableMesh->GetSurfaceCount();
-		}
-
-		Helper_GetLODInfoArray(OutSkeletalMesh)[MeshLODIndex].LODMaterialMap.SetNum(1);
-		Helper_GetLODInfoArray(OutSkeletalMesh)[MeshLODIndex].LODMaterialMap[0] = 0;
-
-		for (int32 SurfaceIndex = 0; SurfaceIndex < SurfaceCount; ++SurfaceIndex)
-		{
-			new(Helper_GetLODRenderSections(OutSkeletalMesh, MeshLODIndex)) Helper_SkelMeshRenderSection();
-		}
-	}
-
-
 	void SetupRenderSections(
 		const USkeletalMesh* OutSkeletalMesh,
 		const mu::MeshPtrConst InMutableMesh,
@@ -181,7 +162,7 @@ namespace UnrealConversionUtils
 		const int32 SurfaceCount = InMutableMesh->GetSurfaceCount();
 		for (int32 SurfaceIndex = 0; SurfaceIndex < SurfaceCount; ++SurfaceIndex)
 		{
-			MUTABLE_CPUPROFILER_SCOPE(UpdateSkeletalMesh_SurfaceLoop);
+			MUTABLE_CPUPROFILER_SCOPE(SetupRenderSections);
 
 			int32 FirstIndex;
 			int32 IndexCount;
@@ -214,8 +195,9 @@ namespace UnrealConversionUtils
 		const int32 MeshLODIndex)
 
 	{
-		FSkeletalMeshLODRenderData& LODModel = Helper_GetLODData(SkeletalMesh)[MeshLODIndex];
+		MUTABLE_CPUPROFILER_SCOPE(CopyMutableVertexBuffers);
 
+		FSkeletalMeshLODRenderData& LODModel = Helper_GetLODData(SkeletalMesh)[MeshLODIndex];
 		const mu::FMeshBufferSet& MutableMeshVertexBuffers = MutableMesh->GetVertexBuffers();
 		const int32 NumVertices = MutableMeshVertexBuffers.GetElementCount();
 
@@ -324,7 +306,7 @@ namespace UnrealConversionUtils
 
 	bool CopyMutableIndexBuffers(mu::MeshPtrConst InMutableMesh, FSkeletalMeshLODRenderData& LODModel)
 	{
-		MUTABLE_CPUPROFILER_SCOPE(UpdateSkeletalMesh_BuildSkeletalMeshRenderData_IndexLoop);
+		MUTABLE_CPUPROFILER_SCOPE(CopyMutableIndexBuffers);
 
 		const int32 IndexCount = InMutableMesh->GetIndexBuffers().GetElementCount();
 
@@ -447,8 +429,151 @@ namespace UnrealConversionUtils
 		Override.BoneIDs.Shrink();
 		Override.BoneWeights.Shrink();
 		Override.VertexIndexToInfluenceOffset.Shrink();
-	} 
-}
+	}
 
 	
-		
+	 void CopySkeletalMeshLODRenderData(const USkeletalMesh* SrcSkeletalMesh,	USkeletalMesh* DestSkeletalMesh, int32 SrcLODIndex, int32 DestLODIndex)
+	 {
+		 MUTABLE_CPUPROFILER_SCOPE(CopySkeletalMeshLODRenderData);
+
+		 if (!SrcSkeletalMesh || !DestSkeletalMesh)
+		 {
+			 return;
+		 }
+
+		 if (!SrcSkeletalMesh->IsValidLODIndex(SrcLODIndex) || !DestSkeletalMesh->IsValidLODIndex(DestLODIndex))
+		 {
+			 // Trying to copy from/to an invalid LOD index
+			 check(false);
+			 return;
+		 }
+
+		 const FSkeletalMeshLODRenderData& SrcLODModel = Helper_GetLODData(SrcSkeletalMesh)[SrcLODIndex];
+		 FSkeletalMeshLODRenderData& DestLODModel = Helper_GetLODData(DestSkeletalMesh)[DestLODIndex];
+
+		 // Copying render sections
+		 {
+			 const FSkeletalMeshLODInfo& SrcLODInfo = Helper_GetLODInfoArray(SrcSkeletalMesh)[SrcLODIndex];
+			 FSkeletalMeshLODInfo& DestLODInfo = Helper_GetLODInfoArray(DestSkeletalMesh)[DestLODIndex];
+
+			 DestLODInfo.LODMaterialMap = SrcLODInfo.LODMaterialMap;
+
+			 const int32 SurfaceCount = SrcLODModel.RenderSections.Num();
+			 for (int32 SurfaceIndex = 0; SurfaceIndex < SurfaceCount; ++SurfaceIndex)
+			 {
+				 const FSkelMeshRenderSection& SrcSection = Helper_GetLODRenderSections(SrcSkeletalMesh, SrcLODIndex)[SurfaceIndex];
+				 FSkelMeshRenderSection* DestSection = new(Helper_GetLODRenderSections(DestSkeletalMesh, DestLODIndex)) Helper_SkelMeshRenderSection();
+
+				 DestSection->DuplicatedVerticesBuffer.Init(1, TMap<int, TArray<int32>>());
+				 DestSection->bDisabled = SrcSection.bDisabled;
+
+				 if (!DestSection->bDisabled)
+				 {
+					 DestSection->BaseIndex = SrcSection.BaseIndex;
+					 DestSection->NumTriangles = SrcSection.NumTriangles;
+					 DestSection->BaseVertexIndex = SrcSection.BaseVertexIndex;
+					 DestSection->MaxBoneInfluences = SrcSection.MaxBoneInfluences;
+					 DestSection->NumVertices = SrcSection.NumVertices;
+					 DestSection->BoneMap = SrcSection.BoneMap;
+				 }
+			 }
+		 }
+
+		 const FStaticMeshVertexBuffers& SrcStaticVertexBuffer = SrcLODModel.StaticVertexBuffers;
+		 FStaticMeshVertexBuffers& DestStaticVertexBuffer = DestLODModel.StaticVertexBuffers;
+
+		 const int32 NumVertices = SrcStaticVertexBuffer.PositionVertexBuffer.GetNumVertices();
+		 const int32 NumTexCoords = SrcStaticVertexBuffer.StaticMeshVertexBuffer.GetNumTexCoords();
+
+		 bool bAllowCPUAccess = Helper_GetLODInfoArray(SrcSkeletalMesh)[SrcLODIndex].bAllowCPUAccess;
+
+		 // Copying Static Vertex Buffers
+		 {
+			 // Position buffer
+			 DestStaticVertexBuffer.PositionVertexBuffer.Init(NumVertices, bAllowCPUAccess);
+			 FMemory::Memcpy(DestStaticVertexBuffer.PositionVertexBuffer.GetVertexData(), SrcStaticVertexBuffer.PositionVertexBuffer.GetVertexData(), NumVertices * DestStaticVertexBuffer.PositionVertexBuffer.GetStride());
+
+			 // Tangent and Texture coords buffers
+			 DestStaticVertexBuffer.StaticMeshVertexBuffer.SetUseFullPrecisionUVs(true);
+			 DestStaticVertexBuffer.StaticMeshVertexBuffer.SetUseHighPrecisionTangentBasis(false);
+			 DestStaticVertexBuffer.StaticMeshVertexBuffer.Init(NumVertices, NumTexCoords, bAllowCPUAccess);
+			 FMemory::Memcpy(DestStaticVertexBuffer.StaticMeshVertexBuffer.GetTangentData(), SrcStaticVertexBuffer.StaticMeshVertexBuffer.GetTangentData(), DestStaticVertexBuffer.StaticMeshVertexBuffer.GetTangentSize());
+			 FMemory::Memcpy(DestStaticVertexBuffer.StaticMeshVertexBuffer.GetTexCoordData(), SrcStaticVertexBuffer.StaticMeshVertexBuffer.GetTexCoordData(), DestStaticVertexBuffer.StaticMeshVertexBuffer.GetTexCoordSize());
+
+			 // Color buffer
+			 if (SrcSkeletalMesh->GetHasVertexColors())
+			 {
+				 DestSkeletalMesh->SetHasVertexColors(true);
+
+				 DestStaticVertexBuffer.ColorVertexBuffer.Init(NumVertices);
+				 FMemory::Memcpy(DestStaticVertexBuffer.ColorVertexBuffer.GetVertexData(), SrcStaticVertexBuffer.ColorVertexBuffer.GetVertexData(), NumVertices * DestStaticVertexBuffer.ColorVertexBuffer.GetStride());
+			 }
+		 }
+
+		 // Copying Skin Buffers
+		 {
+			 const FSkinWeightVertexBuffer& SrcSkinWeightBuffer = SrcLODModel.SkinWeightVertexBuffer;
+			 FSkinWeightVertexBuffer& DestSkinWeightBuffer = DestLODModel.SkinWeightVertexBuffer;
+
+			  int32 NumBoneInfluences = SrcSkinWeightBuffer.GetDataVertexBuffer()->GetMaxBoneInfluences();
+			 int32 NumBones = SrcSkinWeightBuffer.GetDataVertexBuffer()->GetNumBones();
+
+			 const_cast<FSkinWeightDataVertexBuffer*>(DestSkinWeightBuffer.GetDataVertexBuffer())->SetMaxBoneInfluences(NumBoneInfluences);
+			 const_cast<FSkinWeightDataVertexBuffer*>(DestSkinWeightBuffer.GetDataVertexBuffer())->Init(NumBones, NumVertices);
+
+			 if (NumVertices)
+			 {
+				 DestSkinWeightBuffer.SetNeedsCPUAccess(bAllowCPUAccess);
+
+				 const FSkinWeightInfo* SrcData = SrcSkinWeightBuffer.GetDataVertexBuffer()->GetWeightData();
+				 FSkinWeightInfo* Data = DestSkinWeightBuffer.GetDataVertexBuffer()->GetWeightData();
+				 check(SrcData);
+				 check(Data);
+
+				 FMemory::Memcpy(Data, SrcData, DestSkinWeightBuffer.GetVertexDataSize());
+			 }
+
+			 DestSkinWeightBuffer.SetUse16BitBoneIndex(SrcSkinWeightBuffer.Use16BitBoneIndex());
+		 }
+
+		 // Copying Skin Weight Profiles Buffers
+		 {
+			 int32 NumSkinWeightProfiles = SrcSkeletalMesh->GetSkinWeightProfiles().Num();
+			 for (int32 ProfileIndex = 0; ProfileIndex < NumSkinWeightProfiles; ++ProfileIndex)
+			 {
+				 DestSkeletalMesh->AddSkinWeightProfile(SrcSkeletalMesh->GetSkinWeightProfiles()[ProfileIndex]);
+
+				 const FRuntimeSkinWeightProfileData* SrcProfile = SrcLODModel.SkinWeightProfilesData.GetOverrideData(SrcSkeletalMesh->GetSkinWeightProfiles()[ProfileIndex].Name);
+				 FRuntimeSkinWeightProfileData& DestProfile = DestLODModel.SkinWeightProfilesData.AddOverrideData(SrcSkeletalMesh->GetSkinWeightProfiles()[ProfileIndex].Name);
+
+				 DestProfile.NumWeightsPerVertex = SrcProfile->NumWeightsPerVertex;
+				 DestProfile.b16BitBoneIndices = SrcProfile->b16BitBoneIndices;
+
+				 DestProfile.BoneIDs = SrcProfile->BoneIDs;
+				 DestProfile.BoneWeights = SrcProfile->BoneWeights;
+				 DestProfile.VertexIndexToInfluenceOffset = SrcProfile->VertexIndexToInfluenceOffset;
+			 }
+		 }
+
+		 // Copying Indices
+		 {
+			 if (SrcLODModel.MultiSizeIndexContainer.GetIndexBuffer())
+			 {
+				 int32 IndexCount = SrcLODModel.MultiSizeIndexContainer.GetIndexBuffer()->Num();
+				 int32 ElementSize = SrcLODModel.MultiSizeIndexContainer.GetDataTypeSize();
+
+				 const void* Data = Helper_GetLODData(SrcSkeletalMesh)[SrcLODIndex].MultiSizeIndexContainer.GetIndexBuffer()->GetPointerTo(0);
+
+				 DestLODModel.MultiSizeIndexContainer.CreateIndexBuffer(ElementSize);
+				 DestLODModel.MultiSizeIndexContainer.GetIndexBuffer()->Insert(0, IndexCount);
+				 FMemory::Memcpy(DestLODModel.MultiSizeIndexContainer.GetIndexBuffer()->GetPointerTo(0), Data, IndexCount * ElementSize);
+			 }
+		 }
+
+		 DestLODModel.ActiveBoneIndices.Append(SrcLODModel.ActiveBoneIndices);
+		 DestLODModel.RequiredBones.Append(SrcLODModel.RequiredBones);
+		 DestLODModel.bIsLODOptional = SrcLODModel.bIsLODOptional;
+		 DestLODModel.bStreamedDataInlined = SrcLODModel.bStreamedDataInlined;
+		 DestLODModel.BuffersSize = SrcLODModel.BuffersSize;
+	}
+}
diff --git a/Engine/Plugins/Experimental/Mutable/Source/CustomizableObject/Public/MuCO/CustomizableObjectInstanceDescriptor.h b/Engine/Plugins/Experimental/Mutable/Source/CustomizableObject/Public/MuCO/CustomizableObjectInstanceDescriptor.h
--- a/Engine/Plugins/Experimental/Mutable/Source/CustomizableObject/Public/MuCO/CustomizableObjectInstanceDescriptor.h
+++ b/Engine/Plugins/Experimental/Mutable/Source/CustomizableObject/Public/MuCO/CustomizableObjectInstanceDescriptor.h
@@ -72,6 +72,10 @@ struct CUSTOMIZABLEOBJECT_API FCustomizableObjectInstanceDescriptor
 
 	void SetMaxLod(int32 InMaxLOD);
 
+	void SetRequestedLODLevels(const TArray<uint16>& InRequestedLODLevels);
+
+	const TArray<uint16>& GetRequestedLODLevels() const;
+
 	// ------------------------------------------------------------
 	// Parameters
 	// ------------------------------------------------------------
@@ -359,10 +363,13 @@ private:
      * These descriptions get generated with every instance update, so it should be disabled when not needed. */
 	bool bBuildParameterDecorations = false;
 
-	/** These are the LODs we want to have, they MUST NOT be used in an update (Mutable thread). */
+	/** These are the LODs Mutable can generate, they MUST NOT be used in an update (Mutable thread). */
 	int32 MinLOD = 0;
 	int32 MaxLOD = INT32_MAX;
 
+	/** Array of RequestedLODs per component to generate, they MUST NOT be used in an update (Mutable thread). */
+	TArray<uint16> RequestedLODLevels;
+
 	/** Lookup of UCustomizableObjectDescriptor::IntParameters. */
 	TMap<FString, int32> IntParametersLookupTable;
 	
@@ -421,10 +428,17 @@ public:
 	int32 GetMinLOD() const;
 
 	int32 GetMaxLOD() const;
-	
+
+	void UpdateRequestedLODs(const TArray<uint16>& InRequestedLODs);
+
+	const TArray<uint16>& GetRequestedLODs() const;
+
 private:
 	int32 MinLOD = 0;
 	int32 MaxLOD = INT32_MAX;
+
+	// Array of bitmasks that indicate which LODs of each component have been requested
+	TArray<uint16> RequestedLODsPerComponent;
 };
 
 #if UE_ENABLE_INCLUDE_ORDER_DEPRECATED_IN_5_2
diff --git a/Engine/Plugins/Experimental/Mutable/Source/CustomizableObject/Public/MuCO/UnrealConversionUtils.h b/Engine/Plugins/Experimental/Mutable/Source/CustomizableObject/Public/MuCO/UnrealConversionUtils.h
--- a/Engine/Plugins/Experimental/Mutable/Source/CustomizableObject/Public/MuCO/UnrealConversionUtils.h
+++ b/Engine/Plugins/Experimental/Mutable/Source/CustomizableObject/Public/MuCO/UnrealConversionUtils.h
@@ -20,17 +20,6 @@ namespace UnrealConversionUtils
 	 * current pipelines (instance update and USkeletal mesh generation for a mesh viewport)
 	 */
 	
-	/**
-	 * Assigns a new render section for each of the surfaces found on the mutable mesh and sets a default material for
-	 * each of them (same default material)
-	 * @param MeshLODIndex - The index of the LOD found on the OutSkeletalMesh to be targeting.
-	 * @param InMutableMesh - The Mutable mesh to get the surface count from. Determines the amount materials and also
-	 * render sections to be added to the OutSkeletalMesh.
-	 * @param OutSkeletalMesh - The skeletal mesh whose render data is being updated (adding render sections)
-	 */
-	CUSTOMIZABLEOBJECT_API void BuildSkeletalMeshElementDataAtLOD(const int32 MeshLODIndex, mu::MeshPtrConst InMutableMesh,
-										   USkeletalMesh* OutSkeletalMesh);
-	
 	
 	/**
 	 * Builds the reference skeleton by adding bones to the reference skeleton modifier provided. When exiting the scope
@@ -97,6 +86,21 @@ namespace UnrealConversionUtils
 		const FName InProfileName,
 		const mu::FMeshBufferSet& InMutableMeshVertexBuffers,
 		const int32 InBoneIndexBuffer);
+
+
+	/**
+	 *Performs a copy of the render data of a specific Skeletal Mesh LOD to another Skeletal Mesh
+	 * @param SrcSkeletalMesh - Skeletal Mesh with the data to copy.
+	 * @param DestSkeletalMesh - Skeletal Mesh where the data will be copied.
+	 * @param SrcLODIndex - Index of the LOD to copy.
+	 * @param DestSkeletalMesh - Index of the LOD that will be copyed to.
+	 */
+	CUSTOMIZABLEOBJECT_API void CopySkeletalMeshLODRenderData(
+		const USkeletalMesh* SrcSkeletalMesh,
+		USkeletalMesh* DestSkeletalMesh,
+		int32 SrcLODIndex,
+		int32 DestLODIndex);
+
 }
 
 #if UE_ENABLE_INCLUDE_ORDER_DEPRECATED_IN_5_2
diff --git a/Engine/Plugins/Experimental/Mutable/Source/CustomizableObjectEditor/Private/MuCOE/MutableMeshPreviewUtils.cpp b/Engine/Plugins/Experimental/Mutable/Source/CustomizableObjectEditor/Private/MuCOE/MutableMeshPreviewUtils.cpp
--- a/Engine/Plugins/Experimental/Mutable/Source/CustomizableObjectEditor/Private/MuCOE/MutableMeshPreviewUtils.cpp
+++ b/Engine/Plugins/Experimental/Mutable/Source/CustomizableObjectEditor/Private/MuCOE/MutableMeshPreviewUtils.cpp
@@ -549,8 +549,16 @@ namespace MutableMeshPreviewUtils
 			OutSkeletalMesh->GetMaterials().SetNum(1);
 			OutSkeletalMesh->GetMaterials()[0] = UnrealMaterial;
 
-			constexpr int32 MeshLODIndex = 0;
-			UnrealConversionUtils::BuildSkeletalMeshElementDataAtLOD(MeshLODIndex,InMutableMesh,OutSkeletalMesh);
+			Helper_GetLODInfoArray(OutSkeletalMesh)[0].LODMaterialMap.SetNumZeroed(1);
+
+			// Add RenderSections for each surface in the mesh
+			if (InMutableMesh)
+			{
+				for (int32 SurfaceIndex = 0; SurfaceIndex < InMutableMesh->GetSurfaceCount(); ++SurfaceIndex)
+				{
+					new(Helper_GetLODRenderSections(OutSkeletalMesh, 0)) Helper_SkelMeshRenderSection();
+				}
+			}
 		}
 
 
