[Mutable] Fix compilation errors

diff --git a/Engine/Plugins/Experimental/Mutable/Source/CustomizableObject/Private/MuCO/UnrealConversionUtils.cpp b/Engine/Plugins/Experimental/Mutable/Source/CustomizableObject/Private/MuCO/UnrealConversionUtils.cpp
--- a/Engine/Plugins/Experimental/Mutable/Source/CustomizableObject/Private/MuCO/UnrealConversionUtils.cpp
+++ b/Engine/Plugins/Experimental/Mutable/Source/CustomizableObject/Private/MuCO/UnrealConversionUtils.cpp
@@ -516,7 +516,7 @@ namespace UnrealConversionUtils
 			 FSkinWeightVertexBuffer& DestSkinWeightBuffer = DestLODModel.SkinWeightVertexBuffer;
 
 			  int32 NumBoneInfluences = SrcSkinWeightBuffer.GetDataVertexBuffer()->GetMaxBoneInfluences();
-			 int32 NumBones = SrcSkinWeightBuffer.GetDataVertexBuffer()->GetNumBones();
+			 int32 NumBones = SrcSkinWeightBuffer.GetDataVertexBuffer()->GetNumBoneWeights();
 
 			 const_cast<FSkinWeightDataVertexBuffer*>(DestSkinWeightBuffer.GetDataVertexBuffer())->SetMaxBoneInfluences(NumBoneInfluences);
 			 const_cast<FSkinWeightDataVertexBuffer*>(DestSkinWeightBuffer.GetDataVertexBuffer())->Init(NumBones, NumVertices);
@@ -525,8 +525,8 @@ namespace UnrealConversionUtils
 			 {
 				 DestSkinWeightBuffer.SetNeedsCPUAccess(bAllowCPUAccess);
 
-				 const FSkinWeightInfo* SrcData = SrcSkinWeightBuffer.GetDataVertexBuffer()->GetWeightData();
-				 FSkinWeightInfo* Data = DestSkinWeightBuffer.GetDataVertexBuffer()->GetWeightData();
+				 const uint8* SrcData = SrcSkinWeightBuffer.GetDataVertexBuffer()->GetWeightData();
+				 uint8* Data = DestSkinWeightBuffer.GetDataVertexBuffer()->GetWeightData();
 				 check(SrcData);
 				 check(Data);
 
