[mutable] Improved Texture Saturate node tooltip.
- Updated Texture Node tooltip too to improve consistency ("Image" against "Texture" naming)
- Updated Texture Sample node tooltip due to it having the same issue as the Texture Saturate node.

diff --git a/Engine/Plugins/Experimental/Mutable/Source/CustomizableObjectEditor/Private/MuCOE/Nodes/CustomizableObjectNodeTexture.cpp b/Engine/Plugins/Experimental/Mutable/Source/CustomizableObjectEditor/Private/MuCOE/Nodes/CustomizableObjectNodeTexture.cpp
--- a/Engine/Plugins/Experimental/Mutable/Source/CustomizableObjectEditor/Private/MuCOE/Nodes/CustomizableObjectNodeTexture.cpp
+++ b/Engine/Plugins/Experimental/Mutable/Source/CustomizableObjectEditor/Private/MuCOE/Nodes/CustomizableObjectNodeTexture.cpp
@@ -52,7 +52,7 @@ FLinearColor UCustomizableObjectNodeTexture::GetNodeTitleColor() const
 
 FText UCustomizableObjectNodeTexture::GetTooltipText() const
 {
-	return LOCTEXT("Texture_Tooltip", "Defines an image.");
+	return LOCTEXT("Texture_Tooltip", "Defines a texture.");
 }
 
 
diff --git a/Engine/Plugins/Experimental/Mutable/Source/CustomizableObjectEditor/Private/MuCOE/Nodes/CustomizableObjectNodeTextureSample.cpp b/Engine/Plugins/Experimental/Mutable/Source/CustomizableObjectEditor/Private/MuCOE/Nodes/CustomizableObjectNodeTextureSample.cpp
--- a/Engine/Plugins/Experimental/Mutable/Source/CustomizableObjectEditor/Private/MuCOE/Nodes/CustomizableObjectNodeTextureSample.cpp
+++ b/Engine/Plugins/Experimental/Mutable/Source/CustomizableObjectEditor/Private/MuCOE/Nodes/CustomizableObjectNodeTextureSample.cpp
@@ -65,7 +65,7 @@ FLinearColor UCustomizableObjectNodeTextureSample::GetNodeTitleColor() const
 
 FText UCustomizableObjectNodeTextureSample::GetTooltipText() const
 {
-	return LOCTEXT("Texture_Sample_Tooltip","Texture Sample");
+	return LOCTEXT("Texture_Sample_Tooltip","Get the color found in a texture at the targeted X and Y position (from 0.0 to 1.0, both included).");
 }
 
 #undef LOCTEXT_NAMESPACE
diff --git a/Engine/Plugins/Experimental/Mutable/Source/CustomizableObjectEditor/Private/MuCOE/Nodes/CustomizableObjectNodeTextureSaturate.cpp b/Engine/Plugins/Experimental/Mutable/Source/CustomizableObjectEditor/Private/MuCOE/Nodes/CustomizableObjectNodeTextureSaturate.cpp
--- a/Engine/Plugins/Experimental/Mutable/Source/CustomizableObjectEditor/Private/MuCOE/Nodes/CustomizableObjectNodeTextureSaturate.cpp
+++ b/Engine/Plugins/Experimental/Mutable/Source/CustomizableObjectEditor/Private/MuCOE/Nodes/CustomizableObjectNodeTextureSaturate.cpp
@@ -51,7 +51,7 @@ FLinearColor UCustomizableObjectNodeTextureSaturate::GetNodeTitleColor() const
 
 FText UCustomizableObjectNodeTextureSaturate::GetTooltipText() const
 {
-	return LOCTEXT("Texture_Saturate_Tooltip", "Saturate Image");
+	return LOCTEXT("Texture_Saturate_Tooltip", "Get the provided texture with its saturation adjusted based on the numerical input provided where 1 equals to full saturation and 0 to no saturation.");
 }
 
 #undef LOCTEXT_NAMESPACE
