[Mutable] Remove invalid check which will always fail if an instance is in a State other than default.

#rnx

diff --git a/Engine/Plugins/Experimental/Mutable/Source/CustomizableObject/Private/MuCO/CustomizableObjectInstance.cpp b/Engine/Plugins/Experimental/Mutable/Source/CustomizableObject/Private/MuCO/CustomizableObjectInstance.cpp
--- a/Engine/Plugins/Experimental/Mutable/Source/CustomizableObject/Private/MuCO/CustomizableObjectInstance.cpp
+++ b/Engine/Plugins/Experimental/Mutable/Source/CustomizableObject/Private/MuCO/CustomizableObjectInstance.cpp
@@ -109,12 +109,12 @@ void UCustomizableInstancePrivateData::InitLastUpdateData(const TSharedPtr<FMuta
 
 
 void UCustomizableInstancePrivateData::InvalidateGeneratedData()
+{
+	for (FCustomizableInstanceComponentData& ComponentData : ComponentsData)
 	{
-		for (FCustomizableInstanceComponentData& ComponentData : ComponentsData)
-		{
-			ComponentData.LastMeshIdPerLOD.Init(INDEX_NONE, NumLODsAvailable);
-		}
-	
+		ComponentData.LastMeshIdPerLOD.Init(INDEX_NONE, NumLODsAvailable);
+	}
+
 	LastUpdateData.Clear();
 }
 
@@ -1503,7 +1503,7 @@ bool UCustomizableInstancePrivateData::UpdateSkeletalMesh_PostBeginUpdate0(UCust
 {
 	MUTABLE_CPUPROFILER_SCOPE(UCustomizableInstancePrivateData::UpdateSkeletalMesh_PostBeginUpdate0)
 
-		bool bEmptyMesh = false;
+	bool bEmptyMesh = false;
 
 	TArray<bool> ComponentNeedsUpdate;
 	bool bUpdateMeshes = DoComponentsNeedUpdate(Public, OperationData, ComponentNeedsUpdate, bEmptyMesh);
@@ -1538,8 +1538,6 @@ bool UCustomizableInstancePrivateData::UpdateSkeletalMesh_PostBeginUpdate0(UCust
 
 	TArray<TObjectPtr<USkeletalMesh>> OldSkeletalMeshes = Public->SkeletalMeshes;
 
-	check(NumLODsAvailable == OperationData->InstanceUpdateData.LODs.Num());
-
 	bool bSuccess = true;
 	bool bHasSkeletalMesh = false;
 
