[mutable] Clampled min value for the numeric properties of the Texture Project node.

diff --git a/Engine/Plugins/Experimental/Mutable/Source/CustomizableObjectEditor/Classes/MuCOE/Nodes/CustomizableObjectNodeTextureProject.h b/Engine/Plugins/Experimental/Mutable/Source/CustomizableObjectEditor/Classes/MuCOE/Nodes/CustomizableObjectNodeTextureProject.h
--- a/Engine/Plugins/Experimental/Mutable/Source/CustomizableObjectEditor/Classes/MuCOE/Nodes/CustomizableObjectNodeTextureProject.h
+++ b/Engine/Plugins/Experimental/Mutable/Source/CustomizableObjectEditor/Classes/MuCOE/Nodes/CustomizableObjectNodeTextureProject.h
@@ -27,7 +27,7 @@ public:
 
 	// Layout to use for the generated images.
 	UPROPERTY(EditAnywhere, Category = CustomizableObject)
-	int32 Layout = 0;
+	uint8 Layout = 0;
 	
 	UPROPERTY(EditAnywhere, Category = CustomizableObject)
 	bool bEnableAngleFadeOutForRGB = true;
@@ -35,16 +35,16 @@ public:
 	UPROPERTY(EditAnywhere, Category = CustomizableObject)
 	bool bEnableAngleFadeOutForAlpha = true;
 
-	/** Set the width of the Texture. If the reference texture is not null, it overrdies the Size X property */
+	/** Set the width of the Texture. If the reference texture is not null, it overrides the Size X property */
 	UPROPERTY(EditAnywhere, Category = CustomizableObject)
-	int32 TextureSizeX = 0;
+	uint32 TextureSizeX = 0;
 
-	/** Set the height of the Texture. If the reference texture is not null, it overrdies the Size Y property */
+	/** Set the height of the Texture. If the reference texture is not null, it overrides the Size Y property */
 	UPROPERTY(EditAnywhere, Category = CustomizableObject)
-	int32 TextureSizeY = 0;
+	uint32 TextureSizeY = 0;
 
-	UPROPERTY(EditAnywhere, Category = CustomizableObject)
-	int32 Textures = 0;
+	UPROPERTY(EditAnywhere, Category = CustomizableObject, meta=( ClampMin = 1))
+	uint32 Textures = 0;
 
 	/** Reference Texture used to decide the texture properties of the mutable-generated textures
 	* connected to this material (e.g. LODBias, Size X,...). If null, mutable default texture properties will be applied. */
diff --git a/Engine/Plugins/Experimental/Mutable/Source/CustomizableObjectEditor/Private/MuCOE/Nodes/CustomizableObjectNodeTextureProject.cpp b/Engine/Plugins/Experimental/Mutable/Source/CustomizableObjectEditor/Private/MuCOE/Nodes/CustomizableObjectNodeTextureProject.cpp
--- a/Engine/Plugins/Experimental/Mutable/Source/CustomizableObjectEditor/Private/MuCOE/Nodes/CustomizableObjectNodeTextureProject.cpp
+++ b/Engine/Plugins/Experimental/Mutable/Source/CustomizableObjectEditor/Private/MuCOE/Nodes/CustomizableObjectNodeTextureProject.cpp
@@ -77,15 +77,15 @@ void UCustomizableObjectNodeTextureProject::AllocateDefaultPins(UCustomizableObj
 	TexturePinsReferences.Reset(0);
 	OutputPinsReferences.Reset(0);
 	
-	for ( int LayerIndex = 0; LayerIndex < Textures; ++LayerIndex )
+	for ( uint32 LayerIndex = 0; LayerIndex < Textures; ++LayerIndex )
 	{
-		FString PinName = FString::Printf(TEXT("Texture %d"), LayerIndex);
+		FString PinName = FString::Printf(TEXT("Texture %u"), LayerIndex);
 		UEdGraphPin* TexturePin = CustomCreatePin(EGPD_Input, Schema->PC_Image, FName(*PinName));
 		TexturePin->bDefaultValueIsIgnored = true;
 
 		TexturePinsReferences.Add(FEdGraphPinReference(TexturePin));
 		
-		PinName = FString::Printf(TEXT("Texture %d"), LayerIndex);
+		PinName = FString::Printf(TEXT("Texture %u"), LayerIndex);
 		UEdGraphPin* OutputPin = CustomCreatePin(EGPD_Output, Schema->PC_Image, FName(*PinName));
 		OutputPin->bDefaultValueIsIgnored = true;
 
