[Mutable] CodeRunner added Unreal Insights scopes.

#rnx

diff --git a/Engine/Plugins/Experimental/Mutable/Source/MutableRuntime/Private/MuR/CodeRunner.cpp b/Engine/Plugins/Experimental/Mutable/Source/MutableRuntime/Private/MuR/CodeRunner.cpp
--- a/Engine/Plugins/Experimental/Mutable/Source/MutableRuntime/Private/MuR/CodeRunner.cpp
+++ b/Engine/Plugins/Experimental/Mutable/Source/MutableRuntime/Private/MuR/CodeRunner.cpp
@@ -165,7 +165,7 @@ namespace mu
                                           const Model* pModel
                                           )
     {
-        //MUTABLE_CPUPROFILER_SCOPE(RunCode_Conditional);
+        MUTABLE_CPUPROFILER_SCOPE(RunCode_Conditional);
 
 		OP_TYPE type = pModel->GetPrivate()->m_program.GetOpType(item.At);
 		OP::ConditionalArgs args = pModel->GetPrivate()->m_program.GetOpArgs<OP::ConditionalArgs>(item.At);
@@ -354,7 +354,7 @@ namespace mu
                                   uint32 lodMask
                                   )
     {
-        //MUTABLE_CPUPROFILER_SCOPE(RunCode_Instance);
+        MUTABLE_CPUPROFILER_SCOPE(RunCode_Instance);
 
 		OP_TYPE type = pModel->GetPrivate()->m_program.GetOpType(item.At);
         switch (type)
@@ -906,6 +906,8 @@ namespace mu
 
             case 1:
             {
+           		MUTABLE_CPUPROFILER_SCOPE(ME_APPLYLAYOUT)
+            		
                 Ptr<const Mesh> pBase = GetMemory().GetMesh( FCacheAddress(args.mesh,item) );
 
                 MeshPtr pApplied;
@@ -956,6 +958,8 @@ namespace mu
 
             case 1:
             {
+       	        MUTABLE_CPUPROFILER_SCOPE(ME_DIFFERENCE)
+
                 Ptr<const Mesh> pBase = GetMemory().GetMesh(FCacheAddress(BaseAt,item));
                 Ptr<const Mesh> pTarget = GetMemory().GetMesh(FCacheAddress(TargetAt,item));
 
@@ -1031,6 +1035,8 @@ namespace mu
 
             case 1:
             {
+                MUTABLE_CPUPROFILER_SCOPE(ME_MORPH2_1)
+
                 bool baseValid = GetMemory().IsValid( FCacheAddress(BaseAt,item) );
                 float factor = GetMemory().GetScalar( FCacheAddress(FactorAt,item) );
 
@@ -1102,6 +1108,8 @@ namespace mu
 
             case 2:
             {
+       		    MUTABLE_CPUPROFILER_SCOPE(ME_MORPH2_2)
+
                 Ptr<const Mesh> pBase = GetMemory().GetMesh( FCacheAddress(BaseAt,item) );
 
                 // Factor from 0 to 1 between the two targets
@@ -1185,6 +1193,8 @@ namespace mu
 
             case 1:
             {
+				MUTABLE_CPUPROFILER_SCOPE(ME_MORPH2_1)
+
                 MeshPtrConst pA = GetMemory().GetMesh( FCacheAddress(args.base,item) );
                 MeshPtrConst pB = GetMemory().GetMesh( FCacheAddress(args.added,item) );
 
@@ -1255,6 +1265,8 @@ namespace mu
 
             case 1:
             {
+            	MUTABLE_CPUPROFILER_SCOPE(ME_INTERPOLATE_1)
+
                 int count = 1;
                 for ( int i=0
                     ; i<MUTABLE_OP_MAX_INTERPOLATE_COUNT-1 && args.targets[i]
@@ -1334,6 +1346,8 @@ namespace mu
 
             case 2:
             {
+				MUTABLE_CPUPROFILER_SCOPE(ME_INTERPOLATE_2)
+
                 int count = 1;
                 for ( int i=0
                     ; i<MUTABLE_OP_MAX_INTERPOLATE_COUNT-1 && args.targets[i]
@@ -1440,6 +1454,8 @@ namespace mu
 
             case 1:
             {
+            	MUTABLE_CPUPROFILER_SCOPE(ME_MASKCLIPMESH_1)
+            		
                 Ptr<const Mesh> pSource = GetMemory().GetMesh( FCacheAddress(args.source,item) );
                 Ptr<const Mesh> pClip = GetMemory().GetMesh( FCacheAddress(args.clip,item) );
 
@@ -1474,6 +1490,8 @@ namespace mu
 
             case 1:
             {
+           		MUTABLE_CPUPROFILER_SCOPE(ME_MASKDIFF_1)
+            		
                 Ptr<const Mesh> pSource = GetMemory().GetMesh( FCacheAddress(args.source,item) );
                 Ptr<const Mesh> pClip = GetMemory().GetMesh( FCacheAddress(args.fragment,item) );
 
@@ -1508,6 +1526,8 @@ namespace mu
 
             case 1:
             {
+           		MUTABLE_CPUPROFILER_SCOPE(ME_SUBTRACT_1)
+            	
                 Ptr<const Mesh> pA = GetMemory().GetMesh( FCacheAddress(args.a,item) );
                 Ptr<const Mesh> pB = GetMemory().GetMesh( FCacheAddress(args.b,item) );
 
@@ -1544,6 +1564,8 @@ namespace mu
 
             case 1:
             {
+            	MUTABLE_CPUPROFILER_SCOPE(ME_FORMAT_1)
+            		
                 Ptr<const Mesh> pSource = GetMemory().GetMesh( FCacheAddress(args.source,item) );
                 Ptr<const Mesh> pFormat = GetMemory().GetMesh( FCacheAddress(args.format,item) );
 
@@ -1606,6 +1628,8 @@ namespace mu
 
             case 1:
             {
+				MUTABLE_CPUPROFILER_SCOPE(ME_EXTRACTLAYOUTBLOCK_1)
+
                 Ptr<const Mesh> pSource = GetMemory().GetMesh( FCacheAddress(source,item) );
 
                 // Access with memcpy necessary for unaligned arm issues.
@@ -1648,6 +1672,8 @@ namespace mu
 
             case 1:
             {
+				MUTABLE_CPUPROFILER_SCOPE(ME_EXTRACTFACEGROUP_1)
+
                 Ptr<const Mesh> pSource = GetMemory().GetMesh( FCacheAddress(args.source,item) );
 
                 MeshPtr pResult;
@@ -1684,6 +1710,8 @@ namespace mu
 
             case 1:
             {
+           		MUTABLE_CPUPROFILER_SCOPE(ME_TRANSFORM_1)
+            		
                 Ptr<const Mesh> pSource = GetMemory().GetMesh( FCacheAddress(args.source,item) );
 
                 const mat4f& mat = pModel->GetPrivate()->m_program.
@@ -1721,6 +1749,8 @@ namespace mu
 
             case 1:
             {
+           		MUTABLE_CPUPROFILER_SCOPE(ME_CLIPMORPHPLANE_1)
+            		
                 Ptr<const Mesh> pSource = GetMemory().GetMesh( FCacheAddress(args.source,item) );
 
                 MeshPtr pResult;
@@ -1792,6 +1822,8 @@ namespace mu
 
             case 1:
             {
+				MUTABLE_CPUPROFILER_SCOPE(ME_CLIPWITHMESH_1)
+
                 Ptr<const Mesh> pSource = GetMemory().GetMesh( FCacheAddress(args.source,item) );
                 Ptr<const Mesh> pClip = GetMemory().GetMesh( FCacheAddress(args.clipMesh,item) );
 
@@ -1836,11 +1868,13 @@ namespace mu
 			
 			case 1:
 			{
+				MUTABLE_CPUPROFILER_SCOPE(ME_CLIPDEFORM_1)
+
 				Ptr<const Mesh> BaseMesh = GetMemory().GetMesh(FCacheAddress(args.mesh, item));
 				Ptr<const Mesh> ClipShape = GetMemory().GetMesh(FCacheAddress(args.clipShape, item));
 				Ptr<Mesh> pResult;
 
-					if ( BaseMesh && ClipShape )
+				if ( BaseMesh && ClipShape )
 				{
 					pResult = MeshClipDeform(BaseMesh.get(), ClipShape.get(), args.clipWeightThreshold);
 				}
@@ -1880,6 +1914,8 @@ namespace mu
 
             case 1:
             {
+    			MUTABLE_CPUPROFILER_SCOPE(ME_REMAPINDICES_1)
+
                 Ptr<const Mesh> pSource = GetMemory().GetMesh( FCacheAddress(args.source,item) );
                 Ptr<const Mesh> pReference = GetMemory().GetMesh( FCacheAddress(args.reference,item) );
 
@@ -1926,6 +1962,8 @@ namespace mu
 
             case 1:
             {
+          		MUTABLE_CPUPROFILER_SCOPE(ME_APPLYPOSE_1)
+
                 Ptr<const Mesh> pBase = GetMemory().GetMesh( FCacheAddress(args.base,item) );
                 Ptr<const Mesh> pPose = GetMemory().GetMesh( FCacheAddress(args.pose,item) );
 
@@ -1974,6 +2012,8 @@ namespace mu
 
 			case 1:
 			{
+				MUTABLE_CPUPROFILER_SCOPE(ME_GEOMETRYOPERATION_1)
+				
 				Ptr<const Mesh> MeshA = GetMemory().GetMesh(FCacheAddress(args.meshA, item));
 				Ptr<const Mesh> MeshB = GetMemory().GetMesh(FCacheAddress(args.meshB, item));
 				float ScalarA = GetMemory().GetScalar(FCacheAddress(args.scalarA, item));
@@ -2015,6 +2055,8 @@ namespace mu
 
 			case 1:
 			{
+				MUTABLE_CPUPROFILER_SCOPE(ME_BINDSHAPE_1)
+				
 				Ptr<const Mesh> BaseMesh = GetMemory().GetMesh(FCacheAddress(args.mesh, item));
 				Ptr<const Mesh> Shape = GetMemory().GetMesh(FCacheAddress(args.shape, item));
 				
@@ -2093,6 +2135,8 @@ namespace mu
 
 			case 1:
 			{
+				MUTABLE_CPUPROFILER_SCOPE(ME_APPLYSHAPE_1)
+					
 				Ptr<const Mesh> BaseMesh = GetMemory().GetMesh(FCacheAddress(args.mesh, item));
 				Ptr<const Mesh> Shape = GetMemory().GetMesh(FCacheAddress(args.shape, item));
 
@@ -2131,6 +2175,8 @@ namespace mu
 			}
 			case 1:
 			{
+				MUTABLE_CPUPROFILER_SCOPE(ME_MORPHRESHAPE_1)
+					
 				Ptr<const Mesh> MorphedMesh = GetMemory().GetMesh(FCacheAddress(Args.Morph, item));
 				Ptr<const Mesh> ReshapeMesh = GetMemory().GetMesh(FCacheAddress(Args.Reshape, item));
 
@@ -2185,6 +2231,8 @@ namespace mu
 
             case 1:
             {
+            	MUTABLE_CPUPROFILER_SCOPE(ME_SETSKELETON_1)
+            		
                 Ptr<const Mesh> pSource = GetMemory().GetMesh( FCacheAddress(args.source,item) );
                 Ptr<const Mesh> pSkeleton = GetMemory().GetMesh( FCacheAddress(args.skeleton,item) );
 
@@ -2232,6 +2280,8 @@ namespace mu
 
         case OP_TYPE::ME_REMOVEMASK:
         {
+       		MUTABLE_CPUPROFILER_SCOPE(ME_REMOVEMASK)
+        		
             // Decode op
             // TODO: Partial decode for each stage
             const uint8* data = pModel->GetPrivate()->m_program.GetOpArgsPointer(item.At);
@@ -2274,6 +2324,8 @@ namespace mu
 
             case 1:
             {
+        		MUTABLE_CPUPROFILER_SCOPE(ME_REMOVEMASK_1)
+
                 // Request the source and the necessary masks
                 // \todo: store condition values in heap?
                 TArray<FScheduledOp> deps;
@@ -2302,6 +2354,8 @@ namespace mu
 
             case 2:
             {
+            	MUTABLE_CPUPROFILER_SCOPE(ME_REMOVEMASK_2)
+            	
                 // \todo: single remove operation with all masks?
                 Ptr<const Mesh> pSource = GetMemory().GetMesh( FCacheAddress(source,item) );
 
@@ -2357,6 +2411,8 @@ namespace mu
 
             case 1:
             {
+				MUTABLE_CPUPROFILER_SCOPE(ME_PROJECT_1)
+
                 Ptr<const Mesh> pMesh = GetMemory().GetMesh( FCacheAddress(args.mesh,item) );
                 Ptr<const Projector> pProjector = GetMemory().GetProjector( FCacheAddress(args.projector,item) );
 
@@ -2408,6 +2464,8 @@ namespace mu
 
 			case 1:
 			{
+				MUTABLE_CPUPROFILER_SCOPE(ME_OPTIMIZESKINNING_1)
+
 				Ptr<const Mesh> pSource = GetMemory().GetMesh(FCacheAddress(args.source, item));
 
 				MeshPtrConst pResult = MeshOptimizeSkinning(pSource.get());
@@ -2506,6 +2564,8 @@ namespace mu
 
             case 1:
             {
+           		MUTABLE_CPUPROFILER_SCOPE(IM_MULTILAYER_1)
+            		
                 // We now know the number of iterations
                 int32 Iterations = 0;
                 if (args.rangeSize)
@@ -2566,7 +2626,7 @@ namespace mu
 
             default:
             {
-                //MUTABLE_CPUPROFILER_SCOPE(ImageMultiLayer);
+            	MUTABLE_CPUPROFILER_SCOPE(IM_MULTILAYER_default)
 
 				FScheduledOpData& Data = m_heapData[item.CustomState];
 
@@ -2762,6 +2822,8 @@ namespace mu
 
             case 1:
             {
+            	MUTABLE_CPUPROFILER_SCOPE(IM_DIFFERENCE_1)
+            		
                 // TODO: Reuse base if possible
                 Ptr<const Image> pA = GetMemory().GetImage( FCacheAddress(args.a,item) );
                 Ptr<const Image> pB = GetMemory().GetImage( FCacheAddress(args.b,item) );
@@ -2799,6 +2861,8 @@ namespace mu
 
 			case 1:
 			{
+				MUTABLE_CPUPROFILER_SCOPE(IM_NORMALCOMPOSITE_1)
+
 				Ptr<const Image> Base = GetMemory().GetImage(FCacheAddress(args.base, item));
 				Ptr<const Image> Normal = GetMemory().GetImage(FCacheAddress(args.normal, item));
 
@@ -2885,11 +2949,13 @@ namespace mu
             switch (item.Stage)
             {
             case 0:
-                    AddOp( FScheduledOp( item.At, item, 1), FScheduledOp( args.source, item) );
+                AddOp( FScheduledOp( item.At, item, 1), FScheduledOp( args.source, item) );
                 break;
 
             case 1:
             {
+				MUTABLE_CPUPROFILER_SCOPE(IM_RESIZE_1)
+            	
                 Ptr<const Image> pBase = GetMemory().GetImage( FCacheAddress(args.source,item) );
 
 				if (!pBase)
@@ -2965,6 +3031,8 @@ namespace mu
 
             case 1:
             {
+            	MUTABLE_CPUPROFILER_SCOPE(IM_RESIZELIKE_1)
+            	
                 Ptr<const Image> pBase = GetMemory().GetImage( FCacheAddress(args.source,item) );
                 Ptr<const Image> pSizeBase = GetMemory().GetImage( FCacheAddress(args.sizeSource,item) );
                 ImagePtr pResult;
@@ -3027,6 +3095,8 @@ namespace mu
 
             case 1:
             {
+            	MUTABLE_CPUPROFILER_SCOPE(IM_RESIZEREL_1)
+            	
                 Ptr<const Image> pBase = GetMemory().GetImage( FCacheAddress(args.source,item) );
 
                 FImageSize destSize(
@@ -3083,6 +3153,8 @@ namespace mu
 
             case 1:
             {
+            	MUTABLE_CPUPROFILER_SCOPE(IM_BLANKLAYOUT_1)
+            		
                 Ptr<const Layout> pLayout = GetMemory().GetLayout(FScheduledOp::FromOpAndOptions(args.layout, item, 0));
 
                 FIntPoint SizeInBlocks = pLayout->GetGridSize();
@@ -3160,6 +3232,8 @@ namespace mu
 
             case 1:
             {
+            	MUTABLE_CPUPROFILER_SCOPE(IM_COMPOSE_1)
+            		
                 Ptr<const Layout> ComposeLayout = GetMemory().GetLayout( FCacheAddress( args.layout, FScheduledOp::FromOpAndOptions(args.layout, item, 0)) );
 
                 FScheduledOpData data;
@@ -3206,6 +3280,8 @@ namespace mu
 
             case 1:
             {
+           		MUTABLE_CPUPROFILER_SCOPE(IM_INTERPOLATE_1)
+            	
                 // Targets must be consecutive
                 int count = 0;
                 for ( int i=0
@@ -3250,6 +3326,8 @@ namespace mu
 
             case 2:
             {
+           		MUTABLE_CPUPROFILER_SCOPE(IM_INTERPOLATE_2)
+            		
                 // Targets must be consecutive
                 int count = 0;
                 for ( int i=0
@@ -3368,6 +3446,8 @@ namespace mu
 
             case 1:
             {
+           		MUTABLE_CPUPROFILER_SCOPE(IM_INTERPOLATE3_1)
+
                 // \TODO Optimise limit cases
 
                 float factor1 = GetMemory().GetScalar(FScheduledOp::FromOpAndOptions(args.factor1, item, 0));
@@ -3404,6 +3484,8 @@ namespace mu
 
             case 1:
             {
+           		MUTABLE_CPUPROFILER_SCOPE(IM_SATURATE_1)
+           		
                 Ptr<const Image> pBase = GetMemory().GetImage( FCacheAddress(args.base,item) );
                 float factor = GetMemory().GetScalar(FScheduledOp::FromOpAndOptions(args.factor, item, 0));
 
@@ -3432,6 +3514,8 @@ namespace mu
 
             case 1:
             {
+           		MUTABLE_CPUPROFILER_SCOPE(IM_LUMINANCE_1)
+            		
                 Ptr<const Image> pBase = GetMemory().GetImage( FCacheAddress(args.base,item) );
 
                 ImagePtr pResult = ImageLuminance( pBase.get() );
@@ -3484,6 +3568,8 @@ namespace mu
 
             case 1:
             {
+            	MUTABLE_CPUPROFILER_SCOPE(IM_SELECTCOLOUR_1)
+            		
                 Ptr<const Image> pBase = GetMemory().GetImage( FCacheAddress(args.base,item) );
 				FVector4f colour = GetMemory().GetColour(FScheduledOp::FromOpAndOptions(args.colour, item, 0));
 
@@ -3514,6 +3600,8 @@ namespace mu
 
             case 1:
             {
+            	MUTABLE_CPUPROFILER_SCOPE(IM_COLOURMAP_1)
+            		
                 Ptr<const Image> pSource = GetMemory().GetImage( FCacheAddress(args.base,item) );
                 Ptr<const Image> pMask = GetMemory().GetImage( FCacheAddress(args.mask,item) );
                 Ptr<const Image> pMap = GetMemory().GetImage( FCacheAddress(args.map,item) );
@@ -3552,6 +3640,8 @@ namespace mu
 
             case 1:
             {
+				MUTABLE_CPUPROFILER_SCOPE(IM_GRADIENT_1)
+            		
 				FVector4f colour0 = GetMemory().GetColour(FScheduledOp::FromOpAndOptions(args.colour0, item, 0));
 				FVector4f colour1 = GetMemory().GetColour(FScheduledOp::FromOpAndOptions(args.colour1, item, 0));
 
@@ -3583,6 +3673,8 @@ namespace mu
 
             case 1:
             {
+            	MUTABLE_CPUPROFILER_SCOPE(IM_BINARISE_1)
+            		
                 Ptr<const Image> pA = GetMemory().GetImage( FCacheAddress(args.base,item) );
 
                 float c = GetMemory().GetScalar(FScheduledOp::FromOpAndOptions(args.threshold, item, 0));
@@ -3612,6 +3704,8 @@ namespace mu
 
 			case 1:
 			{
+				MUTABLE_CPUPROFILER_SCOPE(IM_INVERT_1)
+					
 				Ptr<const Image> pA = GetMemory().GetImage(FCacheAddress(args.base, item));
 
 				ImagePtr pResult;
@@ -3648,6 +3742,8 @@ namespace mu
 
             case 1:
             {
+            	MUTABLE_CPUPROFILER_SCOPE(IM_PLAINCOLOUR_1)
+            		
 				FVector4f c = GetMemory().GetColour(FScheduledOp::FromOpAndOptions(args.colour, item, 0));
 
 				uint16 SizeX = args.size[0];
@@ -3694,6 +3790,8 @@ namespace mu
 
             case 1:
             {
+            	MUTABLE_CPUPROFILER_SCOPE(IM_CROP_1)
+            		
                 Ptr<const Image> pA = GetMemory().GetImage( FCacheAddress(args.source,item) );
 
                 box< vec2<int> > rect;
@@ -3740,7 +3838,7 @@ namespace mu
 
             case 1:
             {
-                MUTABLE_CPUPROFILER_SCOPE(ImagePatch);
+                MUTABLE_CPUPROFILER_SCOPE(IM_PATCH_1)
 
                 Ptr<const Image> pA = GetMemory().GetImage( FCacheAddress(args.base,item) );
                 Ptr<const Image> pB = GetMemory().GetImage( FCacheAddress(args.patch,item) );
@@ -3864,7 +3962,7 @@ namespace mu
 
             case 1:
             {
-                MUTABLE_CPUPROFILER_SCOPE(RunCode_RasterMesh);
+                MUTABLE_CPUPROFILER_SCOPE(IM_RASTERMESH_1)
 
                 Ptr<const Mesh> pMesh = GetMemory().GetMesh( FScheduledOp::FromOpAndOptions(args.mesh, item, 0) );
 
@@ -4007,7 +4105,7 @@ namespace mu
 
             case 1:
             {
-                MUTABLE_CPUPROFILER_SCOPE(RunCode_ImageMakeGrowMap);
+                MUTABLE_CPUPROFILER_SCOPE(IM_MAKEGROWMAP_1)
 
                 Ptr<const Image> pMask = GetMemory().GetImage( FCacheAddress(args.mask,item) );
 
@@ -4040,7 +4138,7 @@ namespace mu
 
             case 1:
             {
-                MUTABLE_CPUPROFILER_SCOPE(RunCode_ImageDisplace);
+                MUTABLE_CPUPROFILER_SCOPE(IM_DISPLACE_1)
 
                 Ptr<const Image> pSource = GetMemory().GetImage( FCacheAddress(args.source,item) );
                 Ptr<const Image> pMap = GetMemory().GetImage( FCacheAddress(args.displacementMap,item) );
@@ -4104,6 +4202,8 @@ namespace mu
 			}
             case 1:
             {
+            	MUTABLE_CPUPROFILER_SCOPE(IM_TRANSFORM_1)
+            		
                 Ptr<const Image> pBaseImage = GetMemory().GetImage( FCacheAddress(Args.base, item) );
                 
                 const FVector2f Offset = FVector2f(
@@ -4188,7 +4288,7 @@ namespace mu
                                    const Model* pModel
                                    )
     {
-        //MUTABLE_CPUPROFILER_SCOPE(RunCode_Bool);
+        MUTABLE_CPUPROFILER_SCOPE(RunCode_Bool);
 
         const FProgram& program = pModel->GetPrivate()->m_program;
         OP_TYPE type = program.GetOpType(item.At);
