[mutable] Mutable Code Debugger: Added UI visualization to inform of the user of the navigation type being performed.
- A message will appear in the box that currently serves as a way to set the navigation operation type.
- The entry is not selectable by the user (collapsed) as it now happens with the NONE entry too.
- Commented out debugging log.
- Enforced "read only" pattern for navigation dropdown elements.

diff --git a/Engine/Plugins/Experimental/Mutable/Source/CustomizableObjectEditor/Private/MuCOE/SMutableCodeViewer.cpp b/Engine/Plugins/Experimental/Mutable/Source/CustomizableObjectEditor/Private/MuCOE/SMutableCodeViewer.cpp
--- a/Engine/Plugins/Experimental/Mutable/Source/CustomizableObjectEditor/Private/MuCOE/SMutableCodeViewer.cpp
+++ b/Engine/Plugins/Experimental/Mutable/Source/CustomizableObjectEditor/Private/MuCOE/SMutableCodeViewer.cpp
@@ -34,6 +34,7 @@
 #include "Widgets/Input/SComboBox.h"
 #include "Widgets/Layout/SScrollBox.h"
 #include "Widgets/Views/STreeView.h"
+#include "Widgets/SNullWidget.h"
 
 class FExtender;
 class FReferenceCollector;
@@ -438,9 +439,9 @@ void SMutableCodeViewer::SetCurrentModel(const TSharedPtr<mu::Model, ESPMode::Th
 		// Generate list elements for the found operation types so we are able to search over them on our type dropdown
 		GenerateNavigationDropdownElements();
 
-		// Check we did find types (witch should always happen in a normal run) and select the first option as default (NONE)
+		// Check we did find types (witch should always happen in a normal run) and select the NONE option as the default value
 		check(FoundModelOperationTypeElements.Num())
-		CurrentlySelectedOperationTypeElement = FoundModelOperationTypeElements[0];
+		CurrentlySelectedOperationTypeElement = NoneOperationEntry;
 	}
 }
 
@@ -563,7 +564,7 @@ void SMutableCodeViewer::Construct(const FArguments& InArgs, const TSharedPtr<mu
 							+SHorizontalBox::Slot()
 							.AutoWidth()
 							[
-								SAssignNew(TargetedTypeSelector,SComboBox<TSharedPtr<FMutableOperationElement>>)
+								SAssignNew(TargetedTypeSelector,SComboBox<TSharedPtr<const FMutableOperationElement>>)
 								.OptionsSource(&FoundModelOperationTypeElements)
 								.InitiallySelectedItem(CurrentlySelectedOperationTypeElement)
 								[
@@ -777,37 +778,59 @@ void SMutableCodeViewer::GenerateNavigationDropdownElements()
 		const FText OperationTypeName = FText::FromString(ModelOperationTypeNames[OperationTypeIndex]);
 
 		const mu::OP_TYPE RepresentedType = ModelOperationTypes[OperationTypeIndex].Key;
+		const uint32 OperationTypeInstancesCount = ModelOperationTypes[OperationTypeIndex].Value;
 		
 		// Get the Color to be used on the text that will represent the operation on the dropdown
 		const FSlateColor OperationColor = ColorPerComputationalCost[StaticCast<uint8>(GetOperationTypeComputationalCost(RepresentedType))];
-
-		const uint32 OperationTypeInstancesCount = ModelOperationTypes[OperationTypeIndex].Value;
 		
 		// Generate an element to be used by the ComboBox handling the selection of the type to be used during navigation
 		TSharedPtr<FMutableOperationElement> OperationElement = MakeShared<FMutableOperationElement>(RepresentedType, OperationTypeName,OperationTypeInstancesCount,OperationColor);
 		FoundModelOperationTypeElements.Add(OperationElement);
 	}
+
+	// Add an entry for the NONE type of operation
+	{
+		const FText EntryName = FText::FromString("NONE");
+		const FSlateColor EntryColor = ColorPerComputationalCost[StaticCast<uint8>(EOperationComputationalCost::Standard)];
+		NoneOperationEntry = MakeShared<FMutableOperationElement>(mu::OP_TYPE::NONE,EntryName,0,EntryColor);
+
+		// @warn While not visible this element must be part of the collection for the ComboBox to be able to work
+		// properly
+		FoundModelOperationTypeElements.Add(NoneOperationEntry);
+	}
+	
+	// Add an extra operation type that will represent the constant resource based navigation type
+	{
+		const FText EntryName = FText::FromString("CONST_RESOURCE_BASED_NAV");	
+		const FSlateColor EntryColor = FSlateColor(FLinearColor(0.35f ,0.35f,1.0f,1));
+		ConstantBasedNavigationEntry = MakeShared<FMutableOperationElement>(mu::OP_TYPE::NONE,EntryName,0,EntryColor);
+		
+		// @warn While not visible this element must be part of the collection for the ComboBox to be able to work
+		// properly
+		FoundModelOperationTypeElements.Add(ConstantBasedNavigationEntry);
+	}
 }
 
 TSharedRef<SWidget> SMutableCodeViewer::OnGenerateOpNavigationDropDownWidget(
-	TSharedPtr<FMutableOperationElement> MutableOperationElement) const
+	TSharedPtr<const FMutableOperationElement> MutableOperationElement) const
 {
-	return SNew(STextBlock)
-			.Text(MutableOperationElement->OperationTypeText)
-			.ColorAndOpacity(MutableOperationElement->OperationTextColor);
+	TSharedRef<STextBlock> NewSlateObject = SNew(STextBlock)
+		.Text(MutableOperationElement->OperationTypeText)
+		.ColorAndOpacity(MutableOperationElement->OperationTextColor);
+
+	// Set the visibility type for the UI object (currently will be hidden for the NONE type)
+	NewSlateObject->SetVisibility(MutableOperationElement->GetEntryVisibility());
+	
+	return NewSlateObject;
 }
 
 void SMutableCodeViewer::OnNavigationSelectedOperationChanged(
-	TSharedPtr<FMutableOperationElement, ESPMode::ThreadSafe> MutableOperationElement, ESelectInfo::Type Arg)
+	TSharedPtr<const FMutableOperationElement, ESPMode::ThreadSafe> MutableOperationElement, ESelectInfo::Type Arg)
 {
 	check (MutableOperationElement.IsValid());
 
 	// Cache the currently selected operation set on the UI by the user
 	const mu::OP_TYPE NewOperationType = MutableOperationElement->OperationType;
-	if (NewOperationType == OperationTypeToSearch)
-	{
-		return;
-	}
 	OperationTypeToSearch = NewOperationType;
 	CurrentlySelectedOperationTypeElement = MutableOperationElement;
 	
@@ -857,18 +880,19 @@ void SMutableCodeViewer::OnSelectedOperationTypeFromTree()
 	const mu::OP_TYPE OperationType =
 		MutableModel->GetPrivate()->m_program.GetOpType(ReferenceOperationElement->MutableOperation);
 
-	// Search for the index of the provided operation type inside the array of operation types found on the model
-	const int32 OperationIndex = ModelOperationTypes.IndexOfByPredicate(
-		[OperationType](const TPair<mu::OP_TYPE, uint32>& Current)
+	// Find the operation type directly in our array of operation elements (from the drop down)
+	const TSharedPtr<const FMutableOperationElement>* RepresentativeElement = FoundModelOperationTypeElements.
+		FindByPredicate([OperationType](const TSharedPtr<const FMutableOperationElement> Other)
 		{
-			return Current.Key == OperationType;
+			return Other->OperationType == OperationType;
 		});
-	
-	// Failing the next check would mean that we are not caching all the types present on the current operation's tree
-	check (OperationIndex != INDEX_NONE);
+
+	// Ensure an element was found. Failing the next check would mean that we are not caching all the types present on
+	// the current operation's tree
+	check(RepresentativeElement != nullptr);
 	
 	// Set the type operation type to be looking for -> Will invoke OnOptionTypeSelectionChanged
-	TargetedTypeSelector->SetSelectedItem(FoundModelOperationTypeElements[OperationIndex]);
+	TargetedTypeSelector->SetSelectedItem(*RepresentativeElement);
 	
 	// Reset the navigation index so it starts from scratch
 	NavigationIndex = -1;
@@ -1017,14 +1041,8 @@ void SMutableCodeViewer::CacheOperationTypesPresentOnModel()
 		// Then the name of the second element
 		return AString < BString;
 	});	
-
-	// Add the "None" option as a default value at index 0 to the sorted array.
-	// the NONE option serves as a way of telling the user "hey, you are not navigating over operations of X type".
-	// We currently use it to tell the user that if he is navigating it is over constant related operations
-	ModelOperationTypes.Insert(TPair<mu::OP_TYPE,uint32>{mu::OP_TYPE::NONE,0},0);
-	ModelOperationTypes.Shrink();
-
-	// ModelOperationTypes is now an array starting with None and all the types found on the operations tree in alphabetical order
+	
+	// ModelOperationTypes is now an array with all the types found on the operations tree in alphabetical order
 }
 
 
@@ -1493,8 +1511,9 @@ void SMutableCodeViewer::GetVisibleChildren(TSharedPtr<FMutableCodeTreeElement>
 			++ChildIndex;
 		});
 	}
-	
-	UE_LOG(LogTemp,Warning,TEXT("Found a total of %i children elements "),OutChildren.Num());
+
+	// Debug
+	// UE_LOG(LogTemp,Warning,TEXT("Found a total of %i children elements "),OutChildren.Num());
 }
 
 
@@ -1962,9 +1981,8 @@ void SMutableCodeViewer::CacheRootNodeAddresses()
 	// At this point we did get all the addresses of operations that do involve the usage of our resource
 	if (SearchPayload.FoundElements.Num() > 0)
 	{
-		// Reset operation selection object to show no element to navigate to since we are not navigating over op types
-		// Set the type operation type to NONE (we know it is always the first one)
-		TargetedTypeSelector->SetSelectedItem(FoundModelOperationTypeElements[0]);
+		// Set the type operation type to CONST_BASED_NAVIGATION (used to tell the user what is happening)
+		TargetedTypeSelector->SetSelectedItem(ConstantBasedNavigationEntry);
 
 		// Dump the located resources array onto the navigation array since we have content to navigate over
 		NavigationElements = MoveTemp(SearchPayload.FoundElements);
diff --git a/Engine/Plugins/Experimental/Mutable/Source/CustomizableObjectEditor/Private/MuCOE/SMutableCodeViewer.h b/Engine/Plugins/Experimental/Mutable/Source/CustomizableObjectEditor/Private/MuCOE/SMutableCodeViewer.h
--- a/Engine/Plugins/Experimental/Mutable/Source/CustomizableObjectEditor/Private/MuCOE/SMutableCodeViewer.h
+++ b/Engine/Plugins/Experimental/Mutable/Source/CustomizableObjectEditor/Private/MuCOE/SMutableCodeViewer.h
@@ -247,13 +247,21 @@ private:
 	 */
 	
 	/** Slate object that provides the user a way of selecting what kind of operation it wants to navigate.*/
-	TSharedPtr<SComboBox<TSharedPtr<FMutableOperationElement>>> TargetedTypeSelector;
+	TSharedPtr<SComboBox<TSharedPtr<const FMutableOperationElement>>> TargetedTypeSelector;
 
 	/** Data backend for the list displayed for the navigation type selection */
-	TArray<TSharedPtr<FMutableOperationElement>> FoundModelOperationTypeElements;
+	TArray<TSharedPtr<const FMutableOperationElement>> FoundModelOperationTypeElements;
 
+	/** Navigation operation entry representing the NONE type */
+	TSharedPtr<const FMutableOperationElement> NoneOperationEntry;
+
+	/** Navigation operation entry used to be able to set the UI to show we are performing a constant resource based 
+	 * navigation
+	 */
+	TSharedPtr<const FMutableOperationElement> ConstantBasedNavigationEntry;
+	
 	/** Currently selected element on the TargetedTypeSelector slate. Actively used by the ui */
-	TSharedPtr<FMutableOperationElement> CurrentlySelectedOperationTypeElement;
+	TSharedPtr<const FMutableOperationElement> CurrentlySelectedOperationTypeElement;
 	
 	/** Operation type we are using to search for tree nodes. Driven primarily by the UI */
 	mu::OP_TYPE OperationTypeToSearch = mu::OP_TYPE::NONE;
@@ -291,12 +299,12 @@ private:
 	/** Callback invoked by the ComboBox used for displaying and selecting operation types for navigation. it gets invoked
 	 * each time the slate object requires to draw a line representing one of the elements set on FoundModelOperationTypeElements
 	 */
-	TSharedRef<SWidget> OnGenerateOpNavigationDropDownWidget(TSharedPtr<FMutableOperationElement> MutableOperationElement) const;
+	TSharedRef<SWidget> OnGenerateOpNavigationDropDownWidget(TSharedPtr<const FMutableOperationElement> MutableOperationElement) const;
 
 	/** Callback invoked each time the selected operation on our navigation slate changes. It can change due to UI interaction
 	 * or also due to direct change by invoking the SetSelectedOption on the SComboBox TargetedTypeSelector
 	 */
-	void OnNavigationSelectedOperationChanged(TSharedPtr<FMutableOperationElement, ESPMode::ThreadSafe> MutableOperationElement, ESelectInfo::Type Arg);
+	void OnNavigationSelectedOperationChanged(TSharedPtr<const FMutableOperationElement, ESPMode::ThreadSafe> MutableOperationElement, ESelectInfo::Type Arg);
 
 	/** Callback used to print on screen the amount of operations found on the tree that share the same operation type that
 	 * the one currently selected on the navigation system.
@@ -726,6 +734,12 @@ public:
 	mu::OP_TYPE OperationType;
 	FText OperationTypeText;
 	FSlateColor OperationTextColor;
+
+	/** Get the type of visibility an slate representing this objects should have. */
+	EVisibility GetEntryVisibility() const
+	{
+		return OperationType == mu::OP_TYPE::NONE ? EVisibility::Collapsed : EVisibility::Visible;
+	}
 };
 
 
