[Mutable] Fix empty Decriptor using 0 as a valid index Texture Parameter.

Caused some issues when no Texture Parameter was explicitly set.

- Fix indexing issues when using a base different than 0.
- Renamed some variables to be more consistent.

#rnx

diff --git a/Engine/Plugins/Experimental/Mutable/Source/CustomizableObject/Private/MuCO/DefaultImageProvider.cpp b/Engine/Plugins/Experimental/Mutable/Source/CustomizableObject/Private/MuCO/DefaultImageProvider.cpp
--- a/Engine/Plugins/Experimental/Mutable/Source/CustomizableObject/Private/MuCO/DefaultImageProvider.cpp
+++ b/Engine/Plugins/Experimental/Mutable/Source/CustomizableObject/Private/MuCO/DefaultImageProvider.cpp
@@ -36,7 +36,9 @@ namespace UDefaultImageProviderCVars
 
 UCustomizableSystemImageProvider::ValueType UDefaultImageProvider::HasTextureParameterValue(int64 ID)
 {
-	return Textures[ID] ?
+	const int32 TextureIndex = ToIndex(ID);
+	
+	return Textures.IsValidIndex(TextureIndex) && Textures[TextureIndex] ?
 		static_cast<ValueType>(UDefaultImageProviderCVars::ImageMode) :
 		ValueType::None;
 }
@@ -44,7 +46,7 @@ UCustomizableSystemImageProvider::ValueType UDefaultImageProvider::HasTexturePar
 
 UTexture2D* UDefaultImageProvider::GetTextureParameterValue(int64 ID)
 {
-	return Textures[ID];
+	return Textures[ToIndex(ID)];
 }
 
 
@@ -56,7 +58,7 @@ void UDefaultImageProvider::GetTextureParameterValues(TArray<FCustomizableObject
 		{
 			FCustomizableObjectExternalTexture Data;
 			Data.Name = Texture->GetName();
-			Data.Value = TextureIndex;
+			Data.Value = ToTextureId(TextureIndex);
 			OutValues.Add(Data);
 		}
 	}
@@ -125,18 +127,18 @@ void UDefaultImageProvider::CacheTextures(UCustomizableObjectInstance& Instance)
 
 	for (const FCustomizableObjectTextureParameterValue& TextureParameter : Instance.GetDescriptor().GetTextureParameters())
 	{
-		if (const int32 Index = ToIndex(TextureParameter.ParameterValue);
-			Textures.IsValidIndex(Index) && Textures[Index])
+		if (const int32 TextureIndex = ToIndex(TextureParameter.ParameterValue);
+			Textures.IsValidIndex(TextureIndex) && Textures[TextureIndex])
 		{
-			System->CacheImage(Index);
+			System->CacheImage(TextureParameter.ParameterValue);
 		}
 
 		for (const uint64 RangeValue : TextureParameter.ParameterRangeValues)
 		{
-			if (const int32 Index = ToIndex(RangeValue);
-				Textures.IsValidIndex(Index) && Textures[Index])
+			if (const int32 TextureIndex = ToIndex(RangeValue);
+				Textures.IsValidIndex(TextureIndex) && Textures[TextureIndex])
 			{
-				System->CacheImage(Index);
+				System->CacheImage(RangeValue);
 			}
 		}
 	}
@@ -158,18 +160,18 @@ void UDefaultImageProvider::GarbageCollectTextureIds()
 	{
 		for (const FCustomizableObjectTextureParameterValue& TextureParameter :  It->GetDescriptor().GetTextureParameters())
 		{			
-			if (const int32 Index = ToIndex(TextureParameter.ParameterValue);
-				IdUsed.IsValidIndex(Index))
+			if (const int32 TextureIndex = ToIndex(TextureParameter.ParameterValue);
+				IdUsed.IsValidIndex(TextureIndex))
 			{
-				IdUsed[Index] = true;
+				IdUsed[TextureIndex] = true;
 			}
 
 			for (const uint64 RangeValue : TextureParameter.ParameterRangeValues)
 			{
-				if (const int32 Index = ToIndex(RangeValue);
-					IdUsed.IsValidIndex(Index))
+				if (const int32 TextureIndex = ToIndex(RangeValue);
+					IdUsed.IsValidIndex(TextureIndex))
 				{
-					IdUsed[Index] = true;
+					IdUsed[TextureIndex] = true;
 				}
 			}
 		}
@@ -184,7 +186,7 @@ void UDefaultImageProvider::GarbageCollectTextureIds()
 		if (!IdUsed[TextureIndex] && !KeepTextures[TextureIndex])
 		{
 			Textures[TextureIndex] = nullptr;
-			System->UnCacheImage(TextureIndex);
+			System->UnCacheImage(ToTextureId(TextureIndex));
 		}
 		else
 		{
@@ -207,8 +209,8 @@ int32 UDefaultImageProvider::ToIndex(const uint64 TextureId) const
 }
 
 
-uint64 UDefaultImageProvider::ToTextureId(const int32 Index) const
+uint64 UDefaultImageProvider::ToTextureId(const int32 TextureIndex) const
 {
-	return Index + BASE_ID;
+	return TextureIndex + BASE_ID;
 }
 
diff --git a/Engine/Plugins/Experimental/Mutable/Source/CustomizableObject/Public/MuCO/DefaultImageProvider.h b/Engine/Plugins/Experimental/Mutable/Source/CustomizableObject/Public/MuCO/DefaultImageProvider.h
--- a/Engine/Plugins/Experimental/Mutable/Source/CustomizableObject/Public/MuCO/DefaultImageProvider.h
+++ b/Engine/Plugins/Experimental/Mutable/Source/CustomizableObject/Public/MuCO/DefaultImageProvider.h
@@ -40,10 +40,10 @@ public:
 private:
 	int32 ToIndex(uint64 TextureId) const;
 
-	uint64 ToTextureId(int32 Index) const;
+	uint64 ToTextureId(int32 TextureIndex) const;
 	
 	/** Change it to avoid TextureId collisions with other Image Providers. */
-	const int32 BASE_ID = 0;
+	const int32 BASE_ID = 1;
 	const int32 MAX_IDS = 1000;
 
 	UPROPERTY()
diff --git a/Engine/Plugins/Experimental/Mutable/Source/CustomizableObjectEditor/Private/MuCOE/SCustomizableInstanceProperties.cpp b/Engine/Plugins/Experimental/Mutable/Source/CustomizableObjectEditor/Private/MuCOE/SCustomizableInstanceProperties.cpp
--- a/Engine/Plugins/Experimental/Mutable/Source/CustomizableObjectEditor/Private/MuCOE/SCustomizableInstanceProperties.cpp
+++ b/Engine/Plugins/Experimental/Mutable/Source/CustomizableObjectEditor/Private/MuCOE/SCustomizableInstanceProperties.cpp
@@ -77,7 +77,7 @@ void SCustomizableInstanceProperties::Construct(const FArguments& InArgs)
 
 		// Store the texture parameters data required for the ui.
 		TextureParameterValueNames.Add(MakeShareable(new FString("None")));
-		TextureParameterValues.Add(UINT64_MAX);
+		TextureParameterValues.AddDefaulted();
 		TArray<FCustomizableObjectExternalTexture> Textures = UCustomizableObjectSystem::GetInstance()->GetTextureParameterValues();
 		for (int i = 0; i < Textures.Num(); ++i)
 		{
