[Mutable] Fix IORelease tasks overwritting other queued tasks of the same type.
+ Fix MipDataProvider not freeing instances after an update.
+ Fix static analysis warnings.

#rnx

diff --git a/Engine/Plugins/Experimental/Mutable/Source/CustomizableObject/Private/MuCO/CustomizableObjectMipDataProvider.cpp b/Engine/Plugins/Experimental/Mutable/Source/CustomizableObject/Private/MuCO/CustomizableObjectMipDataProvider.cpp
--- a/Engine/Plugins/Experimental/Mutable/Source/CustomizableObject/Private/MuCO/CustomizableObjectMipDataProvider.cpp
+++ b/Engine/Plugins/Experimental/Mutable/Source/CustomizableObject/Private/MuCO/CustomizableObjectMipDataProvider.cpp
@@ -113,6 +113,7 @@ namespace impl
 		{
 			MUTABLE_CPUPROFILER_SCOPE(EndUpdate);
 			System->EndUpdate(InstanceID);
+			System->ReleaseInstance(InstanceID);
 		}
 
 		{
diff --git a/Engine/Plugins/Experimental/Mutable/Source/CustomizableObject/Private/MuCO/CustomizableObjectSystem.cpp b/Engine/Plugins/Experimental/Mutable/Source/CustomizableObject/Private/MuCO/CustomizableObjectSystem.cpp
--- a/Engine/Plugins/Experimental/Mutable/Source/CustomizableObject/Private/MuCO/CustomizableObjectSystem.cpp
+++ b/Engine/Plugins/Experimental/Mutable/Source/CustomizableObject/Private/MuCO/CustomizableObjectSystem.cpp
@@ -107,7 +107,8 @@ void FMutableQueue::Enqueue(const FMutableQueueElem& TaskToEnqueue)
 	{
 		check(QueueElem.Operation);
 		check(TaskToEnqueue.Operation);
-		if (QueueElem.Operation->CustomizableObjectInstance == TaskToEnqueue.Operation->CustomizableObjectInstance)
+		if ((QueueElem.Operation->Type == FMutableOperation::EOperationType::Update || QueueElem.Operation->Type == FMutableOperation::EOperationType::Discard) &&
+			QueueElem.Operation->CustomizableObjectInstance == TaskToEnqueue.Operation->CustomizableObjectInstance)
 		{
 			QueueElem.Operation = TaskToEnqueue.Operation;
 			QueueElem.PriorityType = FMath::Min(QueueElem.PriorityType, TaskToEnqueue.PriorityType);
@@ -755,7 +756,8 @@ void UpdateSkeletalMesh(UCustomizableObjectInstance* CustomizableObjectInstance,
 		{
 			MUTABLE_CPUPROFILER_SCOPE(UpdateSkeletalMesh_UpdatedDelegate);
 
-			CustomizableObjectInstance->Updated(CustomizableObjectInstance->SkeletalMeshStatus == ESkeletalMeshState::Correct ? EUpdateResult::Success : EUpdateResult::Error, UpdatedDescriptorRuntimeHash);
+			EUpdateResult State = CustomizableObjectInstance->SkeletalMeshStatus == ESkeletalMeshState::Correct ? EUpdateResult::Success : EUpdateResult::Error;
+			CustomizableObjectInstance->Updated(State, UpdatedDescriptorRuntimeHash);
 
 
 #if WITH_EDITOR
@@ -2093,7 +2095,7 @@ namespace impl
 		TObjectPtr<UCustomizableObjectInstance> CandidateInstance = Operation->CustomizableObjectInstance.Get();
 		
 		UCustomizableInstancePrivateData* CandidateInstancePrivateData = CandidateInstance->GetPrivate();
-		check(CandidateInstancePrivateData);
+		check(CandidateInstancePrivateData != nullptr);
 
 		if (CandidateInstancePrivateData->HasCOInstanceFlags(PendingLODsUpdate))
 		{
