#rnx

diff --git a/Engine/Plugins/Experimental/Mutable/Source/CustomizableObject/Private/MuCO/CustomizableObjectSystemPrivate.h b/Engine/Plugins/Experimental/Mutable/Source/CustomizableObject/Private/MuCO/CustomizableObjectSystemPrivate.h
--- a/Engine/Plugins/Experimental/Mutable/Source/CustomizableObject/Private/MuCO/CustomizableObjectSystemPrivate.h
+++ b/Engine/Plugins/Experimental/Mutable/Source/CustomizableObject/Private/MuCO/CustomizableObjectSystemPrivate.h
@@ -56,6 +56,10 @@ public:
 
 		// Release the instance ID and all the temp data associated with it. Usually used with the LiveUpdateMode
 		IDRelease
+
+		// Attention! If any new operation type is added, make sure to review FMutableQueue::Enqueue in CustomizableObjectSystem.cpp and
+		// modify it if new operations of this type should override older ones to the same instance. The default behavior is to enqueue 
+		// the new ones so that they are executed after the old ones.
 	};
 
 	// Type of the operation
