[Mutable] Postpone DefaultImageProvider image cache.

- Added documentation.

#rnx

diff --git a/Engine/Plugins/Experimental/Mutable/Source/CustomizableObject/Classes/MuCO/CustomizableObjectSystem.h b/Engine/Plugins/Experimental/Mutable/Source/CustomizableObject/Classes/MuCO/CustomizableObjectSystem.h
--- a/Engine/Plugins/Experimental/Mutable/Source/CustomizableObject/Classes/MuCO/CustomizableObjectSystem.h
+++ b/Engine/Plugins/Experimental/Mutable/Source/CustomizableObject/Classes/MuCO/CustomizableObjectSystem.h
@@ -297,10 +297,14 @@ public:
 	/** [Texture Parameters] Remove all images from the cache. */
 	void ClearImageCache();
 
-	UDefaultImageProvider& GetDefaultImageProvider();
+	/** Get the default image provider. Returns null if it was not already initialized. */
+	TObjectPtr<UDefaultImageProvider> GetDefaultImageProvider() const;
+
+	/** Initialize (if was not already) and get the default image provider. */
+	UDefaultImageProvider& GetOrCreateDefaultImageProvider();
 
 private:
-	
+	/** Mutable default image provider. Used by the COIEditor and Instance/Descriptor APIs. */
 	UPROPERTY()
 	TObjectPtr<UDefaultImageProvider> DefaultImageProvider = nullptr;
 
diff --git a/Engine/Plugins/Experimental/Mutable/Source/CustomizableObject/Private/MuCO/CustomizableObjectInstance.cpp b/Engine/Plugins/Experimental/Mutable/Source/CustomizableObject/Private/MuCO/CustomizableObjectInstance.cpp
--- a/Engine/Plugins/Experimental/Mutable/Source/CustomizableObject/Private/MuCO/CustomizableObjectInstance.cpp
+++ b/Engine/Plugins/Experimental/Mutable/Source/CustomizableObject/Private/MuCO/CustomizableObjectInstance.cpp
@@ -232,7 +232,7 @@ void UCustomizableObjectInstance::PreEditChange(FProperty* PropertyAboutToChange
 
 	if (PropertyName == GET_MEMBER_NAME_CHECKED(UCustomizableObjectInstance, TextureParameterDeclarations))
 	{
-		UDefaultImageProvider& DefaultImageProvider = UCustomizableObjectSystem::GetInstance()->GetDefaultImageProvider();
+		UDefaultImageProvider& DefaultImageProvider = UCustomizableObjectSystem::GetInstance()->GetOrCreateDefaultImageProvider();
 		
 		for (TObjectPtr<UTexture2D> Texture : TextureParameterDeclarations)
 		{
@@ -256,7 +256,7 @@ void UCustomizableObjectInstance::PostEditChangeProperty(FPropertyChangedEvent&
 
 	if (PropertyName == GET_MEMBER_NAME_CHECKED(UCustomizableObjectInstance, TextureParameterDeclarations))
 	{
-		UDefaultImageProvider& DefaultImageProvider = UCustomizableObjectSystem::GetInstance()->GetDefaultImageProvider();
+		UDefaultImageProvider& DefaultImageProvider = UCustomizableObjectSystem::GetInstance()->GetOrCreateDefaultImageProvider();
 		
 		for (TObjectPtr<UTexture2D> Texture : TextureParameterDeclarations)
 		{
diff --git a/Engine/Plugins/Experimental/Mutable/Source/CustomizableObject/Private/MuCO/CustomizableObjectInstanceDescriptor.cpp b/Engine/Plugins/Experimental/Mutable/Source/CustomizableObject/Private/MuCO/CustomizableObjectInstanceDescriptor.cpp
--- a/Engine/Plugins/Experimental/Mutable/Source/CustomizableObject/Private/MuCO/CustomizableObjectInstanceDescriptor.cpp
+++ b/Engine/Plugins/Experimental/Mutable/Source/CustomizableObject/Private/MuCO/CustomizableObjectInstanceDescriptor.cpp
@@ -1315,7 +1315,7 @@ uint64 FCustomizableObjectInstanceDescriptor::GetTextureParameterSelectedOption(
 
 UTexture2D* FCustomizableObjectInstanceDescriptor::GetTextureParameterSelectedOptionT(const FString& TextureParamName, const int32 RangeIndex) const
 {
-	UDefaultImageProvider& ImageProvider = UCustomizableObjectSystem::GetInstance()->GetDefaultImageProvider();
+	UDefaultImageProvider& ImageProvider = UCustomizableObjectSystem::GetInstance()->GetOrCreateDefaultImageProvider();
 
 	const uint64 TextureValue = GetTextureParameterSelectedOption(TextureParamName, RangeIndex);
 	return ImageProvider.Get(TextureValue);
@@ -1357,7 +1357,7 @@ void FCustomizableObjectInstanceDescriptor::SetTextureParameterSelectedOption(co
 
 void FCustomizableObjectInstanceDescriptor::SetTextureParameterSelectedOptionT(const FString& TextureParamName, UTexture2D* TextureValue, const int32 RangeIndex)
 {
-	UDefaultImageProvider& ImageProvider = UCustomizableObjectSystem::GetInstance()->GetDefaultImageProvider();
+	UDefaultImageProvider& ImageProvider = UCustomizableObjectSystem::GetInstance()->GetOrCreateDefaultImageProvider();
 
 	const uint64 NewUid = ImageProvider.GetOrAdd(TextureValue);
 	SetTextureParameterSelectedOption(TextureParamName, NewUid, RangeIndex);
diff --git a/Engine/Plugins/Experimental/Mutable/Source/CustomizableObject/Private/MuCO/CustomizableObjectSystem.cpp b/Engine/Plugins/Experimental/Mutable/Source/CustomizableObject/Private/MuCO/CustomizableObjectSystem.cpp
--- a/Engine/Plugins/Experimental/Mutable/Source/CustomizableObject/Private/MuCO/CustomizableObjectSystem.cpp
+++ b/Engine/Plugins/Experimental/Mutable/Source/CustomizableObject/Private/MuCO/CustomizableObjectSystem.cpp
@@ -830,11 +830,6 @@ void FCustomizableObjectSystemPrivate::InitUpdateSkeletalMesh(UCustomizableObjec
 	{
 		return; // The requested update is equal to the running update.
 	}
-
-	if (const TObjectPtr<UDefaultImageProvider> DefaultImageProvider = UCustomizableObjectSystem::GetInstance()->DefaultImageProvider)
-	{
-		DefaultImageProvider->CacheTextures(Instance);
-	}
 	
 	// These delegates must be called at the end of the begin update.
 	Instance.BeginUpdateDelegate.Broadcast(&Instance);
@@ -1434,8 +1429,6 @@ namespace impl
 			}
 			else
 			{
-				MUTABLE_CPUPROFILER_SCOPE(GetImage);
-
 				int32 MaxSize = FMath::Max(Image.FullImageSizeX, Image.FullImageSizeY);
 				int32 FullLODCount = FMath::CeilLogTwo(MaxSize) + 1;
 				int32 MinMipsInImage = FMath::Min(FullLODCount, UTexture::GetStaticMinTextureResidentMipCount());
@@ -1957,7 +1950,11 @@ namespace impl
 			ObjectInstancePrivateData->UpdateParameterDecorationsEngineResources(OperationData);
 		}
 
-
+		if (const TObjectPtr<UDefaultImageProvider> DefaultImageProvider = System->GetDefaultImageProvider())
+		{
+			DefaultImageProvider->CacheTextures(*ObjectInstance);
+		}
+		
 		// Selectively lock the resource cache for the object used by this instance to avoid the destruction of resources that we may want to reuse.
 		// When protecting textures there mustn't be any left from a previous update
 		check(System->ProtectedCachedTextures.Num() == 0);
@@ -2558,7 +2555,13 @@ void UCustomizableObjectSystem::UnregisterImageProvider(UCustomizableSystemImage
 }
 
 
-UDefaultImageProvider& UCustomizableObjectSystem::GetDefaultImageProvider()
+TObjectPtr<UDefaultImageProvider> UCustomizableObjectSystem::GetDefaultImageProvider() const
+{
+	return DefaultImageProvider;
+}
+
+
+UDefaultImageProvider& UCustomizableObjectSystem::GetOrCreateDefaultImageProvider()
 {
 	if (!DefaultImageProvider)
 	{
diff --git a/Engine/Plugins/Experimental/Mutable/Source/CustomizableObject/Public/MuCO/DefaultImageProvider.h b/Engine/Plugins/Experimental/Mutable/Source/CustomizableObject/Public/MuCO/DefaultImageProvider.h
--- a/Engine/Plugins/Experimental/Mutable/Source/CustomizableObject/Public/MuCO/DefaultImageProvider.h
+++ b/Engine/Plugins/Experimental/Mutable/Source/CustomizableObject/Public/MuCO/DefaultImageProvider.h
@@ -10,6 +10,11 @@ class UTexture2D;
 class UCustomizableObjectInstance;
 
 
+/** Simple image provider that translates UTexture2D to Mutable IDs.
+ *
+ * Allows the reuse of UTexture2D.
+ * IDs will be garbage collected if they are referenced inside an Instance (GetTextureParameters()) or are not mark explicitly to keep (Keep(...)).
+ */
 UCLASS()
 class CUSTOMIZABLEOBJECT_API UDefaultImageProvider : public UCustomizableSystemImageProvider
 {
