[Mutable] Fix memory leak with live update.

#rnx

diff --git a/Engine/Plugins/Experimental/Mutable/Source/MutableRuntime/Private/MuR/SystemPrivate.h b/Engine/Plugins/Experimental/Mutable/Source/MutableRuntime/Private/MuR/SystemPrivate.h
--- a/Engine/Plugins/Experimental/Mutable/Source/MutableRuntime/Private/MuR/SystemPrivate.h
+++ b/Engine/Plugins/Experimental/Mutable/Source/MutableRuntime/Private/MuR/SystemPrivate.h
@@ -227,7 +227,6 @@ namespace mu
 		void resize(size_t s)
 		{
 			m_index0.SetNum(s);
-			FMemory::Memzero(m_index0.GetData(), s * sizeof(DATA));
 		}
 
 		uint32 size_code() const
@@ -470,6 +469,7 @@ namespace mu
 
 		void Init(size_t size)
 		{
+			m_resources.clear();
 			m_resources.resize(size);
 			m_opHitCount.resize(size);
 		}
