[Mutable] Fix dummy texture when changing texture parameters rapidly.

- Added logs for more image acquire failure cases.
- Now the external images to be cached are obtained from the mutable parameters and not from the descriptor.

#rnx

diff --git a/Engine/Plugins/Experimental/Mutable/Source/CustomizableObject/Private/MuCO/CustomizableObjectSystem.cpp b/Engine/Plugins/Experimental/Mutable/Source/CustomizableObject/Private/MuCO/CustomizableObjectSystem.cpp
--- a/Engine/Plugins/Experimental/Mutable/Source/CustomizableObject/Private/MuCO/CustomizableObjectSystem.cpp
+++ b/Engine/Plugins/Experimental/Mutable/Source/CustomizableObject/Private/MuCO/CustomizableObjectSystem.cpp
@@ -1952,7 +1952,7 @@ namespace impl
 
 		if (const TObjectPtr<UDefaultImageProvider> DefaultImageProvider = System->GetDefaultImageProvider())
 		{
-			DefaultImageProvider->CacheTextures(*ObjectInstance);
+			DefaultImageProvider->CacheTextures(*OperationData->MutableParameters);
 		}
 		
 		// Selectively lock the resource cache for the object used by this instance to avoid the destruction of resources that we may want to reuse.
diff --git a/Engine/Plugins/Experimental/Mutable/Source/CustomizableObject/Private/MuCO/DefaultImageProvider.cpp b/Engine/Plugins/Experimental/Mutable/Source/CustomizableObject/Private/MuCO/DefaultImageProvider.cpp
--- a/Engine/Plugins/Experimental/Mutable/Source/CustomizableObject/Private/MuCO/DefaultImageProvider.cpp
+++ b/Engine/Plugins/Experimental/Mutable/Source/CustomizableObject/Private/MuCO/DefaultImageProvider.cpp
@@ -6,6 +6,8 @@
 #include "MuCO/CustomizableObjectInstance.h"
 #include "UObject/UObjectIterator.h"
 
+#include "MuR/Parameters.h"
+
 #include UE_INLINE_GENERATED_CPP_BY_NAME(DefaultImageProvider)
 
 
@@ -119,26 +121,39 @@ uint64 UDefaultImageProvider::GetOrAdd(UTexture2D* Texture)
 }
 
 
-void UDefaultImageProvider::CacheTextures(UCustomizableObjectInstance& Instance)
+void UDefaultImageProvider::CacheTextures(const mu::Parameters& Parameters)
 {
-	GarbageCollectTextureIds();
+	GarbageCollectTextureIds(Parameters);
 
 	UCustomizableObjectSystem* System = UCustomizableObjectSystem::GetInstance();
 
-	for (const FCustomizableObjectTextureParameterValue& TextureParameter : Instance.GetDescriptor().GetTextureParameters())
+	const int32 NumParams = Parameters.GetCount();
+	for (int32 ParamIndex = 0; ParamIndex < NumParams; ++ParamIndex)
 	{
-		if (const int32 TextureIndex = ToIndex(TextureParameter.ParameterValue);
-			Textures.IsValidIndex(TextureIndex) && Textures[TextureIndex])
+		if (Parameters.GetType(ParamIndex) != mu::PARAMETER_TYPE::T_IMAGE)
+		{
+			continue;
+		}
+	
 		{
-			System->CacheImage(TextureParameter.ParameterValue);
+			const mu::EXTERNAL_IMAGE_ID TextureId = Parameters.GetImageValue(ParamIndex);
+			if (const int32 TextureIndex = ToIndex(TextureId);
+				Textures.IsValidIndex(TextureIndex) && Textures[TextureIndex])
+			{
+				System->CacheImage(TextureId);
+			}
 		}
 
-		for (const uint64 RangeValue : TextureParameter.ParameterRangeValues)
+		const int32 NumValues = Parameters.GetValueCount(ParamIndex);
+		for (int32 ValueIndex = 0; ValueIndex < NumValues; ++ValueIndex)
 		{
-			if (const int32 TextureIndex = ToIndex(RangeValue);
+			mu::Ptr<const mu::RangeIndex> Range = Parameters.GetValueIndex(ParamIndex, ValueIndex);
+			const mu::EXTERNAL_IMAGE_ID TextureId = Parameters.GetImageValue(ParamIndex, Range);
+
+			if (const int32 TextureIndex = ToIndex(TextureId);
 				Textures.IsValidIndex(TextureIndex) && Textures[TextureIndex])
 			{
-				System->CacheImage(RangeValue);
+				System->CacheImage(TextureId);
 			}
 		}
 	}
@@ -151,7 +166,7 @@ void UDefaultImageProvider::Keep(const uint64 TextureId, const bool bKeep)
 }
 
 
-void UDefaultImageProvider::GarbageCollectTextureIds()
+void UDefaultImageProvider::GarbageCollectTextureIds(const mu::Parameters& Parameters)
 {
 	TBitArray IdUsed;
 	IdUsed.Init(false, Textures.Num());
@@ -177,6 +192,38 @@ void UDefaultImageProvider::GarbageCollectTextureIds()
 		}
 	}
 
+	const int32 NumParams = Parameters.GetCount();
+	for (int32 ParamIndex = 0; ParamIndex < NumParams; ++ParamIndex)
+	{
+		if (Parameters.GetType(ParamIndex) != mu::PARAMETER_TYPE::T_IMAGE)
+		{
+			continue;
+		}
+
+		{
+			const mu::EXTERNAL_IMAGE_ID TextureId = Parameters.GetImageValue(ParamIndex);
+			if (const int32 TextureIndex = ToIndex(TextureId);
+				Textures.IsValidIndex(TextureIndex) && Textures[TextureIndex])
+			{
+				IdUsed[TextureIndex] = true;
+			}
+		}
+
+		const int32 NumValues = Parameters.GetValueCount(ParamIndex);
+		for (int32 ValueIndex = 0; ValueIndex < NumValues; ++ValueIndex)
+		{
+			mu::Ptr<const mu::RangeIndex> Range = Parameters.GetValueIndex(ParamIndex, ValueIndex);
+			const mu::EXTERNAL_IMAGE_ID TextureId = Parameters.GetImageValue(ParamIndex, Range);
+
+			if (const int32 TextureIndex = ToIndex(TextureId);
+				Textures.IsValidIndex(TextureIndex) && Textures[TextureIndex])
+			{
+				IdUsed[TextureIndex] = true;
+			}
+		}
+
+	}
+
 	UCustomizableObjectSystem* System = UCustomizableObjectSystem::GetInstance();
 
 	int32 LastToKeep = -1;
diff --git a/Engine/Plugins/Experimental/Mutable/Source/CustomizableObject/Private/MuCO/UnrealMutableImageProvider.cpp b/Engine/Plugins/Experimental/Mutable/Source/CustomizableObject/Private/MuCO/UnrealMutableImageProvider.cpp
--- a/Engine/Plugins/Experimental/Mutable/Source/CustomizableObject/Private/MuCO/UnrealMutableImageProvider.cpp
+++ b/Engine/Plugins/Experimental/Mutable/Source/CustomizableObject/Private/MuCO/UnrealMutableImageProvider.cpp
@@ -149,6 +149,8 @@ TTuple<FGraphEventRef, TFunction<void()>> FUnrealMutableImageProvider::GetImageA
 		if (!GlobalExternalImages.Contains(id))
 		{
 			// Null case, no image was provided
+			UE_LOG(LogMutable, Warning, TEXT("Failed to get external image. GlobalExternalImage not found."));
+
 			ResultCallback(CreateDummy());
 			return Invoke(TrivialReturn);
 		}
@@ -244,6 +246,7 @@ TTuple<FGraphEventRef, TFunction<void()>> FUnrealMutableImageProvider::GetImageA
 					// Should we do someting different than returning a dummy image if cancelled?
 					if (bWasCancelled)
 					{
+						UE_LOG(LogMutable, Warning, TEXT("Failed to get external image. Cancelled IO Request"));
 						ResultCallback(CreateDummy());
 						return;
 					}
@@ -467,6 +470,7 @@ void FUnrealMutableImageProvider::CacheImage(mu::EXTERNAL_IMAGE_ID id)
 
 	if (!pResult && !UnrealDeferredTexture)
 	{
+		UE_LOG(LogMutable, Warning, TEXT("Failed to catch external image."));
 		pResult = CreateDummy();
 	}
 
diff --git a/Engine/Plugins/Experimental/Mutable/Source/CustomizableObject/Public/MuCO/DefaultImageProvider.h b/Engine/Plugins/Experimental/Mutable/Source/CustomizableObject/Public/MuCO/DefaultImageProvider.h
--- a/Engine/Plugins/Experimental/Mutable/Source/CustomizableObject/Public/MuCO/DefaultImageProvider.h
+++ b/Engine/Plugins/Experimental/Mutable/Source/CustomizableObject/Public/MuCO/DefaultImageProvider.h
@@ -9,6 +9,10 @@
 class UTexture2D;
 class UCustomizableObjectInstance;
 
+namespace mu
+{
+	class Parameters;
+}
 
 /** Simple image provider that translates UTexture2D to Mutable IDs.
  *
@@ -35,12 +39,12 @@ public:
 	/** If the TextureId does not get assigned to an Instance Texture Parameter it will be GC on the next Instance update. */
 	uint64 GetOrAdd(UTexture2D* Texture); 
 
-	void CacheTextures(UCustomizableObjectInstance& Instance);
+	void CacheTextures(const mu::Parameters&);
 
 	/** Keep the TextureId from being collected by the GC. */
 	void Keep(uint64 TextureId, bool bKeep);
 	
-	void GarbageCollectTextureIds();
+	void GarbageCollectTextureIds(const mu::Parameters&);
 
 private:
 	int32 ToIndex(uint64 TextureId) const;
