[Mutable] Moved some operations to tasks and other small optimizations.

- Moved Resize, Saturate and Invert  operation to tasks.
- Added saturate implace.
- Added generic simd optimizations for saturate invert and BufferLayerComposite with BlendChannelMasked and Lighten Channel (currenly disabled)

#rnx

diff --git a/Engine/Plugins/Experimental/Mutable/Source/MutableRuntime/Private/MuR/CodeRunner.cpp b/Engine/Plugins/Experimental/Mutable/Source/MutableRuntime/Private/MuR/CodeRunner.cpp
--- a/Engine/Plugins/Experimental/Mutable/Source/MutableRuntime/Private/MuR/CodeRunner.cpp
+++ b/Engine/Plugins/Experimental/Mutable/Source/MutableRuntime/Private/MuR/CodeRunner.cpp
@@ -2687,8 +2687,16 @@ namespace mu
 					{
 						// This is a frequent critical-path case because of multilayer projectors.
 						bDone = true;
-						
-						BufferLayerComposite<BlendChannelMasked, LightenChannel, false>(Base.get(), pBlended.get(), bBlendOnlyOneMip);
+					
+						constexpr bool bUseVectorImpl = false;
+						if constexpr (bUseVectorImpl)
+						{
+							BufferLayerCompositeVector<VectorBlendChannelMasked, VectorLightenChannel, false>(Base.get(), pBlended.get(), bBlendOnlyOneMip);
+						}
+						else
+						{
+							BufferLayerComposite<BlendChannelMasked, LightenChannel, false>(Base.get(), pBlended.get(), bBlendOnlyOneMip);
+						}
 					}
 
                     if (!bDone && args.mask)
@@ -2954,61 +2962,8 @@ namespace mu
 
             case 1:
             {
-				MUTABLE_CPUPROFILER_SCOPE(IM_RESIZE_1)
-            	
-                Ptr<const Image> pBase = GetMemory().GetImage( FCacheAddress(args.source,item) );
-
-				if (!pBase)
-				{
-					GetMemory().SetImage(item, nullptr);
-					break;
-				}
-
-                FImageSize destSize = FImageSize
-                    (
-                        args.size[0],
-                        args.size[1]
-                    );
-
-                ImagePtr pResult;
-
-                if ( destSize[0]!=pBase->GetSizeX()
-                     ||
-                     destSize[1]!=pBase->GetSizeY() )
-                {
-                    //pResult = ImageResize( pBase.get(), destSize );
-                    pResult = ImageResizeLinear( m_pSettings->GetPrivate()->m_imageCompressionQuality, pBase.get(), destSize );
-
-                    // If the source image had mips, generate them as well for the resized image.
-                    // This shouldn't happen often since "ResizeLike" should be usually optimised out
-                    // during model compilation. The mipmap generation below is not very precise with
-                    // the number of mips that are needed and will probably generate too many
-                    bool sourceHasMips = pBase->GetLODCount()>1;
-                    if (sourceHasMips)
-                    {
-                        int levelCount = Image::GetMipmapCount( pResult->GetSizeX(), pResult->GetSizeY() );
-                        ImagePtr pMipmapped = new Image( pResult->GetSizeX(), pResult->GetSizeY(),
-                                                         levelCount,
-                                                         pResult->GetFormat() );
-
-                        SCRATCH_IMAGE_MIPMAP scratch;
-						FMipmapGenerationSettings mipSettings{};
-
-                        ImageMipmap_PrepareScratch( pMipmapped.get(), pResult.get(), levelCount, &scratch );
-                        ImageMipmap( m_pSettings->GetPrivate()->m_imageCompressionQuality,
-                                     pMipmapped.get(), pResult.get(), levelCount, &scratch, mipSettings );
-
-                        pResult = pMipmapped;
-                    }
-
-                }
-                else
-                {
-                    pResult = pBase->Clone();
-                }
-
-                GetMemory().SetImage( item, pResult );
-                break;
+				// This has been moved to a task. It should have been intercepted in IssueOp.
+				check(false);
             }
 
             default:
@@ -3095,42 +3050,9 @@ namespace mu
 
             case 1:
             {
-            	MUTABLE_CPUPROFILER_SCOPE(IM_RESIZEREL_1)
-            	
-                Ptr<const Image> pBase = GetMemory().GetImage( FCacheAddress(args.source,item) );
-
-                FImageSize destSize(
-                            uint16( FMath::Max(1.0, pBase->GetSizeX()*args.factor[0] + 0.5f ) ),
-                            uint16( FMath::Max(1.0, pBase->GetSizeY()*args.factor[1] + 0.5f ) ) );
 
-                //pResult = ImageResize( pBase.get(), destSize );
-                ImagePtr pResult = ImageResizeLinear(
-                    m_pSettings->GetPrivate()->m_imageCompressionQuality, pBase.get(), destSize );
-
-                // If the source image had mips, generate them as well for the resized image.
-                // This shouldn't happen often since "ResizeLike" should be usually optimised out
-                // during model compilation. The mipmap generation below is not very precise with
-                // the number of mips that are needed and will probably generate too many
-                bool sourceHasMips = pBase->GetLODCount()>1;
-                if (sourceHasMips)
-                {
-                    int levelCount = Image::GetMipmapCount( pResult->GetSizeX(), pResult->GetSizeY() );
-                    ImagePtr pMipmapped = new Image( pResult->GetSizeX(), pResult->GetSizeY(),
-                                                     levelCount,
-                                                     pResult->GetFormat() );
-
-                    SCRATCH_IMAGE_MIPMAP scratch;
-					FMipmapGenerationSettings mipSettings{};
-
-                    ImageMipmap_PrepareScratch( pMipmapped.get(), pResult.get(), levelCount, &scratch );
-                    ImageMipmap( m_pSettings->GetPrivate()->m_imageCompressionQuality,
-                                 pMipmapped.get(), pResult.get(), levelCount, &scratch, mipSettings );
-
-                    pResult = pMipmapped;
-                }
-
-                GetMemory().SetImage( item, pResult );
-                break;
+				// This has been moved to a task. It should have been intercepted in IssueOp.
+				check(false);
             }
 
             default:
@@ -3484,15 +3406,8 @@ namespace mu
 
             case 1:
             {
-           		MUTABLE_CPUPROFILER_SCOPE(IM_SATURATE_1)
-           		
-                Ptr<const Image> pBase = GetMemory().GetImage( FCacheAddress(args.base,item) );
-                float factor = GetMemory().GetScalar(FScheduledOp::FromOpAndOptions(args.factor, item, 0));
-
-                ImagePtr pResult = ImageSaturate( pBase.get(), factor );
-
-                GetMemory().SetImage( item, pResult );
-                break;
+				// This has been moved to a task. It should have been intercepted in IssueOp.
+				check(false);
             }
 
             default:
@@ -3704,23 +3619,8 @@ namespace mu
 
 			case 1:
 			{
-				MUTABLE_CPUPROFILER_SCOPE(IM_INVERT_1)
-					
-				Ptr<const Image> pA = GetMemory().GetImage(FCacheAddress(args.base, item));
-
-				ImagePtr pResult;
-				if (pA->IsUnique())
-				{
-					pResult = mu::CloneOrTakeOver<>(pA.get());
-					ImageInvertInPlace(pResult.get());
-				}
-				else
-				{
-					pResult = ImageInvert(pA.get());
-				}
-
-				GetMemory().SetImage(item, pResult);
-				break;
+				// This has been moved to a task. It should have been intercepted in IssueOp.
+				check(false);
 			}
 
 			default:
diff --git a/Engine/Plugins/Experimental/Mutable/Source/MutableRuntime/Private/MuR/CodeRunnerTasked.cpp b/Engine/Plugins/Experimental/Mutable/Source/MutableRuntime/Private/MuR/CodeRunnerTasked.cpp
--- a/Engine/Plugins/Experimental/Mutable/Source/MutableRuntime/Private/MuR/CodeRunnerTasked.cpp
+++ b/Engine/Plugins/Experimental/Mutable/Source/MutableRuntime/Private/MuR/CodeRunnerTasked.cpp
@@ -26,6 +26,8 @@
 #include "MuR/OpImagePixelFormat.h"
 #include "MuR/OpImageResize.h"
 #include "MuR/OpImageSwizzle.h"
+#include "MuR/OpImageSaturate.h"
+#include "MuR/OpImageInvert.h"
 #include "MuR/Operations.h"
 #include "MuR/Platform.h"
 #include "MuR/Ptr.h"
@@ -447,7 +449,6 @@ namespace mu
 		}
 	}
 
-
 	//---------------------------------------------------------------------------------------------
 	void FImageLayerTask::DoWork()
 	{
@@ -513,7 +514,15 @@ namespace mu
 				// This is a frequent critical-path case because of multilayer projectors.
 				bDone = true;
 
-				BufferLayerComposite<BlendChannelMasked, LightenChannel, false>(pNew.get(), m_blended.get(), bOnlyOneMip);
+				constexpr bool bUseVectorImplementation = false;	
+				if constexpr (bUseVectorImplementation)
+				{
+					BufferLayerCompositeVector<VectorBlendChannelMasked, VectorLightenChannel, false>(pNew.get(), m_blended.get(), bOnlyOneMip);
+				}
+				else
+				{
+					BufferLayerComposite<BlendChannelMasked, LightenChannel, false>(pNew.get(), m_blended.get(), bOnlyOneMip);
+				}
 			}
 
 
@@ -1099,6 +1108,297 @@ namespace mu
 	}
 
 
+	//---------------------------------------------------------------------------------------------
+	//---------------------------------------------------------------------------------------------
+	//---------------------------------------------------------------------------------------------
+	class FImageSaturateTask : public CodeRunner::FIssuedTask
+	{
+	public:
+		FImageSaturateTask(FScheduledOp, FProgramCache&, const OP::ImageSaturateArgs&);
+
+		// FIssuedTask interface
+		void DoWork() override;
+		void Complete(CodeRunner* staged) override;
+
+	private:
+		Ptr<const Image> Source;
+		Ptr<const Image> Result;
+		float Factor;
+	};
+
+
+	//---------------------------------------------------------------------------------------------
+	FImageSaturateTask::FImageSaturateTask(FScheduledOp InOp, FProgramCache& Memory, const OP::ImageSaturateArgs& InArgs)
+	{
+		Op = InOp;
+		Source = Memory.GetImage(FCacheAddress(InArgs.base, InOp));
+		Factor = Memory.GetScalar(FScheduledOp::FromOpAndOptions(InArgs.factor, InOp, 0));
+	}
+
+	//---------------------------------------------------------------------------------------------
+	void FImageSaturateTask::DoWork()
+	{
+		MUTABLE_CPUPROFILER_SCOPE(FImageSaturateTask);
+
+		// For now don't enable this optimization 
+		const bool bOptimizeUnchanged = false;//FMath::IsNearlyEqual(Factor, 1.0f);
+
+		if (Source->IsUnique())
+		{
+			Ptr<Image> MutableImage = mu::CloneOrTakeOver<>(Source.get());
+			if (!bOptimizeUnchanged)
+			{
+				constexpr bool bUseVectorIntrinsics = false;
+				ImageSaturateInPlace<bUseVectorIntrinsics>(MutableImage.get(), Factor);
+			}
+
+			Result = MutableImage;
+		}
+		else
+		{
+			if (bOptimizeUnchanged)
+			{
+				Result = Source->Clone();
+			}
+			else
+			{
+				Result = ImageSaturate(Source.get(), Factor);
+			}
+		}
+	}
+
+	//---------------------------------------------------------------------------------------------
+	void FImageSaturateTask::Complete(CodeRunner* Runner)
+	{
+		// This runs in the runner thread
+		Runner->GetMemory().SetImage(Op, Result);
+	}
+
+
+	//---------------------------------------------------------------------------------------------
+	//---------------------------------------------------------------------------------------------
+	//---------------------------------------------------------------------------------------------
+	class FImageResizeTask : public CodeRunner::FIssuedTask
+	{
+	public:
+		FImageResizeTask(FScheduledOp, FProgramCache&, const OP::ImageResizeArgs&, int32 ImageCompressionQuality);
+
+		// FIssuedTask interface
+		void DoWork() override;
+		void Complete(CodeRunner* staged) override;
+
+	private:
+		Ptr<const Image> pBase;
+		Ptr<const Image> Result;
+		OP::ImageResizeArgs Args;
+		int32 ImageCompressionQuality;
+	};
+
+
+	//---------------------------------------------------------------------------------------------
+	FImageResizeTask::FImageResizeTask(FScheduledOp InOp, FProgramCache& Memory, const OP::ImageResizeArgs& InArgs, int32 InImageCompressionQuality)
+	{
+		Op = InOp;
+		Args = InArgs;
+		pBase = Memory.GetImage(FCacheAddress(InArgs.source, InOp));
+		ImageCompressionQuality = InImageCompressionQuality;
+	}
+
+	//---------------------------------------------------------------------------------------------
+	void FImageResizeTask::DoWork()
+	{
+		MUTABLE_CPUPROFILER_SCOPE(FImageResizeTask);
+
+		if (!pBase)
+		{
+			Result = nullptr;
+			return;
+		}
+
+		FImageSize destSize = FImageSize
+			(
+				Args.size[0],
+				Args.size[1]
+			);
+
+		ImagePtr pResult;
+
+		if ( destSize[0]!=pBase->GetSizeX()
+			 ||
+			 destSize[1]!=pBase->GetSizeY() )
+		{
+			//pResult = ImageResize( pBase.get(), destSize );
+			pResult = ImageResizeLinear( ImageCompressionQuality, pBase.get(), destSize );
+
+			// If the source image had mips, generate them as well for the resized image.
+			// This shouldn't happen often since "ResizeLike" should be usually optimised out
+			// during model compilation. The mipmap generation below is not very precise with
+			// the number of mips that are needed and will probably generate too many
+			bool sourceHasMips = pBase->GetLODCount()>1;
+			if (sourceHasMips)
+			{
+				int levelCount = Image::GetMipmapCount( pResult->GetSizeX(), pResult->GetSizeY() );
+				ImagePtr pMipmapped = new Image( pResult->GetSizeX(), pResult->GetSizeY(),
+												 levelCount,
+												 pResult->GetFormat() );
+
+				SCRATCH_IMAGE_MIPMAP scratch;
+				FMipmapGenerationSettings mipSettings{};
+
+				ImageMipmap_PrepareScratch( pMipmapped.get(), pResult.get(), levelCount, &scratch );
+				ImageMipmap( ImageCompressionQuality,
+							 pMipmapped.get(), pResult.get(), levelCount, &scratch, mipSettings );
+
+				pResult = pMipmapped;
+			}
+
+		}
+		else
+		{
+			pResult = pBase->Clone();
+		}
+
+		Result = pResult;
+	}
+
+	//---------------------------------------------------------------------------------------------
+	void FImageResizeTask::Complete(CodeRunner* Runner)
+	{
+		// This runs in the runner thread
+		Runner->GetMemory().SetImage(Op, Result);
+	}
+
+
+	//---------------------------------------------------------------------------------------------
+	//---------------------------------------------------------------------------------------------
+	//---------------------------------------------------------------------------------------------
+	class FImageResizeRelTask : public CodeRunner::FIssuedTask
+	{
+	public:
+		FImageResizeRelTask(FScheduledOp, FProgramCache&, const OP::ImageResizeRelArgs&, int32 ImageCompressionQuality);
+
+		// FIssuedTask interface
+		void DoWork() override;
+		void Complete(CodeRunner* staged) override;
+
+	private:
+		Ptr<const Image> pBase;
+		Ptr<const Image> Result;
+		OP::ImageResizeRelArgs Args;
+		int32 ImageCompressionQuality;
+	};
+
+
+	//---------------------------------------------------------------------------------------------
+	FImageResizeRelTask::FImageResizeRelTask(FScheduledOp InOp, FProgramCache& Memory, const OP::ImageResizeRelArgs& InArgs, int32 InImageCompressionQuality)
+	{
+		Op = InOp;
+		Args = InArgs;
+		pBase = Memory.GetImage(FCacheAddress(InArgs.source, InOp));
+		ImageCompressionQuality = InImageCompressionQuality;
+	}
+
+	//---------------------------------------------------------------------------------------------
+	void FImageResizeRelTask::DoWork()
+	{
+		MUTABLE_CPUPROFILER_SCOPE(FImageResizeRelTask);
+
+		if (!pBase)
+		{
+			Result = nullptr;
+			return;
+		}
+
+		FImageSize destSize(
+					uint16( FMath::Max(1.0, pBase->GetSizeX()*Args.factor[0] + 0.5f ) ),
+					uint16( FMath::Max(1.0, pBase->GetSizeY()*Args.factor[1] + 0.5f ) ) );
+
+		//pResult = ImageResize( pBase.get(), destSize );
+		ImagePtr pResult = ImageResizeLinear(
+			ImageCompressionQuality, pBase.get(), destSize );
+
+		// If the source image had mips, generate them as well for the resized image.
+		// This shouldn't happen often since "ResizeLike" should be usually optimised out
+		// during model compilation. The mipmap generation below is not very precise with
+		// the number of mips that are needed and will probably generate too many
+		bool sourceHasMips = pBase->GetLODCount()>1;
+		if (sourceHasMips)
+		{
+			int levelCount = Image::GetMipmapCount(pResult->GetSizeX(), pResult->GetSizeY());
+			ImagePtr pMipmapped = new Image(pResult->GetSizeX(), pResult->GetSizeY(),
+				levelCount,
+				pResult->GetFormat());
+
+			SCRATCH_IMAGE_MIPMAP scratch;
+			FMipmapGenerationSettings mipSettings{};
+
+			ImageMipmap_PrepareScratch(pMipmapped.get(), pResult.get(), levelCount, &scratch);
+			ImageMipmap(ImageCompressionQuality,
+				pMipmapped.get(), pResult.get(), levelCount, &scratch, mipSettings);
+
+			pResult = pMipmapped;
+		}
+
+		Result = pResult;
+	}
+
+	//---------------------------------------------------------------------------------------------
+	void FImageResizeRelTask::Complete(CodeRunner* Runner)
+	{
+		// This runs in the runner thread
+		Runner->GetMemory().SetImage(Op, Result);
+	}
+
+	//---------------------------------------------------------------------------------------------
+	//---------------------------------------------------------------------------------------------
+	//---------------------------------------------------------------------------------------------
+	class FImageInvertTask : public CodeRunner::FIssuedTask
+	{
+	public:
+		FImageInvertTask(FScheduledOp, FProgramCache&, const OP::ImageInvertArgs&);
+
+		// FIssuedTask interface
+		void DoWork() override;
+		void Complete(CodeRunner* Runner) override;
+
+	private:
+		Ptr<const Image> Source;
+		Ptr<const Image> Result;
+	};
+
+
+	//---------------------------------------------------------------------------------------------
+	FImageInvertTask::FImageInvertTask(FScheduledOp InOp, FProgramCache& Memory, const OP::ImageInvertArgs& InArgs)
+	{
+		Op = InOp;
+		Source = Memory.GetImage(FCacheAddress(InArgs.base, InOp));
+	}
+
+	//---------------------------------------------------------------------------------------------
+	void FImageInvertTask::DoWork()
+	{
+		MUTABLE_CPUPROFILER_SCOPE(FImageInvertTask);
+
+		if (Source->IsUnique())
+		{
+			Ptr<Image> MutableImage = mu::CloneOrTakeOver<>(Source.get());
+			ImageInvertInPlace(MutableImage.get());
+			Result = MutableImage;
+		}
+		else
+		{
+			Result = ImageInvert(Source.get());
+		}
+	}
+
+	//---------------------------------------------------------------------------------------------
+	void FImageInvertTask::Complete(CodeRunner* Runner)
+	{
+		// This runs in the runner thread
+		Runner->GetMemory().SetImage(Op, Result);
+	}
+
+
 	//---------------------------------------------------------------------------------------------
 	//---------------------------------------------------------------------------------------------
 	//---------------------------------------------------------------------------------------------
@@ -1698,6 +1998,46 @@ namespace mu
 			break;
 		}
 
+		case OP_TYPE::IM_SATURATE:
+		{
+			if (item.Stage == 1)
+			{
+				OP::ImageSaturateArgs Args = program.GetOpArgs<OP::ImageSaturateArgs>(item.At);
+				Issued = MakeShared<FImageSaturateTask>(item, GetMemory(), Args);
+			}
+			break;
+		}
+
+		case OP_TYPE::IM_INVERT:
+		{
+			if (item.Stage == 1)
+			{
+				OP::ImageInvertArgs Args = program.GetOpArgs<OP::ImageInvertArgs>(item.At);
+				Issued = MakeShared<FImageInvertTask>(item, GetMemory(), Args);
+			}
+			break;
+		}
+
+		case OP_TYPE::IM_RESIZE:
+		{
+			if (item.Stage == 1)
+			{
+				OP::ImageResizeArgs Args = program.GetOpArgs<OP::ImageResizeArgs>(item.At);
+				Issued = MakeShared<FImageResizeTask>(item, GetMemory(), Args, m_pSettings->GetPrivate()->m_imageCompressionQuality);
+			}
+			break;
+		}
+
+		case OP_TYPE::IM_RESIZEREL:
+		{
+			if (item.Stage == 1)
+			{
+				OP::ImageResizeRelArgs Args = program.GetOpArgs<OP::ImageResizeRelArgs>(item.At);
+				Issued = MakeShared<FImageResizeRelTask>(item, GetMemory(), Args, m_pSettings->GetPrivate()->m_imageCompressionQuality);
+			}
+			break;
+		}
+
 		case OP_TYPE::IM_COMPOSE:
 		{
 			if (item.Stage == 2)
diff --git a/Engine/Plugins/Experimental/Mutable/Source/MutableRuntime/Private/MuR/OpImageBlend.h b/Engine/Plugins/Experimental/Mutable/Source/MutableRuntime/Private/MuR/OpImageBlend.h
--- a/Engine/Plugins/Experimental/Mutable/Source/MutableRuntime/Private/MuR/OpImageBlend.h
+++ b/Engine/Plugins/Experimental/Mutable/Source/MutableRuntime/Private/MuR/OpImageBlend.h
@@ -4,7 +4,7 @@
 
 #include "MuR/ImagePrivate.h"
 #include "MuR/OpImageLayer.h"
-
+#include "Math/VectorRegister.h"
 
 namespace mu
 {
@@ -135,8 +135,7 @@ namespace mu
 		uint32 overlay = base + (blended * uint32(255 - base) >> 8);
 		uint32 masked = (((255 - mask) * base) + (mask * overlay)) >> 8;
 		return masked;
-	}
-
+	}	
 
 	inline uint32 LightenChannel(uint32 base, uint32 blended)
 	{
@@ -173,4 +172,21 @@ namespace mu
 		return overlay;
 	}
 
+	//---------------------------------------------------------------------------------------------
+	FORCEINLINE VectorRegister4Int VectorBlendChannelMasked(
+		VectorRegister4Int Base, VectorRegister4Int Blended, VectorRegister4Int Mask)
+	{
+		const VectorRegister4Int Value = VectorIntAdd(
+			VectorIntMultiply(Base, VectorIntSubtract(VectorIntSet1(255), Mask)),
+			VectorIntMultiply(Blended, Mask));
+
+		// fast division by 255 assuming Value is in the range [0, (1 << 16)]
+		return VectorShiftRightImmLogical(VectorIntMultiply(Value, VectorIntSet1(32897)), 23);
+	}
+
+	FORCEINLINE int32 VectorLightenChannel(int32 Base, int32 Blended)
+	{
+		return Base + (Blended * (255 - Base) >> 8);
+	}
+
 }
diff --git a/Engine/Plugins/Experimental/Mutable/Source/MutableRuntime/Private/MuR/OpImageInvert.h b/Engine/Plugins/Experimental/Mutable/Source/MutableRuntime/Private/MuR/OpImageInvert.h
--- a/Engine/Plugins/Experimental/Mutable/Source/MutableRuntime/Private/MuR/OpImageInvert.h
+++ b/Engine/Plugins/Experimental/Mutable/Source/MutableRuntime/Private/MuR/OpImageInvert.h
@@ -5,6 +5,8 @@
 #include "MuR/ImagePrivate.h"
 #include "MuR/Platform.h"
 
+#include "Math/VectorRegister.h"
+
 namespace mu
 {
 
@@ -65,6 +67,7 @@ namespace mu
 	}
 
 	/** Invert the colour (RGB or L) components of an image. Leave alpha untouched. */
+	template<bool bUseVectorImpl = false>
 	inline void ImageInvertInPlace(Image* pA)
 	{
 		MUTABLE_CPUPROFILER_SCOPE(ImageInvertInPlace);
@@ -77,26 +80,81 @@ namespace mu
 		uint8* pABuf = pA->GetData();
 
 		//Generic implementation
-		int32 pixelCount = pA->CalculatePixelCount();
+		const int32 pixelCount = pA->CalculatePixelCount();
 
 		switch (pA->GetFormat())
 		{
 		case EImageFormat::IF_L_UBYTE:
 		{
-			for (int32 i = 0; i < pixelCount; ++i)
+			if constexpr (!bUseVectorImpl)
 			{
-				pABuf[i] = 255 - pABuf[i];
+				for (int32 i = 0; i < pixelCount; ++i)
+				{
+					pABuf[i] = 255 - pABuf[i];
+				}
+			}
+			else
+			{
+				const int32 PixelRem = pixelCount % 4;
+
+				constexpr VectorRegister4Int One255 = MakeVectorRegisterIntConstant(255, 255, 255, 255);
+				for (int32 i = 0; i < pixelCount - PixelRem; i += 4)
+				{
+					VectorRegister4Int Value = MakeVectorRegisterInt(pABuf[i], pABuf[i + 1], pABuf[i + 2], pABuf[i + 3]);
+					Value = VectorIntSubtract(One255, Value);
+
+					alignas(alignof(VectorRegister4Int)) int32 IndexableValue[4];
+					VectorIntStoreAligned(Value, IndexableValue);
+
+					pABuf[i + 0] = static_cast<uint8>(IndexableValue[0]);
+					pABuf[i + 1] = static_cast<uint8>(IndexableValue[1]);
+					pABuf[i + 2] = static_cast<uint8>(IndexableValue[2]);
+					pABuf[i + 3] = static_cast<uint8>(IndexableValue[3]);
+				}
+
+				for (int32 i = pixelCount - PixelRem; i < pixelCount; ++i)
+				{
+					pABuf[i] = 255 - pABuf[i];
+				}
 			}
 			break;
 		}
 
 		case EImageFormat::IF_RGB_UBYTE:
 		{
-			for (int32 i = 0; i < pixelCount; ++i)
+			if constexpr (!bUseVectorImpl)
 			{
-				for (int32 c = 0; c < 3; ++c)
+				for (int32 i = 0; i < pixelCount; ++i)
 				{
-					pABuf[i * 3 + c] = 255 - pABuf[i * 3 + c];
+					for (int32 c = 0; c < 3; ++c)
+					{
+						pABuf[i * 3 + c] = 255 - pABuf[i * 3 + c];
+					}
+				}
+			}
+			else
+			{
+				const int32 ElemCount = pixelCount * 3;
+				const int32 ElemRem = ElemCount % 4;
+
+				constexpr VectorRegister4Int One255 = MakeVectorRegisterIntConstant(255, 255, 255, 255);
+				for (int32 i = 0; i < ElemCount - ElemRem; i += 4)
+				{
+					VectorRegister4Int Value = MakeVectorRegisterInt(pABuf[i], pABuf[i + 1], pABuf[i + 2], pABuf[i + 3]);
+					Value = VectorIntSubtract(One255, Value);
+
+					alignas(alignof(VectorRegister4Int)) int32 IndexableValue[4];
+					VectorIntStoreAligned(Value, IndexableValue);
+
+					pABuf[i + 0] = static_cast<uint8>(IndexableValue[0]);
+					pABuf[i + 1] = static_cast<uint8>(IndexableValue[1]);
+					pABuf[i + 2] = static_cast<uint8>(IndexableValue[2]);
+					pABuf[i + 3] = static_cast<uint8>(IndexableValue[3]);
+				}
+
+				for (int32 i = ElemCount - ElemRem; i < ElemCount; ++i)
+				{
+					pABuf[i] = 255 - pABuf[i];
 				}
 			}
 			break;
@@ -104,11 +162,34 @@ namespace mu
 		case EImageFormat::IF_RGBA_UBYTE:
 		case EImageFormat::IF_BGRA_UBYTE:
 		{
-			for (int32 i = 0; i < pixelCount; ++i)
+			if constexpr (!bUseVectorImpl)
 			{
-				for (int32 c = 0; c < 3; ++c)
+				for (int32 i = 0; i < pixelCount; ++i)
 				{
-					pABuf[i * 4 + c] = 255 - pABuf[i * 4 + c];
+					for (int32 c = 0; c < 3; ++c)
+					{
+						pABuf[i * 4 + c] = 255 - pABuf[i * 4 + c];
+					}
+				}
+			}
+			else
+			{
+				const int32 ElemCount = pixelCount * 4;
+
+				constexpr VectorRegister4Int One255 = MakeVectorRegisterIntConstant(255, 255, 255, 255);
+
+				for (int32 i = 0; i < ElemCount; i += 4)
+				{
+					VectorRegister4Int Value = MakeVectorRegisterInt(pABuf[i], pABuf[i + 1], pABuf[i + 2], pABuf[i + 3]);
+					Value = VectorIntSubtract(One255, Value);
+
+					alignas(alignof(VectorRegister4Int)) int32 IndexableValue[4];
+					VectorIntStoreAligned(Value, IndexableValue);
+
+					pABuf[i + 0] = static_cast<uint8>(IndexableValue[0]);
+					pABuf[i + 1] = static_cast<uint8>(IndexableValue[1]);
+					pABuf[i + 2] = static_cast<uint8>(IndexableValue[2]);
+					//pABuf[i + 3] = static_cast<uint8>(IndexableValue[3]);
 				}
 			}
 			break;
diff --git a/Engine/Plugins/Experimental/Mutable/Source/MutableRuntime/Private/MuR/OpImageLayer.h b/Engine/Plugins/Experimental/Mutable/Source/MutableRuntime/Private/MuR/OpImageLayer.h
--- a/Engine/Plugins/Experimental/Mutable/Source/MutableRuntime/Private/MuR/OpImageLayer.h
+++ b/Engine/Plugins/Experimental/Mutable/Source/MutableRuntime/Private/MuR/OpImageLayer.h
@@ -1235,7 +1235,6 @@ namespace mu
 		}
 	}
 
-
 	//---------------------------------------------------------------------------------------------
 	template< unsigned (*BLEND_FUNC_MASKED)(unsigned, unsigned, unsigned),
 		unsigned (*BLEND_FUNC)(unsigned, unsigned),
@@ -1586,6 +1585,85 @@ namespace mu
 			});
 	}
 
+	//---------------------------------------------------------------------------------------------
+	template< VectorRegister4Int (*RGB_FUNC_MASKED)(VectorRegister4Int, VectorRegister4Int, VectorRegister4Int),
+		int32 (*A_FUNC)(int32, int32),
+		bool CLAMP >
+	inline void BufferLayerCompositeVector(
+		Image* pBase,
+		const Image* pBlend,
+		bool bOnlyFirstLOD)
+	{
+		check(pBase->GetFormat() == EImageFormat::IF_RGBA_UBYTE);
+		check(pBlend->GetFormat() == EImageFormat::IF_RGBA_UBYTE);
+		check(pBase->GetSizeX() == pBlend->GetSizeX() && pBase->GetSizeY() == pBlend->GetSizeY());
+		check(bOnlyFirstLOD || pBase->GetLODCount() <= pBlend->GetLODCount());
+
+		uint8* pBaseBuf = pBase->GetData();
+		const uint8* pBlendedBuf = pBlend->GetData();
+
+		// The base determines the number of lods to process.
+		int32 LODCount = bOnlyFirstLOD ? 1 : pBase->GetLODCount();
+
+		int32 PixelCount = 0;
+		if (bOnlyFirstLOD)
+		{
+			PixelCount = pBase->GetSizeX() * pBase->GetSizeY();
+		}
+		else
+		{
+			PixelCount = pBase->CalculatePixelCount();
+		}
+
+		constexpr int32 NumBatchElems = 4096*2;
+	
+		const int32 NumBatches = (PixelCount / NumBatchElems) + static_cast<int32>((PixelCount % NumBatchElems) != 0);
+		const int32 LastBatchNumElems = (PixelCount % NumBatchElems) > 0 ? (PixelCount % NumBatchElems) : NumBatchElems;
+
+		ParallelFor(NumBatches,
+			[
+				pBaseBuf, pBlendedBuf, NumBatches, LastBatchNumElems
+			] (uint32 BatchId)
+			{
+				const int32 CurrentBatchNumElems = static_cast<int32>(BatchId) < (NumBatches - 1) ? NumBatchElems : LastBatchNumElems;
+				for (int32 BatchElemIndex = 0; BatchElemIndex < CurrentBatchNumElems; ++BatchElemIndex)
+				{
+					const int32 i = BatchId * NumBatchElems + BatchElemIndex;
+
+					// TODO: Optimize this (SIMD?)
+					const int32 BaseAlpha = pBaseBuf[4 * i + 3];
+					const int32 BlendedAlpha = pBlendedBuf[4 * i + 3];
+
+					const VectorRegister4Int Mask = VectorIntSet1(BlendedAlpha);
+					const VectorRegister4Int Blended = MakeVectorRegisterInt(pBlendedBuf[4 * i + 0], pBlendedBuf[4 * i + 1], pBlendedBuf[4 * i + 2], 0);
+					const VectorRegister4Int Base = MakeVectorRegisterInt(pBaseBuf[4 * i + 0], pBaseBuf[4 * i + 1], pBaseBuf[4 * i + 2], 0);
+
+					VectorRegister4Int Result = RGB_FUNC_MASKED(Base, Blended, Mask);
+					if constexpr (CLAMP)
+					{
+						Result = VectorIntMin(VectorIntSet1(255), Result);
+					}
+
+					int32 AlphaResult = A_FUNC(BaseAlpha, BlendedAlpha);
+					if constexpr (CLAMP)
+					{
+						AlphaResult = FMath::Min(255, AlphaResult);
+					}
+
+					struct alignas(alignof(VectorRegister4Int)) AlignedOutputRegister4Int
+					{
+						int32 Data[4];
+					} RegisterOutput;
+
+					VectorIntStoreAligned(Result, &RegisterOutput);
+
+					pBaseBuf[4 * i + 0] = static_cast<uint8>(RegisterOutput.Data[0]);
+					pBaseBuf[4 * i + 1] = static_cast<uint8>(RegisterOutput.Data[1]);
+					pBaseBuf[4 * i + 2] = static_cast<uint8>(RegisterOutput.Data[2]);
+					pBaseBuf[4 * i + 3] = static_cast<uint8>(AlphaResult);
+				}
+			});
+	}
 
 	//---------------------------------------------------------------------------------------------
 	template< unsigned (*BLEND_FUNC)(unsigned,unsigned),
diff --git a/Engine/Plugins/Experimental/Mutable/Source/MutableRuntime/Private/MuR/OpImageSaturate.h b/Engine/Plugins/Experimental/Mutable/Source/MutableRuntime/Private/MuR/OpImageSaturate.h
--- a/Engine/Plugins/Experimental/Mutable/Source/MutableRuntime/Private/MuR/OpImageSaturate.h
+++ b/Engine/Plugins/Experimental/Mutable/Source/MutableRuntime/Private/MuR/OpImageSaturate.h
@@ -4,7 +4,7 @@
 
 #include "MuR/ImagePrivate.h"
 #include "MuR/Platform.h"
-
+#include "Math/VectorRegister.h"
 
 namespace mu
 {
@@ -24,13 +24,13 @@ namespace mu
         const uint8_t* pABuf = pA->GetData();
 
 		// Generic implementation
-		int pixelCount = (int)pA->CalculatePixelCount();
+		int32 PixelCount = (int)pA->CalculatePixelCount();
 
 		switch ( pA->GetFormat() )
 		{
 		case EImageFormat::IF_RGB_UBYTE:
 		{
-			for ( int i=0; i<pixelCount; ++i )
+			for ( int i=0; i<PixelCount; ++i )
 			{
                 int32_t l_16 = 76 * pABuf[3*i+0] + 150 * pABuf[3*i+1] + 29 * pABuf[3*i+2];
 
@@ -41,12 +41,13 @@ namespace mu
                     pDestBuf[3*i+c] = (uint8_t)FMath::Min( (uint32_t)r_16>>8, 255u );
 				}
 			}
+	
 			break;
 		}
 
         case EImageFormat::IF_RGBA_UBYTE:
         {
-            for ( int i=0; i<pixelCount; ++i )
+            for ( int i=0; i<PixelCount; ++i )
             {
                 int32_t l_16 = 76 * pABuf[4*i+0] + 150 * pABuf[4*i+1] + 29 * pABuf[4*i+2];
 
@@ -64,7 +65,7 @@ namespace mu
 
         case EImageFormat::IF_BGRA_UBYTE:
         {
-            for ( int i=0; i<pixelCount; ++i )
+            for ( int i=0; i<PixelCount; ++i )
             {
                 int32_t l_16 = 76 * pABuf[4*i+2] + 150 * pABuf[4*i+1] + 29 * pABuf[4*i+0];
 
@@ -87,4 +88,202 @@ namespace mu
 		return pDest;
 	}
 
+	template<bool bUseVectorImpl = false>
+	inline void ImageSaturateInPlace(Image* pA, float Factor)
+	{
+		MUTABLE_CPUPROFILER_SCOPE(ImageSaturateInplace);
+
+		if (!pA)
+		{
+			return;
+		}
+
+		// Clamp the factor
+		// TODO: See what happens if we don't
+		Factor = FMath::Max(0.0f, Factor);
+        int32_t f_8 = (int32_t)(Factor*255);
+
+        uint8_t* pABuf = pA->GetData();
+
+		// Generic implementation
+		const int32 PixelCount = (int32)pA->CalculatePixelCount();
+
+		switch ( pA->GetFormat() )
+		{
+		case EImageFormat::IF_RGB_UBYTE:
+		{
+			if constexpr (!bUseVectorImpl)
+			{
+				for (int i = 0; i < PixelCount; ++i)
+				{
+					int32_t l_16 = 76 * pABuf[3 * i + 0] + 150 * pABuf[3 * i + 1] + 29 * pABuf[3 * i + 2];
+
+					for (int c = 0; c < 3; ++c)
+					{
+						int32_t d_16 = (pABuf[3 * i + c] << 8) - l_16;
+						int32_t r_16 = l_16 + ((d_16 * f_8) >> 8);
+						pABuf[3 * i + c] = (uint8_t)FMath::Min((uint32_t)r_16 >> 8, 255u);
+					}
+				}
+			}
+			else // if constexpr bUseVectorImpl
+			{
+				constexpr VectorRegister4Int StaurationFactorRGB = MakeVectorRegisterIntConstant(76, 150, 29, 0);
+				const VectorRegister4Int F = VectorIntSet1(static_cast<int32>(Factor * 255.0f));
+				const VectorRegister4Int OneMinusF = VectorIntSet1(255 - static_cast<int32>(Factor * 255.0f));
+				for (int32 PixelIndex = 0; PixelIndex < PixelCount; ++PixelIndex)
+				{
+					const VectorRegister4Int Value = MakeVectorRegisterInt(
+						static_cast<int32>(pABuf[3 * PixelIndex + 0]), 
+						static_cast<int32>(pABuf[3 * PixelIndex + 1]), 
+						static_cast<int32>(pABuf[3 * PixelIndex + 2]), 
+						0);
+
+					const VectorRegister4Int ScaledValue = VectorIntMultiply(Value, StaurationFactorRGB);
+					
+					alignas(VectorRegister4Int) int32 IndexableScaledValue[4];
+					VectorIntStoreAligned(ScaledValue, IndexableScaledValue);
+
+					// Add scaled values and divide by 255
+					const VectorRegister4Int L = VectorShiftRightImmLogical(
+						VectorIntMultiply(
+							VectorIntSet1(IndexableScaledValue[0] + IndexableScaledValue[1] + IndexableScaledValue[2]), 
+							VectorIntSet1(32897)), 23);
+					
+					const VectorRegister4Int Lerp = VectorIntAdd(VectorIntMultiply(L, OneMinusF), VectorIntMultiply(Value, F));
+				
+					alignas(VectorRegister4Int) int32 IndexableResult[4];
+
+					// divide by 255 and store. do we need to clamp?
+					VectorIntStoreAligned(VectorShiftRightImmLogical(VectorIntMultiply(Lerp, VectorIntSet1(32897)), 23),  IndexableResult);
+
+					pABuf[3 * PixelIndex + 0] = static_cast<uint8>(IndexableResult[0]);
+					pABuf[3 * PixelIndex + 1] = static_cast<uint8>(IndexableResult[1]);
+					pABuf[3 * PixelIndex + 2] = static_cast<uint8>(IndexableResult[2]);
+				}
+			}
+			break;
+		}
+
+        case EImageFormat::IF_RGBA_UBYTE:
+		{
+			if constexpr (!bUseVectorImpl)
+			{
+				for (int i = 0; i < PixelCount; ++i)
+				{
+					int32_t l_16 = 76 * pABuf[4 * i + 0] + 150 * pABuf[4 * i + 1] + 29 * pABuf[4 * i + 2];
+
+					for (int c = 0; c < 3; ++c)
+					{
+						int32_t d_16 = (pABuf[4 * i + c] << 8) - l_16;
+						int32_t r_16 = l_16 + ((d_16 * f_8) >> 8);
+						pABuf[4 * i + c] = (uint8_t)FMath::Min((uint32_t)r_16 >> 8, 255u);
+					}
+
+					//pABuf[4 * i + 3] = pABuf[4 * i + 3];
+				}
+			}
+			else // if constexpr bUseVectorImpl
+			{
+				constexpr VectorRegister4Int StaurationFactorRGB = MakeVectorRegisterIntConstant(76, 150, 29, 0);
+				const VectorRegister4Int F = VectorIntSet1(static_cast<int32>(Factor * 255.0f));
+				const VectorRegister4Int OneMinusF = VectorIntSet1(255 - static_cast<int32>(Factor * 255.0f));
+				for (int32 PixelIndex = 0; PixelIndex < PixelCount; ++PixelIndex)
+				{
+					const VectorRegister4Int Value = MakeVectorRegisterInt(
+						static_cast<int32>(pABuf[4 * PixelIndex + 0]), 
+						static_cast<int32>(pABuf[4 * PixelIndex + 1]), 
+						static_cast<int32>(pABuf[4 * PixelIndex + 2]), 
+						0);
+
+					const VectorRegister4Int ScaledValue = VectorIntMultiply(Value, StaurationFactorRGB);
+					
+					alignas(VectorRegister4Int) int32 IndexableScaledValue[4];
+					VectorIntStoreAligned(ScaledValue, IndexableScaledValue);
+
+					// Add scaled values and divide by 255
+					const VectorRegister4Int L = VectorShiftRightImmLogical(
+						VectorIntMultiply(
+							VectorIntSet1(IndexableScaledValue[0] + IndexableScaledValue[1] + IndexableScaledValue[2]), 
+							VectorIntSet1(32897)), 23);
+					
+					const VectorRegister4Int Lerp = VectorIntAdd(VectorIntMultiply(L, OneMinusF), VectorIntMultiply(Value, F));
+				
+					alignas(VectorRegister4Int) int32 IndexableResult[4];
+					// divide by 255 and store. do we need to clamp?
+					VectorIntStoreAligned(VectorShiftRightImmLogical(VectorIntMultiply(Lerp, VectorIntSet1(32897)), 23),  IndexableResult);
+
+					pABuf[4 * PixelIndex + 0] = static_cast<uint8>(IndexableResult[0]);
+					pABuf[4 * PixelIndex + 1] = static_cast<uint8>(IndexableResult[1]);
+					pABuf[4 * PixelIndex + 2] = static_cast<uint8>(IndexableResult[2]);
+					//pABuf[4 * PixelIndex + 3] = pABuf[4 * PixelIndex + 3];
+				}
+			}
+
+            break;
+        }
+
+        case EImageFormat::IF_BGRA_UBYTE:
+        {
+			if constexpr (!bUseVectorImpl)
+			{
+				for (int i = 0; i < PixelCount; ++i)
+				{
+					int32_t l_16 = 76 * pABuf[4 * i + 2] + 150 * pABuf[4 * i + 1] + 29 * pABuf[4 * i + 0];
+
+					for (int c = 0; c < 3; ++c)
+					{
+						int32_t d_16 = (pABuf[4 * i + c] << 8) - l_16;
+						int32_t r_16 = l_16 + ((d_16 * f_8) >> 8);
+						pABuf[4 * i + c] = (uint8_t)FMath::Min((uint32_t)r_16 >> 8, 255u);
+					}
+
+					//pABuf[4 * i + 3] = pABuf[4 * i + 3];
+				}
+			}
+			else // if constexpr bUseVectorImpl
+			{
+				constexpr VectorRegister4Int StaurationFactorBGR = MakeVectorRegisterIntConstant(29, 150, 76, 0);
+				const VectorRegister4Int F = VectorIntSet1(static_cast<int32>(Factor * 255.0f));
+				const VectorRegister4Int OneMinusF = VectorIntSet1(255 - static_cast<int32>(Factor * 255.0f));
+				for (int32 PixelIndex = 0; PixelIndex < PixelCount; ++PixelIndex)
+				{
+					const VectorRegister4Int Value = MakeVectorRegisterInt(
+						static_cast<int32>(pABuf[4 * PixelIndex + 0]), 
+						static_cast<int32>(pABuf[4 * PixelIndex + 1]), 
+						static_cast<int32>(pABuf[4 * PixelIndex + 2]), 
+						0);
+
+					const VectorRegister4Int ScaledValue = VectorIntMultiply(Value, StaurationFactorBGR);
+					
+					alignas(VectorRegister4Int) int32 IndexableScaledValue[4];
+					VectorIntStoreAligned(ScaledValue, IndexableScaledValue);
+
+					// Add scaled values and divide by 255
+					const VectorRegister4Int L = VectorShiftRightImmLogical(
+						VectorIntMultiply(
+							VectorIntSet1(IndexableScaledValue[0] + IndexableScaledValue[1] + IndexableScaledValue[2]), 
+							VectorIntSet1(32897)), 23);
+					
+					const VectorRegister4Int Lerp = VectorIntAdd(VectorIntMultiply(L, OneMinusF), VectorIntMultiply(Value, F));
+				
+					alignas(VectorRegister4Int) int32 IndexableResult[4];
+
+					// divide by 255 and store. do we need to clamp?
+					VectorIntStoreAligned(VectorShiftRightImmLogical(VectorIntMultiply(Lerp, VectorIntSet1(32897)), 23),  IndexableResult);
+
+					pABuf[4 * PixelIndex + 0] = static_cast<uint8>(IndexableResult[0]);
+					pABuf[4 * PixelIndex + 1] = static_cast<uint8>(IndexableResult[1]);
+					pABuf[4 * PixelIndex + 2] = static_cast<uint8>(IndexableResult[2]);
+					//pABuf[4 * PixelIndex + 3] = pABuf[4 * PixelIndex + 3];
+				}
+			}
+            
+			break;
+        }
+
+		default:
+			check(false);
+		}
+	}
 }
