[Mutable] Temporal fix to debug lodbias test crash due to large image resizes.

#rnx

diff --git a/Engine/Plugins/Experimental/Mutable/Source/CustomizableObjectEditor/Private/MuCOE/GenerateMutableSource/GenerateMutableSource.cpp b/Engine/Plugins/Experimental/Mutable/Source/CustomizableObjectEditor/Private/MuCOE/GenerateMutableSource/GenerateMutableSource.cpp
--- a/Engine/Plugins/Experimental/Mutable/Source/CustomizableObjectEditor/Private/MuCOE/GenerateMutableSource/GenerateMutableSource.cpp
+++ b/Engine/Plugins/Experimental/Mutable/Source/CustomizableObjectEditor/Private/MuCOE/GenerateMutableSource/GenerateMutableSource.cpp
@@ -1334,14 +1334,14 @@ void PopulateReferenceSkeletalMeshesData(FMutableGraphGenerationContext& Generat
 int32 ComputeLODBias(const FMutableGraphGenerationContext& GenerationContext, const UTexture2D* ReferenceTexture, int32 MaxTextureSize,
 	const UCustomizableObjectNodeMaterial* MaterialNode, const int32 ImageIndex, bool bUseLODAsBias)
 {
-	int32 LODBias = 0;
+	constexpr int32 MaxAllowedLODBias = 6;
 
 	if (GenerationContext.Options.bForceLargeLODBias)
 	{
-		// This seems to be the highest LODBias we get during cook.
-		return GenerationContext.Options.DebugBias;
+		return FMath::Min(GenerationContext.Options.DebugBias, MaxAllowedLODBias);
 	}
-
+	
+	int32 LODBias = 0;
 	// We used to calculate the lod bias directly from the group like this:
 	//int LODBias = 0;
 	//if (LODSettings.TextureLODGroups.IsValidIndex(ReferenceTexture->LODGroup))
@@ -1393,7 +1393,7 @@ int32 ComputeLODBias(const FMutableGraphGenerationContext& GenerationContext, co
 		UE_LOG(LogMutable, Verbose, TEXT("Compiling texture without reference will have LOD Bias %d."), LODBias);
 	}
 
-	return FMath::Min(6, LODBias);
+	return FMath::Min(LODBias, MaxAllowedLODBias);
 }
 
 
