[mutable] Better code optimization for saturate versus swizzle.
#rnx

diff --git a/Engine/Plugins/Experimental/Mutable/Source/CustomizableObject/Classes/MuCO/CustomizableObject.h b/Engine/Plugins/Experimental/Mutable/Source/CustomizableObject/Classes/MuCO/CustomizableObject.h
--- a/Engine/Plugins/Experimental/Mutable/Source/CustomizableObject/Classes/MuCO/CustomizableObject.h
+++ b/Engine/Plugins/Experimental/Mutable/Source/CustomizableObject/Classes/MuCO/CustomizableObject.h
@@ -1203,7 +1203,7 @@ private:
 	// This is a manual version number for the binary blobs in this asset.
 	// Increasing it invalidates all the previously compiled models.
 	// Warning: If while merging code both versions have changed, take the highest+1.
-	static const int32 CurrentSupportedVersion = 366;
+	static const int32 CurrentSupportedVersion = 368;
 
 public:
 
diff --git a/Engine/Plugins/Experimental/Mutable/Source/CustomizableObjectEditor/Private/MuCOE/SMutableCodeViewer.cpp b/Engine/Plugins/Experimental/Mutable/Source/CustomizableObjectEditor/Private/MuCOE/SMutableCodeViewer.cpp
--- a/Engine/Plugins/Experimental/Mutable/Source/CustomizableObjectEditor/Private/MuCOE/SMutableCodeViewer.cpp
+++ b/Engine/Plugins/Experimental/Mutable/Source/CustomizableObjectEditor/Private/MuCOE/SMutableCodeViewer.cpp
@@ -175,11 +175,9 @@ public:
 				mu::OP::ImageMultiLayerArgs Args = Program.GetOpArgs<mu::OP::ImageMultiLayerArgs>(RowItem->MutableOperation);
 				OpName += TEXT(" rgb: ");
 				OpName += mu::TypeInfo::s_blendModeName[int32(Args.blendType)];
-				if (Args.blendTypeAlpha != int8(mu::EBlendType::BT_NONE))
-				{
-					OpName += TEXT(", a: ");
-					OpName += mu::TypeInfo::s_blendModeName[int32(Args.blendTypeAlpha)];
-				}
+				OpName += TEXT(", a: ");
+				OpName += mu::TypeInfo::s_blendModeName[int32(Args.blendTypeAlpha)];
+				OpName += FString::Printf(TEXT(" a from %d "), Args.BlendAlphaSourceChannel);
 				OpName += FString::Printf(TEXT(" range-id: %d"), Args.rangeId);
 				OpName += FString::Printf(TEXT(" mask-from-alpha: %d"), int32(Args.bUseMaskFromBlended));
 				break;
@@ -190,11 +188,9 @@ public:
 				mu::OP::ImageLayerArgs Args = Program.GetOpArgs<mu::OP::ImageLayerArgs>(RowItem->MutableOperation);
 				OpName += TEXT(" rgb: ");
 				OpName += mu::TypeInfo::s_blendModeName[int32(Args.blendType)];
-				if (Args.blendTypeAlpha != int8(mu::EBlendType::BT_NONE))
-				{
-					OpName += TEXT(", a: ");
-					OpName += mu::TypeInfo::s_blendModeName[int32(Args.blendTypeAlpha)];
-				}
+				OpName += TEXT(", a: ");
+				OpName += mu::TypeInfo::s_blendModeName[int32(Args.blendTypeAlpha)];
+				OpName += FString::Printf(TEXT(" a from %d "), Args.BlendAlphaSourceChannel);
 				OpName += FString::Printf(TEXT(" flags %d"),Args.flags);
 				break;
 			}
@@ -206,6 +202,8 @@ public:
 				OpName += mu::TypeInfo::s_blendModeName[int32(Args.blendType)];
 				OpName += TEXT(" a: ");
 				OpName += mu::TypeInfo::s_blendModeName[int32(Args.blendTypeAlpha)];
+				OpName += TEXT(" a from ");
+				OpName += FString::Printf(TEXT(" a from %d "), Args.BlendAlphaSourceChannel);
 				OpName += FString::Printf(TEXT(" flags %d"), Args.flags);
 				break;
 			}
diff --git a/Engine/Plugins/Experimental/Mutable/Source/MutableRuntime/Private/MuR/CodeRunner.cpp b/Engine/Plugins/Experimental/Mutable/Source/MutableRuntime/Private/MuR/CodeRunner.cpp
--- a/Engine/Plugins/Experimental/Mutable/Source/MutableRuntime/Private/MuR/CodeRunner.cpp
+++ b/Engine/Plugins/Experimental/Mutable/Source/MutableRuntime/Private/MuR/CodeRunner.cpp
@@ -2678,6 +2678,7 @@ namespace mu
 					// This becomes true if we need to update the mips of the resulting image
 					// This could happen in the base image has mips, but one of the blended one doesn't.
 					bool bBlendOnlyOneMip = Data.MultiLayer.bBlendOnlyOneMip;
+					bool bUseBlendSourceFromBlendAlpha = false; // (Args.flags& OP::ImageLayerArgs::F_BLENDED_RGB_FROM_ALPHA) != 0;
 
 					if (!args.mask && args.bUseMaskFromBlended
 						&&
@@ -2691,11 +2692,13 @@ namespace mu
 						constexpr bool bUseVectorImpl = false;
 						if constexpr (bUseVectorImpl)
 						{
+							// TODO: Add support for args.BlendAlphaSourceChannel
+							check(args.BlendAlphaSourceChannel == 3);
 							BufferLayerCompositeVector<VectorBlendChannelMasked, VectorLightenChannel, false>(Base.get(), pBlended.get(), bBlendOnlyOneMip);
 						}
 						else
 						{
-							BufferLayerComposite<BlendChannelMasked, LightenChannel, false>(Base.get(), pBlended.get(), bBlendOnlyOneMip);
+							BufferLayerComposite<BlendChannelMasked, LightenChannel, false>(Base.get(), pBlended.get(), bBlendOnlyOneMip, args.BlendAlphaSourceChannel);
 						}
 					}
 
@@ -2710,6 +2713,9 @@ namespace mu
 							pMask = ImageResizeLinear(0, pMask.get(), ResultSize);
 						}
 
+						// Not implemented yet
+						check(!bUseBlendSourceFromBlendAlpha);
+
                         switch (EBlendType(args.blendType))
                         {
 						case EBlendType::BT_NORMAL_COMBINE: check(false); break;
@@ -2728,6 +2734,9 @@ namespace mu
                     }
 					else if (!bDone && args.bUseMaskFromBlended)
 					{
+						// Not implemented yet
+						check(!bUseBlendSourceFromBlendAlpha);
+
 						switch (EBlendType(args.blendType))
 						{
 						case EBlendType::BT_NORMAL_COMBINE: check(false); break;
@@ -2748,15 +2757,15 @@ namespace mu
                         switch (EBlendType(args.blendType))
                         {
 						case EBlendType::BT_NORMAL_COMBINE: check(false); break;
-                        case EBlendType::BT_SOFTLIGHT: BufferLayer<SoftLightChannel, false>(Base.get(), Base.get(), pBlended.get(), bApplyColorBlendToAlpha, bBlendOnlyOneMip); break;
-                        case EBlendType::BT_HARDLIGHT: BufferLayer<HardLightChannel, false>(Base.get(), Base.get(), pBlended.get(), bApplyColorBlendToAlpha, bBlendOnlyOneMip); break;
-                        case EBlendType::BT_BURN: BufferLayer<BurnChannel, false>(Base.get(), Base.get(), pBlended.get(), bApplyColorBlendToAlpha, bBlendOnlyOneMip); break;
-                        case EBlendType::BT_DODGE: BufferLayer<DodgeChannel, false>(Base.get(), Base.get(), pBlended.get(), bApplyColorBlendToAlpha, bBlendOnlyOneMip); break;
-                        case EBlendType::BT_SCREEN: BufferLayer<ScreenChannel, false>(Base.get(), Base.get(),  pBlended.get(), bApplyColorBlendToAlpha, bBlendOnlyOneMip); break;
-                        case EBlendType::BT_OVERLAY: BufferLayer<OverlayChannel, false>(Base.get(), Base.get(), pBlended.get(), bApplyColorBlendToAlpha, bBlendOnlyOneMip); break;
-                        case EBlendType::BT_LIGHTEN: BufferLayer<LightenChannel, false>(Base.get(), Base.get(), pBlended.get(), bApplyColorBlendToAlpha, bBlendOnlyOneMip); break;
-                        case EBlendType::BT_MULTIPLY: BufferLayer<MultiplyChannel, false>(Base.get(), Base.get(), pBlended.get(), bApplyColorBlendToAlpha, bBlendOnlyOneMip); break;
-                        case EBlendType::BT_BLEND: BufferLayer<BlendChannel, false>(Base.get(), Base.get(), pBlended.get(), bApplyColorBlendToAlpha, bBlendOnlyOneMip); break;
+                        case EBlendType::BT_SOFTLIGHT: BufferLayer<SoftLightChannel, false>(Base.get(), Base.get(), pBlended.get(), bApplyColorBlendToAlpha, bBlendOnlyOneMip, bUseBlendSourceFromBlendAlpha); break;
+                        case EBlendType::BT_HARDLIGHT: BufferLayer<HardLightChannel, false>(Base.get(), Base.get(), pBlended.get(), bApplyColorBlendToAlpha, bBlendOnlyOneMip, bUseBlendSourceFromBlendAlpha); break;
+                        case EBlendType::BT_BURN: BufferLayer<BurnChannel, false>(Base.get(), Base.get(), pBlended.get(), bApplyColorBlendToAlpha, bBlendOnlyOneMip, bUseBlendSourceFromBlendAlpha); break;
+                        case EBlendType::BT_DODGE: BufferLayer<DodgeChannel, false>(Base.get(), Base.get(), pBlended.get(), bApplyColorBlendToAlpha, bBlendOnlyOneMip, bUseBlendSourceFromBlendAlpha); break;
+                        case EBlendType::BT_SCREEN: BufferLayer<ScreenChannel, false>(Base.get(), Base.get(),  pBlended.get(), bApplyColorBlendToAlpha, bBlendOnlyOneMip, bUseBlendSourceFromBlendAlpha); break;
+                        case EBlendType::BT_OVERLAY: BufferLayer<OverlayChannel, false>(Base.get(), Base.get(), pBlended.get(), bApplyColorBlendToAlpha, bBlendOnlyOneMip, bUseBlendSourceFromBlendAlpha); break;
+                        case EBlendType::BT_LIGHTEN: BufferLayer<LightenChannel, false>(Base.get(), Base.get(), pBlended.get(), bApplyColorBlendToAlpha, bBlendOnlyOneMip, bUseBlendSourceFromBlendAlpha); break;
+                        case EBlendType::BT_MULTIPLY: BufferLayer<MultiplyChannel, false>(Base.get(), Base.get(), pBlended.get(), bApplyColorBlendToAlpha, bBlendOnlyOneMip, bUseBlendSourceFromBlendAlpha); break;
+                        case EBlendType::BT_BLEND: BufferLayer<BlendChannel, false>(Base.get(), Base.get(), pBlended.get(), bApplyColorBlendToAlpha, bBlendOnlyOneMip, bUseBlendSourceFromBlendAlpha); break;
                         default: check(false);
                         }
                     }
@@ -2767,15 +2776,15 @@ namespace mu
 						// Separate alpha operation ignores the mask.
 						switch (EBlendType(args.blendTypeAlpha))
 						{
-						case EBlendType::BT_SOFTLIGHT: BufferLayerInPlace<SoftLightChannel, false, 1>(Base.get(), pBlended.get(), bBlendOnlyOneMip, 3, 3); break;
-						case EBlendType::BT_HARDLIGHT: BufferLayerInPlace<HardLightChannel, false, 1>(Base.get(), pBlended.get(), bBlendOnlyOneMip, 3, 3); break;
-						case EBlendType::BT_BURN: BufferLayerInPlace<BurnChannel, false, 1>(Base.get(), pBlended.get(), bBlendOnlyOneMip, 3, 3); break;
-						case EBlendType::BT_DODGE: BufferLayerInPlace<DodgeChannel, false, 1>(Base.get(), pBlended.get(), bBlendOnlyOneMip, 3, 3); break;
-						case EBlendType::BT_SCREEN: BufferLayerInPlace<ScreenChannel, false, 1>(Base.get(), pBlended.get(), bBlendOnlyOneMip, 3, 3); break;
-						case EBlendType::BT_OVERLAY: BufferLayerInPlace<OverlayChannel, false, 1>(Base.get(), pBlended.get(), bBlendOnlyOneMip, 3, 3); break;
-						case EBlendType::BT_LIGHTEN: BufferLayerInPlace<LightenChannel, false, 1>(Base.get(), pBlended.get(), bBlendOnlyOneMip, 3, 3); break;
-						case EBlendType::BT_MULTIPLY: BufferLayerInPlace<MultiplyChannel, false, 1>(Base.get(), pBlended.get(), bBlendOnlyOneMip, 3, 3); break;
-						case EBlendType::BT_BLEND: BufferLayerInPlace<BlendChannel, false, 1>(Base.get(), pBlended.get(), bBlendOnlyOneMip, 3, 3); break;
+						case EBlendType::BT_SOFTLIGHT: BufferLayerInPlace<SoftLightChannel, false, 1>(Base.get(), pBlended.get(), bBlendOnlyOneMip, 3, args.BlendAlphaSourceChannel); break;
+						case EBlendType::BT_HARDLIGHT: BufferLayerInPlace<HardLightChannel, false, 1>(Base.get(), pBlended.get(), bBlendOnlyOneMip, 3, args.BlendAlphaSourceChannel); break;
+						case EBlendType::BT_BURN: BufferLayerInPlace<BurnChannel, false, 1>(Base.get(), pBlended.get(), bBlendOnlyOneMip, 3, args.BlendAlphaSourceChannel); break;
+						case EBlendType::BT_DODGE: BufferLayerInPlace<DodgeChannel, false, 1>(Base.get(), pBlended.get(), bBlendOnlyOneMip, 3, args.BlendAlphaSourceChannel); break;
+						case EBlendType::BT_SCREEN: BufferLayerInPlace<ScreenChannel, false, 1>(Base.get(), pBlended.get(), bBlendOnlyOneMip, 3, args.BlendAlphaSourceChannel); break;
+						case EBlendType::BT_OVERLAY: BufferLayerInPlace<OverlayChannel, false, 1>(Base.get(), pBlended.get(), bBlendOnlyOneMip, 3, args.BlendAlphaSourceChannel); break;
+						case EBlendType::BT_LIGHTEN: BufferLayerInPlace<LightenChannel, false, 1>(Base.get(), pBlended.get(), bBlendOnlyOneMip, 3, args.BlendAlphaSourceChannel); break;
+						case EBlendType::BT_MULTIPLY: BufferLayerInPlace<MultiplyChannel, false, 1>(Base.get(), pBlended.get(), bBlendOnlyOneMip, 3, args.BlendAlphaSourceChannel); break;
+						case EBlendType::BT_BLEND: BufferLayerInPlace<BlendChannel, false, 1>(Base.get(), pBlended.get(), bBlendOnlyOneMip, 3, args.BlendAlphaSourceChannel); break;
 						default: check(false);
 						}
 					}
diff --git a/Engine/Plugins/Experimental/Mutable/Source/MutableRuntime/Private/MuR/CodeRunnerTasked.cpp b/Engine/Plugins/Experimental/Mutable/Source/MutableRuntime/Private/MuR/CodeRunnerTasked.cpp
--- a/Engine/Plugins/Experimental/Mutable/Source/MutableRuntime/Private/MuR/CodeRunnerTasked.cpp
+++ b/Engine/Plugins/Experimental/Mutable/Source/MutableRuntime/Private/MuR/CodeRunnerTasked.cpp
@@ -472,13 +472,11 @@ namespace mu
 			m_base = ImagePixelFormat(m_imageCompressionQuality, m_base.get(), UncompressedFormat);
 		}
 
-		if (m_blended && InitialFormat != m_blended->GetFormat())
+		if (m_blended && InitialFormat != m_blended->GetFormat() && !Args.flags)
 		{
 			m_blended = ImagePixelFormat(m_imageCompressionQuality, m_blended.get(), m_base->GetFormat());
 		}
 
-		bool bApplyColorBlendToAlpha = (Args.flags & OP::ImageLayerArgs::F_APPLY_TO_ALPHA) != 0;
-
 		ImagePtr pNew = mu::CloneOrTakeOver<Image>(m_base.get());
 		//ImagePtr pNew = m_base->Clone();
 		check(pNew->GetDataSize() == m_base->GetDataSize());
@@ -503,28 +501,36 @@ namespace mu
 
 			bool bDone = false;
 
-			bool bUseMaskFromBlendAlpha = (Args.flags & OP::ImageLayerArgs::F_USE_MASK_FROM_BLENDED);
-
-			if (!m_mask && bUseMaskFromBlendAlpha
+			if (!m_mask 
+				&& 
+				// Flags have to match exactly for this optimize case. Other flags are not supported.
+				Args.flags== OP::ImageLayerArgs::F_USE_MASK_FROM_BLENDED
 				&&
 				Args.blendType == uint8(EBlendType::BT_BLEND)
 				&&
 				Args.blendTypeAlpha == uint8(EBlendType::BT_LIGHTEN))
 			{
+				MUTABLE_CPUPROFILER_SCOPE(ImageLayerTask_Optimized);
+
 				// This is a frequent critical-path case because of multilayer projectors.
 				bDone = true;
 
 				constexpr bool bUseVectorImplementation = false;	
 				if constexpr (bUseVectorImplementation)
 				{
+					// TODO: Add support for Args.BlendAlphaSourceChannel
+					check(Args.BlendAlphaSourceChannel==3);
 					BufferLayerCompositeVector<VectorBlendChannelMasked, VectorLightenChannel, false>(pNew.get(), m_blended.get(), bOnlyOneMip);
 				}
 				else
 				{
-					BufferLayerComposite<BlendChannelMasked, LightenChannel, false>(pNew.get(), m_blended.get(), bOnlyOneMip);
+					BufferLayerComposite<BlendChannelMasked, LightenChannel, false>(pNew.get(), m_blended.get(), bOnlyOneMip, Args.BlendAlphaSourceChannel);
 				}
 			}
 
+			bool bApplyColorBlendToAlpha = (Args.flags & OP::ImageLayerArgs::F_APPLY_TO_ALPHA) != 0;
+			bool bUseBlendSourceFromBlendAlpha = (Args.flags & OP::ImageLayerArgs::F_BLENDED_RGB_FROM_ALPHA) != 0;
+			bool bUseMaskFromBlendAlpha = (Args.flags & OP::ImageLayerArgs::F_USE_MASK_FROM_BLENDED);
 
 			if (!bDone && Args.mask)
 			{
@@ -554,6 +560,9 @@ namespace mu
 					m_mask = pDest;
 				}
 
+				// Not implemented yet
+				check(!bUseBlendSourceFromBlendAlpha);
+
 				switch (EBlendType(Args.blendType))
 				{
 				case EBlendType::BT_NORMAL_COMBINE: ImageNormalCombine(pNew.get(), m_base.get(), m_mask.get(), m_blended.get(), bOnlyOneMip); break;
@@ -566,12 +575,16 @@ namespace mu
 				case EBlendType::BT_LIGHTEN: BufferLayer<LightenChannelMasked, LightenChannel, true>(pNew->GetData(), m_base.get(), m_mask.get(), m_blended.get(), bApplyColorBlendToAlpha, bOnlyOneMip); break;
 				case EBlendType::BT_MULTIPLY: BufferLayer<MultiplyChannelMasked, MultiplyChannel, true>(pNew->GetData(), m_base.get(), m_mask.get(), m_blended.get(), bApplyColorBlendToAlpha, bOnlyOneMip); break;
 				case EBlendType::BT_BLEND: BufferLayer<BlendChannelMasked, BlendChannel, true>(pNew->GetData(), m_base.get(), m_mask.get(), m_blended.get(), bApplyColorBlendToAlpha, bOnlyOneMip); break;
+				case EBlendType::BT_NONE: break;
 				default: check(false);
 				}
 
 			}
 			else if (!bDone && bUseMaskFromBlendAlpha)
 			{
+				// Not implemented yet
+				check(!bUseBlendSourceFromBlendAlpha);
+
 				// Apply blend without to RGB using mask in blended alpha
 				switch (EBlendType(Args.blendType))
 				{
@@ -585,6 +598,7 @@ namespace mu
 				case EBlendType::BT_LIGHTEN: BufferLayerEmbeddedMask<LightenChannelMasked, LightenChannel, false>(pNew->GetData(), pNew.get(), m_blended.get(), bApplyColorBlendToAlpha, bOnlyOneMip); break;
 				case EBlendType::BT_MULTIPLY: BufferLayerEmbeddedMask<MultiplyChannelMasked, MultiplyChannel, false>(pNew->GetData(), pNew.get(), m_blended.get(), bApplyColorBlendToAlpha, bOnlyOneMip); break;
 				case EBlendType::BT_BLEND: BufferLayerEmbeddedMask<BlendChannelMasked, BlendChannel, false>(pNew->GetData(), pNew.get(), m_blended.get(), bApplyColorBlendToAlpha, bOnlyOneMip); break;
+				case EBlendType::BT_NONE: break;
 				default: check(false);
 				}
 			}
@@ -593,35 +607,37 @@ namespace mu
 				// Apply blend without mask to RGB
 				switch (EBlendType(Args.blendType))
 				{
-				case EBlendType::BT_NORMAL_COMBINE: ImageNormalCombine(pNew.get(), m_base.get(), m_blended.get(), bOnlyOneMip); break;
-				case EBlendType::BT_SOFTLIGHT: BufferLayer<SoftLightChannel, true>(pNew.get(), m_base.get(), m_blended.get(), bApplyColorBlendToAlpha, bOnlyOneMip); break;
-				case EBlendType::BT_HARDLIGHT: BufferLayer<HardLightChannel, true>(pNew.get(), m_base.get(), m_blended.get(), bApplyColorBlendToAlpha, bOnlyOneMip); break;
-				case EBlendType::BT_BURN: BufferLayer<BurnChannel, true>(pNew.get(), m_base.get(), m_blended.get(), bApplyColorBlendToAlpha, bOnlyOneMip); break;
-				case EBlendType::BT_DODGE: BufferLayer<DodgeChannel, true>(pNew.get(), m_base.get(), m_blended.get(), bApplyColorBlendToAlpha, bOnlyOneMip); break;
-				case EBlendType::BT_SCREEN: BufferLayer<ScreenChannel, true>(pNew.get(), m_base.get(), m_blended.get(), bApplyColorBlendToAlpha, bOnlyOneMip); break;
-				case EBlendType::BT_OVERLAY: BufferLayer<OverlayChannel, true>(pNew.get(), m_base.get(), m_blended.get(), bApplyColorBlendToAlpha, bOnlyOneMip); break;
-				case EBlendType::BT_LIGHTEN: BufferLayer<LightenChannel, true>(pNew.get(), m_base.get(), m_blended.get(), bApplyColorBlendToAlpha, bOnlyOneMip); break;
-				case EBlendType::BT_MULTIPLY: BufferLayer<MultiplyChannel, true>(pNew.get(), m_base.get(), m_blended.get(), bApplyColorBlendToAlpha, bOnlyOneMip); break;
-				case EBlendType::BT_BLEND: BufferLayer<BlendChannel, true>(pNew.get(), m_base.get(), m_blended.get(), bApplyColorBlendToAlpha, bOnlyOneMip); break;
+				case EBlendType::BT_NORMAL_COMBINE: ImageNormalCombine(pNew.get(), m_base.get(), m_blended.get(), bOnlyOneMip); check(!bUseBlendSourceFromBlendAlpha);  break;
+				case EBlendType::BT_SOFTLIGHT: BufferLayer<SoftLightChannel, true>(pNew.get(), m_base.get(), m_blended.get(), bApplyColorBlendToAlpha, bOnlyOneMip, bUseBlendSourceFromBlendAlpha); break;
+				case EBlendType::BT_HARDLIGHT: BufferLayer<HardLightChannel, true>(pNew.get(), m_base.get(), m_blended.get(), bApplyColorBlendToAlpha, bOnlyOneMip, bUseBlendSourceFromBlendAlpha); break;
+				case EBlendType::BT_BURN: BufferLayer<BurnChannel, true>(pNew.get(), m_base.get(), m_blended.get(), bApplyColorBlendToAlpha, bOnlyOneMip, bUseBlendSourceFromBlendAlpha); break;
+				case EBlendType::BT_DODGE: BufferLayer<DodgeChannel, true>(pNew.get(), m_base.get(), m_blended.get(), bApplyColorBlendToAlpha, bOnlyOneMip, bUseBlendSourceFromBlendAlpha); break;
+				case EBlendType::BT_SCREEN: BufferLayer<ScreenChannel, true>(pNew.get(), m_base.get(), m_blended.get(), bApplyColorBlendToAlpha, bOnlyOneMip, bUseBlendSourceFromBlendAlpha); break;
+				case EBlendType::BT_OVERLAY: BufferLayer<OverlayChannel, true>(pNew.get(), m_base.get(), m_blended.get(), bApplyColorBlendToAlpha, bOnlyOneMip, bUseBlendSourceFromBlendAlpha); break;
+				case EBlendType::BT_LIGHTEN: BufferLayer<LightenChannel, true>(pNew.get(), m_base.get(), m_blended.get(), bApplyColorBlendToAlpha, bOnlyOneMip, bUseBlendSourceFromBlendAlpha); break;
+				case EBlendType::BT_MULTIPLY: BufferLayer<MultiplyChannel, true>(pNew.get(), m_base.get(), m_blended.get(), bApplyColorBlendToAlpha, bOnlyOneMip, bUseBlendSourceFromBlendAlpha); break;
+				case EBlendType::BT_BLEND: BufferLayer<BlendChannel, true>(pNew.get(), m_base.get(), m_blended.get(), bApplyColorBlendToAlpha, bOnlyOneMip, bUseBlendSourceFromBlendAlpha); break;
+				case EBlendType::BT_NONE: break;
 				default: check(false);
 				}
 			}
 
 			// Apply the separate blend operation for alpha
-			if (!bDone && !bApplyColorBlendToAlpha && Args.blendTypeAlpha != uint8(EBlendType::BT_NONE))
+			if (!bDone && !bApplyColorBlendToAlpha)
 			{
 				// Separate alpha operation ignores the mask.
 				switch (EBlendType(Args.blendTypeAlpha))
 				{
-				case EBlendType::BT_SOFTLIGHT: BufferLayerInPlace<SoftLightChannel, false, 1>(pNew.get(), m_blended.get(), bOnlyOneMip, 3, 3); break;
-				case EBlendType::BT_HARDLIGHT: BufferLayerInPlace<HardLightChannel, false, 1>(pNew.get(), m_blended.get(), bOnlyOneMip, 3, 3); break;
-				case EBlendType::BT_BURN: BufferLayerInPlace<BurnChannel, false, 1>(pNew.get(), m_blended.get(), bOnlyOneMip, 3, 3); break;
-				case EBlendType::BT_DODGE: BufferLayerInPlace<DodgeChannel, false, 1>(pNew.get(), m_blended.get(), bOnlyOneMip, 3, 3); break;
-				case EBlendType::BT_SCREEN: BufferLayerInPlace<ScreenChannel, false, 1>(pNew.get(), m_blended.get(), bOnlyOneMip, 3, 3); break;
-				case EBlendType::BT_OVERLAY: BufferLayerInPlace<OverlayChannel, false, 1>(pNew.get(), m_blended.get(), bOnlyOneMip, 3, 3); break;
-				case EBlendType::BT_LIGHTEN: BufferLayerInPlace<LightenChannel, false, 1>(pNew.get(), m_blended.get(), bOnlyOneMip, 3, 3); break;
-				case EBlendType::BT_MULTIPLY: BufferLayerInPlace<MultiplyChannel, false, 1>(pNew.get(), m_blended.get(), bOnlyOneMip, 3, 3); break;
-				case EBlendType::BT_BLEND: BufferLayerInPlace<BlendChannel, false, 1>(pNew.get(), m_blended.get(), bOnlyOneMip, 3, 3); break;
+				case EBlendType::BT_SOFTLIGHT: BufferLayerInPlace<SoftLightChannel, false, 1>(pNew.get(), m_blended.get(), bOnlyOneMip, 3, Args.BlendAlphaSourceChannel); break;
+				case EBlendType::BT_HARDLIGHT: BufferLayerInPlace<HardLightChannel, false, 1>(pNew.get(), m_blended.get(), bOnlyOneMip, 3, Args.BlendAlphaSourceChannel); break;
+				case EBlendType::BT_BURN: BufferLayerInPlace<BurnChannel, false, 1>(pNew.get(), m_blended.get(), bOnlyOneMip, 3, Args.BlendAlphaSourceChannel); break;
+				case EBlendType::BT_DODGE: BufferLayerInPlace<DodgeChannel, false, 1>(pNew.get(), m_blended.get(), bOnlyOneMip, 3, Args.BlendAlphaSourceChannel); break;
+				case EBlendType::BT_SCREEN: BufferLayerInPlace<ScreenChannel, false, 1>(pNew.get(), m_blended.get(), bOnlyOneMip, 3, Args.BlendAlphaSourceChannel); break;
+				case EBlendType::BT_OVERLAY: BufferLayerInPlace<OverlayChannel, false, 1>(pNew.get(), m_blended.get(), bOnlyOneMip, 3, Args.BlendAlphaSourceChannel); break;
+				case EBlendType::BT_LIGHTEN: BufferLayerInPlace<LightenChannel, false, 1>(pNew.get(), m_blended.get(), bOnlyOneMip, 3, Args.BlendAlphaSourceChannel); break;
+				case EBlendType::BT_MULTIPLY: BufferLayerInPlace<MultiplyChannel, false, 1>(pNew.get(), m_blended.get(), bOnlyOneMip, 3, Args.BlendAlphaSourceChannel); break;
+				case EBlendType::BT_BLEND: BufferLayerInPlace<BlendChannel, false, 1>(pNew.get(), m_blended.get(), bOnlyOneMip, 3, Args.BlendAlphaSourceChannel); break;
+				case EBlendType::BT_NONE: break;
 				default: check(false);
 				}
 			}
@@ -714,15 +730,14 @@ namespace mu
 		{
 			if (m_base->GetSize() != m_mask->GetSize())
 			{
-				MUTABLE_CPUPROFILER_SCOPE(ImageResize_EmergencyFix);
+				MUTABLE_CPUPROFILER_SCOPE(ImageResize_EmergencyFixSize);
 				m_mask = ImageResizeLinear(m_imageCompressionQuality, m_mask.get(), m_base->GetSize());
 			}
 
 			// emergy fix c36adf47-e40d-490f-b709-41142bafad78
 			if (m_mask->GetLODCount() < m_base->GetLODCount())
 			{
-				//MUTABLE_CPUPROFILER_SCOPE(ImageLayerColour_EmergencyFix);
-
+				MUTABLE_CPUPROFILER_SCOPE(ImageResize_EmergencyFixMips);
 				int32 levelCount = m_base->GetLODCount();
 				Ptr<Image> pDest = new Image(m_mask->GetSizeX(), m_mask->GetSizeY(), levelCount, m_mask->GetFormat());
 
@@ -742,7 +757,6 @@ namespace mu
 		{
 			// TODO: It could be done "in-place"
 			pNew = mu::CloneOrTakeOver<mu::Image>(m_base.get());
-			//pNew = new Image(m_base->GetSizeX(), m_base->GetSizeY(), m_base->GetLODCount(), m_base->GetFormat());
 
 			if (Args.mask && m_mask)
 			{			
@@ -768,7 +782,6 @@ namespace mu
 			}
 			else
 			{
-				// Not implemented yet
 				if (Args.flags & OP::ImageLayerArgs::FLAGS::F_BASE_RGB_FROM_ALPHA)
 				{
 					switch (EBlendType(Args.blendType))
@@ -814,27 +827,21 @@ namespace mu
 				pNew = mu::CloneOrTakeOver<mu::Image>(m_base.get());
 			}
 
-			uint32 AlphaOffset = 3;
-			if (pNew->GetFormat() == EImageFormat::IF_L_UBYTE)
-			{
-				AlphaOffset = 1;
-			}
-
 			if (Args.mask && m_mask)
 			{
 				// \todo: precalculated tables for softlight
 				switch (EBlendType(Args.blendTypeAlpha))
 				{
 				case EBlendType::BT_NORMAL_COMBINE: check(false); break;
-				case EBlendType::BT_SOFTLIGHT: BufferLayerColourInPlace<SoftLightChannelMasked, SoftLightChannel, 1>(pNew.get(), m_mask.get(), m_col, bOnlyOneMip, 3); break;
-				case EBlendType::BT_HARDLIGHT: BufferLayerColourInPlace<HardLightChannelMasked, HardLightChannel, 1>(pNew.get(), m_mask.get(), m_col, bOnlyOneMip, 3); break;
-				case EBlendType::BT_BURN: BufferLayerColourInPlace<BurnChannelMasked, BurnChannel, 1>(pNew.get(), m_mask.get(), m_col, bOnlyOneMip, 3); break;
-				case EBlendType::BT_DODGE: BufferLayerColourInPlace<DodgeChannelMasked, DodgeChannel, 1>(pNew.get(), m_mask.get(), m_col, bOnlyOneMip, 3); break;
-				case EBlendType::BT_SCREEN: BufferLayerColourInPlace<ScreenChannelMasked, ScreenChannel, 1>(pNew.get(), m_mask.get(), m_col, bOnlyOneMip, 3); break;
-				case EBlendType::BT_OVERLAY: BufferLayerColourInPlace<OverlayChannelMasked, OverlayChannel, 1>(pNew.get(), m_mask.get(), m_col, bOnlyOneMip, 3); break;
-				case EBlendType::BT_LIGHTEN: BufferLayerColourInPlace<LightenChannelMasked, LightenChannel, 1>(pNew.get(), m_mask.get(), m_col, bOnlyOneMip, 3); break;
-				case EBlendType::BT_MULTIPLY: BufferLayerColourInPlace<MultiplyChannelMasked, MultiplyChannel, 1>(pNew.get(), m_mask.get(), m_col, bOnlyOneMip, 3); break;
-				case EBlendType::BT_BLEND: BufferLayerColourInPlace<BlendChannelMasked, BlendChannel, 1>(pNew.get(), m_mask.get(), m_col, bOnlyOneMip, 3); break;
+				case EBlendType::BT_SOFTLIGHT: BufferLayerColourInPlace<SoftLightChannelMasked, SoftLightChannel, 1>(pNew.get(), m_mask.get(), m_col, bOnlyOneMip, 3, Args.BlendAlphaSourceChannel); break;
+				case EBlendType::BT_HARDLIGHT: BufferLayerColourInPlace<HardLightChannelMasked, HardLightChannel, 1>(pNew.get(), m_mask.get(), m_col, bOnlyOneMip, 3, Args.BlendAlphaSourceChannel); break;
+				case EBlendType::BT_BURN: BufferLayerColourInPlace<BurnChannelMasked, BurnChannel, 1>(pNew.get(), m_mask.get(), m_col, bOnlyOneMip, 3, Args.BlendAlphaSourceChannel); break;
+				case EBlendType::BT_DODGE: BufferLayerColourInPlace<DodgeChannelMasked, DodgeChannel, 1>(pNew.get(), m_mask.get(), m_col, bOnlyOneMip, 3, Args.BlendAlphaSourceChannel); break;
+				case EBlendType::BT_SCREEN: BufferLayerColourInPlace<ScreenChannelMasked, ScreenChannel, 1>(pNew.get(), m_mask.get(), m_col, bOnlyOneMip, 3, Args.BlendAlphaSourceChannel); break;
+				case EBlendType::BT_OVERLAY: BufferLayerColourInPlace<OverlayChannelMasked, OverlayChannel, 1>(pNew.get(), m_mask.get(), m_col, bOnlyOneMip, 3, Args.BlendAlphaSourceChannel); break;
+				case EBlendType::BT_LIGHTEN: BufferLayerColourInPlace<LightenChannelMasked, LightenChannel, 1>(pNew.get(), m_mask.get(), m_col, bOnlyOneMip, 3, Args.BlendAlphaSourceChannel); break;
+				case EBlendType::BT_MULTIPLY: BufferLayerColourInPlace<MultiplyChannelMasked, MultiplyChannel, 1>(pNew.get(), m_mask.get(), m_col, bOnlyOneMip, 3, Args.BlendAlphaSourceChannel); break;
+				case EBlendType::BT_BLEND: BufferLayerColourInPlace<BlendChannelMasked, BlendChannel, 1>(pNew.get(), m_mask.get(), m_col, bOnlyOneMip, 3, Args.BlendAlphaSourceChannel); break;
 				default: check(false);
 				}
 
@@ -844,15 +851,15 @@ namespace mu
 				switch (EBlendType(Args.blendTypeAlpha))
 				{
 				case EBlendType::BT_NORMAL_COMBINE: check(false); break;
-				case EBlendType::BT_SOFTLIGHT: BufferLayerColourInPlace<SoftLightChannel, 1>(pNew.get(), m_col, bOnlyOneMip, 3); break;
-				case EBlendType::BT_HARDLIGHT: BufferLayerColourInPlace<HardLightChannel, 1>(pNew.get(), m_col, bOnlyOneMip, 3); break;
-				case EBlendType::BT_BURN: BufferLayerColourInPlace<BurnChannel, 1>(pNew.get(), m_col, bOnlyOneMip, 3); break;
-				case EBlendType::BT_DODGE: BufferLayerColourInPlace<DodgeChannel, 1>(pNew.get(), m_col, bOnlyOneMip, 3); break;
-				case EBlendType::BT_SCREEN: BufferLayerColourInPlace<ScreenChannel, 1>(pNew.get(), m_col, bOnlyOneMip, 3); break;
-				case EBlendType::BT_OVERLAY: BufferLayerColourInPlace<OverlayChannel, 1>(pNew.get(), m_col, bOnlyOneMip, 3); break;
-				case EBlendType::BT_LIGHTEN: BufferLayerColourInPlace<LightenChannel, 1>(pNew.get(), m_col, bOnlyOneMip, 3); break;
-				case EBlendType::BT_MULTIPLY: BufferLayerColourInPlace<MultiplyChannel, 1>(pNew.get(), m_col, bOnlyOneMip, 3); break;
-				case EBlendType::BT_BLEND: BufferLayerColourInPlace<BlendChannel, 1>(pNew.get(), m_col, bOnlyOneMip, 3); break;
+				case EBlendType::BT_SOFTLIGHT: BufferLayerColourInPlace<SoftLightChannel, 1>(pNew.get(), m_col, bOnlyOneMip, 3, Args.BlendAlphaSourceChannel); break;
+				case EBlendType::BT_HARDLIGHT: BufferLayerColourInPlace<HardLightChannel, 1>(pNew.get(), m_col, bOnlyOneMip, 3, Args.BlendAlphaSourceChannel); break;
+				case EBlendType::BT_BURN: BufferLayerColourInPlace<BurnChannel, 1>(pNew.get(), m_col, bOnlyOneMip, 3, Args.BlendAlphaSourceChannel); break;
+				case EBlendType::BT_DODGE: BufferLayerColourInPlace<DodgeChannel, 1>(pNew.get(), m_col, bOnlyOneMip, 3, Args.BlendAlphaSourceChannel); break;
+				case EBlendType::BT_SCREEN: BufferLayerColourInPlace<ScreenChannel, 1>(pNew.get(), m_col, bOnlyOneMip, 3, Args.BlendAlphaSourceChannel); break;
+				case EBlendType::BT_OVERLAY: BufferLayerColourInPlace<OverlayChannel, 1>(pNew.get(), m_col, bOnlyOneMip, 3, Args.BlendAlphaSourceChannel); break;
+				case EBlendType::BT_LIGHTEN: BufferLayerColourInPlace<LightenChannel, 1>(pNew.get(), m_col, bOnlyOneMip, 3, Args.BlendAlphaSourceChannel); break;
+				case EBlendType::BT_MULTIPLY: BufferLayerColourInPlace<MultiplyChannel, 1>(pNew.get(), m_col, bOnlyOneMip, 3, Args.BlendAlphaSourceChannel); break;
+				case EBlendType::BT_BLEND: BufferLayerColourInPlace<BlendChannel, 1>(pNew.get(), m_col, bOnlyOneMip, 3, Args.BlendAlphaSourceChannel); break;
 				default: check(false);
 				}
 			}
diff --git a/Engine/Plugins/Experimental/Mutable/Source/MutableRuntime/Private/MuR/OpImageLayer.h b/Engine/Plugins/Experimental/Mutable/Source/MutableRuntime/Private/MuR/OpImageLayer.h
--- a/Engine/Plugins/Experimental/Mutable/Source/MutableRuntime/Private/MuR/OpImageLayer.h
+++ b/Engine/Plugins/Experimental/Mutable/Source/MutableRuntime/Private/MuR/OpImageLayer.h
@@ -15,11 +15,11 @@ namespace mu
 	}
 	
 
-	//---------------------------------------------------------------------------------------------
-	//! Apply a blending function to an image with a colour source
-	//---------------------------------------------------------------------------------------------
+	/** Apply a blending function to an image with a colour source.
+	* It only affects the RGB or L channels, leaving alpha untouched.
+	*/
 	template< unsigned (*BLEND_FUNC)(unsigned,unsigned), bool CLAMP >
-    inline void BufferLayerColour( Image* pResult, const Image* pBase, vec3<float> col )
+    inline void BufferLayerColour( Image* pResult, const Image* pBase, FVector4f col )
 	{
         check( pResult->GetFormat() == pBase->GetFormat() );
         check( pResult->GetSizeX() == pBase->GetSizeX() );
@@ -31,10 +31,10 @@ namespace mu
         uint8* pDestBuf = pResult->GetData();
         const uint8* pBaseBuf = pBase->GetData();
 
-		unsigned top[3];
-		top[0] = (unsigned)(255 * col[0]);
-		top[1] = (unsigned)(255 * col[1]);
-		top[2] = (unsigned)(255 * col[2]);
+		uint32 top[3];
+		top[0] = uint32(255 * col[0]);
+		top[1] = uint32(255 * col[1]);
+		top[2] = uint32(255 * col[2]);
 
 		// Generic implementation
 		int32 pixelCount = pBase->CalculatePixelCount();
@@ -390,14 +390,15 @@ namespace mu
 		uint32 CHANNELS_TO_BLEND,
 		// Number of total channels in the base image
 		int32 BASE_CHANNEL_STRIDE >
-    inline void BufferLayerColourFormat( uint8* StartDestBuf, const Image* Base, const Image* Mask, FVector4f Col, uint32 BaseOffset, bool bOnlyOneMip )
+    inline void BufferLayerColourFormat( uint8* StartDestBuf, const Image* Base, const Image* Mask, FVector4f Col, uint32 BaseOffset, uint8 ColorOffset, bool bOnlyOneMip )
 	{
 		check(CHANNELS_TO_BLEND+BaseOffset<=BASE_CHANNEL_STRIDE);
 
-		unsigned top[3];
-		top[0] = (unsigned)(255 * Col[0]);
-		top[1] = (unsigned)(255 * Col[1]);
-		top[2] = (unsigned)(255 * Col[2]);
+		uint32 top[4];
+		top[0] = uint32(255 * Col[0]);
+		top[1] = uint32(255 * Col[1]);
+		top[2] = uint32(255 * Col[2]);
+		top[3] = uint32(255 * Col[3]);
 
 		uint8* pDestBuf = StartDestBuf+ BaseOffset;
         const uint8* pMaskBuf = Mask->GetData();
@@ -414,6 +415,8 @@ namespace mu
         bool isUncompressed = ( maskFormat == EImageFormat::IF_L_UBYTE );
 
 		constexpr uint32 NumColorChannels = FMath::Min(CHANNELS_TO_BLEND,3u);
+		check(NumColorChannels+ColorOffset<4);
+
         if ( isUncompressed )
         {
             int32 pixelCount = Base->CalculatePixelCount();
@@ -423,7 +426,7 @@ namespace mu
                 for ( int32 c=0; c<NumColorChannels; ++c )
                 {
                     unsigned base = pBaseBuf[BASE_CHANNEL_STRIDE *i+c];
-                    unsigned result = BLEND_FUNC_MASKED( base, top[c], mask );
+                    unsigned result = BLEND_FUNC_MASKED( base, top[c+ColorOffset], mask );
                     if ( CLAMP )
                     {
                         pDestBuf[BASE_CHANNEL_STRIDE *i+c] = (uint8)FMath::Min( 255u, result );
@@ -480,7 +483,7 @@ namespace mu
                                 for ( int32 c=0; c<NumColorChannels; ++c )
                                 {
                                     unsigned base = pBaseBuf[BASE_CHANNEL_STRIDE *i+c];
-                                    unsigned result = BLEND_FUNC( base, top[c] );
+                                    unsigned result = BLEND_FUNC( base, top[c+ColorOffset] );
                                     if ( CLAMP )
                                     {
                                         pDestBuf[BASE_CHANNEL_STRIDE *i+c] = (uint8)FMath::Min( 255u, result );
@@ -505,7 +508,7 @@ namespace mu
                                 for ( int32 c=0; c<NumColorChannels; ++c )
                                 {
                                     unsigned base = pBaseBuf[BASE_CHANNEL_STRIDE *i+c];
-                                    unsigned result = BLEND_FUNC_MASKED( base, top[c], equalPixel );
+                                    unsigned result = BLEND_FUNC_MASKED( base, top[c+ColorOffset], equalPixel );
                                     if ( CLAMP )
                                     {
                                         pDestBuf[BASE_CHANNEL_STRIDE *i+c] = (uint8)FMath::Min( 255u, result );
@@ -542,7 +545,7 @@ namespace mu
                             {
                                 unsigned mask = pMaskBuf[i];
                                 unsigned base = pBaseBuf[BASE_CHANNEL_STRIDE *i+c];
-                                unsigned result = BLEND_FUNC_MASKED( base, top[c], mask );
+                                unsigned result = BLEND_FUNC_MASKED( base, top[c+ColorOffset], mask );
                                 if ( CLAMP )
                                 {
                                     pDestBuf[BASE_CHANNEL_STRIDE *i+c] = (uint8)FMath::Min( 255u, result );
@@ -598,8 +601,8 @@ namespace mu
 			(Base->GetSizeY() == Mask->GetSizeY())
 			&&
 			(Mask->GetFormat() == EImageFormat::IF_L_UBYTE
-				||
-				Mask->GetFormat() == EImageFormat::IF_L_UBYTE_RLE);
+			||
+			Mask->GetFormat() == EImageFormat::IF_L_UBYTE_RLE);
 		if (!bValid)
 		{
 			return;
@@ -608,22 +611,22 @@ namespace mu
 		EImageFormat baseFormat = Base->GetFormat();
 		if ( baseFormat==EImageFormat::IF_RGB_UBYTE )
 		{
-			BufferLayerColourFormat<BLEND_FUNC_MASKED, BLEND_FUNC, CLAMP, 3, 3>( DestBuf, Base, Mask, Col, 0, false);
+			BufferLayerColourFormat<BLEND_FUNC_MASKED, BLEND_FUNC, CLAMP, 3, 3>( DestBuf, Base, Mask, Col, 0, 0, false);
 		}
         else if ( baseFormat==EImageFormat::IF_RGBA_UBYTE )
         {
-            BufferLayerColourFormat<BLEND_FUNC_MASKED, BLEND_FUNC, CLAMP, 3, 4>( DestBuf, Base, Mask, Col, 0, false);
+            BufferLayerColourFormat<BLEND_FUNC_MASKED, BLEND_FUNC, CLAMP, 3, 4>( DestBuf, Base, Mask, Col, 0, 0, false);
         }
         else if ( baseFormat==EImageFormat::IF_BGRA_UBYTE )
         {
             float temp = Col[0];
             Col[0] = Col[2];
             Col[2] = temp;
-            BufferLayerColourFormat<BLEND_FUNC_MASKED, BLEND_FUNC, CLAMP, 3, 4>( DestBuf, Base, Mask, Col, 0, false);
+            BufferLayerColourFormat<BLEND_FUNC_MASKED, BLEND_FUNC, CLAMP, 3, 4>( DestBuf, Base, Mask, Col, 0, 0, false);
         }
         else if ( baseFormat==EImageFormat::IF_L_UBYTE )
 		{
-			BufferLayerColourFormat<BLEND_FUNC_MASKED, BLEND_FUNC, CLAMP, 1, 1>( DestBuf, Base, Mask, Col, 0, false);
+			BufferLayerColourFormat<BLEND_FUNC_MASKED, BLEND_FUNC, CLAMP, 1, 1>( DestBuf, Base, Mask, Col, 0, 0, false);
 		}
 		else
 		{
@@ -659,7 +662,7 @@ namespace mu
 		// Number of total channels in the base image
 		int32 BASE_CHANNEL_STRIDE>
 
-	inline void BufferLayerColourFormatInPlace(Image* Base, FVector4f Col, uint32 BaseChannelOffset, bool bOnlyFirstLOD)
+	inline void BufferLayerColourFormatInPlace(Image* Base, FVector4f Col, uint32 BaseChannelOffset, uint8 ColOffset, bool bOnlyFirstLOD)
 	{
 		uint8* pBaseBuf = Base->GetData() + BaseChannelOffset;
 
@@ -667,7 +670,7 @@ namespace mu
 		top[0] = (uint32)(255 * Col[0]);
 		top[1] = (uint32)(255 * Col[1]);
 		top[2] = (uint32)(255 * Col[2]);
-		top[3] = (uint32)(255 * Col[0]); // used for 1 channel blends to alpha
+		top[3] = (uint32)(255 * Col[3]);
 
 		// Generic implementation
 		int32 PixelCount = 0;
@@ -682,13 +685,13 @@ namespace mu
 
 		ParallelFor(PixelCount,
 			[
-				pBaseBuf, top, BaseChannelOffset
+				pBaseBuf, top, BaseChannelOffset, ColOffset
 			] (uint32 i)
 			{
 				for (int c = 0; c < CHANNELS_TO_BLEND; ++c)
 				{
 					uint32 base = pBaseBuf[BASE_CHANNEL_STRIDE * i + c];
-					uint32 blended = top[c + BaseChannelOffset];
+					uint32 blended = top[c + ColOffset];
 					uint32 result = BLEND_FUNC(base, blended);
 					if (CLAMP)
 					{
@@ -705,32 +708,32 @@ namespace mu
 
 	//---------------------------------------------------------------------------------------------
 	template< unsigned (*BLEND_FUNC)(unsigned, unsigned), bool CLAMP, uint32 CHANNEL_COUNT >
-	inline void BufferLayerColourInPlace(Image* Base, FVector4f Col, bool bOnlyOneMip, uint32 BaseOffset)
+	inline void BufferLayerColourInPlace(Image* Base, FVector4f Col, bool bOnlyOneMip, uint32 BaseOffset, uint8 ColOffset)
 	{
 		EImageFormat baseFormat = Base->GetFormat();
 
 		if (baseFormat == EImageFormat::IF_RGB_UBYTE)
 		{
 			check(BaseOffset + CHANNEL_COUNT <= 3);
-			BufferLayerColourFormatInPlace< BLEND_FUNC, CLAMP, CHANNEL_COUNT, 3 >(Base, Col, BaseOffset, bOnlyOneMip);
+			BufferLayerColourFormatInPlace< BLEND_FUNC, CLAMP, CHANNEL_COUNT, 3 >(Base, Col, BaseOffset, ColOffset, bOnlyOneMip);
 		}
 		else if (baseFormat == EImageFormat::IF_RGBA_UBYTE)
 		{
 			check(BaseOffset + CHANNEL_COUNT <= 4);
-			BufferLayerColourFormatInPlace< BLEND_FUNC, CLAMP, CHANNEL_COUNT, 4 >(Base, Col, BaseOffset, bOnlyOneMip);
+			BufferLayerColourFormatInPlace< BLEND_FUNC, CLAMP, CHANNEL_COUNT, 4 >(Base, Col, BaseOffset, ColOffset, bOnlyOneMip);
 		}
 		else if (baseFormat == EImageFormat::IF_BGRA_UBYTE)
 		{
 			float temp = Col[0];
 			Col[0] = Col[2];
 			Col[2] = temp;
-			BufferLayerColourFormatInPlace< BLEND_FUNC, CLAMP, CHANNEL_COUNT, 4>(Base, Col, BaseOffset, bOnlyOneMip);
+			BufferLayerColourFormatInPlace< BLEND_FUNC, CLAMP, CHANNEL_COUNT, 4>(Base, Col, BaseOffset, ColOffset, bOnlyOneMip);
 		}
 
 		else if (baseFormat == EImageFormat::IF_L_UBYTE)
 		{
 			check(BaseOffset + CHANNEL_COUNT <= 1);
-			BufferLayerColourFormatInPlace< BLEND_FUNC, CLAMP, CHANNEL_COUNT, 1 >(Base, Col, BaseOffset, bOnlyOneMip);
+			BufferLayerColourFormatInPlace< BLEND_FUNC, CLAMP, CHANNEL_COUNT, 1 >(Base, Col, BaseOffset, ColOffset, bOnlyOneMip);
 		}
 		else
 		{
@@ -741,16 +744,16 @@ namespace mu
 
 	//---------------------------------------------------------------------------------------------
 	template< unsigned (*BLEND_FUNC)(unsigned, unsigned), uint32 CHANNEL_COUNT >
-	inline void BufferLayerColourInPlace(Image* Base, FVector4f Col, bool bOnlyOneMip, uint32 BaseOffset)
+	inline void BufferLayerColourInPlace(Image* Base, FVector4f Col, bool bOnlyOneMip, uint32 BaseOffset, uint8 ColOffset)
 	{
 		bool bIsClampNeeded = IsAnyComponentLargerThan1(Col);
 		if (bIsClampNeeded)
 		{
-			BufferLayerColourInPlace<BLEND_FUNC, true, CHANNEL_COUNT>(Base, Col, bOnlyOneMip, BaseOffset);
+			BufferLayerColourInPlace<BLEND_FUNC, true, CHANNEL_COUNT>(Base, Col, bOnlyOneMip, BaseOffset, ColOffset);
 		}
 		else
 		{
-			BufferLayerColourInPlace<BLEND_FUNC, false, CHANNEL_COUNT>(Base, Col, bOnlyOneMip, BaseOffset);
+			BufferLayerColourInPlace<BLEND_FUNC, false, CHANNEL_COUNT>(Base, Col, bOnlyOneMip, BaseOffset, ColOffset);
 		}
 	}
 
@@ -761,7 +764,7 @@ namespace mu
 		unsigned (*BLEND_FUNC)(unsigned, unsigned), 
 		bool CLAMP, 
 		int32 CHANNEL_COUNT >
-	inline void BufferLayerColourInPlace(Image* Base, const Image* Mask, FVector4f Col, bool bOnlyOneMip, uint32 BaseOffset)
+	inline void BufferLayerColourInPlace(Image* Base, const Image* Mask, FVector4f Col, bool bOnlyOneMip, uint32 BaseOffset, uint8 ColOffset)
 	{
 		check(Base->GetSizeX() == Mask->GetSizeX());
 		check(Base->GetSizeY() == Mask->GetSizeY());
@@ -771,12 +774,12 @@ namespace mu
 		if (baseFormat == EImageFormat::IF_RGB_UBYTE)
 		{
 			check(BaseOffset + CHANNEL_COUNT <= 3);
-			BufferLayerColourFormat< BLEND_FUNC_MASKED, BLEND_FUNC, CLAMP, CHANNEL_COUNT, 3 >(Base->GetData(), Base, Mask, Col, BaseOffset, bOnlyOneMip);
+			BufferLayerColourFormat< BLEND_FUNC_MASKED, BLEND_FUNC, CLAMP, CHANNEL_COUNT, 3 >(Base->GetData(), Base, Mask, Col, BaseOffset, ColOffset, bOnlyOneMip);
 		}
 		else if (baseFormat == EImageFormat::IF_RGBA_UBYTE)
 		{
 			check(BaseOffset + CHANNEL_COUNT <= 4);
-			BufferLayerColourFormat< BLEND_FUNC_MASKED, BLEND_FUNC, CLAMP, CHANNEL_COUNT, 4 >(Base->GetData(), Base, Mask, Col, BaseOffset, bOnlyOneMip);
+			BufferLayerColourFormat< BLEND_FUNC_MASKED, BLEND_FUNC, CLAMP, CHANNEL_COUNT, 4 >(Base->GetData(), Base, Mask, Col, BaseOffset, ColOffset, bOnlyOneMip);
 		}
 		else if (baseFormat == EImageFormat::IF_BGRA_UBYTE)
 		{
@@ -784,12 +787,12 @@ namespace mu
 			float temp = Col[0];
 			Col[0] = Col[2];
 			Col[2] = temp;
-			BufferLayerColourFormat< BLEND_FUNC_MASKED, BLEND_FUNC, CLAMP, CHANNEL_COUNT, 4 >(Base->GetData(), Base, Mask, Col, BaseOffset, bOnlyOneMip);
+			BufferLayerColourFormat< BLEND_FUNC_MASKED, BLEND_FUNC, CLAMP, CHANNEL_COUNT, 4 >(Base->GetData(), Base, Mask, Col, BaseOffset, ColOffset, bOnlyOneMip);
 		}
 		else if (baseFormat == EImageFormat::IF_L_UBYTE)
 		{
 			check(BaseOffset + CHANNEL_COUNT <= 1);
-			BufferLayerColourFormat< BLEND_FUNC_MASKED, BLEND_FUNC, CLAMP, CHANNEL_COUNT, 1 >(Base->GetData(), Base, Mask, Col, BaseOffset, bOnlyOneMip);
+			BufferLayerColourFormat< BLEND_FUNC_MASKED, BLEND_FUNC, CLAMP, CHANNEL_COUNT, 1 >(Base->GetData(), Base, Mask, Col, BaseOffset, ColOffset, bOnlyOneMip);
 		}
 		else
 		{
@@ -800,16 +803,16 @@ namespace mu
 
 	//---------------------------------------------------------------------------------------------
 	template< unsigned (*BLEND_FUNC_MASKED)(unsigned, unsigned, unsigned), unsigned (*BLEND_FUNC)(unsigned, unsigned), uint32 CHANNEL_COUNT >
-	inline void BufferLayerColourInPlace(Image* Base, const Image* Mask, FVector4f Col, bool bOnlyOneMip, uint32 BaseOffset)
+	inline void BufferLayerColourInPlace(Image* Base, const Image* Mask, FVector4f Col, bool bOnlyOneMip, uint32 BaseOffset, uint8 ColOffset)
 	{
 		bool bIsClampNeeded = IsAnyComponentLargerThan1(Col);
 		if (bIsClampNeeded)
 		{
-			BufferLayerColourInPlace<BLEND_FUNC_MASKED, BLEND_FUNC, true, CHANNEL_COUNT>(Base, Mask, Col, bOnlyOneMip, BaseOffset);
+			BufferLayerColourInPlace<BLEND_FUNC_MASKED, BLEND_FUNC, true, CHANNEL_COUNT>(Base, Mask, Col, bOnlyOneMip, BaseOffset, ColOffset);
 		}
 		else
 		{
-			BufferLayerColourInPlace<BLEND_FUNC_MASKED, BLEND_FUNC, false, CHANNEL_COUNT>(Base, Mask, Col, bOnlyOneMip, BaseOffset);
+			BufferLayerColourInPlace<BLEND_FUNC_MASKED, BLEND_FUNC, false, CHANNEL_COUNT>(Base, Mask, Col, bOnlyOneMip, BaseOffset, ColOffset);
 		}
 	}
 
@@ -832,7 +835,12 @@ namespace mu
 	{
 		check(pBase->GetSizeX() == pBlended->GetSizeX());
 		check(pBase->GetSizeY() == pBlended->GetSizeY());
-		check(pBase->GetFormat() == pBlended->GetFormat());
+
+		// No longer required.
+		//check(pBase->GetFormat() == pBlended->GetFormat());
+		check(BaseChannelOffset + CHANNELS_TO_BLEND <= GetImageFormatData(pBase->GetFormat()).m_channels);
+		check(BlendedChannelOffset + CHANNELS_TO_BLEND <= GetImageFormatData(pBlended->GetFormat()).m_channels);
+
 		check(bOnlyFirstLOD || pBase->GetLODCount() <= pBlended->GetLODCount());
 
         uint8* pBaseBuf = pBase->GetData() + BaseChannelOffset;
@@ -880,18 +888,20 @@ namespace mu
 		// Number of total channels in the base image
 		int32 BASE_CHANNEL_STRIDE,
 		// Number of total channels in the blend image
-		int32 BLENDED_CHANNEL_STRIDE>
+		int32 BLENDED_CHANNEL_STRIDE,
+		int32 BLENDED_CHANNEL_OFFSET>
 
 	inline void BufferLayerFormat(uint8* pDestBuf, const Image* pBase, const Image* pBlended,
 		bool bOnlyFirstLOD)
 	{
 		check(pBase->GetSizeX() == pBlended->GetSizeX());
 		check(pBase->GetSizeY() == pBlended->GetSizeY());
-		check(pBase->GetFormat() == pBlended->GetFormat());
+		// Not true anymore, since the BLENDED_CHANNEL_OFFSET has been added.
+		// check(pBase->GetFormat() == pBlended->GetFormat());
 		check(bOnlyFirstLOD || pBase->GetLODCount() <= pBlended->GetLODCount());
 
 		const uint8* pBaseBuf = pBase->GetData();
-		const uint8* pBlendedBuf = pBlended->GetData();
+		const uint8* pBlendedBuf = pBlended->GetData()+BLENDED_CHANNEL_OFFSET;
 
 		// Generic implementation
 		int32 PixelCount = 0;
@@ -940,7 +950,7 @@ namespace mu
 	//! Apply a blending function to an image with another image as blending layer
 	//---------------------------------------------------------------------------------------------
 	template< unsigned (*BLEND_FUNC)(unsigned, unsigned), bool CLAMP >
-	inline void BufferLayer(Image* pResult, const Image* pBase, const Image* pBlended, bool bApplyToAlpha, bool bOnlyOneMip)
+	inline void BufferLayer(Image* pResult, const Image* pBase, const Image* pBlended, bool bApplyToAlpha, bool bOnlyOneMip, bool bUseBlendSourceFromBlendAlpha)
 	{
 		check(pResult->GetFormat() == pBase->GetFormat());
 		check(pResult->GetSizeX() == pBase->GetSizeX());
@@ -948,34 +958,58 @@ namespace mu
 		check(bOnlyOneMip || pResult->GetLODCount() == pBase->GetLODCount());
 		check(pBase->GetSizeX() == pBlended->GetSizeX());
 		check(pBase->GetSizeY() == pBlended->GetSizeY());
-		check(pBase->GetFormat() == pBlended->GetFormat());
 		check(bOnlyOneMip || pResult->GetLODCount() <= pBlended->GetLODCount());
 
-		EImageFormat baseFormat = pBase->GetFormat();
+		EImageFormat BaseFormat = pBase->GetFormat();
+		EImageFormat BlendedFormat = pBlended->GetFormat();
 		uint8* pDestBuf = pResult->GetData();
 
-		if (baseFormat == EImageFormat::IF_RGB_UBYTE)
+		if (bUseBlendSourceFromBlendAlpha)
 		{
-			BufferLayerFormat< BLEND_FUNC, CLAMP, 3, 3, 3 >(pDestBuf, pBase, pBlended, bOnlyOneMip);
-		}
-		else if (baseFormat == EImageFormat::IF_RGBA_UBYTE || baseFormat == EImageFormat::IF_BGRA_UBYTE)
-		{
-			if (bApplyToAlpha)
+			if (BlendedFormat == EImageFormat::IF_RGBA_UBYTE || BlendedFormat == EImageFormat::IF_BGRA_UBYTE)
 			{
-				BufferLayerFormat< BLEND_FUNC, CLAMP, 4, 4, 4 >(pDestBuf, pBase, pBlended, bOnlyOneMip);
+				if (BaseFormat == EImageFormat::IF_L_UBYTE)
+				{
+					BufferLayerFormat< BLEND_FUNC, CLAMP, 1, 1, 4, 3 >(pDestBuf, pBase, pBlended, bOnlyOneMip);
+				}
+				else
+				{
+					checkf(false, TEXT("Unsupported format."));
+				}
 			}
-			else
+			else if (BlendedFormat == EImageFormat::IF_L_UBYTE)
 			{
-				BufferLayerFormat< BLEND_FUNC, CLAMP, 3, 4, 4 >(pDestBuf, pBase, pBlended, bOnlyOneMip);
+				BufferLayerFormat< BLEND_FUNC, CLAMP, 1, 1, 1, 0 >(pDestBuf, pBase, pBlended, bOnlyOneMip);
 			}
 		}
-		else if (baseFormat == EImageFormat::IF_L_UBYTE)
-		{
-			BufferLayerFormat< BLEND_FUNC, CLAMP, 1, 1, 1 >(pDestBuf, pBase, pBlended, bOnlyOneMip);
-		}
 		else
 		{
-			checkf(false, TEXT("Unsupported format."));
+			check(BaseFormat == BlendedFormat);
+			if (BaseFormat == EImageFormat::IF_RGB_UBYTE)
+			{
+				check(!bUseBlendSourceFromBlendAlpha);
+				BufferLayerFormat< BLEND_FUNC, CLAMP, 3, 3, 3, 0 >(pDestBuf, pBase, pBlended, bOnlyOneMip);
+			}
+			else if (BaseFormat == EImageFormat::IF_RGBA_UBYTE || BaseFormat == EImageFormat::IF_BGRA_UBYTE)
+			{
+				check(!bUseBlendSourceFromBlendAlpha);
+				if (bApplyToAlpha)
+				{
+					BufferLayerFormat< BLEND_FUNC, CLAMP, 4, 4, 4, 0 >(pDestBuf, pBase, pBlended, bOnlyOneMip);
+				}
+				else
+				{
+					BufferLayerFormat< BLEND_FUNC, CLAMP, 3, 4, 4, 0 >(pDestBuf, pBase, pBlended, bOnlyOneMip);
+				}
+			}
+			else if (BaseFormat == EImageFormat::IF_L_UBYTE)
+			{
+				BufferLayerFormat< BLEND_FUNC, CLAMP, 1, 1, 1, 0 >(pDestBuf, pBase, pBlended, bOnlyOneMip);
+			}
+			else
+			{
+				checkf(false, TEXT("Unsupported format."));
+			}
 		}
 	}
 
@@ -986,23 +1020,35 @@ namespace mu
 	{
 		check(pBase->GetSizeX() == pBlended->GetSizeX());
 		check(pBase->GetSizeY() == pBlended->GetSizeY());
-		check(pBase->GetFormat() == pBlended->GetFormat());
+		// Not required since we have the CHANNEL_COUNT and offsets. 
+		// check(pBase->GetFormat() == pBlended->GetFormat());
 
-		EImageFormat baseFormat = pBase->GetFormat();
+		EImageFormat BaseFormat = pBase->GetFormat();
+		EImageFormat BlendFormat = pBlended->GetFormat();
 
-		if (baseFormat == EImageFormat::IF_RGB_UBYTE)
+		if (BaseFormat==EImageFormat::IF_RGB_UBYTE && BlendFormat==EImageFormat::IF_RGB_UBYTE)
 		{
 			check(BaseOffset+CHANNEL_COUNT <= 3);
 			check(BlendedOffset+CHANNEL_COUNT <= 3);
 			BufferLayerFormatInPlace< BLEND_FUNC, CLAMP, CHANNEL_COUNT, 3, 3 >(pBase, pBlended, BaseOffset, BlendedOffset, bOnlyOneMip);
 		}
-		else if (baseFormat == EImageFormat::IF_RGBA_UBYTE || baseFormat == EImageFormat::IF_BGRA_UBYTE)
+		else if ( (BaseFormat == EImageFormat::IF_RGBA_UBYTE || BaseFormat == EImageFormat::IF_BGRA_UBYTE)
+			&&
+			(BlendFormat == EImageFormat::IF_RGBA_UBYTE || BlendFormat == EImageFormat::IF_BGRA_UBYTE) )
 		{
-			check(BaseOffset+CHANNEL_COUNT <= 4);
-			check(BlendedOffset+CHANNEL_COUNT <= 4);
+			check(BaseOffset + CHANNEL_COUNT <= 4);
+			check(BlendedOffset + CHANNEL_COUNT <= 4);
 			BufferLayerFormatInPlace< BLEND_FUNC, CLAMP, CHANNEL_COUNT, 4, 4 >(pBase, pBlended, BaseOffset, BlendedOffset, bOnlyOneMip);
 		}
-		else if (baseFormat == EImageFormat::IF_L_UBYTE)
+		else if ( (BaseFormat == EImageFormat::IF_RGBA_UBYTE || BaseFormat == EImageFormat::IF_BGRA_UBYTE)
+			&&
+			BlendFormat == EImageFormat::IF_L_UBYTE )
+		{
+			check(BaseOffset + CHANNEL_COUNT <= 4);
+			check(BlendedOffset + CHANNEL_COUNT <= 1);
+			BufferLayerFormatInPlace< BLEND_FUNC, CLAMP, CHANNEL_COUNT, 4, 1 >(pBase, pBlended, BaseOffset, BlendedOffset, bOnlyOneMip);
+		}
+		else if (BaseFormat==EImageFormat::IF_L_UBYTE && BlendFormat==EImageFormat::IF_L_UBYTE)
 		{
 			check(BaseOffset+CHANNEL_COUNT <= 1);
 			BufferLayerFormatInPlace< BLEND_FUNC, CLAMP, CHANNEL_COUNT, 1, 1 >(pBase, pBlended, BaseOffset, BlendedOffset, bOnlyOneMip);
@@ -1521,7 +1567,8 @@ namespace mu
 	inline void BufferLayerComposite(
 		Image* pBase,
 		const Image* pBlend,
-		bool bOnlyFirstLOD)
+		bool bOnlyFirstLOD,
+		uint8 BlendAlphaSourceChannel )
 	{
 		check(pBase->GetFormat() == EImageFormat::IF_RGBA_UBYTE);
 		check(pBlend->GetFormat() == EImageFormat::IF_RGBA_UBYTE);
@@ -1546,7 +1593,7 @@ namespace mu
 
 		ParallelFor(PixelCount,
 			[
-				pBaseBuf, pBlendedBuf
+				pBaseBuf, pBlendedBuf, BlendAlphaSourceChannel
 			] (uint32 i)
 			{
 				// TODO: Optimize this (SIMD?)
@@ -1571,7 +1618,7 @@ namespace mu
 				// A
 				{
 					uint32 base = pBaseBuf[4 * i + 3];
-					uint32 blended = pBlendedBuf[4 * i + 3];
+					uint32 blended = pBlendedBuf[4 * i + BlendAlphaSourceChannel];
 					uint32 result = A_FUNC(base, blended);
 					if (CLAMP)
 					{
@@ -1957,7 +2004,7 @@ namespace mu
 	template< 
 			uint32 (*BLEND_FUNC)(uint32, uint32), 
 			size_t NC >
-	inline void BufferLayerCombineColour(uint8* pDestBuf, const Image* pBase, vec3<float> col)
+	inline void BufferLayerCombineColour(uint8* pDestBuf, const Image* pBase, FVector4f col)
 	{
 		//static_assert(NC > 0 && NC <= 4);
 
@@ -2057,7 +2104,7 @@ namespace mu
 	template< 
 			uint32 (*BLEND_FUNC_MASKED)(uint32, uint32, uint32), 
 			size_t NC >
-	inline void BufferLayerCombineColour(uint8* pDestBuf, const Image* pBase, const Image* pMask, vec3<float> col)
+	inline void BufferLayerCombineColour(uint8* pDestBuf, const Image* pBase, const Image* pMask, FVector4f col)
 	{
 		//static_assert(NC > 0 && NC <= 4);
 
@@ -2166,7 +2213,7 @@ namespace mu
 	}
 
 	template< uint32 (*BLEND_FUNC)(uint32, uint32)>
-	inline void ImageLayerCombineColour(Image* pResult, const Image* pBase, vec3<float> col)
+	inline void ImageLayerCombineColour(Image* pResult, const Image* pBase, FVector4f col)
 	{
 		check(pResult->GetFormat() == pBase->GetFormat());
 		check(pResult->GetSizeX() == pBase->GetSizeX());
@@ -2199,7 +2246,7 @@ namespace mu
 	template< 
 		uint32 (*BLEND_FUNC)(uint32, uint32),
 		uint32 (*BLEND_FUNC_MASKED)(uint32, uint32, uint32) >
-	inline void ImageLayerCombineColour(Image* pResult, const Image* pBase, const Image* pMask, vec3<float> col)
+	inline void ImageLayerCombineColour(Image* pResult, const Image* pBase, const Image* pMask, FVector4f col)
 	{
 		check(pResult->GetFormat() == pBase->GetFormat());
 		check(pResult->GetSizeX() == pBase->GetSizeX());
diff --git a/Engine/Plugins/Experimental/Mutable/Source/MutableRuntime/Private/MuR/OpImageNormalCombine.h b/Engine/Plugins/Experimental/Mutable/Source/MutableRuntime/Private/MuR/OpImageNormalCombine.h
--- a/Engine/Plugins/Experimental/Mutable/Source/MutableRuntime/Private/MuR/OpImageNormalCombine.h
+++ b/Engine/Plugins/Experimental/Mutable/Source/MutableRuntime/Private/MuR/OpImageNormalCombine.h
@@ -61,7 +61,7 @@ namespace mu
 			(mu::clamp(static_cast<unsigned>((n.z() + 1.0f) * 0.5f * 255.0f), 0u, 255u) << 16);
 	}
 
-	inline void ImageNormalCombine(Image* pResult, const Image* pBase, vec3<float> col)
+	inline void ImageNormalCombine(Image* pResult, const Image* pBase, FVector4f col)
 	{
 		ImageLayerCombineColour<CombineNormal>(pResult, pBase, col);
 	}
@@ -71,7 +71,7 @@ namespace mu
 		ImageLayerCombine<CombineNormal>(pResult, pBase, pBlended, bOnlyFirstLOD);
 	}
 
-	inline void ImageNormalCombine(Image* pResult, const Image* pBase, const Image* pMask, vec3<float> col)
+	inline void ImageNormalCombine(Image* pResult, const Image* pBase, const Image* pMask, FVector4f col)
 	{
 		ImageLayerCombineColour<CombineNormal, CombineNormalMasked>(pResult, pBase, pMask, col);
 	}
diff --git a/Engine/Plugins/Experimental/Mutable/Source/MutableRuntime/Private/MuR/Operations.h b/Engine/Plugins/Experimental/Mutable/Source/MutableRuntime/Private/MuR/Operations.h
--- a/Engine/Plugins/Experimental/Mutable/Source/MutableRuntime/Private/MuR/Operations.h
+++ b/Engine/Plugins/Experimental/Mutable/Source/MutableRuntime/Private/MuR/Operations.h
@@ -531,7 +531,13 @@ namespace mu
             ADDRESS mask;
             ADDRESS blended;
 			uint8 blendType;		// One of EBlendType
-			uint8 blendTypeAlpha;	// One of EBlendType
+
+			/** One of EBlendType. If this is different than NONE, it will be applied to the alpha with the channel BlendAlphaSourceChannel of the blended. */
+			uint8 blendTypeAlpha;
+
+			/** Channel to use from the source blended argument to apply blendTypeAlpha, if any. */
+			uint8 BlendAlphaSourceChannel;
+
 			uint8 flags;
 			typedef enum
             {
@@ -544,6 +550,8 @@ namespace mu
 				F_USE_MASK_FROM_BLENDED = 1 << 2,
 				/** Use the alpha channel of the base image as its RGB.*/
 				F_BASE_RGB_FROM_ALPHA = 1 << 3,
+				/** Use the alpha channel of the blended image as its RGB.*/
+				F_BLENDED_RGB_FROM_ALPHA = 1 << 4,
 			} FLAGS;
         };
 
@@ -555,7 +563,13 @@ namespace mu
             ADDRESS rangeSize;
             uint16 rangeId;
 			uint8 blendType;		// One of EBlendType
-			uint8 blendTypeAlpha;	// One of EBlendType
+			
+			/** One of EBlendType. If this is different than NONE, it will be applied to the alpha with the channel BlendAlphaSourceChannel of the blended. */
+			uint8 blendTypeAlpha;
+
+			/** Channel to use from the source color argument to apply blendTypeAlpha, if any. */
+			uint8 BlendAlphaSourceChannel;
+
 			uint8 bUseMaskFromBlended;	
 		};
 
@@ -565,8 +579,15 @@ namespace mu
             ADDRESS mask;
             ADDRESS colour;
 			uint8 blendType;		// One of EBlendType
-			uint8 blendTypeAlpha;	// One of EBlendType
-			uint8 flags;			// Like in ImageLayerArgs
+
+			/** One of EBlendType. If this is different than NONE, it will be applied to the alpha but with the channel BlendAlphaSourceChannel of the color. */
+			uint8 blendTypeAlpha;	
+
+			/** Channel to use from the source color argument to apply blendTypeAlpha, if any. */
+			uint8 BlendAlphaSourceChannel;
+
+			/** Like in ImageLayerArgs. */
+			uint8 flags;
 		};
 
         struct ImagePixelFormatArgs
diff --git a/Engine/Plugins/Experimental/Mutable/Source/MutableRuntime/Public/MuR/Model.h b/Engine/Plugins/Experimental/Mutable/Source/MutableRuntime/Public/MuR/Model.h
--- a/Engine/Plugins/Experimental/Mutable/Source/MutableRuntime/Public/MuR/Model.h
+++ b/Engine/Plugins/Experimental/Mutable/Source/MutableRuntime/Public/MuR/Model.h
@@ -13,7 +13,7 @@
 //! data structures. Compiled models are not necessarily compatible when the runtime is updated,
 //! so this version number can be used externally to verify this. It is not used internally, and
 //! serializing models from different versions than this runtime will probably result in a crash.
-#define MUTABLE_COMPILED_MODEL_CODE_VERSION		uint32( 54 )
+#define MUTABLE_COMPILED_MODEL_CODE_VERSION		uint32( 55 )
 #define MUTABLE_PARAMETERS_VERSION              uint32( 1 )
 
 
diff --git a/Engine/Plugins/Experimental/Mutable/Source/MutableTools/Private/MuT/AST.h b/Engine/Plugins/Experimental/Mutable/Source/MutableTools/Private/MuT/AST.h
--- a/Engine/Plugins/Experimental/Mutable/Source/MutableTools/Private/MuT/AST.h
+++ b/Engine/Plugins/Experimental/Mutable/Source/MutableTools/Private/MuT/AST.h
@@ -449,7 +449,7 @@ namespace mu
 
         //!
         virtual Ptr<ASTOp> OptimiseSize() const { return nullptr; }
-        virtual Ptr<ASTOp> OptimiseSemantic(const FModelOptimizationOptions&) const { return nullptr; }
+        virtual Ptr<ASTOp> OptimiseSemantic(const FModelOptimizationOptions&, int32 Pass) const { return nullptr; }
         virtual Ptr<ASTOp> OptimiseSink(const FModelOptimizationOptions&, FOptimizeSinkContext& ) const { return nullptr; }
         virtual Ptr<ImageSizeExpression> GetImageSizeExpression() const
         {
@@ -986,7 +986,7 @@ namespace mu
 		FImageDesc GetImageDesc( bool returnBestOption, FGetImageDescContext* context ) const override;
         void GetLayoutBlockSize( int* pBlockX, int* pBlockY ) override;
         Ptr<ASTOp> OptimiseSize() const override;
-        Ptr<ASTOp> OptimiseSemantic(const FModelOptimizationOptions&) const override;
+        Ptr<ASTOp> OptimiseSemantic(const FModelOptimizationOptions&, int32 Pass) const override;
         Ptr<ASTOp> OptimiseSink(const FModelOptimizationOptions&, FOptimizeSinkContext&) const override;
         FBoolEvalResult EvaluateBool( ASTOpList& facts, FEvaluateBoolCache* cache ) const override;
         int EvaluateInt( ASTOpList& facts, bool &unknown ) const override;
diff --git a/Engine/Plugins/Experimental/Mutable/Source/MutableTools/Private/MuT/ASTOpConditional.cpp b/Engine/Plugins/Experimental/Mutable/Source/MutableTools/Private/MuT/ASTOpConditional.cpp
--- a/Engine/Plugins/Experimental/Mutable/Source/MutableTools/Private/MuT/ASTOpConditional.cpp
+++ b/Engine/Plugins/Experimental/Mutable/Source/MutableTools/Private/MuT/ASTOpConditional.cpp
@@ -281,7 +281,7 @@ namespace mu
 
 
 	//-------------------------------------------------------------------------------------------------
-	mu::Ptr<ASTOp> ASTOpConditional::OptimiseSemantic(const FModelOptimizationOptions&) const
+	mu::Ptr<ASTOp> ASTOpConditional::OptimiseSemantic(const FModelOptimizationOptions&, int32 Pass) const
 	{
 		if (!condition)
 		{
diff --git a/Engine/Plugins/Experimental/Mutable/Source/MutableTools/Private/MuT/ASTOpConditional.h b/Engine/Plugins/Experimental/Mutable/Source/MutableTools/Private/MuT/ASTOpConditional.h
--- a/Engine/Plugins/Experimental/Mutable/Source/MutableTools/Private/MuT/ASTOpConditional.h
+++ b/Engine/Plugins/Experimental/Mutable/Source/MutableTools/Private/MuT/ASTOpConditional.h
@@ -47,7 +47,7 @@ struct FProgram;
 		FImageDesc GetImageDesc(bool returnBestOption, class FGetImageDescContext* context) const override;
 		void GetLayoutBlockSize(int* pBlockX, int* pBlockY) override;
 		void GetBlockLayoutSize(int blockIndex, int* pBlockX, int* pBlockY, FBlockLayoutSizeCache* cache) override;
-		Ptr<ASTOp> OptimiseSemantic(const FModelOptimizationOptions&) const override;
+		Ptr<ASTOp> OptimiseSemantic(const FModelOptimizationOptions&, int32 Pass) const override;
 		bool GetNonBlackRect(FImageRect& maskUsage) const override;
 		Ptr<ImageSizeExpression> GetImageSizeExpression() const override;
 	};
diff --git a/Engine/Plugins/Experimental/Mutable/Source/MutableTools/Private/MuT/ASTOpImageCompose.cpp b/Engine/Plugins/Experimental/Mutable/Source/MutableTools/Private/MuT/ASTOpImageCompose.cpp
--- a/Engine/Plugins/Experimental/Mutable/Source/MutableTools/Private/MuT/ASTOpImageCompose.cpp
+++ b/Engine/Plugins/Experimental/Mutable/Source/MutableTools/Private/MuT/ASTOpImageCompose.cpp
@@ -216,7 +216,7 @@ void ASTOpImageCompose::GetLayoutBlockSize(int* pBlockX, int* pBlockY)
 }
 
 
-mu::Ptr<ASTOp> ASTOpImageCompose::OptimiseSemantic(const FModelOptimizationOptions& options) const
+mu::Ptr<ASTOp> ASTOpImageCompose::OptimiseSemantic(const FModelOptimizationOptions& options, int32 Pass) const
 {
 	mu::Ptr<ASTOp> at;
 
diff --git a/Engine/Plugins/Experimental/Mutable/Source/MutableTools/Private/MuT/ASTOpImageCompose.h b/Engine/Plugins/Experimental/Mutable/Source/MutableTools/Private/MuT/ASTOpImageCompose.h
--- a/Engine/Plugins/Experimental/Mutable/Source/MutableTools/Private/MuT/ASTOpImageCompose.h
+++ b/Engine/Plugins/Experimental/Mutable/Source/MutableTools/Private/MuT/ASTOpImageCompose.h
@@ -37,7 +37,7 @@ template <class SCALAR> class vec4;
 		void ForEachChild(const TFunctionRef<void(ASTChild&)>) override;
 		void Link(FProgram& program, const FLinkerOptions* Options) override;
 		//TODO: Ptr<ASTOp> OptimiseSink(const FModelOptimizationOptions& options, FOptimizeSinkContext& context) const override;
-		Ptr<ASTOp> OptimiseSemantic(const FModelOptimizationOptions& options) const;
+		Ptr<ASTOp> OptimiseSemantic(const FModelOptimizationOptions& options, int32 Pass) const;
 		FImageDesc GetImageDesc(bool returnBestOption, FGetImageDescContext* context) const override;
 		void GetLayoutBlockSize(int* pBlockX, int* pBlockY) override;
 		bool IsImagePlainConstant(FVector4f& colour) const override;
diff --git a/Engine/Plugins/Experimental/Mutable/Source/MutableTools/Private/MuT/ASTOpImageLayer.cpp b/Engine/Plugins/Experimental/Mutable/Source/MutableTools/Private/MuT/ASTOpImageLayer.cpp
--- a/Engine/Plugins/Experimental/Mutable/Source/MutableTools/Private/MuT/ASTOpImageLayer.cpp
+++ b/Engine/Plugins/Experimental/Mutable/Source/MutableTools/Private/MuT/ASTOpImageLayer.cpp
@@ -46,6 +46,7 @@ namespace mu
 				mask == Other->mask &&
 				blendType == Other->blendType &&
 				blendTypeAlpha == Other->blendTypeAlpha &&
+				BlendAlphaSourceChannel == Other->BlendAlphaSourceChannel &&
 				Flags == Other->Flags;
 		}
 		return false;
@@ -72,6 +73,7 @@ namespace mu
 		n->mask = mapChild(mask.child());
 		n->blendType = blendType;
 		n->blendTypeAlpha = blendTypeAlpha;
+		n->BlendAlphaSourceChannel = BlendAlphaSourceChannel;
 		n->Flags = Flags;
 		return n;
 	}
@@ -97,6 +99,7 @@ namespace mu
 
 			args.blendType = (uint8)blendType;
 			args.blendTypeAlpha = (uint8)blendTypeAlpha;
+			args.BlendAlphaSourceChannel = BlendAlphaSourceChannel;
 			args.flags = Flags;
 
 			if (base) args.base = base->linkedAddress;
@@ -173,7 +176,7 @@ namespace mu
 
 
 	//---------------------------------------------------------------------------------------------
-	Ptr<ASTOp> ASTOpImageLayer::OptimiseSemantic(const FModelOptimizationOptions& options) const
+	Ptr<ASTOp> ASTOpImageLayer::OptimiseSemantic(const FModelOptimizationOptions& options, int32 Pass) const
 	{
 		Ptr<ASTOp> at;
 
@@ -182,19 +185,19 @@ namespace mu
 		auto maskAt = mask.child();
 
 		// Convert to image layer color if blend is plain
-		if (!at && blendAt->GetOpType() == OP_TYPE::IM_PLAINCOLOUR)
+		if (!at && blendAt && blendAt->GetOpType() == OP_TYPE::IM_PLAINCOLOUR)
 		{
-			bool bUseMaskFromBlendAlpha = (Flags & OP::ImageLayerArgs::F_USE_MASK_FROM_BLENDED);
-
-			const ASTOpFixed* BlendPlainColor = dynamic_cast<const ASTOpFixed*>(blendAt.get());
-
-			if (!bUseMaskFromBlendAlpha)
+			// TODO: May some blags be supported?
+			if (Flags == 0)
 			{
+				const ASTOpFixed* BlendPlainColor = dynamic_cast<const ASTOpFixed*>(blendAt.get());
+
 				Ptr<ASTOpImageLayerColor> NewLayerColor = new ASTOpImageLayerColor;
 				NewLayerColor->base = baseAt;
 				NewLayerColor->mask = maskAt;
 				NewLayerColor->blendType = blendType;
 				NewLayerColor->blendTypeAlpha = blendTypeAlpha;
+				NewLayerColor->BlendAlphaSourceChannel = BlendAlphaSourceChannel;
 				NewLayerColor->color = BlendPlainColor->children[BlendPlainColor->op.args.ImagePlainColour.colour].child();
 				at = NewLayerColor;
 			}
@@ -398,7 +401,7 @@ namespace mu
 			}
 
 			// Is the mask a swizzle expanding alpha from a texture?
-			if (maskAt && maskAt->GetOpType() == OP_TYPE::IM_SWIZZLE)
+			if (!at && maskAt && maskAt->GetOpType() == OP_TYPE::IM_SWIZZLE)
 			{
 				const ASTOpImageSwizzle* TypedBase = dynamic_cast<const ASTOpImageSwizzle*>(maskAt.get());
 				bool bAreAllAlpha = true;
@@ -416,6 +419,23 @@ namespace mu
 					// TODO
 				}
 			}
+
+			// Is the blend a swizzle selecting alpha?
+			if (Pass>0 && !at && blendAt && blendAt->GetOpType() == OP_TYPE::IM_SWIZZLE && Flags==0)
+			{
+				const ASTOpImageSwizzle* TypedBlend = dynamic_cast<const ASTOpImageSwizzle*>(blendAt.get());
+				if (TypedBlend->Format==EImageFormat::IF_L_UBYTE
+					&&
+					TypedBlend->SourceChannels[0]==3)
+				{
+					Ptr<ASTOpImageLayer> nop = mu::Clone<ASTOpImageLayer>(this);
+					nop->Flags = Flags | OP::ImageLayerArgs::FLAGS::F_BLENDED_RGB_FROM_ALPHA;
+					nop->blend = TypedBlend->Sources[0].child();
+					auto Desc = nop->blend->GetImageDesc(true);
+					check(Desc.m_format==EImageFormat::IF_RGBA_UBYTE);
+					at = nop;					
+				}
+			}
 		}
 
 		// Introduce crop if mask is constant and smaller than the base
@@ -623,7 +643,7 @@ namespace mu
 		}
 
 		// If we failed to optimize so far, see if it is worth optimizing the blended branch only.
-		if (!at)
+		if (!at && blendAt)
 		{
 			OP_TYPE BlendType = blendAt->GetOpType();
 			switch (BlendType)
diff --git a/Engine/Plugins/Experimental/Mutable/Source/MutableTools/Private/MuT/ASTOpImageLayer.h b/Engine/Plugins/Experimental/Mutable/Source/MutableTools/Private/MuT/ASTOpImageLayer.h
--- a/Engine/Plugins/Experimental/Mutable/Source/MutableTools/Private/MuT/ASTOpImageLayer.h
+++ b/Engine/Plugins/Experimental/Mutable/Source/MutableTools/Private/MuT/ASTOpImageLayer.h
@@ -31,6 +31,9 @@ namespace mu
 		/** Blend type used for the alpha channel if any. */
 		EBlendType blendTypeAlpha = EBlendType::BT_NONE;
 
+		/** Channel to use from the source color argument to apply blendTypeAlpha, if any. */
+		uint8 BlendAlphaSourceChannel = 0;
+
 		/** See ImageLayerArgs::Flags .*/
 		uint8 Flags = 0;
 
@@ -49,7 +52,7 @@ namespace mu
 		FImageDesc GetImageDesc(bool returnBestOption, FGetImageDescContext*) const override;
 		void GetLayoutBlockSize(int* pBlockX, int* pBlockY) override;
 		Ptr<ImageSizeExpression> GetImageSizeExpression() const override;
-		Ptr<ASTOp> OptimiseSemantic(const FModelOptimizationOptions&) const override;
+		Ptr<ASTOp> OptimiseSemantic(const FModelOptimizationOptions&, int32 Pass) const override;
 		Ptr<ASTOp> OptimiseSink(const FModelOptimizationOptions&, FOptimizeSinkContext&) const;
 
 	};
diff --git a/Engine/Plugins/Experimental/Mutable/Source/MutableTools/Private/MuT/ASTOpImageLayerColor.cpp b/Engine/Plugins/Experimental/Mutable/Source/MutableTools/Private/MuT/ASTOpImageLayerColor.cpp
--- a/Engine/Plugins/Experimental/Mutable/Source/MutableTools/Private/MuT/ASTOpImageLayerColor.cpp
+++ b/Engine/Plugins/Experimental/Mutable/Source/MutableTools/Private/MuT/ASTOpImageLayerColor.cpp
@@ -42,6 +42,7 @@ namespace mu
 				mask == Other->mask &&
 				blendType == Other->blendType &&
 				blendTypeAlpha == Other->blendTypeAlpha &&
+				BlendAlphaSourceChannel == Other->BlendAlphaSourceChannel &&
 				Flags == Other->Flags;
 		}
 		return false;
@@ -68,6 +69,7 @@ namespace mu
 		n->mask = mapChild(mask.child());
 		n->blendType = blendType;
 		n->blendTypeAlpha = blendTypeAlpha;
+		n->BlendAlphaSourceChannel = BlendAlphaSourceChannel;
 		n->Flags = Flags;
 		return n;
 	}
@@ -93,6 +95,7 @@ namespace mu
 
 			args.blendType = (uint8)blendType;
 			args.blendTypeAlpha = (uint8)blendTypeAlpha;
+			args.BlendAlphaSourceChannel = BlendAlphaSourceChannel;
 			args.flags = Flags;
 
 			if (base) args.base = base->linkedAddress;
@@ -168,7 +171,7 @@ namespace mu
 
 
 	//---------------------------------------------------------------------------------------------
-	Ptr<ASTOp> ASTOpImageLayerColor::OptimiseSemantic(const FModelOptimizationOptions& options) const
+	Ptr<ASTOp> ASTOpImageLayerColor::OptimiseSemantic(const FModelOptimizationOptions& options, int32 Pass) const
 	{
 		Ptr<ASTOp> at;
 
@@ -214,8 +217,8 @@ namespace mu
 					// TODO: Update when alpha may come from alpha in the color.
 					switch (blendTypeAlpha)
 					{
-					case EBlendType::BT_LIGHTEN: bAlphaUnchanged = ColorConst.IsNearlyZero3(UE_SMALL_NUMBER); break;
-					case EBlendType::BT_MULTIPLY: bAlphaUnchanged = ColorConst.Equals(FVector4f(1, 1, 1, 1)); break;
+					case EBlendType::BT_LIGHTEN: bAlphaUnchanged = FMath::IsNearlyEqual(ColorConst[BlendAlphaSourceChannel], 0.0f); break;
+					case EBlendType::BT_MULTIPLY: bAlphaUnchanged = FMath::IsNearlyEqual(ColorConst[BlendAlphaSourceChannel], 1.0f); break;
 					default: break;
 					}
 				}
@@ -226,7 +229,7 @@ namespace mu
 					{
 						switch (blendType)
 						{
-						case EBlendType::BT_LIGHTEN: bRGBUnchanged = ColorConst[3]<UE_SMALL_NUMBER; break;
+						case EBlendType::BT_LIGHTEN: bRGBUnchanged = FMath::IsNearlyEqual(ColorConst[3],0.0f); break;
 						case EBlendType::BT_MULTIPLY: bRGBUnchanged = FMath::IsNearlyEqual(ColorConst[3],1.0f); break;
 						default: break;
 						}
diff --git a/Engine/Plugins/Experimental/Mutable/Source/MutableTools/Private/MuT/ASTOpImageLayerColor.h b/Engine/Plugins/Experimental/Mutable/Source/MutableTools/Private/MuT/ASTOpImageLayerColor.h
--- a/Engine/Plugins/Experimental/Mutable/Source/MutableTools/Private/MuT/ASTOpImageLayerColor.h
+++ b/Engine/Plugins/Experimental/Mutable/Source/MutableTools/Private/MuT/ASTOpImageLayerColor.h
@@ -28,9 +28,12 @@ namespace mu
 		/** Blend type used for the colour channels. */
 		EBlendType blendType = EBlendType::BT_NONE;
 
-		/** Blend type used for the alpha channel if any. */
+		/** Blend type used for the alpha channel if any. This will be applied to the alpha with the channel BlendAlphaSourceChannel of the color. */
 		EBlendType blendTypeAlpha = EBlendType::BT_NONE;
 
+		/** Channel to use from the source color argument to apply blendTypeAlpha, if any. */
+		uint8 BlendAlphaSourceChannel = 0;
+
 		/** See ImageLayerArgs::Flags .*/
 		uint8 Flags = 0;
 
@@ -49,7 +52,7 @@ namespace mu
 		FImageDesc GetImageDesc(bool returnBestOption, FGetImageDescContext* context) const override;
 		void GetLayoutBlockSize(int* pBlockX, int* pBlockY) override;
 		Ptr<ImageSizeExpression> GetImageSizeExpression() const override;
-		Ptr<ASTOp> OptimiseSemantic(const FModelOptimizationOptions& options) const;
+		Ptr<ASTOp> OptimiseSemantic(const FModelOptimizationOptions& options, int32 Pass) const;
 		Ptr<ASTOp> OptimiseSink(const FModelOptimizationOptions&, FOptimizeSinkContext&) const;
 
 	};
diff --git a/Engine/Plugins/Experimental/Mutable/Source/MutableTools/Private/MuT/ASTOpImageMultiLayer.cpp b/Engine/Plugins/Experimental/Mutable/Source/MutableTools/Private/MuT/ASTOpImageMultiLayer.cpp
--- a/Engine/Plugins/Experimental/Mutable/Source/MutableTools/Private/MuT/ASTOpImageMultiLayer.cpp
+++ b/Engine/Plugins/Experimental/Mutable/Source/MutableTools/Private/MuT/ASTOpImageMultiLayer.cpp
@@ -45,6 +45,7 @@ namespace mu
 				range == Other->range &&
 				blendType == Other->blendType &&
 				blendTypeAlpha == Other->blendTypeAlpha &&
+				BlendAlphaSourceChannel == Other->BlendAlphaSourceChannel &&
 				bUseMaskFromBlended == Other->bUseMaskFromBlended;
 		}
 		return false;
@@ -74,6 +75,7 @@ namespace mu
 		n->range.rangeSize = mapChild(range.rangeSize.child());
 		n->blendType = blendType;
 		n->blendTypeAlpha = blendTypeAlpha;
+		n->BlendAlphaSourceChannel = BlendAlphaSourceChannel;
 		n->bUseMaskFromBlended = bUseMaskFromBlended;
 		return n;
 	}
@@ -100,6 +102,7 @@ namespace mu
 
 			args.blendType = (uint8)blendType;
 			args.blendTypeAlpha = (uint8)blendTypeAlpha;
+			args.BlendAlphaSourceChannel = BlendAlphaSourceChannel;
 			args.bUseMaskFromBlended = bUseMaskFromBlended;
 
 			if (base) args.base = base->linkedAddress;
diff --git a/Engine/Plugins/Experimental/Mutable/Source/MutableTools/Private/MuT/ASTOpImageMultiLayer.h b/Engine/Plugins/Experimental/Mutable/Source/MutableTools/Private/MuT/ASTOpImageMultiLayer.h
--- a/Engine/Plugins/Experimental/Mutable/Source/MutableTools/Private/MuT/ASTOpImageMultiLayer.h
+++ b/Engine/Plugins/Experimental/Mutable/Source/MutableTools/Private/MuT/ASTOpImageMultiLayer.h
@@ -32,6 +32,9 @@ namespace mu
 		/** Blend type used for the alpha channel if any. */
 		EBlendType blendTypeAlpha = EBlendType::BT_NONE;
 
+		/** Channel to use from the source color argument to apply blendTypeAlpha, if any. */
+		uint8 BlendAlphaSourceChannel = 0;
+
 		/** If true, use the alpha channel of the blended image as mask. Mask should be null.*/
 		bool bUseMaskFromBlended = false;
 
diff --git a/Engine/Plugins/Experimental/Mutable/Source/MutableTools/Private/MuT/ASTOpImageRasterMesh.cpp b/Engine/Plugins/Experimental/Mutable/Source/MutableTools/Private/MuT/ASTOpImageRasterMesh.cpp
--- a/Engine/Plugins/Experimental/Mutable/Source/MutableTools/Private/MuT/ASTOpImageRasterMesh.cpp
+++ b/Engine/Plugins/Experimental/Mutable/Source/MutableTools/Private/MuT/ASTOpImageRasterMesh.cpp
@@ -186,7 +186,7 @@ namespace mu
 
 
 	//---------------------------------------------------------------------------------------------
-	Ptr<ASTOp> ASTOpImageRasterMesh::OptimiseSemantic(const FModelOptimizationOptions& options) const
+	Ptr<ASTOp> ASTOpImageRasterMesh::OptimiseSemantic(const FModelOptimizationOptions& options, int32 Pass) const
 	{
 		Ptr<ASTOp> at;
 
diff --git a/Engine/Plugins/Experimental/Mutable/Source/MutableTools/Private/MuT/ASTOpImageRasterMesh.h b/Engine/Plugins/Experimental/Mutable/Source/MutableTools/Private/MuT/ASTOpImageRasterMesh.h
--- a/Engine/Plugins/Experimental/Mutable/Source/MutableTools/Private/MuT/ASTOpImageRasterMesh.h
+++ b/Engine/Plugins/Experimental/Mutable/Source/MutableTools/Private/MuT/ASTOpImageRasterMesh.h
@@ -42,7 +42,7 @@ namespace mu
 		void Link(FProgram& program, const FLinkerOptions* Options) override;
 		FImageDesc GetImageDesc(bool returnBestOption, FGetImageDescContext* context) const override;
 		Ptr<ImageSizeExpression> GetImageSizeExpression() const override;
-		Ptr<ASTOp> OptimiseSemantic(const FModelOptimizationOptions& options) const;
+		Ptr<ASTOp> OptimiseSemantic(const FModelOptimizationOptions& options, int32 Pass) const;
 		Ptr<ASTOp> OptimiseSink(const FModelOptimizationOptions&, FOptimizeSinkContext&) const;
 
 	};
diff --git a/Engine/Plugins/Experimental/Mutable/Source/MutableTools/Private/MuT/ASTOpImageSwizzle.cpp b/Engine/Plugins/Experimental/Mutable/Source/MutableTools/Private/MuT/ASTOpImageSwizzle.cpp
--- a/Engine/Plugins/Experimental/Mutable/Source/MutableTools/Private/MuT/ASTOpImageSwizzle.cpp
+++ b/Engine/Plugins/Experimental/Mutable/Source/MutableTools/Private/MuT/ASTOpImageSwizzle.cpp
@@ -119,7 +119,7 @@ namespace mu
 	}
 
 
-	mu::Ptr<ASTOp> ASTOpImageSwizzle::OptimiseSemantic(const FModelOptimizationOptions&) const
+	mu::Ptr<ASTOp> ASTOpImageSwizzle::OptimiseSemantic(const FModelOptimizationOptions&, int32 Pass) const
 	{
 		Ptr<ASTOpImageSwizzle> sat;
 
@@ -204,7 +204,7 @@ namespace mu
 		MUTABLE_CPUPROFILER_SCOPE(OptimiseSwizzleAST);
 
 		//! Basic optimisation first
-		Ptr<ASTOp> at = OptimiseSemantic(options);
+		Ptr<ASTOp> at = OptimiseSemantic(options, 0);
 		if (at)
 		{
 			return at;
@@ -303,15 +303,18 @@ namespace mu
 				// We move the swizzle down the two paths
 				Ptr<ASTOpImageLayer> nop = mu::Clone<ASTOpImageLayer>(channelSourceAt);
 
-				Ptr<ASTOpImageSwizzle> aOp = mu::Clone<ASTOpImageSwizzle>(this);
-				ReplaceAllSources(aOp, nop->base.child());
-				nop->base = aOp;
+				if (nop->Flags == 0)
+				{
+					Ptr<ASTOpImageSwizzle> aOp = mu::Clone<ASTOpImageSwizzle>(this);
+					ReplaceAllSources(aOp, nop->base.child());
+					nop->base = aOp;
 
-				Ptr<ASTOpImageSwizzle> bOp = mu::Clone<ASTOpImageSwizzle>(this);
-				ReplaceAllSources(bOp, nop->blend.child());
-				nop->blend = bOp;
+					Ptr<ASTOpImageSwizzle> bOp = mu::Clone<ASTOpImageSwizzle>(this);
+					ReplaceAllSources(bOp, nop->blend.child());
+					nop->blend = bOp;
 
-				at = nop;
+					at = nop;
+				}
 				break;
 			}
 
@@ -348,6 +351,16 @@ namespace mu
 				break;
 			}
 
+			case OP_TYPE::IM_INVERT:
+			{
+				Ptr<ASTOpFixed> NewInvert = mu::Clone<ASTOpFixed>(channelSourceAt);
+				Ptr<ASTOpImageSwizzle> NewSwizzle = mu::Clone<ASTOpImageSwizzle>(this);
+				ReplaceAllSources(NewSwizzle, NewInvert->children[NewInvert->op.args.ImageInvert.base].child());
+				NewInvert->SetChild(NewInvert->op.args.ImageInvert.base, NewSwizzle);
+				at = NewInvert;
+				break;
+			}
+
 			case OP_TYPE::IM_RASTERMESH:
 			{
 				Ptr<ASTOpImageRasterMesh> NewRaster = mu::Clone<ASTOpImageRasterMesh>(channelSourceAt);
@@ -358,15 +371,15 @@ namespace mu
 				break;
 			}
 
-			case OP_TYPE::IM_RESIZE:
-			{
-				Ptr<ASTOpFixed> NewResize = mu::Clone<ASTOpFixed>(channelSourceAt);
-				Ptr<ASTOpImageSwizzle> NewSwizzle = mu::Clone<ASTOpImageSwizzle>(this);
-				ReplaceAllSources(NewSwizzle, NewResize->children[NewResize->op.args.ImageResize.source].child());
-				NewResize->SetChild(NewResize->op.args.ImageResize.source, NewSwizzle);
-				at = NewResize;
-				break;
-			}
+			//case OP_TYPE::IM_RESIZE:
+			//{
+			//	Ptr<ASTOpFixed> NewResize = mu::Clone<ASTOpFixed>(channelSourceAt);
+			//	Ptr<ASTOpImageSwizzle> NewSwizzle = mu::Clone<ASTOpImageSwizzle>(this);
+			//	ReplaceAllSources(NewSwizzle, NewResize->children[NewResize->op.args.ImageResize.source].child());
+			//	NewResize->SetChild(NewResize->op.args.ImageResize.source, NewSwizzle);
+			//	at = NewResize;
+			//	break;
+			//}
 
 			// This is not valid: binarize always forces format L8
 			//case OP_TYPE::IM_BINARISE:
@@ -442,6 +455,7 @@ namespace mu
 
 					Ptr<ASTOpImageMultiLayer> NewMultiLayer = mu::Clone<ASTOpImageMultiLayer>(ColorMultiLayer);
 					NewMultiLayer->blendTypeAlpha = AlphaMultiLayer->blendType;
+					NewMultiLayer->BlendAlphaSourceChannel = 3;
 					NewMultiLayer->base = NewBase;
 					NewMultiLayer->blend = NewBlended;
 
@@ -498,6 +512,7 @@ namespace mu
 
 					Ptr<ASTOpImageLayer> NewLayer = mu::Clone<ASTOpImageLayer>(ColorLayer);
 					NewLayer->blendTypeAlpha = AlphaLayer->blendType;
+					NewLayer->BlendAlphaSourceChannel = 3;
 					NewLayer->base = NewBase;
 					NewLayer->blend = NewBlended;
 
@@ -675,51 +690,51 @@ namespace mu
 			}
 
 			// Swizzle down compatible resizes.
-			if (!at && sourceType == OP_TYPE::IM_RESIZE)
-			{
-				const ASTOpFixed* FirstResize = dynamic_cast<const ASTOpFixed*>(Sources[0].child().get());
-				check(FirstResize);
-
-				bool bAreAllResizesCompatible = true;
-				for (int32 c = 1; c < MUTABLE_OP_MAX_SWIZZLE_CHANNELS; ++c)
-				{
-					if (Sources[c])
-					{
-						const ASTOpFixed* Typed = dynamic_cast<const ASTOpFixed*>(Sources[c].child().get());
-						check(Typed);
-						// Compare all args but the source image
-						OP::ImageResizeArgs ArgCopy = FirstResize->op.args.ImageResize;
-						ArgCopy.source = Typed->op.args.ImageResize.source;
-
-						if (FMemory::Memcmp(&ArgCopy, &Typed->op.args.ImageResize, sizeof(OP::ImageResizeArgs)) != 0)
-						{
-							bAreAllResizesCompatible = false;
-							break;
-						}
-					}
-				}
-
-				if (bAreAllResizesCompatible)
-				{
-					// Move the swizzle down all the paths
-					Ptr<ASTOpFixed> NewResize = mu::Clone<ASTOpFixed>(FirstResize);
-
-					Ptr<ASTOpImageSwizzle> NewSwizzle = mu::Clone<ASTOpImageSwizzle>(this);
-					for (int c = 0; c < MUTABLE_OP_MAX_SWIZZLE_CHANNELS; ++c)
-					{
-						const ASTOpFixed* ChannelResize = dynamic_cast<const ASTOpFixed*>(Sources[c].child().get());
-						if (ChannelResize)
-						{
-							NewSwizzle->Sources[c] = ChannelResize->children[ChannelResize->op.args.ImageResize.source].child();
-						}
-					}
-
-					NewResize->SetChild(NewResize->op.args.ImageResize.source, NewSwizzle);
-
-					at = NewResize;
-				}
+			//if (!at && sourceType == OP_TYPE::IM_RESIZE)
+			//{
+			//	const ASTOpFixed* FirstResize = dynamic_cast<const ASTOpFixed*>(Sources[0].child().get());
+			//	check(FirstResize);
+
+			//	bool bAreAllResizesCompatible = true;
+			//	for (int32 c = 1; c < MUTABLE_OP_MAX_SWIZZLE_CHANNELS; ++c)
+			//	{
+			//		if (Sources[c])
+			//		{
+			//			const ASTOpFixed* Typed = dynamic_cast<const ASTOpFixed*>(Sources[c].child().get());
+			//			check(Typed);
+			//			// Compare all args but the source image
+			//			OP::ImageResizeArgs ArgCopy = FirstResize->op.args.ImageResize;
+			//			ArgCopy.source = Typed->op.args.ImageResize.source;
+
+			//			if (FMemory::Memcmp(&ArgCopy, &Typed->op.args.ImageResize, sizeof(OP::ImageResizeArgs)) != 0)
+			//			{
+			//				bAreAllResizesCompatible = false;
+			//				break;
+			//			}
+			//		}
+			//	}
+
+			//	if (bAreAllResizesCompatible)
+			//	{
+			//		// Move the swizzle down all the paths
+			//		Ptr<ASTOpFixed> NewResize = mu::Clone<ASTOpFixed>(FirstResize);
+
+			//		Ptr<ASTOpImageSwizzle> NewSwizzle = mu::Clone<ASTOpImageSwizzle>(this);
+			//		for (int c = 0; c < MUTABLE_OP_MAX_SWIZZLE_CHANNELS; ++c)
+			//		{
+			//			const ASTOpFixed* ChannelResize = dynamic_cast<const ASTOpFixed*>(Sources[c].child().get());
+			//			if (ChannelResize)
+			//			{
+			//				NewSwizzle->Sources[c] = ChannelResize->children[ChannelResize->op.args.ImageResize.source].child();
+			//			}
+			//		}
+
+			//		NewResize->SetChild(NewResize->op.args.ImageResize.source, NewSwizzle);
+
+			//		at = NewResize;
+			//	}
 
-			}
+			//}
 
 			// Swizzle down compatible pixelformats.
 			if (!at && sourceType == OP_TYPE::IM_PIXELFORMAT && bSameChannelOrder)
@@ -775,27 +790,54 @@ namespace mu
 
 		}
 
+		// TODO \warning: probably wrong because it doesn't check if the layer colour is doing a separated alpha operation.
 		// Swizzle of RGB from a source + A from a layer colour
 		// This can be optimized to apply the layer colour on-base directly to the alpha channel to skip the swizzle
-		if ( !at 
-			&&
-			Sources[0] && Sources[0]==Sources[1] && Sources[0]==Sources[2]
-			&&
-			Sources[3] && Sources[3]->GetOpType()==OP_TYPE::IM_LAYERCOLOUR )
-		{
-			// Move the swizzle down all the paths
-			Ptr<ASTOpImageLayerColor> NewLayerColour = mu::Clone<ASTOpImageLayerColor>(Sources[3].child());
+		//if ( !at 
+		//	&&
+		//	Sources[0] && Sources[0]==Sources[1] && Sources[0]==Sources[2]
+		//	&&
+		//	Sources[3] && Sources[3]->GetOpType()==OP_TYPE::IM_LAYERCOLOUR
+		//	)
+		//{
+		//	// Move the swizzle down all the paths
+		//	Ptr<ASTOpImageLayerColor> NewLayerColour = mu::Clone<ASTOpImageLayerColor>(Sources[3].child());
 
-			Ptr<ASTOpImageSwizzle> NewSwizzle = mu::Clone<ASTOpImageSwizzle>(this);
-			NewSwizzle->Sources[3] = NewLayerColour->base.child();
+		//	Ptr<ASTOpImageSwizzle> NewSwizzle = mu::Clone<ASTOpImageSwizzle>(this);
+		//	NewSwizzle->Sources[3] = NewLayerColour->base.child();
 
-			NewLayerColour->blendTypeAlpha = NewLayerColour->blendType;
-			NewLayerColour->blendType = EBlendType::BT_NONE;
-			NewLayerColour->base = NewSwizzle;
+		//	NewLayerColour->blendTypeAlpha = NewLayerColour->blendType;
+		//	NewLayerColour->BlendAlphaSourceChannel = SourceChannels[3];
+		//	NewLayerColour->blendType = EBlendType::BT_NONE;
+		//	NewLayerColour->base = NewSwizzle;
 
-			at = NewLayerColour;
+		//	at = NewLayerColour;
 
-		}
+		//}
+
+		// Swizzle of RGB from a source + A from a layer
+		// This can be optimized to apply the layer on-base directly to the alpha channel to skip the swizzle
+		// \TODO: wrong: the new layer colour will always use the alpha from the colour, instead of the channel that the swizzle is selecting.
+		// \TODO: wrong: it ignores the possibility of separate alpha operation
+		//if (!at
+		//	&&
+		//	Sources[0] && Sources[0] == Sources[1] && Sources[0] == Sources[2]
+		//	&&
+		//	Sources[3] && Sources[3]->GetOpType() == OP_TYPE::IM_LAYER)
+		//{
+		//	// Move the swizzle down all the paths
+		//	Ptr<ASTOpImageLayer> NewLayer = mu::Clone<ASTOpImageLayer>(Sources[3].child());
+
+		//	Ptr<ASTOpImageSwizzle> NewSwizzle = mu::Clone<ASTOpImageSwizzle>(this);
+		//	NewSwizzle->Sources[3] = NewLayer->base.child();
+
+		//	NewLayer->blendTypeAlpha = NewLayer->blendType;
+		//	NewLayer->BlendAlphaSourceChannel = SourceChannels[3];
+		//	NewLayer->blendType = EBlendType::BT_NONE;
+		//	NewLayer->base = NewSwizzle;
+
+		//	at = NewLayer;
+		//}
 
 		// Swizzle of RGB from a layer colour + A from a different source
 		// This can be optimized to apply the layer colour on-base directly to the rgb channel to skip the swizzle
@@ -807,7 +849,9 @@ namespace mu
 			&& 
 			(!Sources[2] || Sources[0] == Sources[2])
 			&&
-			Sources[0]->GetOpType() == OP_TYPE::IM_LAYERCOLOUR)
+			Sources[0]->GetOpType() == OP_TYPE::IM_LAYERCOLOUR
+			&& 
+			!(Sources[3]==Sources[0]) )
 		{
 			// Move the swizzle down all the rgb path
 			Ptr<ASTOpImageLayerColor> NewLayerColour = mu::Clone<ASTOpImageLayerColor>(Sources[0].child());
@@ -824,28 +868,119 @@ namespace mu
 			at = NewLayerColour;
 		}
 
-		// Swizzle of RGB from a binarise + A from a different source
-		// This can be optimized to apply the layer colour on-base directly to the rgb channel to skip the swizzle
-		// This is not valid because binarise always forces L8 format.
-		//if (!at
-		//	&&
-		//	Sources[0] && (!Sources[0] ||Sources[0] == Sources[1]) && (!Sources[2] || Sources[0] == Sources[2])
-		//	&&
-		//	Sources[0]->GetOpType() == OP_TYPE::IM_BINARISE)
-		//{
-		//	// Move the swizzle down all the rgb path
-		//	Ptr<ASTOpFixed> NewBinarise = mu::Clone<ASTOpFixed>(Sources[0].child());
-		//	check(NewBinarise);
+		// Swizzle getting an A from a saturate
+		// The saturate doesn't affect A channel so it can be removed.
+		if (!at)
+		{
+			Ptr<ASTOpImageSwizzle> NewSwizzle;
 
-		//	Ptr<ASTOpImageSwizzle> NewSwizzle = mu::Clone<ASTOpImageSwizzle>(this);
-		//	NewSwizzle->Sources[0] = NewBinarise->children[NewBinarise->op.args.ImageBinarise.base].child();
-		//	NewSwizzle->Sources[1] = Sources[1] ? NewSwizzle->Sources[0].child() : nullptr;
-		//	NewSwizzle->Sources[2] = Sources[2] ? NewSwizzle->Sources[0].child() : nullptr;
+			for (int32 Channel = 0; Channel < MUTABLE_OP_MAX_SWIZZLE_CHANNELS; ++Channel)
+			{
+				if (Sources[Channel] && SourceChannels[Channel]==3 && Sources[Channel]->GetOpType() == OP_TYPE::IM_SATURATE)
+				{
+					// Remove the saturate for this channel
+					if (!NewSwizzle)
+					{
+						NewSwizzle = mu::Clone<ASTOpImageSwizzle>(this);
+					}
 
-		//	NewBinarise->SetChild(NewBinarise->op.args.ImageBinarise.base, NewSwizzle);
+					const ASTOpFixed* OldSaturate = dynamic_cast<const ASTOpFixed*>(Sources[Channel].child().get());
+					Ptr<ASTOp> OldSaturateBase = OldSaturate->children[OldSaturate->op.args.ImageSaturate.base].child();
 
-		//	at = NewBinarise;
-		//}
+					NewSwizzle->Sources[Channel] = OldSaturateBase;
+					at = NewSwizzle;
+				}
+			}
+		}
+
+		// Swizzle of RGB from a saturate + A from a different source
+		// This can be optimized to apply the saturate after the swizzle, since it doesn't touch A
+		if (!at
+			&&
+			Sources[0]->GetOpType() == OP_TYPE::IM_SATURATE 
+			&&
+			Sources[0]
+			&&
+			(Sources[0] == Sources[1]) && (Sources[0] == Sources[2])
+			&&
+			// Actually it would be enough with all the RGB channels to be present in any order
+			SourceChannels[0]==0 && SourceChannels[1] == 1 && SourceChannels[2] == 2
+			)
+		{
+			// Move the swizzle down 
+			Ptr<ASTOpFixed> NewSaturate = mu::Clone<ASTOpFixed>(Sources[0].child());
+			check(NewSaturate);
+
+			Ptr<ASTOpImageSwizzle> NewSwizzle = mu::Clone<ASTOpImageSwizzle>(this);
+			Ptr<ASTOp> OldSaturateBase = NewSaturate->children[NewSaturate->op.args.ImageSaturate.base].child();
+			NewSwizzle->Sources[0] = OldSaturateBase;
+			NewSwizzle->Sources[1] = OldSaturateBase;
+			NewSwizzle->Sources[2] = OldSaturateBase;
+
+			// Remove the saturate from the alpha if it is there.
+			if (Sources[3] == Sources[0] && SourceChannels[3] == 3)
+			{
+				NewSwizzle->Sources[3] = OldSaturateBase;
+			}
+
+			NewSaturate->SetChild(NewSaturate->op.args.ImageSaturate.base, NewSwizzle);
+
+			at = NewSaturate;
+		}
+
+
+		// Swizzle with the same op as identity in RGB, a Layer op in the A that has one of the operands matching 
+		// the one in the swizzle RGB, but using its A. 
+		// The Layer operation can be flagged as alpha only and moved up the swizzle, then the swizzle is identity 
+		// and can be removed, so remove it here anyway.
+		// This is another very specific optimization case that happens with certain combination of operations.
+		// from:
+		//- SWIZZLE
+		//	r -> r from A
+		//	g -> g from A
+		//	b -> b from A
+		//	a -> (r or a) from LAYER
+		//		- 3 from A (on alpha only using flags)
+		//		- B
+		// to:
+		//- LAYER (on alpha only)
+		//	- A
+		//	- B
+		// In addition, if the blend operation done by LAYER is commutative, see if X is 3 from I instead.
+		if (!at
+			&&
+			(Sources[0] == Sources[1]) && (Sources[0] == Sources[2])
+			&&
+			SourceChannels[0] == 0 && SourceChannels[1] == 1 && SourceChannels[2] == 2
+			&&
+			Sources[3] && Sources[3]->GetOpType() == OP_TYPE::IM_LAYER
+			)
+		{
+			const ASTOpImageLayer* OldLayer = dynamic_cast<const ASTOpImageLayer*>(Sources[3].child().get());
+
+			// For now just check the case that we are observing in the working data: 
+			// A is in the blended of a multiply, and we take its alpha channel
+			// \TODO: Implement the other cases when we find instances of them.
+			if ( OldLayer->Flags==OP::ImageLayerArgs::FLAGS::F_BLENDED_RGB_FROM_ALPHA
+				&&
+				OldLayer->blend==Sources[0]
+				&& 
+				OldLayer->blendType==EBlendType::BT_MULTIPLY 
+				&&
+				OldLayer->blendTypeAlpha == EBlendType::BT_NONE 
+				&&
+				SourceChannels[3]==0)
+			{
+				Ptr<ASTOpImageLayer> NewLayer = mu::Clone<ASTOpImageLayer>(OldLayer);
+				NewLayer->blend = NewLayer->base.child();
+				NewLayer->base = Sources[0].child();
+				NewLayer->blendTypeAlpha = NewLayer->blendType;
+				NewLayer->blendType = EBlendType::BT_NONE;
+				NewLayer->Flags = OP::ImageLayerArgs::FLAGS::F_BASE_RGB_FROM_ALPHA;
+
+				at = NewLayer;
+			}
+		}
 
 
 		return at;
diff --git a/Engine/Plugins/Experimental/Mutable/Source/MutableTools/Private/MuT/ASTOpImageSwizzle.h b/Engine/Plugins/Experimental/Mutable/Source/MutableTools/Private/MuT/ASTOpImageSwizzle.h
--- a/Engine/Plugins/Experimental/Mutable/Source/MutableTools/Private/MuT/ASTOpImageSwizzle.h
+++ b/Engine/Plugins/Experimental/Mutable/Source/MutableTools/Private/MuT/ASTOpImageSwizzle.h
@@ -35,7 +35,7 @@ namespace mu
 		Ptr<ASTOp> Clone(MapChildFuncRef mapChild) const override;
 		void ForEachChild(const TFunctionRef<void(ASTChild&)>) override;
 		void Link(FProgram& program, const FLinkerOptions* Options) override;
-		Ptr<ASTOp> OptimiseSemantic(const FModelOptimizationOptions&) const override;
+		Ptr<ASTOp> OptimiseSemantic(const FModelOptimizationOptions&, int32 Pass) const override;
 		Ptr<ASTOp> OptimiseSink(const FModelOptimizationOptions& options, FOptimizeSinkContext& context) const override;
 		FImageDesc GetImageDesc(bool returnBestOption, FGetImageDescContext* context) const override;
 		void GetLayoutBlockSize(int* pBlockX, int* pBlockY) override;
diff --git a/Engine/Plugins/Experimental/Mutable/Source/MutableTools/Private/MuT/ASTOpSwitch.cpp b/Engine/Plugins/Experimental/Mutable/Source/MutableTools/Private/MuT/ASTOpSwitch.cpp
--- a/Engine/Plugins/Experimental/Mutable/Source/MutableTools/Private/MuT/ASTOpSwitch.cpp
+++ b/Engine/Plugins/Experimental/Mutable/Source/MutableTools/Private/MuT/ASTOpSwitch.cpp
@@ -411,7 +411,7 @@ namespace mu
 
 
 	//-------------------------------------------------------------------------------------------------
-	mu::Ptr<ASTOp> ASTOpSwitch::OptimiseSemantic(const FModelOptimizationOptions&) const
+	mu::Ptr<ASTOp> ASTOpSwitch::OptimiseSemantic(const FModelOptimizationOptions&, int32 Pass) const
 	{
 		// Constant condition?
 		if (variable->GetOpType() == OP_TYPE::NU_CONSTANT)
diff --git a/Engine/Plugins/Experimental/Mutable/Source/MutableTools/Private/MuT/ASTOpSwitch.h b/Engine/Plugins/Experimental/Mutable/Source/MutableTools/Private/MuT/ASTOpSwitch.h
--- a/Engine/Plugins/Experimental/Mutable/Source/MutableTools/Private/MuT/ASTOpSwitch.h
+++ b/Engine/Plugins/Experimental/Mutable/Source/MutableTools/Private/MuT/ASTOpSwitch.h
@@ -69,7 +69,7 @@ template <class SCALAR> class vec4;
 		FImageDesc GetImageDesc(bool returnBestOption, class FGetImageDescContext* context) const  override;
 		void GetBlockLayoutSize(int blockIndex, int* pBlockX, int* pBlockY, FBlockLayoutSizeCache* cache) override;
 		void GetLayoutBlockSize(int* pBlockX, int* pBlockY) override;
-		Ptr<ASTOp> OptimiseSemantic(const FModelOptimizationOptions&) const override;
+		Ptr<ASTOp> OptimiseSemantic(const FModelOptimizationOptions&, int32 Pass) const override;
 		bool GetNonBlackRect(FImageRect& maskUsage) const override;
 		bool IsImagePlainConstant(FVector4f& colour) const override;
 		Ptr<ImageSizeExpression> GetImageSizeExpression() const override;
diff --git a/Engine/Plugins/Experimental/Mutable/Source/MutableTools/Private/MuT/CodeOptimiser.cpp b/Engine/Plugins/Experimental/Mutable/Source/MutableTools/Private/MuT/CodeOptimiser.cpp
--- a/Engine/Plugins/Experimental/Mutable/Source/MutableTools/Private/MuT/CodeOptimiser.cpp
+++ b/Engine/Plugins/Experimental/Mutable/Source/MutableTools/Private/MuT/CodeOptimiser.cpp
@@ -792,7 +792,7 @@ namespace mu
 
 
 	//-------------------------------------------------------------------------------------------------
-	void CodeOptimiser::FullOptimiseAST( ASTOpList& roots )
+	void CodeOptimiser::FullOptimiseAST( ASTOpList& roots, int32 Pass )
 	{
 		bool modified = true;
 		int numIterations = 0;
@@ -808,7 +808,7 @@ namespace mu
 			// All kind of optimisations that depend on the meaning of each operation
 			// \TODO: We are doing it for all states.
 			UE_LOG(LogMutableCore, Verbose, TEXT(" - semantic optimiser"));
-			modified |= SemanticOptimiserAST( roots, m_options->GetPrivate()->m_optimisationOptions );
+			modified |= SemanticOptimiserAST( roots, m_options->GetPrivate()->m_optimisationOptions, Pass );
 			//AXE_INT_VALUE("Mutable", Verbose, "ast size", (int64_t)ASTOp::CountNodes(roots));
 			UE_LOG(LogMutableCore, Verbose, TEXT("(int) %s : %ld"), TEXT("ast size"), int64(ASTOp::CountNodes(roots)));
 			ASTOp::LogHistogram(roots);
@@ -1226,7 +1226,11 @@ namespace mu
 			// Main optimisation stage
 			{
 				MUTABLE_CPUPROFILER_SCOPE(MainStage);
-				FullOptimiseAST( roots );
+				FullOptimiseAST( roots, 0 );
+				UE_LOG(LogMutableCore, Verbose, TEXT("(int) %s : %ld"), TEXT("ast size"), int64(ASTOp::CountNodes(roots)));
+				ASTOp::LogHistogram(roots);
+
+				FullOptimiseAST(roots, 1);
 				UE_LOG(LogMutableCore, Verbose, TEXT("(int) %s : %ld"), TEXT("ast size"), int64(ASTOp::CountNodes(roots)));
 				ASTOp::LogHistogram(roots);
 			}
diff --git a/Engine/Plugins/Experimental/Mutable/Source/MutableTools/Private/MuT/CodeOptimiser.h b/Engine/Plugins/Experimental/Mutable/Source/MutableTools/Private/MuT/CodeOptimiser.h
--- a/Engine/Plugins/Experimental/Mutable/Source/MutableTools/Private/MuT/CodeOptimiser.h
+++ b/Engine/Plugins/Experimental/Mutable/Source/MutableTools/Private/MuT/CodeOptimiser.h
@@ -38,7 +38,7 @@ namespace mu
         int m_optimizeIterationsLeft=0;
 
         //! Full optimisation pass
-        void FullOptimiseAST( ASTOpList& roots );
+        void FullOptimiseAST( ASTOpList& roots, int32 Pass);
 
         //! Optimise the code of a model for a specific state, generating new instructions and
         //! state information.
@@ -118,8 +118,7 @@ namespace mu
     //---------------------------------------------------------------------------------------------
     //! All kinds of optimisations that depend on the meaning of each operation
     //---------------------------------------------------------------------------------------------
-    extern bool SemanticOptimiserAST( ASTOpList& roots,
-                                      const FModelOptimizationOptions& optimisationOptions );
+    extern bool SemanticOptimiserAST( ASTOpList& roots, const FModelOptimizationOptions& optimisationOptions, int32 Pass);
 
     //---------------------------------------------------------------------------------------------
     //! Semantic operator that reorders instructions moving expensive ones down to the
diff --git a/Engine/Plugins/Experimental/Mutable/Source/MutableTools/Private/MuT/CodeOptimiser_Parameters.cpp b/Engine/Plugins/Experimental/Mutable/Source/MutableTools/Private/MuT/CodeOptimiser_Parameters.cpp
--- a/Engine/Plugins/Experimental/Mutable/Source/MutableTools/Private/MuT/CodeOptimiser_Parameters.cpp
+++ b/Engine/Plugins/Experimental/Mutable/Source/MutableTools/Private/MuT/CodeOptimiser_Parameters.cpp
@@ -2039,7 +2039,7 @@ namespace mu
 
                     // All kind of optimisations that depend on the meaning of each operation
                     UE_LOG(LogMutableCore, Verbose, TEXT(" - semantic optimiser"));
-                    modified |= SemanticOptimiserAST( roots, m_options->GetPrivate()->m_optimisationOptions );
+                    modified |= SemanticOptimiserAST( roots, m_options->GetPrivate()->m_optimisationOptions, 1 );
 					UE_LOG(LogMutableCore, Verbose, TEXT("(int) %s : %ld"), TEXT("ast size"), int64(ASTOp::CountNodes(roots)));
 					//ASTOp::LogHistogram(roots);
 
@@ -2142,7 +2142,7 @@ namespace mu
 
                 UE_LOG(LogMutableCore, Verbose, TEXT(" - semantic optimiser"));
                 modified |=
-                    SemanticOptimiserAST( roots, m_options->GetPrivate()->m_optimisationOptions );
+                    SemanticOptimiserAST( roots, m_options->GetPrivate()->m_optimisationOptions, 1 );
 				UE_LOG(LogMutableCore, Verbose, TEXT("(int) %s : %ld"), TEXT("ast size"), int64(ASTOp::CountNodes(roots)));
 
                 // Image size operations are treated separately
diff --git a/Engine/Plugins/Experimental/Mutable/Source/MutableTools/Private/MuT/CodeOptimiser_Semantic.cpp b/Engine/Plugins/Experimental/Mutable/Source/MutableTools/Private/MuT/CodeOptimiser_Semantic.cpp
--- a/Engine/Plugins/Experimental/Mutable/Source/MutableTools/Private/MuT/CodeOptimiser_Semantic.cpp
+++ b/Engine/Plugins/Experimental/Mutable/Source/MutableTools/Private/MuT/CodeOptimiser_Semantic.cpp
@@ -69,9 +69,10 @@ namespace mu
     //---------------------------------------------------------------------------------------------
     //---------------------------------------------------------------------------------------------
     bool SemanticOptimiserAST(
-            ASTOpList& roots,
-            const FModelOptimizationOptions& optimisationOptions
-            )
+		ASTOpList& roots,
+		const FModelOptimizationOptions& optimisationOptions,
+		int32 Pass
+	)
     {
         MUTABLE_CPUPROFILER_SCOPE(SemanticOptimiserAST);
 
@@ -80,7 +81,7 @@ namespace mu
         // TODO: isn't top down better suited?
         ASTOp::Traverse_BottomUp_Unique( roots, [&](Ptr<ASTOp>& n)
         {
-            auto o = n->OptimiseSemantic(optimisationOptions);
+            auto o = n->OptimiseSemantic(optimisationOptions, Pass);
 
             // If the returned value is null it means no change.
             if (o && o!=n)
@@ -95,7 +96,7 @@ namespace mu
 
 
     //---------------------------------------------------------------------------------------------
-    Ptr<ASTOp> ASTOpFixed::OptimiseSemantic( const FModelOptimizationOptions& options ) const
+    Ptr<ASTOp> ASTOpFixed::OptimiseSemantic( const FModelOptimizationOptions& options, int32 Pass) const
     {
         Ptr<ASTOp> at;
 
@@ -1318,6 +1319,23 @@ namespace mu
                 break;
             }
 
+			case OP_TYPE::IM_SWIZZLE:
+			{
+				Ptr<ASTOpImageSwizzle> newOp = mu::Clone<ASTOpImageSwizzle>(sourceAt);
+				for (int s = 0; s < MUTABLE_OP_MAX_SWIZZLE_CHANNELS; ++s)
+				{
+					Ptr<ASTOp> OldChannelOp = newOp->Sources[s].child();
+					if (OldChannelOp)
+					{
+						Ptr<ASTOpFixed> ChannelResize = mu::Clone<ASTOpFixed>(this);
+						ChannelResize->SetChild(ChannelResize->op.args.ImageResize.source, OldChannelOp);
+						newOp->Sources[s] = ChannelResize;
+					}
+				}
+				at = newOp;
+				break;
+			}
+
             case OP_TYPE::IM_COMPOSE:
             {
                 Ptr<ASTOpImageCompose> newOp = mu::Clone<ASTOpImageCompose>(sourceAt);
diff --git a/Engine/Plugins/Experimental/Mutable/Source/MutableTools/Private/MuT/DataPacker.cpp b/Engine/Plugins/Experimental/Mutable/Source/MutableTools/Private/MuT/DataPacker.cpp
--- a/Engine/Plugins/Experimental/Mutable/Source/MutableTools/Private/MuT/DataPacker.cpp
+++ b/Engine/Plugins/Experimental/Mutable/Source/MutableTools/Private/MuT/DataPacker.cpp
@@ -28,10 +28,6 @@
 #include "MuT/ASTOpMeshMorph.h"
 #include "MuT/ASTOpLayoutFromMesh.h"
 
-#include <cstdint>
-#include <memory>
-#include <utility>
-
 
 namespace mu
 {
@@ -49,7 +45,7 @@ namespace mu
         m_visited.SetNum( program.m_opAddress.Num() );
         if ( m_visited.Num() )
         {
-            memset( &m_visited[0], 0, m_visited.Num() );
+            FMemory::Memset( &m_visited[0], 0, m_visited.Num() );
         }
     }
 
@@ -95,7 +91,7 @@ namespace mu
                     OP_TYPE thisOpType = m_program.GetOpType(item.Value);
                     if ( m_opType == thisOpType )
                     {
-                        auto args = m_program.GetOpArgs<OP::ResourceConstantArgs>(item.Value);
+						OP::ResourceConstantArgs args = m_program.GetOpArgs<OP::ResourceConstantArgs>(item.Value);
                         if ( m_constant == args.value )
                         {
                             found = true;
@@ -179,7 +175,7 @@ namespace mu
                 case OP_TYPE::LA_CONDITIONAL:
                 case OP_TYPE::IN_CONDITIONAL:
                 {
-                    auto args = program.GetOpArgs<OP::ConditionalArgs>(at);
+					OP::ConditionalArgs args = program.GetOpArgs<OP::ConditionalArgs>(at);
 
                     // If the constant is present in only one of the 2 branches
                     bool foundYes = m_constSearch.Run( args.yes  );
@@ -235,7 +231,7 @@ namespace mu
                          &&
                          program.GetOpType(at) == m_opType )
                     {
-                        auto args = program.GetOpArgs<OP::MeshConstantArgs>(at);
+						OP::MeshConstantArgs args = program.GetOpArgs<OP::MeshConstantArgs>(at);
                         if ( args.value == m_constant )
                         {
                             // Accumulate the currently relevant parameters
@@ -248,7 +244,7 @@ namespace mu
                     }
                     else if ( program.GetOpType(at) == m_opType )
                     {
-                        auto args = program.GetOpArgs<OP::ResourceConstantArgs>(at);
+						OP::ResourceConstantArgs args = program.GetOpArgs<OP::ResourceConstantArgs>(at);
                         if ( args.value == m_constant )
                         {
                             // Accumulate the currently relevant parameters
@@ -336,7 +332,7 @@ namespace mu
             case OP_TYPE::IM_CONSTANT:
             {
                 // Remove unsupported formats
-                auto op = dynamic_cast<const ASTOpConstantResource*>(node.get());
+				const ASTOpConstantResource* op = dynamic_cast<const ASTOpConstantResource*>(node.get());
                 if (!m_supportedFormats.count(op))
                 {
 					TArray<bool> initial;
@@ -571,7 +567,7 @@ namespace mu
                 uint64_t newState = currentSemantics;
                 newState |= (UINT64_C(1)<<MBS_VERTEXINDEX);
                 RecurseWithState( op->source.child(), newState );
-                for( auto& r: op->removes )
+                for( const TPair<ASTChild, ASTChild>& r: op->removes )
                 {
                     RecurseWithState( r.Value.child(), newState );
                 }
