[Mutable] Fix Static Analysis warning.

#rnx

diff --git a/Engine/Plugins/Experimental/Mutable/Source/CustomizableObject/Private/MuCO/CustomizableObjectSystem.cpp b/Engine/Plugins/Experimental/Mutable/Source/CustomizableObject/Private/MuCO/CustomizableObjectSystem.cpp
--- a/Engine/Plugins/Experimental/Mutable/Source/CustomizableObject/Private/MuCO/CustomizableObjectSystem.cpp
+++ b/Engine/Plugins/Experimental/Mutable/Source/CustomizableObject/Private/MuCO/CustomizableObjectSystem.cpp
@@ -698,10 +698,11 @@ static FAutoConsoleVariableRef CVarEnableMutableReusePreviousUpdateData(
 	ECVF_Default);
 
 
-bool FCustomizableObjectSystemPrivate::bEnableOnlyGenerateRequestedLODs = true;
+int32 FCustomizableObjectSystemPrivate::EnableOnlyGenerateRequestedLODs = 1;
+
 static FAutoConsoleVariableRef CVarEnableOnlyGenerateRequestedLODs(
-	TEXT("mutable.EnableOnlyGenerateRequestedLODs"), FCustomizableObjectSystemPrivate::bEnableOnlyGenerateRequestedLODs,
-	TEXT("If true, Only the RequestedLODLevels will be generated. If false, all LODs will be build."),
+	TEXT("mutable.EnableOnlyGenerateRequestedLODs"), FCustomizableObjectSystemPrivate::EnableOnlyGenerateRequestedLODs,
+	TEXT("If 1 or greater, Only the RequestedLODLevels will be generated. If 0, all LODs will be build."),
 	ECVF_Default);
 
 
@@ -2094,7 +2095,7 @@ namespace impl
 		UCustomizableInstancePrivateData* CandidateInstancePrivateData = CandidateInstance->GetPrivate();
 		check(CandidateInstancePrivateData != nullptr);
 
-		if (CandidateInstancePrivateData->HasCOInstanceFlags(PendingLODsUpdate))
+		if (CandidateInstancePrivateData && CandidateInstancePrivateData->HasCOInstanceFlags(PendingLODsUpdate))
 		{
 			CandidateInstancePrivateData->ClearCOInstanceFlags(PendingLODsUpdate);
 			// TODO: Is anything needed for this now?
@@ -2121,6 +2122,7 @@ namespace impl
 		// Only update resources if the instance is in range (it could have got far from the player since the task was queued)
 		check(System->CurrentInstanceLODManagement != nullptr);
 		if (System->CurrentInstanceLODManagement->IsOnlyUpdateCloseCustomizableObjectsEnabled()
+			&& CandidateInstancePrivateData
 			&& CandidateInstancePrivateData->LastMinSquareDistFromComponentToPlayer > FMath::Square(System->CurrentInstanceLODManagement->GetOnlyUpdateCloseCustomizableObjectsDist())
 			&& CandidateInstancePrivateData->LastMinSquareDistFromComponentToPlayer != FLT_MAX // This means it is the first frame so it has to be updated
 		   )
@@ -2200,9 +2202,6 @@ namespace impl
 			}
 		}
 
-		FCustomizableObjectSystemPrivate* CustomizableObjectSystemPrivateData = System->GetPrivate();
-		check(CustomizableObjectSystemPrivateData != nullptr);
-
 		if (!bLiveUpdateMode && CandidateInstancePrivateData->LiveUpdateModeInstanceID != 0)
 		{
 			// The instance was in live update mode last update, but now it's not. So the Id and resources have to be released.
@@ -2229,7 +2228,7 @@ namespace impl
 		CurrentOperationData->MutableParameters = Parameters;
 		CurrentOperationData->State = CandidateInstance->GetState();
 
-		if (SystemPrivateData->bEnableOnlyGenerateRequestedLODs && System->CurrentInstanceLODManagement->IsOnlyGenerateRequestedLODLevelsEnabled() && !CandidateInstancePrivateData->HasCOInstanceFlags(ForceGenerateAllLODs))
+		if (System->IsOnlyGenerateRequestedLODsEnabled() && System->CurrentInstanceLODManagement->IsOnlyGenerateRequestedLODLevelsEnabled() && !CandidateInstancePrivateData->HasCOInstanceFlags(ForceGenerateAllLODs))
 		{
 			CurrentOperationData->RequestedLODs = Operation->InstanceDescriptorRuntimeHash.GetRequestedLODs();
 		}
@@ -2714,7 +2713,7 @@ bool UCustomizableObjectSystem::IsProgressiveMipStreamingEnabled() const
 bool UCustomizableObjectSystem::IsOnlyGenerateRequestedLODsEnabled() const
 {
 	check(Private != nullptr);
-	return Private->bEnableOnlyGenerateRequestedLODs;
+	return Private->EnableOnlyGenerateRequestedLODs != 0;
 }
 
 
diff --git a/Engine/Plugins/Experimental/Mutable/Source/CustomizableObject/Private/MuCO/CustomizableObjectSystemPrivate.h b/Engine/Plugins/Experimental/Mutable/Source/CustomizableObject/Private/MuCO/CustomizableObjectSystemPrivate.h
--- a/Engine/Plugins/Experimental/Mutable/Source/CustomizableObject/Private/MuCO/CustomizableObjectSystemPrivate.h
+++ b/Engine/Plugins/Experimental/Mutable/Source/CustomizableObject/Private/MuCO/CustomizableObjectSystemPrivate.h
@@ -552,8 +552,8 @@ public:
 	static int32 EnableMutableLiveUpdate;
 	static int32 EnableReuseInstanceTextures;
 	static int32 EnableMutableAnimInfoDebugging;
+	static int32 EnableOnlyGenerateRequestedLODs;
 	static bool bEnableMutableReusePreviousUpdateData;
-	static bool bEnableOnlyGenerateRequestedLODs;
 
 	/** */
 	inline void AddGameThreadTask(const FMutableTask& Task)
