[mutable] Another case of specific optimization to remove unnecessary swizzles.
#rnx

diff --git a/Engine/Plugins/Experimental/Mutable/Source/MutableTools/Private/MuT/ASTOpImageSwizzle.cpp b/Engine/Plugins/Experimental/Mutable/Source/MutableTools/Private/MuT/ASTOpImageSwizzle.cpp
--- a/Engine/Plugins/Experimental/Mutable/Source/MutableTools/Private/MuT/ASTOpImageSwizzle.cpp
+++ b/Engine/Plugins/Experimental/Mutable/Source/MutableTools/Private/MuT/ASTOpImageSwizzle.cpp
@@ -983,6 +983,56 @@ namespace mu
 		}
 
 
+		// If we have an alpha channel that has as children something that expands a single channel texture
+		// skip the expansion,. since we know we just want one channel.
+		// Very specific, based on observed code patterns.
+		// \TODO: Make more general.
+		if (!at
+			&&
+			Sources[3] && Sources[3]->GetOpType() == OP_TYPE::IM_LAYER
+			)
+		{
+			ASTOpImageLayer* OldLayer = dynamic_cast<ASTOpImageLayer*>(Sources[3].child().get());
+
+			// For now just check the case that we are observing in the working data: 
+			if (OldLayer->Flags == 0
+				&&
+				SourceChannels[3] == 0
+				&&
+				OldLayer->blend->GetOpType()==OP_TYPE::IM_PIXELFORMAT )
+			{
+				const ASTOpImagePixelFormat* OldFormat = dynamic_cast<const ASTOpImagePixelFormat*>(OldLayer->blend.child().get());
+				if (OldFormat->Source->GetOpType() == OP_TYPE::IM_SWIZZLE
+					&&
+					OldFormat->Format == EImageFormat::IF_RGB_UBYTE)
+				{
+					const ASTOpImageSwizzle* OldChildSwizzle = dynamic_cast<const ASTOpImageSwizzle*>(OldFormat->Source.child().get());
+					if (OldChildSwizzle->Format == EImageFormat::IF_L_UBYTE)
+					{
+						Ptr<ASTOpImageSwizzle> NewBaseSwizzle = new ASTOpImageSwizzle;
+						NewBaseSwizzle->Format = EImageFormat::IF_L_UBYTE;
+						NewBaseSwizzle->Sources[0] = OldLayer->base.child();
+						NewBaseSwizzle->SourceChannels[0] = SourceChannels[3];
+
+						Ptr<ASTOpImageSwizzle> NewBlendSwizzle = new ASTOpImageSwizzle;
+						NewBlendSwizzle->Format = EImageFormat::IF_L_UBYTE;
+						NewBlendSwizzle->Sources[0] = OldLayer->blend.child();
+						NewBlendSwizzle->SourceChannels[0] = SourceChannels[3];
+
+						Ptr<ASTOpImageLayer> NewLayer = mu::Clone<ASTOpImageLayer>(OldLayer);
+						NewLayer->base = NewBaseSwizzle;
+						NewLayer->blend = NewBlendSwizzle;
+
+						Ptr<ASTOpImageSwizzle> NewSwizzle = mu::Clone<ASTOpImageSwizzle>(this);
+						NewSwizzle->Sources[3] = NewLayer;
+
+						at = NewSwizzle;
+					}
+				}
+			}
+		}
+
+
 		return at;
 	}
 
