[Mutable] Fix mutable image parameter MipsToSkip so they have into account the LOD being generated.

#rnx

diff --git a/Engine/Plugins/Experimental/Mutable/Source/CustomizableObject/Private/MuCO/CustomizableObjectMipDataProvider.cpp b/Engine/Plugins/Experimental/Mutable/Source/CustomizableObject/Private/MuCO/CustomizableObjectMipDataProvider.cpp
--- a/Engine/Plugins/Experimental/Mutable/Source/CustomizableObject/Private/MuCO/CustomizableObjectMipDataProvider.cpp
+++ b/Engine/Plugins/Experimental/Mutable/Source/CustomizableObject/Private/MuCO/CustomizableObjectMipDataProvider.cpp
@@ -101,7 +101,7 @@ namespace impl
 			{
 				MUTABLE_CPUPROFILER_SCOPE(GetImage);
 
-				Image = System->GetImage(InstanceID, MipImageID, OperationData->MipsToSkip);
+				Image = System->GetImage(InstanceID, MipImageID, OperationData->MipsToSkip, ImageRef.LOD);
 			}
 
 			check(Image);
diff --git a/Engine/Plugins/Experimental/Mutable/Source/CustomizableObject/Private/MuCO/CustomizableObjectSystem.cpp b/Engine/Plugins/Experimental/Mutable/Source/CustomizableObject/Private/MuCO/CustomizableObjectSystem.cpp
--- a/Engine/Plugins/Experimental/Mutable/Source/CustomizableObject/Private/MuCO/CustomizableObjectSystem.cpp
+++ b/Engine/Plugins/Experimental/Mutable/Source/CustomizableObject/Private/MuCO/CustomizableObjectSystem.cpp
@@ -1350,6 +1350,8 @@ namespace impl
 							Image.ImageID = Instance->GetImageId(MutableLODIndex, ComponentIndex, InstanceSurfaceIndex, ImageIndex);
 							Image.FullImageSizeX = 0;
 							Image.FullImageSizeY = 0;
+							Image.LOD = MutableLODIndex;
+
 						}
 
 						// Vectors
@@ -1435,7 +1437,7 @@ namespace impl
 				int32 MinMipsInImage = FMath::Min(FullLODCount, UTexture::GetStaticMinTextureResidentMipCount());
 				int32 MaxMipsToSkip = FullLODCount - MinMipsInImage;
 				int32 MipsToSkip = FMath::Min(MaxMipsToSkip, OperationData->MipsToSkip);
-				Image.Image = System->GetImage(OperationData->InstanceID, Image.ImageID, MipsToSkip);
+				Image.Image = System->GetImage(OperationData->InstanceID, Image.ImageID, MipsToSkip, Image.LOD);
 
 				// We need one mip or the complete chain. Otherwise there was a bug.
 				check(Image.Image);
diff --git a/Engine/Plugins/Experimental/Mutable/Source/CustomizableObject/Private/MuCO/CustomizableObjectSystemPrivate.h b/Engine/Plugins/Experimental/Mutable/Source/CustomizableObject/Private/MuCO/CustomizableObjectSystemPrivate.h
--- a/Engine/Plugins/Experimental/Mutable/Source/CustomizableObject/Private/MuCO/CustomizableObjectSystemPrivate.h
+++ b/Engine/Plugins/Experimental/Mutable/Source/CustomizableObject/Private/MuCO/CustomizableObjectSystemPrivate.h
@@ -317,6 +317,7 @@ struct FInstanceUpdateData
 		uint16 FullImageSizeX, FullImageSizeY;
 		mu::ImagePtrConst Image;
 		TWeakObjectPtr<UTexture2D> Cached;
+		int32 LOD;
 	};
 
 	struct FVector
diff --git a/Engine/Plugins/Experimental/Mutable/Source/CustomizableObjectEditor/Private/MuCOE/SMutableCodeViewer.cpp b/Engine/Plugins/Experimental/Mutable/Source/CustomizableObjectEditor/Private/MuCOE/SMutableCodeViewer.cpp
--- a/Engine/Plugins/Experimental/Mutable/Source/CustomizableObjectEditor/Private/MuCOE/SMutableCodeViewer.cpp
+++ b/Engine/Plugins/Experimental/Mutable/Source/CustomizableObjectEditor/Private/MuCOE/SMutableCodeViewer.cpp
@@ -2334,7 +2334,7 @@ void SMutableCodeViewer::Tick(const FGeometry& AllottedGeometry, const double In
 	case mu::DT_IMAGE:
 	{
 		check(PreviewImageViewer);
-		mu::ImagePtrConst MutableImage = System->GetPrivate()->BuildImage(MutableModel, PreviewParameters.get(), SelectedOperationAddress, MipsToSkip);
+		mu::ImagePtrConst MutableImage = System->GetPrivate()->BuildImage(MutableModel, PreviewParameters.get(), SelectedOperationAddress, MipsToSkip, 0);
 		PreviewImageViewer->SetImage(MutableImage, 0);
 		break;
 	}
diff --git a/Engine/Plugins/Experimental/Mutable/Source/MutableRuntime/Private/MuR/CodeRunner.cpp b/Engine/Plugins/Experimental/Mutable/Source/MutableRuntime/Private/MuR/CodeRunner.cpp
--- a/Engine/Plugins/Experimental/Mutable/Source/MutableRuntime/Private/MuR/CodeRunner.cpp
+++ b/Engine/Plugins/Experimental/Mutable/Source/MutableRuntime/Private/MuR/CodeRunner.cpp
@@ -83,7 +83,7 @@ namespace mu
 		const TSharedPtr<const Model>& pModel,
 		const Parameters* pParams,
 		OP::ADDRESS at,
-		uint32 lodMask, uint8 executionOptions, FScheduledOp::EType Type )
+		uint32 lodMask, uint8 executionOptions, int32 InImageLOD, FScheduledOp::EType Type )
 		: m_pSettings(pSettings), m_pSystem(s), m_pModel(pModel), m_pParams(pParams), m_lodMask(lodMask)
 	{
 		ScheduledStagePerOp.resize(m_pModel->GetPrivate()->m_program.m_opAddress.Num());
@@ -94,6 +94,8 @@ namespace mu
 		FProgram& program = m_pModel->GetPrivate()->m_program;
 		m_romPendingOps.SetNum(program.m_roms.Num());
 
+		ImageLOD = InImageLOD;
+
 		// Push the root operation
 		FScheduledOp rootOp;
 		rootOp.At = at;
diff --git a/Engine/Plugins/Experimental/Mutable/Source/MutableRuntime/Private/MuR/CodeRunner.h b/Engine/Plugins/Experimental/Mutable/Source/MutableRuntime/Private/MuR/CodeRunner.h
--- a/Engine/Plugins/Experimental/Mutable/Source/MutableRuntime/Private/MuR/CodeRunner.h
+++ b/Engine/Plugins/Experimental/Mutable/Source/MutableRuntime/Private/MuR/CodeRunner.h
@@ -44,7 +44,7 @@ namespace  mu
     public:
 		CodeRunner(const SettingsPtrConst&, class System::Private*, 
 			const TSharedPtr<const Model>&, const Parameters* pParams,
-			OP::ADDRESS at, uint32 lodMask, uint8 executionOptions, FScheduledOp::EType );
+			OP::ADDRESS at, uint32 lodMask, uint8 executionOptions, int32 InImageLOD, FScheduledOp::EType );
 
     protected:
 
@@ -115,6 +115,9 @@ namespace  mu
 		TArray< FScheduledOpData > m_heapData;
 		TArray< FImageDesc > m_heapImageDesc;
 
+		//! Only used for correct mip skipping with external images. Is the LOD for which the image is build.
+		int32 ImageLOD;
+
 	public:
 
 		void Run();
diff --git a/Engine/Plugins/Experimental/Mutable/Source/MutableRuntime/Private/MuR/CodeRunnerTasked.cpp b/Engine/Plugins/Experimental/Mutable/Source/MutableRuntime/Private/MuR/CodeRunnerTasked.cpp
--- a/Engine/Plugins/Experimental/Mutable/Source/MutableRuntime/Private/MuR/CodeRunnerTasked.cpp
+++ b/Engine/Plugins/Experimental/Mutable/Source/MutableRuntime/Private/MuR/CodeRunnerTasked.cpp
@@ -1948,8 +1948,12 @@ namespace mu
 			Ptr<RangeIndex> Index = BuildCurrentOpRangeIndex(item, m_pParams, m_pModel.Get(), args.variable);
 
 			const EXTERNAL_IMAGE_ID Id = m_pParams->GetImageValue(args.variable, Index);
-			const uint8 MipmapsToSkip = item.ExecutionOptions;
-			
+
+			check(ImageLOD < TNumericLimits<uint8>::Max() && ImageLOD >= 0);
+			check(ImageLOD + static_cast<int32>(item.ExecutionOptions) < TNumericLimits<uint8>::Max());
+
+			const uint8 MipmapsToSkip = item.ExecutionOptions + static_cast<uint8>(ImageLOD);
+
 			Issued = MakeShared<FImageExternalLoadTask>(item, MipmapsToSkip, Id);
 
 			break;
@@ -2063,5 +2067,4 @@ namespace mu
 		return Issued;
 	}
 
-
 } // namespace mu
diff --git a/Engine/Plugins/Experimental/Mutable/Source/MutableRuntime/Private/MuR/System.cpp b/Engine/Plugins/Experimental/Mutable/Source/MutableRuntime/Private/MuR/System.cpp
--- a/Engine/Plugins/Experimental/Mutable/Source/MutableRuntime/Private/MuR/System.cpp
+++ b/Engine/Plugins/Experimental/Mutable/Source/MutableRuntime/Private/MuR/System.cpp
@@ -290,7 +290,7 @@ namespace mu
 
 
 	//---------------------------------------------------------------------------------------------
-	ImagePtrConst System::GetImage(Instance::ID instanceID, RESOURCE_ID imageId, int32 MipsToSkip)
+	ImagePtrConst System::GetImage(Instance::ID instanceID, RESOURCE_ID imageId, int32 MipsToSkip, int32 InImageLOD)
 	{
 		LLM_SCOPE_BYNAME(TEXT("MutableRuntime"));
 		MUTABLE_CPUPROFILER_SCOPE(SystemGetImage);
@@ -309,7 +309,7 @@ namespace mu
 			{
 				pResult = m_pD->BuildImage(pLiveInstance->m_pModel,
 					pLiveInstance->m_pOldParameters.get(),
-					res.m_rootAddress, MipsToSkip);
+					res.m_rootAddress, MipsToSkip, InImageLOD);
 
 				// We always need ot return something valid.
 				if (!pResult)
@@ -372,7 +372,7 @@ namespace mu
 				if (GetOpDataType(opType) == DT_IMAGE)
 				{
 					int8 executionOptions = 0;
-					CodeRunner Runner(m_pD->m_pSettings, m_pD, pLiveInstance->m_pModel, pLiveInstance->m_pOldParameters.get(), at, System::AllLODs, executionOptions, FScheduledOp::EType::ImageDesc);
+					CodeRunner Runner(m_pD->m_pSettings, m_pD, pLiveInstance->m_pModel, pLiveInstance->m_pOldParameters.get(), at, System::AllLODs, executionOptions, 0, FScheduledOp::EType::ImageDesc);
 					Runner.Run();
 					Runner.GetImageDescResult(OutDesc);
 				}
@@ -483,7 +483,7 @@ namespace mu
 		if (imageAddress)
 		{
 			m_pD->BeginBuild(pModel);
-			pResult = m_pD->BuildImage( pModel, pParams.get(), imageAddress, 0 );
+			pResult = m_pD->BuildImage( pModel, pParams.get(), imageAddress, 0, 0 );
 			m_pD->EndBuild();
 		}
 
@@ -639,10 +639,10 @@ namespace mu
 
 	//---------------------------------------------------------------------------------------------
 	void System::Private::RunCode(const TSharedPtr<const Model>& InModel,
-		const Parameters* InParameters, OP::ADDRESS InCodeRoot, uint32 InLODs, uint8 executionOptions)
+		const Parameters* InParameters, OP::ADDRESS InCodeRoot, uint32 InLODs, uint8 executionOptions, int32 InImageLOD)
 	{
 		CodeRunner Runner(m_pSettings, this, InModel, InParameters, InCodeRoot, InLODs, 
-			executionOptions, FScheduledOp::EType::Full);
+			executionOptions, InImageLOD, FScheduledOp::EType::Full);
 		Runner.Run();
 		bUnrecoverableError = Runner.bUnrecoverableError;
 	}
@@ -739,13 +739,13 @@ namespace mu
 	//---------------------------------------------------------------------------------------------
 	Ptr<const Image> System::Private::BuildImage(const TSharedPtr<const Model>& pModel,
 		const Parameters* pParams,
-		OP::ADDRESS at, int MipsToSkip)
+		OP::ADDRESS at, int32 MipsToSkip, int32 InImageLOD)
 	{
 		mu::OP_TYPE opType = pModel->GetPrivate()->m_program.GetOpType(at);
 		if (GetOpDataType(opType) == DT_IMAGE)
 		{
 			m_memory->ClearCacheLayer0();
-			RunCode(pModel, pParams, at, System::AllLODs, uint8(MipsToSkip));
+			RunCode(pModel, pParams, at, System::AllLODs, uint8(MipsToSkip), InImageLOD);
 			if (bUnrecoverableError)
 			{
 				return nullptr;
diff --git a/Engine/Plugins/Experimental/Mutable/Source/MutableRuntime/Private/MuR/SystemPrivate.h b/Engine/Plugins/Experimental/Mutable/Source/MutableRuntime/Private/MuR/SystemPrivate.h
--- a/Engine/Plugins/Experimental/Mutable/Source/MutableRuntime/Private/MuR/SystemPrivate.h
+++ b/Engine/Plugins/Experimental/Mutable/Source/MutableRuntime/Private/MuR/SystemPrivate.h
@@ -919,7 +919,7 @@ namespace mu
 		MUTABLERUNTIME_API float BuildScalar(const TSharedPtr<const Model>&, const Parameters* pParams, OP::ADDRESS at) ;
 		MUTABLERUNTIME_API void BuildColour(const TSharedPtr<const Model>&, const Parameters*, OP::ADDRESS, float* OutR, float* OutG, float* OutB, float* OutA) ;
 		MUTABLERUNTIME_API Ptr<const String> BuildString(const TSharedPtr<const Model>&, const Parameters* pParams, OP::ADDRESS at) ;
-		MUTABLERUNTIME_API Ptr<const Image> BuildImage(const TSharedPtr<const Model>&, const Parameters* pParams, OP::ADDRESS at, int32 MipsToSkip) ;
+		MUTABLERUNTIME_API Ptr<const Image> BuildImage(const TSharedPtr<const Model>&, const Parameters* pParams, OP::ADDRESS at, int32 MipsToSkip, int32 LOD) ;
 		MUTABLERUNTIME_API Ptr<const Mesh> BuildMesh(const TSharedPtr<const Model>&, const Parameters* pParams, OP::ADDRESS at) ;
 		MUTABLERUNTIME_API Ptr<const Layout> BuildLayout(const TSharedPtr<const Model>&, const Parameters* pParams, OP::ADDRESS at) ;
     	MUTABLERUNTIME_API Ptr<const Projector> BuildProjector(const TSharedPtr<const Model>&, const Parameters* pParams, OP::ADDRESS at) ;
@@ -1009,7 +1009,7 @@ namespace mu
 
 		void RunCode(const TSharedPtr<const Model>& pModel,
 			const Parameters* pParams,
-			OP::ADDRESS at, uint32 LODs = System::AllLODs, uint8 executionOptions = 0);
+			OP::ADDRESS at, uint32 LODs = System::AllLODs, uint8 executionOptions = 0, int32 LOD = 0);
 
 		//!
 		void PrepareCache(const Model*, int32 State);
diff --git a/Engine/Plugins/Experimental/Mutable/Source/MutableRuntime/Public/MuR/System.h b/Engine/Plugins/Experimental/Mutable/Source/MutableRuntime/Public/MuR/System.h
--- a/Engine/Plugins/Experimental/Mutable/Source/MutableRuntime/Public/MuR/System.h
+++ b/Engine/Plugins/Experimental/Mutable/Source/MutableRuntime/Public/MuR/System.h
@@ -146,7 +146,7 @@ namespace mu
 		//! \param MipsToSkip Number of mips to skip compared from the full image.
 		//! If 0, all mip levels will be generated. If more levels than possible to discard are specified, 
 		//! the image will still contain a minimum number of mips specified at model compile time.
-		Ptr<const Image> GetImage(Instance::ID InstanceID, RESOURCE_ID ImageId, int32 MipsToSkip = 0);
+		Ptr<const Image> GetImage(Instance::ID InstanceID, RESOURCE_ID ImageId, int32 MipsToSkip = 0, int32 LOD = 0);
 
         //! Only valid between BeginUpdate and EndUpdate
         Ptr<const Mesh> GetMesh(Instance::ID InstanceID, RESOURCE_ID MeshId);
diff --git a/Engine/Plugins/Experimental/Mutable/Source/MutableTools/Private/MuT/CodeOptimiser.cpp b/Engine/Plugins/Experimental/Mutable/Source/MutableTools/Private/MuT/CodeOptimiser.cpp
--- a/Engine/Plugins/Experimental/Mutable/Source/MutableTools/Private/MuT/CodeOptimiser.cpp
+++ b/Engine/Plugins/Experimental/Mutable/Source/MutableTools/Private/MuT/CodeOptimiser.cpp
@@ -496,7 +496,7 @@ namespace mu
 			{
 				MUTABLE_CPUPROFILER_SCOPE(ConstantImage);
 
-				mu::Ptr<const Image> pImage = pSystem->GetPrivate()->BuildImage( model, localParams.get(), at, 0 );
+				mu::Ptr<const Image> pImage = pSystem->GetPrivate()->BuildImage( model, localParams.get(), at, 0, 0 );
 
 				if (pImage)
 				{
