[Mutable] ParallelFor pixel batching optimizations and ASTC decompression

#rnx

diff --git a/Engine/Plugins/Experimental/Mutable/Source/MutableRuntime/Private/MuR/BlockCompression/Miro/Miro.cpp b/Engine/Plugins/Experimental/Mutable/Source/MutableRuntime/Private/MuR/BlockCompression/Miro/Miro.cpp
--- a/Engine/Plugins/Experimental/Mutable/Source/MutableRuntime/Private/MuR/BlockCompression/Miro/Miro.cpp
+++ b/Engine/Plugins/Experimental/Mutable/Source/MutableRuntime/Private/MuR/BlockCompression/Miro/Miro.cpp
@@ -7548,37 +7548,45 @@ namespace miro
 	void Generic_ASTCRGBAL_to_RGBA(uint32 sx, uint32 sy, const uint8* from, uint8* to)
 	{
 		init_astc_decompress();
+	
+		const int32 NumBlocksY = FMath::DivideAndRoundUp(sy, BLOCK_SIZE);
+		const int32 NumBlocksX = FMath::DivideAndRoundUp(sx, BLOCK_SIZE);
 
-		for (uint32 y = 0; y < sy; y += BLOCK_SIZE)
+		if (sx <= 0 || sy <= 0)
 		{
-			for (uint32 x = 0; x < sx; x += BLOCK_SIZE)
+			return;
+		}
+
+		//for (uint32 BlockY = 0; BlockY < NumBlocksY; ++BlockY)
+		ParallelFor(NumBlocksY, [NumBlocksX, sx, sy, from, to](int32 BlockY)
+		{
+			for (int32 BlockX = 0; BlockX < NumBlocksX; ++BlockX)
 			{
-				bool bIsSRGB = false;
-				uint8 Block[BLOCK_SIZE * BLOCK_SIZE *4 ];
-				bool bSuccess = astcdec::decompress( Block, from, bIsSRGB, BLOCK_SIZE, BLOCK_SIZE);
+				constexpr bool bIsSRGB = false;
+				uint8 Block[BLOCK_SIZE * BLOCK_SIZE * 4];
+
+				constexpr int32 CompressedBlockSize = 16;
+				const uint8* SrcBlockPtr = from + (BlockY * NumBlocksX + BlockX) * CompressedBlockSize;
+
+				bool bSuccess = astcdec::decompress(Block, SrcBlockPtr, bIsSRGB, BLOCK_SIZE, BLOCK_SIZE);
 				check(bSuccess);
 
 				for (uint32 py = 0; py < BLOCK_SIZE; py++)
 				{
 					for (uint32 px = 0; px < BLOCK_SIZE; px++)
 					{
-						uint32 xi = x + px;
-						uint32 yi = y + py;
-
-						if (xi < sx && yi < sy)
-						{
-							uint8* toPixel = to + (yi * sx + xi) * 4;
-							toPixel[0] = Block[py * BLOCK_SIZE * 4 + px * 4 + 0];
-							toPixel[1] = Block[py * BLOCK_SIZE * 4 + px * 4 + 1];
-							toPixel[2] = Block[py * BLOCK_SIZE * 4 + px * 4 + 2];
-							toPixel[3] = Block[py * BLOCK_SIZE * 4 + px * 4 + 3];
-						}
+						uint32 xi = FMath::Min(BlockX * BLOCK_SIZE + px, sx - 1);
+						uint32 yi = FMath::Min(BlockY * BLOCK_SIZE + py, sy - 1);
+
+						uint8* toPixel = to + (yi * sx + xi) * 4;
+						toPixel[0] = Block[py * BLOCK_SIZE * 4 + px * 4 + 0];
+						toPixel[1] = Block[py * BLOCK_SIZE * 4 + px * 4 + 1];
+						toPixel[2] = Block[py * BLOCK_SIZE * 4 + px * 4 + 2];
+						toPixel[3] = Block[py * BLOCK_SIZE * 4 + px * 4 + 3];
 					}
 				}
-
-				from += 16;
 			}
-		}
+		});
 	}
 
 
@@ -7588,21 +7596,29 @@ namespace miro
 	{
 		init_astc_decompress();
 
-		for (uint32 y = 0; y < sy; y += BLOCK_SIZE)
+		const int32 NumBlocksY = FMath::DivideAndRoundUp(sy, BLOCK_SIZE);
+		const int32 NumBlocksX = FMath::DivideAndRoundUp(sx, BLOCK_SIZE);
+
+		//for (uint32 BlockY = 0; BlockY < NumBlocksY; ++BlockY)
+		ParallelFor(NumBlocksY, [NumBlocksX, sx, sy, from, to](int32 BlockY)
 		{
-			for (uint32 x = 0; x < sx; x += BLOCK_SIZE)
+			for (int32 BlockX = 0; BlockX < NumBlocksX; ++BlockX)
 			{
-				bool bIsSRGB = false;
+				constexpr bool bIsSRGB = false;
 				uint8 Block[BLOCK_SIZE * BLOCK_SIZE * 4];
-				bool bSuccess = astcdec::decompress(Block, from, bIsSRGB, BLOCK_SIZE, BLOCK_SIZE);
+
+				constexpr int32 CompressedBlockSize = 16;
+				const uint8* SrcBlockPtr = from + (BlockY * NumBlocksX + BlockX) * CompressedBlockSize;
+
+				bool bSuccess = astcdec::decompress(Block, SrcBlockPtr, bIsSRGB, BLOCK_SIZE, BLOCK_SIZE);
 				check(bSuccess);
 
 				for (uint32 py = 0; py < BLOCK_SIZE; py++)
 				{
 					for (uint32 px = 0; px < BLOCK_SIZE; px++)
 					{
-						uint32 xi = x + px;
-						uint32 yi = y + py;
+						uint32 xi = BlockX * BLOCK_SIZE + px;
+						uint32 yi = BlockY * BLOCK_SIZE + py;
 
 						if (xi < sx && yi < sy)
 						{
@@ -7613,10 +7629,8 @@ namespace miro
 						}
 					}
 				}
-
-				from += 16;
 			}
-		}
+		});
 	}
 
 
diff --git a/Engine/Plugins/Experimental/Mutable/Source/MutableRuntime/Private/MuR/OpImageLayer.h b/Engine/Plugins/Experimental/Mutable/Source/MutableRuntime/Private/MuR/OpImageLayer.h
--- a/Engine/Plugins/Experimental/Mutable/Source/MutableRuntime/Private/MuR/OpImageLayer.h
+++ b/Engine/Plugins/Experimental/Mutable/Source/MutableRuntime/Private/MuR/OpImageLayer.h
@@ -37,86 +37,88 @@ namespace mu
 		top[2] = uint32(255 * col[2]);
 
 		// Generic implementation
-		int32 pixelCount = pBase->CalculatePixelCount();
+		const int32 PixelCount = pBase->CalculatePixelCount();
 
+		// Not used, ParallelFor + baching already consider low element count. 
 		constexpr int PixelCountConcurrencyThreshold = 0xff;
 
+		constexpr int32 NumBatchElems = 4096*2;
+		const int32 NumBatches = FMath::DivideAndRoundUp(PixelCount, NumBatchElems);
+
 		switch (baseFormat)
 		{
 		case EImageFormat::IF_RGB_UBYTE:
 		{
-			const auto& ProcessPixel = [
-				pBaseBuf, pDestBuf, top
-			] (uint32 i)
+			ParallelFor(NumBatches,
+			[
+				pBaseBuf, pDestBuf, top, PixelCount, NumBatchElems
+			] (int32 BatchId)
 			{
-				for (int c = 0; c < 3; ++c)
+				const int32 BatchBegin = BatchId * NumBatchElems;
+				const int32 BatchEnd = FMath::Min(BatchBegin + NumBatchElems, PixelCount);
+
+				for (int32 i = BatchBegin; i < BatchEnd; ++i)
 				{
-					unsigned base = pBaseBuf[3 * i + c];
-					unsigned result = BLEND_FUNC(base, top[c]);
-					if (CLAMP)
-					{
-						pDestBuf[3 * i + c] = (uint8)FMath::Min(255u, result);
-					}
-					else
+					for (int c = 0; c < 3; ++c)
 					{
-						pDestBuf[3 * i + c] = (uint8)result;
+						unsigned base = pBaseBuf[3 * i + c];
+						unsigned result = BLEND_FUNC(base, top[c]);
+						if constexpr (CLAMP)
+						{
+							pDestBuf[3 * i + c] = (uint8)FMath::Min(255u, result);
+						}
+						else
+						{
+							pDestBuf[3 * i + c] = (uint8)result;
+						}
 					}
 				}
-			};
+			});
 
-			if (pixelCount > PixelCountConcurrencyThreshold)
-			{
-				ParallelFor(pixelCount, ProcessPixel);
-			}
-			else
-			{
-				for ( int p=0; p< pixelCount; ++p )
-				{
-					ProcessPixel(p);
-				}
-			}
 			break;
 		}
 		case EImageFormat::IF_RGBA_UBYTE:
 		{
-			const auto& ProcessPixel = [
-				pBaseBuf, pDestBuf, top
-			] (uint32 i)
+			ParallelFor(NumBatches, 
+			[
+				pBaseBuf, pDestBuf, top, PixelCount, NumBatchElems
+			] (int32 BatchId)
 			{
-				for (int c = 0; c < 3; ++c)
+				const int32 BatchBegin = BatchId * NumBatchElems;
+				const int32 BatchEnd = FMath::Min(BatchBegin + NumBatchElems, PixelCount);
+
+				for (int32 i = BatchBegin; i < BatchEnd; ++i)
 				{
-					unsigned base = pBaseBuf[4 * i + c];
-					unsigned result = BLEND_FUNC(base, top[c]);
-					if (CLAMP)
-					{
-						pDestBuf[4 * i + c] = (uint8)FMath::Min(255u, result);
-					}
-					else
+					for (int c = 0; c < 3; ++c)
 					{
-						pDestBuf[4 * i + c] = (uint8)result;
+						unsigned base = pBaseBuf[4 * i + c];
+						unsigned result = BLEND_FUNC(base, top[c]);
+						if constexpr (CLAMP)
+						{
+							pDestBuf[4 * i + c] = (uint8)FMath::Min(255u, result);
+						}
+						else
+						{
+							pDestBuf[4 * i + c] = (uint8)result;
+						}
 					}
+					pDestBuf[4 * i + 3] = pBaseBuf[4 * i + 3];
 				}
-				pDestBuf[4 * i + 3] = pBaseBuf[4 * i + 3];
-			};
+			});
 
-			if (pixelCount > PixelCountConcurrencyThreshold)
-			{
-				ParallelFor(pixelCount, ProcessPixel);
-			}
-			else
-			{
-				for (int p = 0; p < pixelCount; ++p)
-				{
-					ProcessPixel(p);
-				}
-			}
 			break;
 		}
 		case EImageFormat::IF_BGRA_UBYTE:
 		{
-			const auto& ProcessPixel = [
-					pBaseBuf, pDestBuf, top
-				] (uint32 i)
+			ParallelFor(NumBatches,
+			[
+				pBaseBuf, pDestBuf, top, PixelCount, NumBatchElems
+			] (int32 BatchId)
+			{
+				const int32 BatchBegin = BatchId * NumBatchElems;
+				const int32 BatchEnd = FMath::Min(BatchBegin + NumBatchElems, PixelCount);
+
+				for (int32 i = BatchBegin; i < BatchEnd; ++i)
 				{
 					for (int c = 0; c < 3; ++c)
 					{
@@ -132,30 +134,26 @@ namespace mu
 						}
 					}
 					pDestBuf[4 * i + 3] = pBaseBuf[4 * i + 3];
-				};
-
-			if (pixelCount > PixelCountConcurrencyThreshold)
-			{
-				ParallelFor(pixelCount, ProcessPixel);
-			}
-			else
-			{
-				for (int p = 0; p < pixelCount; ++p)
-				{
-					ProcessPixel(p);
 				}
-			}
+			});
+
 			break;
 		}
 		case EImageFormat::IF_L_UBYTE:
 		{
-			const auto& ProcessPixel = [
-					pBaseBuf, pDestBuf, top
-				] (uint32 i)
+			ParallelFor(NumBatches, 
+			[
+				pBaseBuf, pDestBuf, top, PixelCount, NumBatchElems
+			] (int32 BatchId)
+			{
+				const int32 BatchBegin = BatchId * NumBatchElems;
+				const int32 BatchEnd = FMath::Min(BatchBegin + NumBatchElems, PixelCount);
+
+				for (int32 i = BatchBegin; i < BatchEnd; ++i)
 				{
 					unsigned base = pBaseBuf[i];
 					unsigned result = BLEND_FUNC(base, top[0]);
-					if (CLAMP)
+					if constexpr (CLAMP)
 					{
 						pDestBuf[i] = (uint8)FMath::Min(255u, result);
 					}
@@ -163,20 +161,9 @@ namespace mu
 					{
 						pDestBuf[i] = (uint8)result;
 					}
-				};
-
-
-			if (pixelCount > PixelCountConcurrencyThreshold)
-			{
-				ParallelFor(pixelCount, ProcessPixel);
-			}
-			else
-			{
-				for (int p = 0; p < pixelCount; ++p)
-				{
-					ProcessPixel(p);
 				}
-			}
+			});
+
 			break;
 		}
 
@@ -222,141 +209,128 @@ namespace mu
 		top[2] = (unsigned)(255 * Col[2]);
 
 		// Generic implementation
-		int32 pixelCount = Base->CalculatePixelCount();
+		int32 PixelCount = Base->CalculatePixelCount();
 
 		constexpr int PixelCountConcurrencyThreshold = 0xff;
 
+
+		constexpr int32 NumBatchElems = 4096*2;
+		const int32 NumBatches = FMath::DivideAndRoundUp(PixelCount, NumBatchElems);
+
 		switch (baseFormat)
 		{
 		case EImageFormat::IF_RGB_UBYTE:
 		{
 			// There is no alpha so we assume 255
-			const auto& ProcessPixel = [ pDestBuf, top ] (uint32 i)
+			ParallelFor(NumBatches,
+				[ pDestBuf, top, PixelCount, NumBatchElems ] (uint32 BatchId)
 			{
-				for (int c = 0; c < 3; ++c)
+				const int32 BatchBegin = BatchId * NumBatchElems;
+				const int32 BatchEnd = FMath::Min(BatchBegin + NumBatchElems, PixelCount);
+
+				for (int32 i = BatchBegin; i < BatchEnd; ++i)
 				{
-					unsigned base = 255;
-					unsigned result = BLEND_FUNC(base, top[c]);
-					if (CLAMP)
-					{
-						pDestBuf[3 * i + c] = (uint8)FMath::Min(255u, result);
-					}
-					else
+					for (int c = 0; c < 3; ++c)
 					{
-						pDestBuf[3 * i + c] = (uint8)result;
+						unsigned base = 255;
+						unsigned result = BLEND_FUNC(base, top[c]);
+						if constexpr (CLAMP)
+						{
+							pDestBuf[3 * i + c] = (uint8)FMath::Min(255u, result);
+						}
+						else
+						{
+							pDestBuf[3 * i + c] = (uint8)result;
+						}
 					}
 				}
-			};
+			});
 
-			if (pixelCount > PixelCountConcurrencyThreshold)
-			{
-				ParallelFor(pixelCount, ProcessPixel);
-			}
-			else
-			{
-				for (int p = 0; p < pixelCount; ++p)
-				{
-					ProcessPixel(p);
-				}
-			}
 			break;
 		}
 		case EImageFormat::IF_RGBA_UBYTE:
 		{
-			const auto& ProcessPixel = [ pBaseBuf, pDestBuf, top ] (uint32 i)
-			{
-				for (int c = 0; c < 3; ++c)
+			ParallelFor(NumBatches, 
+				[ pBaseBuf, pDestBuf, top, PixelCount, NumBatchElems ] (uint32 BatchId)
+			{	
+				const int32 BatchBegin = BatchId * NumBatchElems;
+				const int32 BatchEnd = FMath::Min(BatchBegin + NumBatchElems, PixelCount);
+
+				for (int32 i = BatchBegin; i < BatchEnd; ++i)
 				{
-					unsigned base = pBaseBuf[4 * i + 3];
-					unsigned result = BLEND_FUNC(base, top[c]);
-					if (CLAMP)
-					{
-						pDestBuf[4 * i + c] = (uint8)FMath::Min(255u, result);
-					}
-					else
+					for (int c = 0; c < 3; ++c)
 					{
-						pDestBuf[4 * i + c] = (uint8)result;
+						unsigned base = pBaseBuf[4 * i + 3];
+						unsigned result = BLEND_FUNC(base, top[c]);
+						if constexpr (CLAMP)
+						{
+							pDestBuf[4 * i + c] = (uint8)FMath::Min(255u, result);
+						}
+						else
+						{
+							pDestBuf[4 * i + c] = (uint8)result;
+						}
 					}
+					pDestBuf[4 * i + 3] = pBaseBuf[4 * i + 3];
 				}
-				pDestBuf[4 * i + 3] = pBaseBuf[4 * i + 3];
-			};
+			});
 
-			if (pixelCount > PixelCountConcurrencyThreshold)
-			{
-				ParallelFor(pixelCount, ProcessPixel);
-			}
-			else
-			{
-				for (int p = 0; p < pixelCount; ++p)
-				{
-					ProcessPixel(p);
-				}
-			}
 			break;
 		}
 		case EImageFormat::IF_BGRA_UBYTE:
-		{
-			const auto& ProcessPixel = [ pBaseBuf, pDestBuf, top ] (uint32 i)
+		{	
+			ParallelFor(NumBatches,
+				[ pBaseBuf, pDestBuf, top, PixelCount, NumBatchElems ] (uint32 BatchId)
 			{
-				for (int c = 0; c < 3; ++c)
+				const int32 BatchBegin = BatchId * NumBatchElems;
+				const int32 BatchEnd = FMath::Min(BatchBegin + NumBatchElems, PixelCount);
+
+				for (int32 i = BatchBegin; i < BatchEnd; ++i)
 				{
-					unsigned base = pBaseBuf[4 * i + 3];
-					unsigned result = BLEND_FUNC(base, top[2 - c]);
-					if (CLAMP)
-					{
-						pDestBuf[4 * i + c] = (uint8)FMath::Min(255u, result);
-					}
-					else
+					for (int c = 0; c < 3; ++c)
 					{
-						pDestBuf[4 * i + c] = (uint8)result;
+						unsigned base = pBaseBuf[4 * i + 3];
+						unsigned result = BLEND_FUNC(base, top[2 - c]);
+						if (CLAMP)
+						{
+							pDestBuf[4 * i + c] = (uint8)FMath::Min(255u, result);
+						}
+						else
+						{
+							pDestBuf[4 * i + c] = (uint8)result;
+						}
 					}
+					pDestBuf[4 * i + 3] = pBaseBuf[4 * i + 3];
 				}
-				pDestBuf[4 * i + 3] = pBaseBuf[4 * i + 3];
-			};
+			});
 
-			if (pixelCount > PixelCountConcurrencyThreshold)
-			{
-				ParallelFor(pixelCount, ProcessPixel);
-			}
-			else
-			{
-				for (int p = 0; p < pixelCount; ++p)
-				{
-					ProcessPixel(p);
-				}
-			}
 			break;
 		}
 		case EImageFormat::IF_L_UBYTE:
 		{
-			const auto& ProcessPixel = [ pBaseBuf, pDestBuf, top ] (uint32 i)
+			ParallelFor(NumBatches, 
+				[ pBaseBuf, pDestBuf, top, PixelCount, NumBatchElems ] (uint32 BatchId)
 			{
-				// There is no alpha so we assume 255
-				unsigned base = 255;
-				unsigned result = BLEND_FUNC(base, top[0]);
-				if (CLAMP)
-				{
-					pDestBuf[i] = (uint8)FMath::Min(255u, result);
-				}
-				else
-				{
-					pDestBuf[i] = (uint8)result;
-				}
-			};
+				const int32 BatchBegin = BatchId * NumBatchElems;
+				const int32 BatchEnd = FMath::Min(BatchBegin + NumBatchElems, PixelCount);
 
-
-				if (pixelCount > PixelCountConcurrencyThreshold)
-				{
-					ParallelFor(pixelCount, ProcessPixel);
-				}
-				else
+				for (int32 i = BatchBegin; i < BatchEnd; ++i)
 				{
-					for (int p = 0; p < pixelCount; ++p)
+					// There is no alpha so we assume 255
+					unsigned base = 255;
+					unsigned result = BLEND_FUNC(base, top[0]);
+					if constexpr (CLAMP)
 					{
-						ProcessPixel(p);
+						pDestBuf[i] = (uint8)FMath::Min(255u, result);
+					}
+					else
+					{
+						pDestBuf[i] = (uint8)result;
 					}
 				}
-				break;
+			});
+
+			break;
 		}
 
 		default:
@@ -683,23 +657,32 @@ namespace mu
 			PixelCount = Base->CalculatePixelCount();
 		}
 
-		ParallelFor(PixelCount,
+		constexpr int32 NumBatchElems = 4096*2;
+		const int32 NumBatches = FMath::DivideAndRoundUp(PixelCount, NumBatchElems);
+
+		ParallelFor(NumBatches,
 			[
-				pBaseBuf, top, BaseChannelOffset, ColOffset
-			] (uint32 i)
+				pBaseBuf, top, BaseChannelOffset, ColOffset, PixelCount, NumBatchElems
+			] (uint32 BatchId)
 			{
-				for (int c = 0; c < CHANNELS_TO_BLEND; ++c)
+				const int32 BatchBegin = BatchId * NumBatchElems;
+				const int32 BatchEnd = FMath::Min(BatchBegin + NumBatchElems, PixelCount);
+				
+				for (int32 i = BatchBegin; i < BatchEnd; ++i)
 				{
-					uint32 base = pBaseBuf[BASE_CHANNEL_STRIDE * i + c];
-					uint32 blended = top[c + ColOffset];
-					uint32 result = BLEND_FUNC(base, blended);
-					if (CLAMP)
-					{
-						pBaseBuf[BASE_CHANNEL_STRIDE * i + c] = (uint8)FMath::Min(255u, result);
-					}
-					else
+					for (int c = 0; c < CHANNELS_TO_BLEND; ++c)
 					{
-						pBaseBuf[BASE_CHANNEL_STRIDE * i + c] = (uint8)result;
+						uint32 base = pBaseBuf[BASE_CHANNEL_STRIDE * i + c];
+						uint32 blended = top[c + ColOffset];
+						uint32 result = BLEND_FUNC(base, blended);
+						if (CLAMP)
+						{
+							pBaseBuf[BASE_CHANNEL_STRIDE * i + c] = (uint8)FMath::Min(255u, result);
+						}
+						else
+						{
+							pBaseBuf[BASE_CHANNEL_STRIDE * i + c] = (uint8)result;
+						}
 					}
 				}
 			});
@@ -857,25 +840,34 @@ namespace mu
 			PixelCount = pBase->CalculatePixelCount();
 		}
 
-		ParallelFor(PixelCount,
+		constexpr int32 NumBatchElems = 4096*2;
+		const int32 NumBatches = FMath::DivideAndRoundUp(PixelCount, NumBatchElems);
+
+		ParallelFor(NumBatches,
 			[
-				pBaseBuf, pBlendedBuf
-			] (uint32 i)
+				pBaseBuf, pBlendedBuf, PixelCount, NumBatchElems
+			] (int32 BatchId)
 			{
-				for (int c = 0; c < CHANNELS_TO_BLEND; ++c)
+				const int32 BatchBegin = BatchId * NumBatchElems;
+				const int32 BatchEnd = FMath::Min(BatchBegin + NumBatchElems, PixelCount);
+
+				for (int32 i = BatchBegin; i < BatchEnd; ++i)
 				{
-					unsigned base = pBaseBuf[BASE_CHANNEL_STRIDE * i + c];
-					unsigned blended = pBlendedBuf[BLENDED_CHANNEL_STRIDE * i + c];
-					unsigned result = BLEND_FUNC(base, blended);
-					if (CLAMP)
-					{
-						pBaseBuf[BASE_CHANNEL_STRIDE * i + c] = (uint8)FMath::Min(255u, result);
-					}
-					else
+					for (int c = 0; c < CHANNELS_TO_BLEND; ++c)
 					{
-						pBaseBuf[BASE_CHANNEL_STRIDE * i + c] = (uint8)result;
+						unsigned base = pBaseBuf[BASE_CHANNEL_STRIDE * i + c];
+						unsigned blended = pBlendedBuf[BLENDED_CHANNEL_STRIDE * i + c];
+						unsigned result = BLEND_FUNC(base, blended);
+						if (CLAMP)
+						{
+							pBaseBuf[BASE_CHANNEL_STRIDE * i + c] = (uint8)FMath::Min(255u, result);
+						}
+						else
+						{
+							pBaseBuf[BASE_CHANNEL_STRIDE * i + c] = (uint8)result;
+						}
 					}
-				}				
+				}
 			});
 	}
 	
@@ -916,32 +908,41 @@ namespace mu
 
 		int32 UnblendedChannels = BASE_CHANNEL_STRIDE - CHANNELS_TO_BLEND;
 
-		ParallelFor(PixelCount,
+		constexpr int32 NumBatchElems = 4096*2;
+		const int32 NumBatches = FMath::DivideAndRoundUp(PixelCount, NumBatchElems);
+		
+		ParallelFor(NumBatches,
 			[
-				pBaseBuf, pBlendedBuf, pDestBuf, UnblendedChannels
-			] (uint32 i)
+				pBaseBuf, pBlendedBuf, pDestBuf, UnblendedChannels, PixelCount, NumBatchElems
+			] (int32 BatchId)
 			{
-				for (int c = 0; c < CHANNELS_TO_BLEND; ++c)
+				const int32 BatchBegin = BatchId * NumBatchElems;
+				const int32 BatchEnd = FMath::Min(BatchBegin + NumBatchElems, PixelCount);
+
+				for (int32 i = BatchBegin; i < BatchEnd; ++i)
 				{
-					unsigned base = pBaseBuf[BASE_CHANNEL_STRIDE * i + c];
-					unsigned blended = pBlendedBuf[BLENDED_CHANNEL_STRIDE * i + c];
-					unsigned result = BLEND_FUNC(base, blended);
-					if (CLAMP)
+					for (int c = 0; c < CHANNELS_TO_BLEND; ++c)
 					{
-						pDestBuf[BASE_CHANNEL_STRIDE * i + c] = (uint8)FMath::Min(255u, result);
+						unsigned base = pBaseBuf[BASE_CHANNEL_STRIDE * i + c];
+						unsigned blended = pBlendedBuf[BLENDED_CHANNEL_STRIDE * i + c];
+						unsigned result = BLEND_FUNC(base, blended);
+						if constexpr (CLAMP)
+						{
+							pDestBuf[BASE_CHANNEL_STRIDE * i + c] = (uint8)FMath::Min(255u, result);
+						}
+						else
+						{
+							pDestBuf[BASE_CHANNEL_STRIDE * i + c] = (uint8)result;
+						}
 					}
-					else
+
+					// Copy the unblended channels
+					// \TODO: unnecessary when doing it in-place?
+					for (int32 c = 0; c < UnblendedChannels; ++c)
 					{
-						pDestBuf[BASE_CHANNEL_STRIDE * i + c] = (uint8)result;
+						pDestBuf[BASE_CHANNEL_STRIDE * i + CHANNELS_TO_BLEND + c] = pBaseBuf[BASE_CHANNEL_STRIDE * i + CHANNELS_TO_BLEND + c];
 					}
 				}
-
-				// Copy the unblended channels
-				// \TODO: unnecessary when doing it in-place?
-				for (int32 c = 0; c < UnblendedChannels; ++c)
-				{
-					pDestBuf[BASE_CHANNEL_STRIDE * i + CHANNELS_TO_BLEND + c] = pBaseBuf[BASE_CHANNEL_STRIDE * i + CHANNELS_TO_BLEND + c];
-				}
 			});
 	}
 
@@ -1111,32 +1112,41 @@ namespace mu
 				PixelCount = pBase->CalculatePixelCount();
 			}
 
-			ParallelFor(PixelCount,
+			constexpr int32 NumBatchElems = 4096*2;
+			const int32 NumBatches = FMath::DivideAndRoundUp(PixelCount, NumBatchElems);
+
+			ParallelFor(NumBatches,
 				[
-					pBaseBuf, pBlendedBuf, pMaskBuf, pDestBuf, UnblendedChannels
-				] (uint32 i)
+					pBaseBuf, pBlendedBuf, pMaskBuf, pDestBuf, UnblendedChannels, PixelCount, NumBatchElems
+				] (uint32 BatchId)
 				{
-					unsigned mask = pMaskBuf[i];
-					for (int c = 0; c < CHANNELS_TO_BLEND; ++c)
+					const int32 BatchBegin = BatchId * NumBatchElems;
+					const int32 BatchEnd = FMath::Min(BatchBegin + NumBatchElems, PixelCount);
+
+					for (int32 i = BatchBegin; i < BatchEnd; ++i)
 					{
-						uint32 base = pBaseBuf[BASE_CHANNEL_STRIDE * i + c];
-						uint32 blended = pBlendedBuf[BLENDED_CHANNEL_STRIDE * i + c];
-						uint32 result = BLEND_FUNC_MASKED(base, blended, mask);
-						if (CLAMP)
+						unsigned mask = pMaskBuf[i];
+						for (int c = 0; c < CHANNELS_TO_BLEND; ++c)
 						{
-							pDestBuf[BASE_CHANNEL_STRIDE * i + c] = (uint8)FMath::Min(255u, result);
+							uint32 base = pBaseBuf[BASE_CHANNEL_STRIDE * i + c];
+							uint32 blended = pBlendedBuf[BLENDED_CHANNEL_STRIDE * i + c];
+							uint32 result = BLEND_FUNC_MASKED(base, blended, mask);
+							if constexpr (CLAMP)
+							{
+								pDestBuf[BASE_CHANNEL_STRIDE * i + c] = (uint8)FMath::Min(255u, result);
+							}
+							else
+							{
+								pDestBuf[BASE_CHANNEL_STRIDE * i + c] = (uint8)result;
+							}
 						}
-						else
+						// Copy the unblended channels
+						// \TODO: unnecessary when doing it in-place?
+						for (int c = 0; c < UnblendedChannels; ++c)
 						{
-							pDestBuf[BASE_CHANNEL_STRIDE * i + c] = (uint8)result;
+							pDestBuf[BASE_CHANNEL_STRIDE * i + CHANNELS_TO_BLEND + c] = pBaseBuf[BASE_CHANNEL_STRIDE * i + CHANNELS_TO_BLEND + c];
 						}
 					}
-					// Copy the unblended channels
-					// \TODO: unnecessary when doing it in-place?
-					for (int c = 0; c < UnblendedChannels; ++c)
-					{
-						pDestBuf[BASE_CHANNEL_STRIDE * i + CHANNELS_TO_BLEND + c] = pBaseBuf[BASE_CHANNEL_STRIDE * i + CHANNELS_TO_BLEND + c];
-					}
 				});
 		}
 		else if (maskFormat == EImageFormat::IF_L_UBYTE_RLE)
@@ -1323,10 +1333,18 @@ namespace mu
 		{
 			const uint8* pMaskBuf = pBlend->GetData() + 3;
 
-			ParallelFor(PixelCount,
-				[
-					pBaseBuf, pBlendedBuf, pDestBuf, pMaskBuf, UnblendedChannels
-				] (uint32 i)
+			constexpr int32 NumBatchElems = 4096*2;
+			const int32 NumBatches = FMath::DivideAndRoundUp(PixelCount, NumBatchElems);
+
+			ParallelFor(NumBatches,
+			[
+				pBaseBuf, pBlendedBuf, pDestBuf, pMaskBuf, UnblendedChannels, PixelCount, NumBatchElems
+			] (int32 BatchId)
+			{
+				const int32 BatchBegin = BatchId * NumBatchElems;
+				const int32 BatchEnd = FMath::Min(BatchBegin + NumBatchElems, PixelCount);
+
+				for (int32 i = BatchBegin; i < BatchEnd; ++i)
 				{
 					uint32 mask = pMaskBuf[BLENDED_CHANNEL_STRIDE * i];
 					for (int c = 0; c < CHANNELS_TO_BLEND; ++c)
@@ -1334,7 +1352,7 @@ namespace mu
 						uint32 base = pBaseBuf[BASE_CHANNEL_STRIDE * i + c];
 						uint32 blended = pBlendedBuf[BLENDED_CHANNEL_STRIDE * i + c];
 						uint32 result = BLEND_FUNC_MASKED(base, blended, mask);
-						if (CLAMP)
+						if constexpr (CLAMP)
 						{
 							pDestBuf[BASE_CHANNEL_STRIDE * i + c] = (uint8)FMath::Min(255u, result);
 						}
@@ -1349,21 +1367,30 @@ namespace mu
 					{
 						pDestBuf[BASE_CHANNEL_STRIDE * i + CHANNELS_TO_BLEND + c] = pBaseBuf[BASE_CHANNEL_STRIDE * i + CHANNELS_TO_BLEND + c];
 					}
-				});
+				}
+			});
 		}
 		else
 		{
-			ParallelFor(PixelCount,
-				[
-					pBaseBuf, pBlendedBuf, pDestBuf, UnblendedChannels
-				] (uint32 i)
+			constexpr int32 NumBatchElems = 4096*2;
+			const int32 NumBatches = FMath::DivideAndRoundUp(PixelCount, NumBatchElems);
+
+			ParallelFor(NumBatches,
+			[
+				pBaseBuf, pBlendedBuf, pDestBuf, UnblendedChannels, PixelCount, NumBatchElems
+			] (uint32 BatchId)
+			{
+				const int32 BatchBegin = BatchId * NumBatchElems;
+				const int32 BatchEnd = FMath::Min(BatchBegin + NumBatchElems, PixelCount);
+
+				for (int32 i = BatchBegin; i < BatchEnd; ++i)
 				{
 					for (int c = 0; c < CHANNELS_TO_BLEND; ++c)
 					{
 						uint32 base = pBaseBuf[BASE_CHANNEL_STRIDE * i + c];
 						uint32 blended = pBlendedBuf[BLENDED_CHANNEL_STRIDE * i + c];
 						uint32 result = BLEND_FUNC(base, blended);
-						if (CLAMP)
+						if constexpr (CLAMP)
 						{
 							pDestBuf[BASE_CHANNEL_STRIDE * i + c] = (uint8)FMath::Min(255u, result);
 						}
@@ -1378,8 +1405,8 @@ namespace mu
 					{
 						pDestBuf[BASE_CHANNEL_STRIDE * i + CHANNELS_TO_BLEND + c] = pBaseBuf[BASE_CHANNEL_STRIDE * i + CHANNELS_TO_BLEND + c];
 					}
-				});
-
+				}
+			});
 		}
 	}
 
@@ -1591,10 +1618,18 @@ namespace mu
 			PixelCount = pBase->CalculatePixelCount();
 		}
 
-		ParallelFor(PixelCount,
-			[
-				pBaseBuf, pBlendedBuf, BlendAlphaSourceChannel
-			] (uint32 i)
+		constexpr int32 NumBatchElems = 4096*2;
+		const int32 NumBatches = FMath::DivideAndRoundUp(PixelCount, NumBatchElems);
+
+		ParallelFor(NumBatches,
+		[
+			pBaseBuf, pBlendedBuf, BlendAlphaSourceChannel, PixelCount, NumBatchElems
+		] (uint32 BatchId)
+		{
+			const int32 BatchBegin = BatchId * NumBatchElems;
+			const int32 BatchEnd = FMath::Min(BatchBegin + NumBatchElems, PixelCount);
+
+			for (int32 i = BatchBegin; i < BatchEnd; ++i)
 			{
 				// TODO: Optimize this (SIMD?)
 				uint32 mask = pBlendedBuf[4 * i + 3];
@@ -1605,7 +1640,7 @@ namespace mu
 					uint32 base = pBaseBuf[4 * i + c];
 					uint32 blended = pBlendedBuf[4 * i + c];
 					uint32 result = RGB_FUNC_MASKED(base, blended, mask);
-					if (CLAMP)
+					if constexpr (CLAMP)
 					{
 						pBaseBuf[4 * i + c] = (uint8)FMath::Min(255u, result);
 					}
@@ -1620,7 +1655,7 @@ namespace mu
 					uint32 base = pBaseBuf[4 * i + 3];
 					uint32 blended = pBlendedBuf[4 * i + BlendAlphaSourceChannel];
 					uint32 result = A_FUNC(base, blended);
-					if (CLAMP)
+					if constexpr (CLAMP)
 					{
 						pBaseBuf[4 * i + 3] = (uint8)FMath::Min(255u, result);
 					}
@@ -1629,7 +1664,8 @@ namespace mu
 						pBaseBuf[4 * i + 3] = (uint8)result;
 					}
 				}
-			});
+			}
+		});
 	}
 
 	//---------------------------------------------------------------------------------------------
@@ -1663,53 +1699,51 @@ namespace mu
 		}
 
 		constexpr int32 NumBatchElems = 4096*2;
-	
-		const int32 NumBatches = (PixelCount / NumBatchElems) + static_cast<int32>((PixelCount % NumBatchElems) != 0);
-		const int32 LastBatchNumElems = (PixelCount % NumBatchElems) > 0 ? (PixelCount % NumBatchElems) : NumBatchElems;
+		const int32 NumBatches = FMath::DivideAndRoundUp(PixelCount, NumBatchElems);
 
 		ParallelFor(NumBatches,
-			[
-				pBaseBuf, pBlendedBuf, NumBatches, LastBatchNumElems
-			] (uint32 BatchId)
-			{
-				const int32 CurrentBatchNumElems = static_cast<int32>(BatchId) < (NumBatches - 1) ? NumBatchElems : LastBatchNumElems;
-				for (int32 BatchElemIndex = 0; BatchElemIndex < CurrentBatchNumElems; ++BatchElemIndex)
-				{
-					const int32 i = BatchId * NumBatchElems + BatchElemIndex;
+		[
+			pBaseBuf, pBlendedBuf, PixelCount, NumBatchElems
+		] (uint32 BatchId)
+		{
 
-					// TODO: Optimize this (SIMD?)
-					const int32 BaseAlpha = pBaseBuf[4 * i + 3];
-					const int32 BlendedAlpha = pBlendedBuf[4 * i + 3];
+			const int32 BatchBegin = BatchId * NumBatchElems;
+			const int32 BatchEnd = FMath::Min(BatchBegin + NumBatchElems, PixelCount);
+			for (int32 i = BatchBegin; i < BatchEnd; ++i)
+			{
+				// TODO: Optimize this (SIMD?)
+				const int32 BaseAlpha = pBaseBuf[4 * i + 3];
+				const int32 BlendedAlpha = pBlendedBuf[4 * i + 3];
 
-					const VectorRegister4Int Mask = VectorIntSet1(BlendedAlpha);
-					const VectorRegister4Int Blended = MakeVectorRegisterInt(pBlendedBuf[4 * i + 0], pBlendedBuf[4 * i + 1], pBlendedBuf[4 * i + 2], 0);
-					const VectorRegister4Int Base = MakeVectorRegisterInt(pBaseBuf[4 * i + 0], pBaseBuf[4 * i + 1], pBaseBuf[4 * i + 2], 0);
+				const VectorRegister4Int Mask = VectorIntSet1(BlendedAlpha);
+				const VectorRegister4Int Blended = MakeVectorRegisterInt(pBlendedBuf[4 * i + 0], pBlendedBuf[4 * i + 1], pBlendedBuf[4 * i + 2], 0);
+				const VectorRegister4Int Base = MakeVectorRegisterInt(pBaseBuf[4 * i + 0], pBaseBuf[4 * i + 1], pBaseBuf[4 * i + 2], 0);
 
-					VectorRegister4Int Result = RGB_FUNC_MASKED(Base, Blended, Mask);
-					if constexpr (CLAMP)
-					{
-						Result = VectorIntMin(VectorIntSet1(255), Result);
-					}
+				VectorRegister4Int Result = RGB_FUNC_MASKED(Base, Blended, Mask);
+				if constexpr (CLAMP)
+				{
+					Result = VectorIntMin(VectorIntSet1(255), Result);
+				}
 
-					int32 AlphaResult = A_FUNC(BaseAlpha, BlendedAlpha);
-					if constexpr (CLAMP)
-					{
-						AlphaResult = FMath::Min(255, AlphaResult);
-					}
+				int32 AlphaResult = A_FUNC(BaseAlpha, BlendedAlpha);
+				if constexpr (CLAMP)
+				{
+					AlphaResult = FMath::Min(255, AlphaResult);
+				}
 
-					struct alignas(alignof(VectorRegister4Int)) AlignedOutputRegister4Int
-					{
-						int32 Data[4];
-					} RegisterOutput;
+				struct alignas(alignof(VectorRegister4Int)) AlignedOutputRegister4Int
+				{
+					int32 Data[4];
+				} RegisterOutput;
 
-					VectorIntStoreAligned(Result, &RegisterOutput);
+				VectorIntStoreAligned(Result, &RegisterOutput);
 
-					pBaseBuf[4 * i + 0] = static_cast<uint8>(RegisterOutput.Data[0]);
-					pBaseBuf[4 * i + 1] = static_cast<uint8>(RegisterOutput.Data[1]);
-					pBaseBuf[4 * i + 2] = static_cast<uint8>(RegisterOutput.Data[2]);
-					pBaseBuf[4 * i + 3] = static_cast<uint8>(AlphaResult);
-				}
-			});
+				pBaseBuf[4 * i + 0] = static_cast<uint8>(RegisterOutput.Data[0]);
+				pBaseBuf[4 * i + 1] = static_cast<uint8>(RegisterOutput.Data[1]);
+				pBaseBuf[4 * i + 2] = static_cast<uint8>(RegisterOutput.Data[2]);
+				pBaseBuf[4 * i + 3] = static_cast<uint8>(AlphaResult);
+			}
+		});
 	}
 
 	//---------------------------------------------------------------------------------------------
diff --git a/Engine/Plugins/Experimental/Mutable/Source/MutableRuntime/Private/MuR/OpImageSwizzle.h b/Engine/Plugins/Experimental/Mutable/Source/MutableRuntime/Private/MuR/OpImageSwizzle.h
--- a/Engine/Plugins/Experimental/Mutable/Source/MutableRuntime/Private/MuR/OpImageSwizzle.h
+++ b/Engine/Plugins/Experimental/Mutable/Source/MutableRuntime/Private/MuR/OpImageSwizzle.h
@@ -28,7 +28,7 @@ namespace mu
                                     format );
 
 		// Very slow generic implementations
-        size_t pixelCount = pDest->CalculatePixelCount();
+        int32 PixelCount = pDest->CalculatePixelCount();
 
         // Pixelcount should already match, but due to bugs it may not be the case. Try to detect it,
         // but avoid crashing below:
@@ -38,12 +38,12 @@ namespace mu
             if (pSources[c])
             {
                 size_t sourcePixelCount = pSources[c]->CalculatePixelCount();
-                if (pixelCount>sourcePixelCount)
+                if (PixelCount>sourcePixelCount)
                 {
                     check(false);
 
                     // Something went wrong
-                    pixelCount = sourcePixelCount;
+					PixelCount = sourcePixelCount;
                 }
             }
         }
@@ -89,19 +89,29 @@ namespace mu
 
 			bool filled = false;
 
+			constexpr int32 NumBatchElems = 4096*2;
+			const int32 NumBatches = FMath::DivideAndRoundUp(PixelCount, NumBatchElems);
+
 			if (pSources[i])
 			{
 				const uint8* pSourceBuf = pSources[i]->GetData() + channels[i];
 
+
 				switch (pSources[i]->GetFormat())
 				{
 				case EImageFormat::IF_L_UBYTE:
 					if (channels[i] < 1)
 					{
 						//for (size_t p = 0; p < pixelCount; ++p)
-						ParallelFor(pixelCount,[pDestBuf, pSourceBuf, NumDestChannels](int p)
+						ParallelFor(NumBatches,[pDestBuf, pSourceBuf, NumDestChannels, PixelCount, NumBatchElems ](int32 BatchId)
 							{
-								pDestBuf[p * NumDestChannels] = pSourceBuf[p];
+								const int32 BatchBegin = BatchId * NumBatchElems;
+								const int32 BatchEnd = FMath::Min(BatchBegin + NumBatchElems, PixelCount);
+
+								for (int32 p = BatchBegin; p < BatchEnd; ++p)
+								{
+									pDestBuf[p * NumDestChannels] = pSourceBuf[p];
+								}
 							});
 
 						filled = true;
@@ -111,9 +121,16 @@ namespace mu
 				case EImageFormat::IF_RGB_UBYTE:
 					if (channels[i] < 3)
 					{
-						ParallelFor(pixelCount, [pDestBuf, pSourceBuf, NumDestChannels](int p)
+
+						ParallelFor(NumBatches, [pDestBuf, pSourceBuf, NumDestChannels, PixelCount, NumBatchElems](int32 BatchId)
 							{
-								pDestBuf[p * NumDestChannels] = pSourceBuf[p * 3];
+								const int32 BatchBegin = BatchId * NumBatchElems;
+								const int32 BatchEnd = FMath::Min(BatchBegin + NumBatchElems, PixelCount);
+
+								for (int32 p = BatchBegin; p < BatchEnd; ++p)
+								{
+									pDestBuf[p * NumDestChannels] = pSourceBuf[p * 3];
+								}
 							});
 						filled = true;
 					}
@@ -122,9 +139,15 @@ namespace mu
 				case EImageFormat::IF_RGBA_UBYTE:
 					if (channels[i] < 4)
 					{
-						ParallelFor(pixelCount, [pDestBuf, pSourceBuf, NumDestChannels](int p)
+						ParallelFor(NumBatches, [pDestBuf, pSourceBuf, NumDestChannels, PixelCount, NumBatchElems](int32 BatchId)
 							{
-								pDestBuf[p * NumDestChannels] = pSourceBuf[p * 4];
+								const int32 BatchBegin = BatchId * NumBatchElems;
+								const int32 BatchEnd = FMath::Min(BatchBegin + NumBatchElems, PixelCount);
+
+								for (int32 p = BatchBegin; p < BatchEnd; ++p)
+								{
+									pDestBuf[p * NumDestChannels] = pSourceBuf[p * 4];
+								}
 							});
 						filled = true;
 					}
@@ -141,9 +164,14 @@ namespace mu
 					}
 					if (channels[i] < 4)
 					{
-						ParallelFor(pixelCount, [pDestBuf, pSourceBuf, NumDestChannels](int p)
+						ParallelFor(NumBatches, [pDestBuf, pSourceBuf, NumDestChannels, PixelCount, NumBatchElems](int32 BatchId)
 							{
-								pDestBuf[p * NumDestChannels] = pSourceBuf[p * 4];
+								const int32 BatchBegin = BatchId * NumBatchElems;
+								const int32 BatchEnd = FMath::Min(BatchBegin + NumBatchElems, PixelCount);
+								for (int32 p = BatchBegin; p < BatchEnd; ++p)
+								{
+									pDestBuf[p * NumDestChannels] = pSourceBuf[p * 4];
+								}
 							});
 						filled = true;
 					}
@@ -157,9 +185,14 @@ namespace mu
 			if (!filled)
 			{
 				// Source not set. Clear to 0
-				ParallelFor(pixelCount, [pDestBuf, NumDestChannels](int p)
+				ParallelFor(NumBatches, [pDestBuf, NumDestChannels, PixelCount, NumBatchElems](int32 BatchId)
 					{
-						pDestBuf[p * NumDestChannels] = 0;
+						const int32 BatchBegin = BatchId * NumBatchElems;
+						const int32 BatchEnd = FMath::Min(BatchBegin + NumBatchElems, PixelCount);
+						for (int32 p = BatchBegin; p < BatchEnd; ++p)
+						{
+							pDestBuf[p * NumDestChannels] = 0;
+						}
 					});
 			}
 		}
