[Mutable] Ajusted Mutable tasks priorities to avoid hitches.

#rnx

diff --git a/Engine/Plugins/Experimental/Mutable/Source/CustomizableObject/Private/MuCO/CustomizableObjectSystemPrivate.h b/Engine/Plugins/Experimental/Mutable/Source/CustomizableObject/Private/MuCO/CustomizableObjectSystemPrivate.h
--- a/Engine/Plugins/Experimental/Mutable/Source/CustomizableObject/Private/MuCO/CustomizableObjectSystemPrivate.h
+++ b/Engine/Plugins/Experimental/Mutable/Source/CustomizableObject/Private/MuCO/CustomizableObjectSystemPrivate.h
@@ -600,7 +600,7 @@ public:
 			MoveTemp(TaskBody),
 			TStatId{},
 			&Prerequisites, 
-			ENamedThreads::AnyThread);
+			ENamedThreads::AnyBackgroundHiPriTask);
 		LastMutableTask->SetDebugName(DebugName);
 
 		return LastMutableTask;
diff --git a/Engine/Plugins/Experimental/Mutable/Source/MutableRuntime/Private/MuR/CodeRunnerTasked.cpp b/Engine/Plugins/Experimental/Mutable/Source/MutableRuntime/Private/MuR/CodeRunnerTasked.cpp
--- a/Engine/Plugins/Experimental/Mutable/Source/MutableRuntime/Private/MuR/CodeRunnerTasked.cpp
+++ b/Engine/Plugins/Experimental/Mutable/Source/MutableRuntime/Private/MuR/CodeRunnerTasked.cpp
@@ -47,12 +47,7 @@
 DECLARE_CYCLE_STAT(TEXT("MutableCoreTask"), STAT_MutableCoreTask, STATGROUP_Game);
 
 
-int32 CoreRunnerTaskPriority =
-#ifdef PLATFORM_SWITCH // Required due to low end platforms hitches. Mutable is occupying all threads with long tasks preventing Niagara from executing its time sensitive tasks.
-		5;
-#else
-		0;
-#endif
+int32 CoreRunnerTaskPriority = 5;
 
 static FAutoConsoleVariableRef CVarCoreRunnerTaskPriority(
 	TEXT("mutable.CoreRunnerTaskPriority"),
