[mutable] Turn plain images into plain-color-image operations to save constant loading operations at runtime.
Cleanup "state runtime parameters" limit message.
#rnx

diff --git a/Engine/Plugins/Experimental/Mutable/Source/MutableRuntime/Private/MuR/Image.cpp b/Engine/Plugins/Experimental/Mutable/Source/MutableRuntime/Private/MuR/Image.cpp
--- a/Engine/Plugins/Experimental/Mutable/Source/MutableRuntime/Private/MuR/Image.cpp
+++ b/Engine/Plugins/Experimental/Mutable/Source/MutableRuntime/Private/MuR/Image.cpp
@@ -588,26 +588,86 @@ namespace mu
         {
             switch( m_format )
             {
-            case EImageFormat::IF_L_UBYTE:
-            {
-                const uint8_t* pData = &m_data[0];
-                uint8_t v = pData[0];
+			case EImageFormat::IF_L_UBYTE:
+			{
+				const uint8* pData = m_data.GetData();
+				uint8 v = pData[0];
 
-                for ( int p=0; res && p<pixelCount; ++p )
-                {
-                    uint8_t nv = *pData++;
-                    res &= ( nv==v );
-                }
+				for (int p = 0; res && p < pixelCount; ++p)
+				{
+					uint8 nv = *pData++;
+					res &= (nv == v);
+				}
 
-                if (res)
-                {
-                    colour[0] = colour[1] = colour[2] = float(v)/255.0f;
-                    colour[3] = 1.0f;
-                }
-                break;
-            }
+				if (res)
+				{
+					colour[0] = colour[1] = colour[2] = float(v) / 255.0f;
+					colour[3] = 1.0f;
+				}
+				break;
+			}
 
-            default:
+			case EImageFormat::IF_RGB_UBYTE:
+			{
+				const uint8* pData = m_data.GetData();
+				uint8 r = pData[0];
+				uint8 g = pData[1];
+				uint8 b = pData[2];
+
+				for (int p = 0; res && p < pixelCount; ++p)
+				{
+					uint8 nr = *pData++;
+					uint8 ng = *pData++;
+					uint8 nb = *pData++;
+					res &= (nr == r) && (ng==g) && (nb==b);
+				}
+
+				if (res)
+				{
+					colour[0] = r / 255.0f;
+					colour[0] = g / 255.0f;
+					colour[0] = b / 255.0f;
+					colour[3] = 1.0f;
+				}
+				break;
+			}
+
+			case EImageFormat::IF_RGBA_UBYTE:
+			case EImageFormat::IF_BGRA_UBYTE:
+			{
+				const uint32* pData = reinterpret_cast<const uint32*>(m_data.GetData());
+				uint32 v = pData[0];
+
+				for (int p = 0; res && p < pixelCount; ++p)
+				{
+					uint32 nv = *pData++;
+					res &= (nv == v);
+				}
+
+				if (res)
+				{
+					const uint8* pByteData = m_data.GetData();
+					if (m_format == EImageFormat::IF_RGBA_UBYTE)
+					{
+						colour[0] = float(pByteData[0]) / 255.0f;
+						colour[1] = float(pByteData[1]) / 255.0f;
+						colour[2] = float(pByteData[2]) / 255.0f;
+					}
+					else
+					{
+						colour[0] = float(pByteData[2]) / 255.0f;
+						colour[1] = float(pByteData[1]) / 255.0f;
+						colour[2] = float(pByteData[0]) / 255.0f;
+					}
+					colour[3] = float(pByteData[3]) / 255.0f;
+				}
+				break;
+			}
+
+			// TODO: Other formats could also be implemented. For compressed types,
+			// the compressed block could be compared and only uncompress if all are the same to
+			// check if all pixels in the block are also equal.
+			default:
                 res = false;
                 break;
             }
diff --git a/Engine/Plugins/Experimental/Mutable/Source/MutableRuntime/Private/MuR/ModelPrivate.h b/Engine/Plugins/Experimental/Mutable/Source/MutableRuntime/Private/MuR/ModelPrivate.h
--- a/Engine/Plugins/Experimental/Mutable/Source/MutableRuntime/Private/MuR/ModelPrivate.h
+++ b/Engine/Plugins/Experimental/Mutable/Source/MutableRuntime/Private/MuR/ModelPrivate.h
@@ -13,6 +13,7 @@
 #include "MuR/OpImageResize.h"
 #include "MuR/OpImageMipmap.h"
 
+#define MUTABLE_MAX_RUNTIME_PARAMETERS_PER_STATE	64
 
 namespace mu
 {
@@ -92,6 +93,7 @@ namespace mu
             //! parameters of this state, with a mask of relevant runtime parameters.
             //! The mask has a bit on for every runtime parameter in the m_runtimeParameters
             //! vector.
+			//! The uint64 is linked to MUTABLE_MAX_RUNTIME_PARAMETERS_PER_STATE
 			TArray< TPair<OP::ADDRESS,uint64> > m_dynamicResources;
 
             //!
@@ -114,16 +116,19 @@ namespace mu
                 arch >> m_dynamicResources;
             }
 
-            //! Returns the parameters mask
+            /** Returns the mask of parameters (from the runtime parameter list of this state) including the parameters that 
+			* are relevant for the dynamic resource at the given address.
+			*/
             uint64 IsDynamic( OP::ADDRESS at ) const
             {
                 uint64 res = 0;
 
-                for ( int32 i=0; !res && i<m_dynamicResources.Num(); ++i )
+                for ( int32 i=0; i<m_dynamicResources.Num(); ++i )
                 {
                     if ( m_dynamicResources[i].Key==at )
                     {
                         res = m_dynamicResources[i].Value;
+						break;
                     }
                 }
 
diff --git a/Engine/Plugins/Experimental/Mutable/Source/MutableTools/Private/MuT/CodeGenerator_FirstPass.cpp b/Engine/Plugins/Experimental/Mutable/Source/MutableTools/Private/MuT/CodeGenerator_FirstPass.cpp
--- a/Engine/Plugins/Experimental/Mutable/Source/MutableTools/Private/MuT/CodeGenerator_FirstPass.cpp
+++ b/Engine/Plugins/Experimental/Mutable/Source/MutableTools/Private/MuT/CodeGenerator_FirstPass.cpp
@@ -481,14 +481,12 @@ namespace mu
         {
             m_states.Emplace( s, &node );
 
-            // \todo move this 32 constant to a macro or constexpr with a meaningful name and location
-            if ( s.m_runtimeParams.Num() > 32 )
+            if ( s.m_runtimeParams.Num() > MUTABLE_MAX_RUNTIME_PARAMETERS_PER_STATE )
             {
-                char temp[256];
-                mutable_snprintf( temp, 256,
-                                  "State [%s] has more than 32 runtime parameters. Their update may fail.",
+                FString Msg = FString::Printf( TEXT("State [%s] has more than %d runtime parameters. Their update may fail."),
+								  MUTABLE_MAX_RUNTIME_PARAMETERS_PER_STATE,
                                   s.m_name.c_str() );
-                m_pErrorLog->GetPrivate()->Add( temp, ELMT_ERROR, node.m_errorContext );
+                m_pErrorLog->GetPrivate()->Add(Msg, ELMT_ERROR, node.m_errorContext );
             }
         }
 
diff --git a/Engine/Plugins/Experimental/Mutable/Source/MutableTools/Private/MuT/Compiler.cpp b/Engine/Plugins/Experimental/Mutable/Source/MutableTools/Private/MuT/Compiler.cpp
--- a/Engine/Plugins/Experimental/Mutable/Source/MutableTools/Private/MuT/Compiler.cpp
+++ b/Engine/Plugins/Experimental/Mutable/Source/MutableTools/Private/MuT/Compiler.cpp
@@ -283,7 +283,7 @@ namespace mu
             // Generate the mask of dynamic resources
             for ( const auto& a: s.m_dynamicResources )
             {
-                uint64_t relevantMask = 0;
+                uint64 relevantMask = 0;
                 for ( const auto& b: a.Value )
                 {
                     // Find the index in the model parameter list
@@ -298,11 +298,11 @@ namespace mu
                     check(paramIndex>=0);
 
                     // Find the position in the state data vector.
-                    int32 it = s.state.m_runtimeParameters.Find( paramIndex );
+                    int32 IndexInRuntimeList = s.state.m_runtimeParameters.Find( paramIndex );
 
-                    if ( it!=INDEX_NONE )
+                    if (IndexInRuntimeList !=INDEX_NONE )
                     {
-                        relevantMask |= uint64_t(1) << it;
+                        relevantMask |= uint64(1) << IndexInRuntimeList;
                     }
                 }
 
diff --git a/Engine/Plugins/Experimental/Mutable/Source/MutableTools/Private/MuT/DataPacker.cpp b/Engine/Plugins/Experimental/Mutable/Source/MutableTools/Private/MuT/DataPacker.cpp
--- a/Engine/Plugins/Experimental/Mutable/Source/MutableTools/Private/MuT/DataPacker.cpp
+++ b/Engine/Plugins/Experimental/Mutable/Source/MutableTools/Private/MuT/DataPacker.cpp
@@ -738,7 +738,6 @@ namespace mu
         AccumulateImageFormatsAST accFormat;
         accFormat.Run( roots );
 
-		// TEMP: Disabled for RLE bug
         // See if we can convert some constants to more efficient formats
         ASTOp::Traverse_BottomUp_Unique_NonReentrant( roots, [&](Ptr<ASTOp>& n)
         {
@@ -746,10 +745,33 @@ namespace mu
             {
 				ASTOpConstantResource* typed = dynamic_cast<ASTOpConstantResource*>(n.get());
                 Ptr<const Image> pOld = static_cast<const Image*>(typed->GetValue().get());
-                if ( accFormat.m_supportedFormats[typed][(size_t)EImageFormat::IF_L_UBIT_RLE] )
+
+				// See if there is a better format for this image
+				FVector4f PlainColor;
+				if ( pOld->IsPlainColour(PlainColor) )
+				{
+					// It is more efficient to just have an instruction for it instead, to avoid the overhead
+					// of data loading.
+					Ptr<ASTOpFixed> NewColor = new ASTOpFixed;
+					NewColor->op.type = OP_TYPE::CO_CONSTANT;
+					NewColor->op.args.ColourConstant.value[0] = PlainColor[0];
+					NewColor->op.args.ColourConstant.value[1] = PlainColor[1];
+					NewColor->op.args.ColourConstant.value[2] = PlainColor[2];
+					NewColor->op.args.ColourConstant.value[3] = PlainColor[3];
+
+					Ptr<ASTOpFixed> NewPlain = new ASTOpFixed;
+					NewPlain->op.type = OP_TYPE::IM_PLAINCOLOUR;
+					NewPlain->SetChild( NewPlain->op.args.ImagePlainColour.colour, NewColor );
+					NewPlain->op.args.ImagePlainColour.format = pOld->GetFormat();
+					NewPlain->op.args.ImagePlainColour.size[0] = pOld->GetSizeX();
+					NewPlain->op.args.ImagePlainColour.size[1] = pOld->GetSizeY();
+
+					ASTOp::Replace(n, NewPlain);
+
+				}
+				else if ( accFormat.m_supportedFormats[typed][(size_t)EImageFormat::IF_L_UBIT_RLE] )
                 {
-                    ImagePtr pNew =
-                        ImagePixelFormat( imageCompressionQuality, pOld.get(), EImageFormat::IF_L_UBIT_RLE );
+                    ImagePtr pNew = ImagePixelFormat( imageCompressionQuality, pOld.get(), EImageFormat::IF_L_UBIT_RLE );
 
                     // Only replace if the compression was worth!
                     size_t oldSize = pOld->GetDataSize();
@@ -761,8 +783,7 @@ namespace mu
                 }
                 else if ( accFormat.m_supportedFormats[typed][(size_t)EImageFormat::IF_L_UBYTE_RLE] )
                 {
-                    ImagePtr pNew =
-                        ImagePixelFormat( imageCompressionQuality, pOld.get(), EImageFormat::IF_L_UBYTE_RLE );
+                    ImagePtr pNew = ImagePixelFormat( imageCompressionQuality, pOld.get(), EImageFormat::IF_L_UBYTE_RLE );
 
                     // Only replace if the compression was worth!
                     size_t oldSize = pOld->GetDataSize();
