[Mutable] Enable full saturation image op optimization.

#rnx

diff --git a/Engine/Plugins/Experimental/Mutable/Source/MutableRuntime/Private/MuR/CodeRunnerTasked.cpp b/Engine/Plugins/Experimental/Mutable/Source/MutableRuntime/Private/MuR/CodeRunnerTasked.cpp
--- a/Engine/Plugins/Experimental/Mutable/Source/MutableRuntime/Private/MuR/CodeRunnerTasked.cpp
+++ b/Engine/Plugins/Experimental/Mutable/Source/MutableRuntime/Private/MuR/CodeRunnerTasked.cpp
@@ -1142,8 +1142,7 @@ namespace mu
 	{
 		MUTABLE_CPUPROFILER_SCOPE(FImageSaturateTask);
 
-		// For now don't enable this optimization 
-		const bool bOptimizeUnchanged = false;//FMath::IsNearlyEqual(Factor, 1.0f);
+		const bool bOptimizeUnchanged = FMath::IsNearlyEqual(Factor, 1.0f);
 
 		if (Source->IsUnique())
 		{
