[Mutable] Add support for BC4 to L_UBYTE image pixel format.

#rnx

diff --git a/Engine/Plugins/Experimental/Mutable/Source/MutableRuntime/Private/MuR/OpImagePixelFormat.cpp b/Engine/Plugins/Experimental/Mutable/Source/MutableRuntime/Private/MuR/OpImagePixelFormat.cpp
--- a/Engine/Plugins/Experimental/Mutable/Source/MutableRuntime/Private/MuR/OpImagePixelFormat.cpp
+++ b/Engine/Plugins/Experimental/Mutable/Source/MutableRuntime/Private/MuR/OpImagePixelFormat.cpp
@@ -257,6 +257,15 @@ bool ImagePixelFormatInPlace( int imageCompressionQuality, Image* pResult, const
                 }
                 break;
             }
+			case EImageFormat::IF_BC4:
+			{
+                for ( int m=0; m<resultLODCount; ++m )
+                {
+					FIntVector2 mipSize = pResult->CalculateMipSize( m );
+                    miro::BC4_to_L( mipSize[0], mipSize[1], pBase->GetMipData(baseLOD+m), pResult->GetMipData(m) );
+                }
+                break;
+			}
 
             default:
                 // Case not implemented
