#rnx

diff --git a/Engine/Plugins/Experimental/Mutable/Source/CustomizableObject/Private/MuCO/UnrealMutableImageProvider.cpp b/Engine/Plugins/Experimental/Mutable/Source/CustomizableObject/Private/MuCO/UnrealMutableImageProvider.cpp
--- a/Engine/Plugins/Experimental/Mutable/Source/CustomizableObject/Private/MuCO/UnrealMutableImageProvider.cpp
+++ b/Engine/Plugins/Experimental/Mutable/Source/CustomizableObject/Private/MuCO/UnrealMutableImageProvider.cpp
@@ -536,7 +536,12 @@ mu::ImagePtr FUnrealMutableImageProvider::CreateDummy()
 	const int size = DUMMY_IMAGE_DESC.m_size[0];
 	const int checkerSize = 4;
 	constexpr int checkerTileCount = 2;
-	uint8_t colours[checkerTileCount][4] = { { 255,255,0,255 },{ 0,0,255,255 } };
+	
+#if !UE_BUILD_SHIPPING
+	uint8_t colours[checkerTileCount][4] = { { 255, 255, 0, 255 },{ 0, 0, 255, 255 } };
+#else
+	uint8_t colours[checkerTileCount][4] = { { 255, 255, 0, 0 },  { 0, 0, 255, 0 } };
+#endif
 
 	mu::ImagePtr pResult = new mu::Image(size, size, DUMMY_IMAGE_DESC.m_lods, DUMMY_IMAGE_DESC.m_format);
 
