[Mutable] Add support for BC4 to RGB[A] formats conversion.

#rnx

diff --git a/Engine/Plugins/Experimental/Mutable/Source/MutableRuntime/Private/MuR/OpImagePixelFormat.cpp b/Engine/Plugins/Experimental/Mutable/Source/MutableRuntime/Private/MuR/OpImagePixelFormat.cpp
--- a/Engine/Plugins/Experimental/Mutable/Source/MutableRuntime/Private/MuR/OpImagePixelFormat.cpp
+++ b/Engine/Plugins/Experimental/Mutable/Source/MutableRuntime/Private/MuR/OpImagePixelFormat.cpp
@@ -1154,14 +1154,32 @@ bool ImagePixelFormatInPlace( int imageCompressionQuality, Image* pResult, const
             switch ( pBase->GetFormat() )
             {
             case EImageFormat::IF_L_UBYTE:
-                for ( int m=0; m<resultLODCount; ++m )
+			{
+				for (int m = 0; m < resultLODCount; ++m)
+				{
+					FIntVector2 mipSize = pResult->CalculateMipSize(m);
+					miro::L_to_BC4(
+						mipSize[0], mipSize[1], pBase->GetMipData(baseLOD + m),
+						pResult->GetMipData(m), imageCompressionQuality);
+				}
+				break;
+			}
+			case EImageFormat::IF_RGB_UBYTE:
+			case EImageFormat::IF_RGBA_UBYTE:
+			case EImageFormat::IF_BGRA_UBYTE:
+			{	
+                ImagePtr pTempBase =
+                        ImagePixelFormat( imageCompressionQuality, pBase, EImageFormat::IF_L_UBYTE );
+                for (int m = 0; m < resultLODCount; ++m)
                 {
-                    FIntVector2 mipSize = pResult->CalculateMipSize( m );
+                    FIntVector2 mipSize = pTempBase->CalculateMipSize( m );
+
                     miro::L_to_BC4(
                         mipSize[0], mipSize[1], pBase->GetMipData( baseLOD + m ),
                         pResult->GetMipData( m ), imageCompressionQuality );
                 }
-                break;
+				break;
+			}
 
             default:
                 // Case not implemented
