[Mutable] Fix DefaultImageProvider running out of Texture Paramater ids.

#rnx

diff --git a/Engine/Plugins/Experimental/Mutable/Source/CustomizableObject/Private/MuCO/CustomizableObjectInstanceDescriptor.cpp b/Engine/Plugins/Experimental/Mutable/Source/CustomizableObject/Private/MuCO/CustomizableObjectInstanceDescriptor.cpp
--- a/Engine/Plugins/Experimental/Mutable/Source/CustomizableObject/Private/MuCO/CustomizableObjectInstanceDescriptor.cpp
+++ b/Engine/Plugins/Experimental/Mutable/Source/CustomizableObject/Private/MuCO/CustomizableObjectInstanceDescriptor.cpp
@@ -990,7 +990,13 @@ void FCustomizableObjectInstanceDescriptor::ReloadParameters()
 
 						if (!Param.ParameterRangeValues.IsValidIndex(ValueIndex))
 						{
-							Param.ParameterRangeValues.AddDefaulted(ValueIndex + 1 - Param.ParameterRangeValues.Num());
+							const int32 PreviousNum = Param.ParameterRangeValues.Num();
+							Param.ParameterRangeValues.AddUninitialized(ValueIndex + 1 - Param.ParameterRangeValues.Num());
+
+							for (int32 Index = PreviousNum; Index < Param.ParameterRangeValues.Num(); ++Index)
+							{
+								Param.ParameterRangeValues[Index] = FCustomizableObjectTextureParameterValue::DEFAULT_PARAMETER_VALUE;
+							}
 						}
 
 						Param.ParameterRangeValues[RangeIndex] = MutableParameters->GetImageValue(ParamIndex, RangeValueIdxPtr);
diff --git a/Engine/Plugins/Experimental/Mutable/Source/CustomizableObject/Private/MuCO/DefaultImageProvider.cpp b/Engine/Plugins/Experimental/Mutable/Source/CustomizableObject/Private/MuCO/DefaultImageProvider.cpp
--- a/Engine/Plugins/Experimental/Mutable/Source/CustomizableObject/Private/MuCO/DefaultImageProvider.cpp
+++ b/Engine/Plugins/Experimental/Mutable/Source/CustomizableObject/Private/MuCO/DefaultImageProvider.cpp
@@ -89,17 +89,24 @@ int64 UDefaultImageProvider::Get(UTexture2D* Texture) const
 
 uint64 UDefaultImageProvider::GetOrAdd(UTexture2D* Texture)
 {
+	if (!Texture)
+	{
+		return FCustomizableObjectTextureParameterValue::DEFAULT_PARAMETER_VALUE;
+	}
+	
 	int32 Hole = INDEX_NONE;
 
 	for (int32 TextureIndex = 0; TextureIndex < Textures.Num(); ++TextureIndex)
 	{
-		if (Textures[TextureIndex] == Texture)
+		const TObjectPtr<UTexture2D> UsedTexture = Textures[TextureIndex];
+		if (UsedTexture == Texture)
 		{
 			return ToTextureId(TextureIndex);
 		}
-		else if (!Texture && Hole == INDEX_NONE)
+		else if (!UsedTexture)
 		{
 			Hole = TextureIndex;
+			break;
 		}
 	}
 	
diff --git a/Engine/Plugins/Experimental/Mutable/Source/CustomizableObject/Public/MuCO/CustomizableObjectParameterTypeDefinitions.h b/Engine/Plugins/Experimental/Mutable/Source/CustomizableObject/Public/MuCO/CustomizableObjectParameterTypeDefinitions.h
--- a/Engine/Plugins/Experimental/Mutable/Source/CustomizableObject/Public/MuCO/CustomizableObjectParameterTypeDefinitions.h
+++ b/Engine/Plugins/Experimental/Mutable/Source/CustomizableObject/Public/MuCO/CustomizableObjectParameterTypeDefinitions.h
@@ -147,13 +147,15 @@ inline uint32 GetTypeHash(const FCustomizableObjectFloatParameterValue& Key)
 USTRUCT(BlueprintType)
 struct FCustomizableObjectTextureParameterValue
 {
+	static constexpr uint64 DEFAULT_PARAMETER_VALUE = 0;
+	
 	GENERATED_USTRUCT_BODY()
 
 	UPROPERTY(Category = CustomizableObjectTextureParameterValue, VisibleAnywhere)
 	FString ParameterName;
 
 	UPROPERTY(Category = CustomizableObjectTextureParameterValue, VisibleAnywhere)
-	uint64 ParameterValue = 0;
+	uint64 ParameterValue = DEFAULT_PARAMETER_VALUE;
 
 	UPROPERTY(Category = CustomizableObjectTextureParameterValue, VisibleAnywhere)
 	FString Uid;
diff --git a/Engine/Plugins/Experimental/Mutable/Source/CustomizableObjectEditor/Private/MuCOE/SCustomizableInstanceProperties.cpp b/Engine/Plugins/Experimental/Mutable/Source/CustomizableObjectEditor/Private/MuCOE/SCustomizableInstanceProperties.cpp
--- a/Engine/Plugins/Experimental/Mutable/Source/CustomizableObjectEditor/Private/MuCOE/SCustomizableInstanceProperties.cpp
+++ b/Engine/Plugins/Experimental/Mutable/Source/CustomizableObjectEditor/Private/MuCOE/SCustomizableInstanceProperties.cpp
@@ -77,7 +77,7 @@ void SCustomizableInstanceProperties::Construct(const FArguments& InArgs)
 
 		// Store the texture parameters data required for the ui.
 		TextureParameterValueNames.Add(MakeShareable(new FString("None")));
-		TextureParameterValues.AddDefaulted();
+		TextureParameterValues.Add(FCustomizableObjectTextureParameterValue::DEFAULT_PARAMETER_VALUE);
 		TArray<FCustomizableObjectExternalTexture> Textures = UCustomizableObjectSystem::GetInstance()->GetTextureParameterValues();
 		for (int i = 0; i < Textures.Num(); ++i)
 		{
