[mutable] Avoid some unnecessary buffer clearing after image creation.
Fix the image colourmap operations when mip number differs.
#rnx

diff --git a/Engine/Plugins/Experimental/Mutable/Source/MutableRuntime/Private/MuR/CodeRunner.cpp b/Engine/Plugins/Experimental/Mutable/Source/MutableRuntime/Private/MuR/CodeRunner.cpp
--- a/Engine/Plugins/Experimental/Mutable/Source/MutableRuntime/Private/MuR/CodeRunner.cpp
+++ b/Engine/Plugins/Experimental/Mutable/Source/MutableRuntime/Private/MuR/CodeRunner.cpp
@@ -3532,6 +3532,8 @@ namespace mu
                 Ptr<const Image> pMask = GetMemory().GetImage( FCacheAddress(args.mask,item) );
                 Ptr<const Image> pMap = GetMemory().GetImage( FCacheAddress(args.map,item) );
 
+				bool bOnlyOneMip = (pMask->GetLODCount() < pSource->GetLODCount());
+
 				// Be defensive: ensure image sizes match.
 				if (pMask->GetSize() != pSource->GetSize())
 				{
@@ -3540,7 +3542,14 @@ namespace mu
 				}
 
 
-                ImagePtr pResult = ImageColourMap( pSource.get(), pMask.get(), pMap.get() );
+                Ptr<Image> pResult = ImageColourMap( pSource.get(), pMask.get(), pMap.get(), bOnlyOneMip);
+
+				if (bOnlyOneMip)
+				{
+					MUTABLE_CPUPROFILER_SCOPE(ImageColourMap_MipFix);
+					FMipmapGenerationSettings DummyMipSettings{};
+					ImageMipmapInPlace(m_pSettings->GetPrivate()->m_imageCompressionQuality, pResult.get(), DummyMipSettings);
+				}
 
                 GetMemory().SetImage( item, pResult );
                 break;
@@ -3668,7 +3677,8 @@ namespace mu
                 ImagePtr pA = new Image( SizeX,
                                          SizeY,
                                          1,
-                                         (EImageFormat)args.format );
+                                         (EImageFormat)args.format,
+										 EInitializationType::NotInitialized );
 
                 FillPlainColourImage(pA.get(), c);
 
diff --git a/Engine/Plugins/Experimental/Mutable/Source/MutableRuntime/Private/MuR/CodeRunnerTasked.cpp b/Engine/Plugins/Experimental/Mutable/Source/MutableRuntime/Private/MuR/CodeRunnerTasked.cpp
--- a/Engine/Plugins/Experimental/Mutable/Source/MutableRuntime/Private/MuR/CodeRunnerTasked.cpp
+++ b/Engine/Plugins/Experimental/Mutable/Source/MutableRuntime/Private/MuR/CodeRunnerTasked.cpp
@@ -1019,7 +1019,7 @@ namespace mu
 		int startLevel = m_base->GetLODCount();
 		levelCount = FMath::Max(startLevel, levelCount);
 
-		ImagePtr pDest = new Image(m_base->GetSizeX(), m_base->GetSizeY(), levelCount, m_base->GetFormat());
+		ImagePtr pDest = new Image(m_base->GetSizeX(), m_base->GetSizeY(), levelCount, m_base->GetFormat(),EInitializationType::NotInitialized);
 		pDest->m_flags = m_base->m_flags;
 
 		SCRATCH_IMAGE_MIPMAP scratch;
diff --git a/Engine/Plugins/Experimental/Mutable/Source/MutableRuntime/Private/MuR/Image.cpp b/Engine/Plugins/Experimental/Mutable/Source/MutableRuntime/Private/MuR/Image.cpp
--- a/Engine/Plugins/Experimental/Mutable/Source/MutableRuntime/Private/MuR/Image.cpp
+++ b/Engine/Plugins/Experimental/Mutable/Source/MutableRuntime/Private/MuR/Image.cpp
@@ -70,7 +70,7 @@ namespace mu
 
 
     //---------------------------------------------------------------------------------------------
-    Image::Image( uint32_t sizeX, uint32_t sizeY, uint32_t lods, EImageFormat format )
+    Image::Image( uint32_t sizeX, uint32_t sizeY, uint32_t lods, EImageFormat format, EInitializationType Init )
     {
 		MUTABLE_CPUPROFILER_SCOPE(NewImage)
 			LLM_SCOPE_BYNAME(TEXT("MutableRuntime"));
@@ -88,7 +88,14 @@ namespace mu
         if (PixelsPerBlock)
         {
 			int32 DataSize = CalculateDataSize();
-            m_data.SetNum( DataSize );
+			if (Init == EInitializationType::Black)
+			{
+				m_data.SetNumZeroed(DataSize);
+			}
+			else
+			{
+				m_data.SetNumUninitialized(DataSize);
+			}
         }
     }
 
diff --git a/Engine/Plugins/Experimental/Mutable/Source/MutableRuntime/Private/MuR/OpImageColourMap.h b/Engine/Plugins/Experimental/Mutable/Source/MutableRuntime/Private/MuR/OpImageColourMap.h
--- a/Engine/Plugins/Experimental/Mutable/Source/MutableRuntime/Private/MuR/OpImageColourMap.h
+++ b/Engine/Plugins/Experimental/Mutable/Source/MutableRuntime/Private/MuR/OpImageColourMap.h
@@ -8,23 +8,31 @@
 namespace mu
 {
 
-	inline ImagePtr ImageColourMap( const Image* pSource, const Image* pMask, const Image* pMap )
+	inline Ptr<Image> ImageColourMap( const Image* pSource, const Image* pMask, const Image* pMap, bool bOnlyOneMip )
 	{
 		check( pSource->GetSizeX() == pMask->GetSizeX() );
 		check( pSource->GetSizeY() == pMask->GetSizeY() );
-        check( pMask->GetFormat() == EImageFormat::IF_L_UBYTE );
+		check( pSource->GetLODCount() == pMask->GetLODCount() || bOnlyOneMip );
+		check( pMask->GetFormat() == EImageFormat::IF_L_UBYTE );
 
-		ImagePtr pDest = new Image( pSource->GetSizeX(),
-									pSource->GetSizeY(),
-                                    pSource->GetLODCount(),
-									pSource->GetFormat() );
+		Ptr<Image> pDest = new Image( pSource->GetSizeX(), pSource->GetSizeY(), pSource->GetLODCount(),
+									  pSource->GetFormat(),
+									  EInitializationType::NotInitialized );
 
         uint8* pDestBuf = pDest->GetData();
         const uint8* pSourceBuf = pSource->GetData();
         const uint8* pMaskBuf = pMask->GetData();
 
 		// Generic implementation
-		int pixelCount = (int)pSource->CalculatePixelCount();
+		int32 PixelCount = 0;
+		if (bOnlyOneMip)
+		{
+			PixelCount = pSource->CalculatePixelCount(0);
+		}
+		else
+		{
+			PixelCount = pSource->CalculatePixelCount();
+		}
 
 		// Make a palette for faster conversion
         uint8 palette[256][4];
@@ -41,7 +49,7 @@ namespace mu
 		{
 		case EImageFormat::IF_L_UBYTE:
 		{
-			for ( int i=0; i<pixelCount; ++i )
+			for ( int i=0; i< PixelCount; ++i )
 			{
                 uint8 m_8 = pMaskBuf[i];
 				if (m_8>127)
@@ -58,7 +66,7 @@ namespace mu
 
 		case EImageFormat::IF_RGB_UBYTE:
 		{
-			for ( int i=0; i<pixelCount; ++i )
+			for ( int i=0; i< PixelCount; ++i )
 			{
                 uint8 m_8 = pMaskBuf[i];
 				if (m_8>127)
@@ -79,7 +87,7 @@ namespace mu
 
         case EImageFormat::IF_RGBA_UBYTE:
         {
-            for ( int i=0; i<pixelCount; ++i )
+            for ( int i=0; i< PixelCount; ++i )
             {
                 uint8 m_8 = pMaskBuf[i];
                 if (m_8>127)
@@ -102,7 +110,7 @@ namespace mu
 
         case EImageFormat::IF_BGRA_UBYTE:
         {
-            for ( int i=0; i<pixelCount; ++i )
+            for ( int i=0; i< PixelCount; ++i )
             {
                 uint8 m_8 = pMaskBuf[i];
                 if (m_8>127)
diff --git a/Engine/Plugins/Experimental/Mutable/Source/MutableRuntime/Private/MuR/OpImageMipmap.h b/Engine/Plugins/Experimental/Mutable/Source/MutableRuntime/Private/MuR/OpImageMipmap.h
--- a/Engine/Plugins/Experimental/Mutable/Source/MutableRuntime/Private/MuR/OpImageMipmap.h
+++ b/Engine/Plugins/Experimental/Mutable/Source/MutableRuntime/Private/MuR/OpImageMipmap.h
@@ -597,10 +597,10 @@ namespace mu
                 pBase->GetFormat());
 
             // Preallocate ample memory for the compressed data
-            scratch->pCompressedMips->m_data.SetNum(scratch->pUncompressedMips->GetDataSize());
+            scratch->pCompressedMips->m_data.SetNumUninitialized(scratch->pUncompressedMips->GetDataSize());
 
             // Preallocate ample memory for the destination data
-            pDest->m_data.SetNum(pBase->GetDataSize() + scratch->pUncompressedMips->GetDataSize());
+            pDest->m_data.SetNumUninitialized(pBase->GetDataSize() + scratch->pUncompressedMips->GetDataSize());
 
             break;
         }
@@ -714,7 +714,7 @@ namespace mu
 				// Bad case: this should almost never happen.
 				MUTABLE_CPUPROFILER_SCOPE(Mipmap_Recompression_OutOfSpace);
 
-				scratch->pCompressedMips->m_data.SetNum(ExcessDataSize);
+				scratch->pCompressedMips->m_data.SetNumUninitialized(ExcessDataSize);
 				bSuccess = ImagePixelFormatInPlace(imageCompressionQuality, scratch->pCompressedMips.get(), scratch->pUncompressedMips.get());
 				ExcessDataSize *= 4;
 			}
@@ -725,7 +725,7 @@ namespace mu
 			if (FinalDestSize > pDest->GetDataSize())
 			{
 				// Bad case: this should almost never happen.
-				pDest->m_data.SetNum(FinalDestSize);
+				pDest->m_data.SetNumUninitialized(FinalDestSize);
 				pDestBuf = !bGenerateOnlyTail ? pDest->GetMipData(startLevel + 1) : pDest->GetMipData(0);
 			}
 
diff --git a/Engine/Plugins/Experimental/Mutable/Source/MutableRuntime/Private/MuR/OpImagePixelFormat.cpp b/Engine/Plugins/Experimental/Mutable/Source/MutableRuntime/Private/MuR/OpImagePixelFormat.cpp
--- a/Engine/Plugins/Experimental/Mutable/Source/MutableRuntime/Private/MuR/OpImagePixelFormat.cpp
+++ b/Engine/Plugins/Experimental/Mutable/Source/MutableRuntime/Private/MuR/OpImagePixelFormat.cpp
@@ -42,7 +42,7 @@ ImagePtr ImagePixelFormat( int imageCompressionQuality, const Image* pBase,
         resultLODCount = 1;
     }
 
-    ImagePtr result = new Image( (uint16)resultSize[0], (uint16)resultSize[1], resultLODCount, targetFormat );
+    ImagePtr result = new Image( (uint16)resultSize[0], (uint16)resultSize[1], resultLODCount, targetFormat, EInitializationType::NotInitialized );
 	result->m_flags = pBase->m_flags;
 
 	if (pBase->GetSizeX()<=0 ||pBase->GetSizeY()<=0)
diff --git a/Engine/Plugins/Experimental/Mutable/Source/MutableRuntime/Public/MuR/Image.h b/Engine/Plugins/Experimental/Mutable/Source/MutableRuntime/Public/MuR/Image.h
--- a/Engine/Plugins/Experimental/Mutable/Source/MutableRuntime/Public/MuR/Image.h
+++ b/Engine/Plugins/Experimental/Mutable/Source/MutableRuntime/Public/MuR/Image.h
@@ -156,6 +156,12 @@ namespace mu
 		_CIM_COUNT
 	};
 
+	enum class EInitializationType
+	{
+		NotInitialized,
+		Black
+	};
+
     //! \brief 2D image resource with mipmaps.
 	//! \ingroup runtime
     class MUTABLERUNTIME_API Image : public RefCounted
@@ -176,7 +182,7 @@ namespace mu
         //!         base level. It must be a number between 1 and the maximum possible levels, which
         //!         depends on the image size.
         //! \param format Pixel format.
-        Image( uint32 sizeX, uint32 sizeY, uint32 lods, EImageFormat format );
+        Image( uint32 sizeX, uint32 sizeY, uint32 lods, EImageFormat format, EInitializationType Init=EInitializationType::Black);
 
 		/** */
 		ImagePtr ExtractMip(int32 Mip) const;
diff --git a/Engine/Plugins/Experimental/Mutable/Source/MutableTools/Private/MuT/CodeGenerator_Image.cpp b/Engine/Plugins/Experimental/Mutable/Source/MutableTools/Private/MuT/CodeGenerator_Image.cpp
--- a/Engine/Plugins/Experimental/Mutable/Source/MutableTools/Private/MuT/CodeGenerator_Image.cpp
+++ b/Engine/Plugins/Experimental/Mutable/Source/MutableTools/Private/MuT/CodeGenerator_Image.cpp
@@ -1086,6 +1086,7 @@ namespace mu
         op->SetChild( op->op.args.ImageColourMap.mask, mask);
 
         // Map image
+		// \TODO: We are forcing an map image size as it was a normal texture, and this is incorrect.
         Ptr<ASTOp> mapImage = 0;
         if ( Node* pMap = node.m_pMap.get() )
         {
diff --git a/Engine/Plugins/Experimental/Mutable/Source/MutableTools/Private/MuT/DataPacker.cpp b/Engine/Plugins/Experimental/Mutable/Source/MutableTools/Private/MuT/DataPacker.cpp
--- a/Engine/Plugins/Experimental/Mutable/Source/MutableTools/Private/MuT/DataPacker.cpp
+++ b/Engine/Plugins/Experimental/Mutable/Source/MutableTools/Private/MuT/DataPacker.cpp
@@ -751,7 +751,8 @@ namespace mu
 				if ( pOld->IsPlainColour(PlainColor) )
 				{
 					// It is more efficient to just have an instruction for it instead, to avoid the overhead
-					// of data loading.
+					// of data loading. 
+					// Warning This eliminates the mips. \TODO: Add support for mips in plaincolour instruction?
 					Ptr<ASTOpFixed> NewColor = new ASTOpFixed;
 					NewColor->op.type = OP_TYPE::CO_CONSTANT;
 					NewColor->op.args.ColourConstant.value[0] = PlainColor[0];
