[Mutable] ASTC decompression optiomizatoin.

- Removed unnecessary branches.

#rnx

diff --git a/Engine/Plugins/Experimental/Mutable/Source/MutableRuntime/Private/MuR/BlockCompression/Miro/Miro.cpp b/Engine/Plugins/Experimental/Mutable/Source/MutableRuntime/Private/MuR/BlockCompression/Miro/Miro.cpp
--- a/Engine/Plugins/Experimental/Mutable/Source/MutableRuntime/Private/MuR/BlockCompression/Miro/Miro.cpp
+++ b/Engine/Plugins/Experimental/Mutable/Source/MutableRuntime/Private/MuR/BlockCompression/Miro/Miro.cpp
@@ -6058,7 +6058,7 @@ namespace impl
 					const UVec4& e0 = colorEndpoints[colorEndpointNdx].e0;
 					const UVec4& e1 = colorEndpoints[colorEndpointNdx].e1;
 					const TexelWeightPair& weight = texelWeights[texelNdx];
-					if (isLDRMode && isHDREndpoint[colorEndpointNdx])
+					if (isLDRMode && isHDREndpoint[colorEndpointNdx] && false) // REMOVING HDR SUPPORT FOR NOW
 					{
 						if (isSRGB)
 						{
@@ -6080,7 +6080,7 @@ namespace impl
 					{
 						for (int channelNdx = 0; channelNdx < 4; channelNdx++)
 						{
-							if (!isHDREndpoint[colorEndpointNdx] || (channelNdx == 3 && colorEndpointModes[colorEndpointNdx] == 14)) // \note Alpha for mode 14 is treated the same as LDR.
+							if (!isHDREndpoint[colorEndpointNdx] || (channelNdx == 3 && colorEndpointModes[colorEndpointNdx] == 14) || true) // \note Alpha for mode 14 is treated the same as LDR. // always true because there HDR is currently not supported.
 							{
 								const deUint32 c0 = (e0[channelNdx] << 8) | (isSRGB ? 0x80 : e0[channelNdx]);
 								const deUint32 c1 = (e1[channelNdx] << 8) | (isSRGB ? 0x80 : e1[channelNdx]);
