[Mutable] Add option to mark textures as low priority to improve initial generation times. Low priority textures will generate defaulted resident mips if possible.

#rnx

diff --git a/Engine/Plugins/Experimental/Mutable/Source/CustomizableObject/Classes/MuCO/CustomizableObject.h b/Engine/Plugins/Experimental/Mutable/Source/CustomizableObject/Classes/MuCO/CustomizableObject.h
--- a/Engine/Plugins/Experimental/Mutable/Source/CustomizableObject/Classes/MuCO/CustomizableObject.h
+++ b/Engine/Plugins/Experimental/Mutable/Source/CustomizableObject/Classes/MuCO/CustomizableObject.h
@@ -1203,7 +1203,7 @@ private:
 	// This is a manual version number for the binary blobs in this asset.
 	// Increasing it invalidates all the previously compiled models.
 	// Warning: If while merging code both versions have changed, take the highest+1.
-	static const int32 CurrentSupportedVersion = 368;
+	static const int32 CurrentSupportedVersion = 369;
 
 public:
 
@@ -1329,6 +1329,11 @@ public:
 	UPROPERTY()
 	TMap<FString, FParameterUIData> ParameterUIDataMap;
 
+	/** Textures marked as low priority will generate defaulted resident mips (if texture streaming is enabled).
+	  * Generating defaulted resident mips greatly reduce initial generation times. */
+	UPROPERTY(EditAnywhere, Category = CustomizableObject)
+	TArray<FName> LowPriorityTextures;
+
 	/** Stores all the state UI metadata information for all the dependencies of this Customizable Object */
 	UPROPERTY()
 	TMap<FString, FParameterUIData> StateUIDataMap;
@@ -1408,6 +1413,9 @@ public:
 	/** Modify the provided mutable parameters so that the forced values for the given customizable object state are applied. */
 	void ApplyStateForcedValuesToParameters(int32 State, mu::Parameters* Parameters);
 
+	/** Return the names used by mutable to identify which mu::Image should be considered of LowPriority. */
+	void GetLowPriorityTextureNames(TArray<FString>& OutTextureNames);
+
 #if WITH_EDITOR
 
 	// UObject Interface -> Data validation
diff --git a/Engine/Plugins/Experimental/Mutable/Source/CustomizableObject/Private/MuCO/CustomizableObject.cpp b/Engine/Plugins/Experimental/Mutable/Source/CustomizableObject/Private/MuCO/CustomizableObject.cpp
--- a/Engine/Plugins/Experimental/Mutable/Source/CustomizableObject/Private/MuCO/CustomizableObject.cpp
+++ b/Engine/Plugins/Experimental/Mutable/Source/CustomizableObject/Private/MuCO/CustomizableObject.cpp
@@ -1468,6 +1468,25 @@ void UCustomizableObject::ApplyStateForcedValuesToParameters( int32 State, mu::P
 
 }
 
+
+void UCustomizableObject::GetLowPriorityTextureNames(TArray<FString>& OutTextureNames)
+{
+	OutTextureNames.Reset(LowPriorityTextures.Num());
+
+	if (!LowPriorityTextures.IsEmpty())
+	{
+		const int32 ImageCount = ImageProperties.Num();
+		for (int32 ImageIndex = 0; ImageIndex < ImageCount; ++ImageIndex)
+		{
+			if (LowPriorityTextures.Find(FName(ImageProperties[ImageIndex].TextureParameterName)) != INDEX_NONE)
+			{
+				OutTextureNames.Add(FString::FromInt(ImageIndex));
+			}
+		}
+	}
+}
+
+
 #if WITH_EDITOR
 void UCustomizableObject::PreSaveRoot(FObjectPreSaveRootContext ObjectSaveContext)
 {
diff --git a/Engine/Plugins/Experimental/Mutable/Source/CustomizableObject/Private/MuCO/CustomizableObjectSystem.cpp b/Engine/Plugins/Experimental/Mutable/Source/CustomizableObject/Private/MuCO/CustomizableObjectSystem.cpp
--- a/Engine/Plugins/Experimental/Mutable/Source/CustomizableObject/Private/MuCO/CustomizableObjectSystem.cpp
+++ b/Engine/Plugins/Experimental/Mutable/Source/CustomizableObject/Private/MuCO/CustomizableObjectSystem.cpp
@@ -705,6 +705,13 @@ static FAutoConsoleVariableRef CVarEnableOnlyGenerateRequestedLODs(
 	TEXT("If 1 or greater, Only the RequestedLODLevels will be generated. If 0, all LODs will be build."),
 	ECVF_Default);
 
+int32 FCustomizableObjectSystemPrivate::EnableSkipGenerateResidentMips = 1;
+
+static FAutoConsoleVariableRef CVarSkipGenerateResidentMips(
+	TEXT("mutable.EnableSkipGenerateResidentMips"), FCustomizableObjectSystemPrivate::EnableSkipGenerateResidentMips,
+	TEXT("If 1 or greater, resident mip generation will be optional. If 0, resident mips will be always generated"),
+	ECVF_Default);
+
 
 /** Update the given Instance Skeletal Meshes and call its callbacks. */
 void UpdateSkeletalMesh(UCustomizableObjectInstance* CustomizableObjectInstance, const FDescriptorRuntimeHash& UpdatedDescriptorRuntimeHash)
@@ -1406,9 +1413,10 @@ namespace impl
 		{
 			MUTABLE_CPUPROFILER_SCOPE(GetImage);
 
+			mu::FImageDesc ImageDesc;
+
 			// This should only be done when using progressive images, since GetImageDesc does some actual processing.
 			{
-				mu::FImageDesc ImageDesc;
 				System->GetImageDesc(OperationData->InstanceID, Image.ImageID, ImageDesc);
 				Image.FullImageSizeX = ImageDesc.m_size[0];
 				Image.FullImageSizeY = ImageDesc.m_size[1];
@@ -1437,7 +1445,17 @@ namespace impl
 				int32 MinMipsInImage = FMath::Min(FullLODCount, UTexture::GetStaticMinTextureResidentMipCount());
 				int32 MaxMipsToSkip = FullLODCount - MinMipsInImage;
 				int32 MipsToSkip = FMath::Min(MaxMipsToSkip, OperationData->MipsToSkip);
-				Image.Image = System->GetImage(OperationData->InstanceID, Image.ImageID, MipsToSkip, Image.LOD);
+
+				if (MipsToSkip > 0 && CustomizableObjectSystemPrivateData->EnableSkipGenerateResidentMips != 0 && OperationData->LowPriorityTextures.Find(Image.Name) != INDEX_NONE)
+				{
+					const int32 MipSizeX = FMath::Max(Image.FullImageSizeX >> MipsToSkip, 1);
+					const int32 MipSizeY = FMath::Max(Image.FullImageSizeY >> MipsToSkip, 1);
+					Image.Image = new mu::Image(MipSizeX, MipSizeY, FullLODCount - MipsToSkip, ImageDesc.m_format);
+				}
+				else
+				{
+					Image.Image = System->GetImage(OperationData->InstanceID, Image.ImageID, MipsToSkip, Image.LOD);
+				}
 
 				// We need one mip or the complete chain. Otherwise there was a bug.
 				check(Image.Image);
@@ -2229,6 +2247,7 @@ namespace impl
 		CurrentOperationData->MipsToSkip = Operation->MipsToSkip;
 		CurrentOperationData->MutableParameters = Parameters;
 		CurrentOperationData->State = CandidateInstance->GetState();
+		CustomizableObject->GetLowPriorityTextureNames(CurrentOperationData->LowPriorityTextures);
 
 		if (System->IsOnlyGenerateRequestedLODsEnabled() && System->CurrentInstanceLODManagement->IsOnlyGenerateRequestedLODLevelsEnabled() && !CandidateInstancePrivateData->HasCOInstanceFlags(ForceGenerateAllLODs))
 		{
diff --git a/Engine/Plugins/Experimental/Mutable/Source/CustomizableObject/Private/MuCO/CustomizableObjectSystemPrivate.h b/Engine/Plugins/Experimental/Mutable/Source/CustomizableObject/Private/MuCO/CustomizableObjectSystemPrivate.h
--- a/Engine/Plugins/Experimental/Mutable/Source/CustomizableObject/Private/MuCO/CustomizableObjectSystemPrivate.h
+++ b/Engine/Plugins/Experimental/Mutable/Source/CustomizableObject/Private/MuCO/CustomizableObjectSystemPrivate.h
@@ -472,6 +472,8 @@ struct FMutableOperationData
 	FParameterDecorationsUpdateData ParametersUpdateData;
 	TArray<int> RelevantParametersInProgress;
 
+	TArray<FString> LowPriorityTextures;
+
 	/** This option comes from the operation request */
 	bool bNeverStream = false;
 	/** When this option is enabled it will reuse the Mutable core instance and its temp data between updates.  */
@@ -553,6 +555,7 @@ public:
 	static int32 EnableMutableLiveUpdate;
 	static int32 EnableReuseInstanceTextures;
 	static int32 EnableMutableAnimInfoDebugging;
+	static int32 EnableSkipGenerateResidentMips;
 	static int32 EnableOnlyGenerateRequestedLODs;
 	static bool bEnableMutableReusePreviousUpdateData;
 
diff --git a/Engine/Plugins/Experimental/Mutable/Source/MutableRuntime/Private/MuR/CodeRunner.cpp b/Engine/Plugins/Experimental/Mutable/Source/MutableRuntime/Private/MuR/CodeRunner.cpp
--- a/Engine/Plugins/Experimental/Mutable/Source/MutableRuntime/Private/MuR/CodeRunner.cpp
+++ b/Engine/Plugins/Experimental/Mutable/Source/MutableRuntime/Private/MuR/CodeRunner.cpp
@@ -5328,7 +5328,7 @@ namespace mu
 			check(item.Stage == 0);
 			OP::ResourceConstantArgs args = program.GetOpArgs<OP::ResourceConstantArgs>(item.At);
 			int32 ImageIndex = args.value;
-			m_heapImageDesc[item.CustomState].m_format = EImageFormat::IF_NONE;	// TODO: precalculate if necessary
+			m_heapImageDesc[item.CustomState].m_format = program.m_constantImages[ImageIndex].ImageFormat;
 			m_heapImageDesc[item.CustomState].m_size[0] = program.m_constantImages[ImageIndex].ImageSizeX;
 			m_heapImageDesc[item.CustomState].m_size[1] = program.m_constantImages[ImageIndex].ImageSizeY;
 			m_heapImageDesc[item.CustomState].m_lods = program.m_constantImages[ImageIndex].LODCount;
@@ -5687,6 +5687,7 @@ namespace mu
 				FImageDesc& ResultAndBaseDesc = m_heapImageDesc[item.CustomState];
 				FImageSize destSize(uint16(ImageSizeInPixels.X), uint16(ImageSizeInPixels.Y));
 				ResultAndBaseDesc.m_size = destSize;
+				ResultAndBaseDesc.m_format = args.format;
 				
 				if (args.generateMipmaps)
 				{
@@ -5882,7 +5883,7 @@ namespace mu
 			m_heapImageDesc[item.CustomState].m_size[0] = args.size[0];
 			m_heapImageDesc[item.CustomState].m_size[1] = args.size[1];
 			m_heapImageDesc[item.CustomState].m_lods = 1;
-			m_heapImageDesc[item.CustomState].m_format = EImageFormat::IF_RGB_UBYTE;
+			m_heapImageDesc[item.CustomState].m_format = args.format;
 			GetMemory().SetValidDesc(item);
 			break;
 		}
diff --git a/Engine/Plugins/Experimental/Mutable/Source/MutableRuntime/Private/MuR/CodeRunnerTasked.cpp b/Engine/Plugins/Experimental/Mutable/Source/MutableRuntime/Private/MuR/CodeRunnerTasked.cpp
--- a/Engine/Plugins/Experimental/Mutable/Source/MutableRuntime/Private/MuR/CodeRunnerTasked.cpp
+++ b/Engine/Plugins/Experimental/Mutable/Source/MutableRuntime/Private/MuR/CodeRunnerTasked.cpp
@@ -920,14 +920,7 @@ namespace mu
 			&&
 			GetImageFormatData(m_base->GetFormat()).m_channels > 3)
 		{
-			// If the alpha channel is all white, we don't need to compress in a format with
-			// alpha support, as the default value is always white in shaders, and mutable
-			// code. This assumption exist in other commercial engines as well.
-			bool hasNonWhiteAlpha = !m_base->IsFullAlpha();
-			if (hasNonWhiteAlpha)
-			{
-				targetFormat = Args.formatIfAlpha;
-			}
+			targetFormat = Args.formatIfAlpha;
 		}
 
 		if (targetFormat != EImageFormat::IF_NONE)
diff --git a/Engine/Plugins/Experimental/Mutable/Source/MutableRuntime/Private/MuR/ModelPrivate.h b/Engine/Plugins/Experimental/Mutable/Source/MutableRuntime/Private/MuR/ModelPrivate.h
--- a/Engine/Plugins/Experimental/Mutable/Source/MutableRuntime/Private/MuR/ModelPrivate.h
+++ b/Engine/Plugins/Experimental/Mutable/Source/MutableRuntime/Private/MuR/ModelPrivate.h
@@ -60,6 +60,7 @@ namespace mu
 		int32 LODCount = 0;
 		uint16 ImageSizeX = 0;
 		uint16 ImageSizeY = 0;
+		EImageFormat ImageFormat = EImageFormat::IF_NONE;
 	};
 	MUTABLE_DEFINE_POD_SERIALISABLE(FImageLODRange);
 
@@ -393,7 +394,7 @@ namespace mu
 				}
 			}
 
-			int32 ImageIndex = m_constantImages.Add({ FirstLODIndexIndex, MipsToStore, pImage->GetSizeX(), pImage->GetSizeY() });
+			int32 ImageIndex = m_constantImages.Add({ FirstLODIndexIndex, MipsToStore, pImage->GetSizeX(), pImage->GetSizeY(), pImage->GetFormat() });
 			return ImageIndex;
 		}
 
diff --git a/Engine/Plugins/Experimental/Mutable/Source/MutableRuntime/Public/MuR/Image.h b/Engine/Plugins/Experimental/Mutable/Source/MutableRuntime/Public/MuR/Image.h
--- a/Engine/Plugins/Experimental/Mutable/Source/MutableRuntime/Public/MuR/Image.h
+++ b/Engine/Plugins/Experimental/Mutable/Source/MutableRuntime/Public/MuR/Image.h
@@ -74,7 +74,7 @@ namespace mu
 
         IF_COUNT
 	};
-
+	MUTABLE_DEFINE_ENUM_SERIALISABLE(EImageFormat);
 
 	struct MUTABLERUNTIME_API FImageDesc
 	{
diff --git a/Engine/Plugins/Experimental/Mutable/Source/MutableRuntime/Public/MuR/Model.h b/Engine/Plugins/Experimental/Mutable/Source/MutableRuntime/Public/MuR/Model.h
--- a/Engine/Plugins/Experimental/Mutable/Source/MutableRuntime/Public/MuR/Model.h
+++ b/Engine/Plugins/Experimental/Mutable/Source/MutableRuntime/Public/MuR/Model.h
@@ -13,7 +13,7 @@
 //! data structures. Compiled models are not necessarily compatible when the runtime is updated,
 //! so this version number can be used externally to verify this. It is not used internally, and
 //! serializing models from different versions than this runtime will probably result in a crash.
-#define MUTABLE_COMPILED_MODEL_CODE_VERSION		uint32( 55 )
+#define MUTABLE_COMPILED_MODEL_CODE_VERSION		uint32( 56 )
 #define MUTABLE_PARAMETERS_VERSION              uint32( 1 )
 
 
