#rnx

diff --git a/Engine/Plugins/Experimental/Mutable/Source/CustomizableObject/Private/MuCO/CustomizableObjectSystem.cpp b/Engine/Plugins/Experimental/Mutable/Source/CustomizableObject/Private/MuCO/CustomizableObjectSystem.cpp
--- a/Engine/Plugins/Experimental/Mutable/Source/CustomizableObject/Private/MuCO/CustomizableObjectSystem.cpp
+++ b/Engine/Plugins/Experimental/Mutable/Source/CustomizableObject/Private/MuCO/CustomizableObjectSystem.cpp
@@ -1261,6 +1261,16 @@ namespace impl
 				{
 					OperationData->RequestedLODs[ComponentIndex] |= (1 << OperationData->CurrentMaxLOD);
 				}
+
+				// Make sure we are not requesting a LOD that doesn't exist in this state (Essentially for states with bBuildOnlyFirstLOD 
+				// and NumExtraLODsToBuildPerPlatform when the ExtraLOD is not needed)
+				if ((OperationData->RequestedLODs[ComponentIndex] & (1 << OperationData->CurrentMaxLOD)) == 0)
+				{
+					OperationData->CurrentMaxLOD = OperationData->CurrentMinLOD;
+
+					// Ensure the fallback LOD actually exists
+					OperationData->RequestedLODs[ComponentIndex] |= (1 << OperationData->CurrentMaxLOD);
+				}
 			}
 		}
 
diff --git a/Engine/Plugins/Experimental/Mutable/Source/CustomizableObjectEditor/Classes/MuCOE/Nodes/CustomizableObjectNodeObject.h b/Engine/Plugins/Experimental/Mutable/Source/CustomizableObjectEditor/Classes/MuCOE/Nodes/CustomizableObjectNodeObject.h
--- a/Engine/Plugins/Experimental/Mutable/Source/CustomizableObjectEditor/Classes/MuCOE/Nodes/CustomizableObjectNodeObject.h
+++ b/Engine/Plugins/Experimental/Mutable/Source/CustomizableObjectEditor/Classes/MuCOE/Nodes/CustomizableObjectNodeObject.h
@@ -43,6 +43,10 @@ struct CUSTOMIZABLEOBJECTEDITOR_API FCustomizableObjectState
 	UPROPERTY(EditAnywhere, Category = CustomizableObject)
 	bool bBuildOnlyFirstLOD = false;
 
+	/** If there's an entry for a specific platform, when compiling for that platform Num LODs will be built after bBuildOnlyFirstLOD */
+	UPROPERTY(EditAnywhere, Category = CustomizableObject)
+	TMap<FString, int32> NumExtraLODsToBuildPerPlatform;
+
 	UPROPERTY(EditAnywhere, Category = CustomizableObject)
 	TMap<FString, FString> ForcedParameterValues;
 
diff --git a/Engine/Plugins/Experimental/Mutable/Source/CustomizableObjectEditor/Private/MuCOE/GenerateMutableSource/GenerateMutableSource.cpp b/Engine/Plugins/Experimental/Mutable/Source/CustomizableObjectEditor/Private/MuCOE/GenerateMutableSource/GenerateMutableSource.cpp
--- a/Engine/Plugins/Experimental/Mutable/Source/CustomizableObjectEditor/Private/MuCOE/GenerateMutableSource/GenerateMutableSource.cpp
+++ b/Engine/Plugins/Experimental/Mutable/Source/CustomizableObjectEditor/Private/MuCOE/GenerateMutableSource/GenerateMutableSource.cpp
@@ -704,7 +704,19 @@ mu::NodeObjectPtr GenerateMutableSource(const UEdGraphPin * Pin, FMutableGraphGe
 			{
 				ObjectNode->AddStateParam(StateIndex, TCHAR_TO_ANSI(*State.RuntimeParameters[ParamIndex]));
 			}
-			ObjectNode->SetStateProperties(StateIndex, State.bDontCompressRuntimeTextures, State.bBuildOnlyFirstLOD, GenerationContext.FirstLODAvailable);
+
+			const ITargetPlatform* TargetPlatform = GenerationContext.Options.TargetPlatform;
+
+			int32 NumExtraLODsToBuildAfterFirstLOD = 0;
+
+			const int32* AuxNumExtraLODs = State.NumExtraLODsToBuildPerPlatform.Find(TargetPlatform->PlatformName());
+
+			if (AuxNumExtraLODs)
+			{
+				NumExtraLODsToBuildAfterFirstLOD = *AuxNumExtraLODs;
+			}
+
+			ObjectNode->SetStateProperties(StateIndex, State.bDontCompressRuntimeTextures, State.bBuildOnlyFirstLOD, GenerationContext.FirstLODAvailable, NumExtraLODsToBuildAfterFirstLOD);
 
 			// UI Data
 			FParameterUIData ParameterUIData(State.Name, State.StateUIMetadata, EMutableParameterType::None);
diff --git a/Engine/Plugins/Experimental/Mutable/Source/MutableTools/Private/MuT/CodeOptimiser_Parameters.cpp b/Engine/Plugins/Experimental/Mutable/Source/MutableTools/Private/MuT/CodeOptimiser_Parameters.cpp
--- a/Engine/Plugins/Experimental/Mutable/Source/MutableTools/Private/MuT/CodeOptimiser_Parameters.cpp
+++ b/Engine/Plugins/Experimental/Mutable/Source/MutableTools/Private/MuT/CodeOptimiser_Parameters.cpp
@@ -2000,7 +2000,7 @@ namespace mu
             // Remove the unnecessary lods
             if (m_states[s].optimisationFlags.m_onlyFirstLOD)
             {
-                LODCountReducerAST(m_states[s].root, m_states[s].optimisationFlags.m_firstLOD + 1);
+				LODCountReducerAST(m_states[s].root, m_states[s].optimisationFlags.m_firstLOD + 1 + m_states[s].optimisationFlags.m_numExtraLODsToBuildAfterFirstLOD);
             }
 
             // If a state has no runtime parameters, skip its optimisation alltogether
diff --git a/Engine/Plugins/Experimental/Mutable/Source/MutableTools/Private/MuT/CompilerPrivate.h b/Engine/Plugins/Experimental/Mutable/Source/MutableTools/Private/MuT/CompilerPrivate.h
--- a/Engine/Plugins/Experimental/Mutable/Source/MutableTools/Private/MuT/CompilerPrivate.h
+++ b/Engine/Plugins/Experimental/Mutable/Source/MutableTools/Private/MuT/CompilerPrivate.h
@@ -22,20 +22,23 @@ namespace mu
 			m_firstLOD = 0;
             m_onlyFirstLOD = false;
             m_avoidRuntimeCompression = false;
+			m_numExtraLODsToBuildAfterFirstLOD = 0;
         }
 
 		uint8 m_firstLOD;
 		bool m_onlyFirstLOD;
         bool m_avoidRuntimeCompression;
+		int32 m_numExtraLODsToBuildAfterFirstLOD;
 
         void Serialise( OutputArchive& arch ) const
         {
-            const int32_t ver = 2;
+            const int32_t ver = 3;
             arch << ver;
 
 			arch << m_firstLOD;
 			arch << m_onlyFirstLOD;
             arch << m_avoidRuntimeCompression;
+			arch << m_numExtraLODsToBuildAfterFirstLOD;
         }
 
 
@@ -43,7 +46,16 @@ namespace mu
         {
             int32_t ver = 0;
             arch >> ver;
-			check(ver <= 2);
+			check(ver <= 3);
+
+			if (ver >= 3)
+			{
+				arch >> m_numExtraLODsToBuildAfterFirstLOD;
+			}
+			else
+			{
+				m_numExtraLODsToBuildAfterFirstLOD = 0;
+			}
 
 			if (ver >= 2)
 			{
diff --git a/Engine/Plugins/Experimental/Mutable/Source/MutableTools/Private/MuT/NodeObjectNew.cpp b/Engine/Plugins/Experimental/Mutable/Source/MutableTools/Private/MuT/NodeObjectNew.cpp
--- a/Engine/Plugins/Experimental/Mutable/Source/MutableTools/Private/MuT/NodeObjectNew.cpp
+++ b/Engine/Plugins/Experimental/Mutable/Source/MutableTools/Private/MuT/NodeObjectNew.cpp
@@ -317,13 +317,14 @@ namespace mu
 
 
     //---------------------------------------------------------------------------------------------
-    void NodeObjectNew::SetStateProperties( int s, bool avoidRuntimeCompression, bool onlyFirstLOD, int firstLOD )
+    void NodeObjectNew::SetStateProperties( int s, bool avoidRuntimeCompression, bool onlyFirstLOD, int firstLOD, int numExtraLODsToBuildAfterFirstLOD )
     {
         check( s>=0 && s<GetStateCount() );
 
         m_pD->m_states[s].m_optimisation.m_avoidRuntimeCompression = avoidRuntimeCompression;
         m_pD->m_states[s].m_optimisation.m_onlyFirstLOD = onlyFirstLOD;
-        m_pD->m_states[s].m_optimisation.m_firstLOD = firstLOD;
+		m_pD->m_states[s].m_optimisation.m_firstLOD = firstLOD;
+		m_pD->m_states[s].m_optimisation.m_numExtraLODsToBuildAfterFirstLOD = numExtraLODsToBuildAfterFirstLOD;
     }
 
 }
diff --git a/Engine/Plugins/Experimental/Mutable/Source/MutableTools/Public/MuT/NodeObjectNew.h b/Engine/Plugins/Experimental/Mutable/Source/MutableTools/Public/MuT/NodeObjectNew.h
--- a/Engine/Plugins/Experimental/Mutable/Source/MutableTools/Public/MuT/NodeObjectNew.h
+++ b/Engine/Plugins/Experimental/Mutable/Source/MutableTools/Public/MuT/NodeObjectNew.h
@@ -102,7 +102,7 @@ namespace mu
 		void RemoveStateParam( int s, const char* param );
 
         //! Set the optimisation properties of a state
-        void SetStateProperties( int s, bool avoidRuntimeCompression, bool onlyFirstLOD, int firstLOD );
+        void SetStateProperties( int s, bool avoidRuntimeCompression, bool onlyFirstLOD, int firstLOD, int numExtraLODsToBuildAfterFirstLOD);
 
 		//-----------------------------------------------------------------------------------------
 		// Interface pattern
