[Mutable] Fix LastMeshIds being cleared every now and then

#rnx

diff --git a/Engine/Plugins/Experimental/Mutable/Source/CustomizableObject/Private/MuCO/CustomizableInstancePrivateData.h b/Engine/Plugins/Experimental/Mutable/Source/CustomizableObject/Private/MuCO/CustomizableInstancePrivateData.h
--- a/Engine/Plugins/Experimental/Mutable/Source/CustomizableObject/Private/MuCO/CustomizableInstancePrivateData.h
+++ b/Engine/Plugins/Experimental/Mutable/Source/CustomizableObject/Private/MuCO/CustomizableInstancePrivateData.h
@@ -204,7 +204,7 @@ public:
 	mu::ParametersPtr GetParameters(UCustomizableObjectInstance* Public);
 
 	/** See FCustomizableObjectInstanceDescriptor::ReloadParameters(...). */
-	void ReloadParameters(UCustomizableObjectInstance* Public);
+	void ReloadParameters(UCustomizableObjectInstance* Public, bool bInvalidatePreviousData = false);
 
 	void TickUpdateCloseCustomizableObjects(UCustomizableObjectInstance& Public);
 	void UpdateInstanceIfNotGenerated(UCustomizableObjectInstance& Public);
diff --git a/Engine/Plugins/Experimental/Mutable/Source/CustomizableObject/Private/MuCO/CustomizableObjectInstance.cpp b/Engine/Plugins/Experimental/Mutable/Source/CustomizableObject/Private/MuCO/CustomizableObjectInstance.cpp
--- a/Engine/Plugins/Experimental/Mutable/Source/CustomizableObject/Private/MuCO/CustomizableObjectInstance.cpp
+++ b/Engine/Plugins/Experimental/Mutable/Source/CustomizableObject/Private/MuCO/CustomizableObjectInstance.cpp
@@ -561,7 +561,7 @@ mu::ParametersPtr UCustomizableInstancePrivateData::GetParameters(UCustomizableO
 }
 
 
-void UCustomizableInstancePrivateData::ReloadParameters(UCustomizableObjectInstance* Public)
+void UCustomizableInstancePrivateData::ReloadParameters(UCustomizableObjectInstance* Public, bool bInvalidatePreviousData)
 {
 	const UCustomizableObject* CustomizableObject = Public->GetCustomizableObject();
 	if (!CustomizableObject)
@@ -584,8 +584,12 @@ void UCustomizableInstancePrivateData::ReloadParameters(UCustomizableObjectInsta
 		ProjectorStates.Reset();
 	}
 	
-	InvalidateGeneratedData();
-	
+
+	if (bInvalidatePreviousData)
+	{
+		InvalidateGeneratedData();
+	}
+
 	Public->GetDescriptor().ReloadParameters();
 }
 
@@ -593,7 +597,7 @@ void UCustomizableInstancePrivateData::ReloadParameters(UCustomizableObjectInsta
 void UCustomizableObjectInstance::SetObject(UCustomizableObject* InObject)
 {
 	Descriptor.SetCustomizableObject(*InObject);
-	PrivateData->ReloadParameters(this);
+	PrivateData->ReloadParameters(this, true);
 	SetRequestedLODs(Descriptor.MinLOD, Descriptor.MaxLOD, Descriptor.RequestedLODLevels);
 }
 
