[Mutable] ParallelFor min batch execution optimizations.

- Add support for BlendAlphaSourceChannel in vector layer composite implementation. (currently not used)

#rnx

diff --git a/Engine/Plugins/Experimental/Mutable/Source/MutableRuntime/Private/MuR/CodeRunner.cpp b/Engine/Plugins/Experimental/Mutable/Source/MutableRuntime/Private/MuR/CodeRunner.cpp
--- a/Engine/Plugins/Experimental/Mutable/Source/MutableRuntime/Private/MuR/CodeRunner.cpp
+++ b/Engine/Plugins/Experimental/Mutable/Source/MutableRuntime/Private/MuR/CodeRunner.cpp
@@ -2694,9 +2694,7 @@ namespace mu
 						constexpr bool bUseVectorImpl = false;
 						if constexpr (bUseVectorImpl)
 						{
-							// TODO: Add support for args.BlendAlphaSourceChannel
-							check(args.BlendAlphaSourceChannel == 3);
-							BufferLayerCompositeVector<VectorBlendChannelMasked, VectorLightenChannel, false>(Base.get(), pBlended.get(), bBlendOnlyOneMip);
+							BufferLayerCompositeVector<VectorBlendChannelMasked, VectorLightenChannel, false>(Base.get(), pBlended.get(), bBlendOnlyOneMip, args.BlendAlphaSourceChannel);
 						}
 						else
 						{
diff --git a/Engine/Plugins/Experimental/Mutable/Source/MutableRuntime/Private/MuR/CodeRunnerTasked.cpp b/Engine/Plugins/Experimental/Mutable/Source/MutableRuntime/Private/MuR/CodeRunnerTasked.cpp
--- a/Engine/Plugins/Experimental/Mutable/Source/MutableRuntime/Private/MuR/CodeRunnerTasked.cpp
+++ b/Engine/Plugins/Experimental/Mutable/Source/MutableRuntime/Private/MuR/CodeRunnerTasked.cpp
@@ -513,9 +513,7 @@ namespace mu
 				constexpr bool bUseVectorImplementation = false;	
 				if constexpr (bUseVectorImplementation)
 				{
-					// TODO: Add support for Args.BlendAlphaSourceChannel
-					check(Args.BlendAlphaSourceChannel==3);
-					BufferLayerCompositeVector<VectorBlendChannelMasked, VectorLightenChannel, false>(pNew.get(), m_blended.get(), bOnlyOneMip);
+					BufferLayerCompositeVector<VectorBlendChannelMasked, VectorLightenChannel, false>(pNew.get(), m_blended.get(), bOnlyOneMip, Args.BlendAlphaSourceChannel);
 				}
 				else
 				{
diff --git a/Engine/Plugins/Experimental/Mutable/Source/MutableRuntime/Private/MuR/OpImageBlend.h b/Engine/Plugins/Experimental/Mutable/Source/MutableRuntime/Private/MuR/OpImageBlend.h
--- a/Engine/Plugins/Experimental/Mutable/Source/MutableRuntime/Private/MuR/OpImageBlend.h
+++ b/Engine/Plugins/Experimental/Mutable/Source/MutableRuntime/Private/MuR/OpImageBlend.h
@@ -174,14 +174,16 @@ namespace mu
 
 	//---------------------------------------------------------------------------------------------
 	FORCEINLINE VectorRegister4Int VectorBlendChannelMasked(
-		VectorRegister4Int Base, VectorRegister4Int Blended, VectorRegister4Int Mask)
+		const VectorRegister4Int& Base, const VectorRegister4Int& Blended, const VectorRegister4Int& Mask)
 	{
 		const VectorRegister4Int Value = VectorIntAdd(
-			VectorIntMultiply(Base, VectorIntSubtract(VectorIntSet1(255), Mask)),
-			VectorIntMultiply(Blended, Mask));
+				VectorIntMultiply(Base, VectorIntSubtract(MakeVectorRegisterIntConstant(255, 255, 255, 255), Mask)),
+				VectorIntMultiply(Blended, Mask));
 
 		// fast division by 255 assuming Value is in the range [0, (1 << 16)]
-		return VectorShiftRightImmLogical(VectorIntMultiply(Value, VectorIntSet1(32897)), 23);
+		return VectorShiftRightImmLogical(
+				VectorIntMultiply(Value, MakeVectorRegisterIntConstant(32897, 32897, 32897, 32897)),
+				23);
 	}
 
 	FORCEINLINE int32 VectorLightenChannel(int32 Base, int32 Blended)
diff --git a/Engine/Plugins/Experimental/Mutable/Source/MutableRuntime/Private/MuR/OpImageLayer.h b/Engine/Plugins/Experimental/Mutable/Source/MutableRuntime/Private/MuR/OpImageLayer.h
--- a/Engine/Plugins/Experimental/Mutable/Source/MutableRuntime/Private/MuR/OpImageLayer.h
+++ b/Engine/Plugins/Experimental/Mutable/Source/MutableRuntime/Private/MuR/OpImageLayer.h
@@ -39,9 +39,6 @@ namespace mu
 		// Generic implementation
 		const int32 PixelCount = pBase->CalculatePixelCount();
 
-		// Not used, ParallelFor + baching already consider low element count. 
-		constexpr int PixelCountConcurrencyThreshold = 0xff;
-
 		constexpr int32 NumBatchElems = 4096*2;
 		const int32 NumBatches = FMath::DivideAndRoundUp(PixelCount, NumBatchElems);
 
@@ -49,7 +46,7 @@ namespace mu
 		{
 		case EImageFormat::IF_RGB_UBYTE:
 		{
-			ParallelFor(NumBatches,
+			auto ProcessBatch =
 			[
 				pBaseBuf, pDestBuf, top, PixelCount, NumBatchElems
 			] (int32 BatchId)
@@ -73,13 +70,22 @@ namespace mu
 						}
 					}
 				}
-			});
+			};
+
+			if (NumBatches == 1)
+			{
+				ProcessBatch(0);
+			}
+			else if (NumBatches > 1)
+			{
+				ParallelFor(NumBatches, ProcessBatch);
+			}
 
 			break;
 		}
 		case EImageFormat::IF_RGBA_UBYTE:
 		{
-			ParallelFor(NumBatches, 
+			auto ProcessBatch = 
 			[
 				pBaseBuf, pDestBuf, top, PixelCount, NumBatchElems
 			] (int32 BatchId)
@@ -104,13 +110,22 @@ namespace mu
 					}
 					pDestBuf[4 * i + 3] = pBaseBuf[4 * i + 3];
 				}
-			});
+			};
+
+			if (NumBatches == 1)
+			{
+				ProcessBatch(0);
+			}
+			else if (NumBatches > 1)
+			{
+				ParallelFor(NumBatches, ProcessBatch);
+			}
 
 			break;
 		}
 		case EImageFormat::IF_BGRA_UBYTE:
 		{
-			ParallelFor(NumBatches,
+			auto ProcessBatch =
 			[
 				pBaseBuf, pDestBuf, top, PixelCount, NumBatchElems
 			] (int32 BatchId)
@@ -135,13 +150,22 @@ namespace mu
 					}
 					pDestBuf[4 * i + 3] = pBaseBuf[4 * i + 3];
 				}
-			});
+			};
+
+			if (NumBatches == 1)
+			{
+				ProcessBatch(0);
+			}
+			else if (NumBatches > 1)
+			{
+				ParallelFor(NumBatches, ProcessBatch);
+			}
 
 			break;
 		}
 		case EImageFormat::IF_L_UBYTE:
 		{
-			ParallelFor(NumBatches, 
+			auto ProcessBatch = 
 			[
 				pBaseBuf, pDestBuf, top, PixelCount, NumBatchElems
 			] (int32 BatchId)
@@ -162,7 +186,16 @@ namespace mu
 						pDestBuf[i] = (uint8)result;
 					}
 				}
-			});
+			};
+
+			if (NumBatches == 1)
+			{
+				ProcessBatch(0);
+			}
+			else if (NumBatches > 1)
+			{
+				ParallelFor(NumBatches, ProcessBatch);
+			}
 
 			break;
 		}
@@ -211,9 +244,6 @@ namespace mu
 		// Generic implementation
 		int32 PixelCount = Base->CalculatePixelCount();
 
-		constexpr int PixelCountConcurrencyThreshold = 0xff;
-
-
 		constexpr int32 NumBatchElems = 4096*2;
 		const int32 NumBatches = FMath::DivideAndRoundUp(PixelCount, NumBatchElems);
 
@@ -222,7 +252,7 @@ namespace mu
 		case EImageFormat::IF_RGB_UBYTE:
 		{
 			// There is no alpha so we assume 255
-			ParallelFor(NumBatches,
+			auto ProcessBatch =
 				[ pDestBuf, top, PixelCount, NumBatchElems ] (uint32 BatchId)
 			{
 				const int32 BatchBegin = BatchId * NumBatchElems;
@@ -244,13 +274,22 @@ namespace mu
 						}
 					}
 				}
-			});
+			};
+
+			if (NumBatches == 1)
+			{
+				ProcessBatch(0);
+			}
+			else if (NumBatches > 1)
+			{
+				ParallelFor(NumBatches, ProcessBatch);
+			}
 
 			break;
 		}
 		case EImageFormat::IF_RGBA_UBYTE:
 		{
-			ParallelFor(NumBatches, 
+			auto ProcessBatch =
 				[ pBaseBuf, pDestBuf, top, PixelCount, NumBatchElems ] (uint32 BatchId)
 			{	
 				const int32 BatchBegin = BatchId * NumBatchElems;
@@ -273,13 +312,22 @@ namespace mu
 					}
 					pDestBuf[4 * i + 3] = pBaseBuf[4 * i + 3];
 				}
-			});
+			};
+
+			if (NumBatches == 1)
+			{
+				ProcessBatch(0);
+			}
+			else if (NumBatches > 1)
+			{
+				ParallelFor(NumBatches, ProcessBatch);
+			}
 
 			break;
 		}
 		case EImageFormat::IF_BGRA_UBYTE:
 		{	
-			ParallelFor(NumBatches,
+			auto ProcessBatch =
 				[ pBaseBuf, pDestBuf, top, PixelCount, NumBatchElems ] (uint32 BatchId)
 			{
 				const int32 BatchBegin = BatchId * NumBatchElems;
@@ -291,7 +339,7 @@ namespace mu
 					{
 						unsigned base = pBaseBuf[4 * i + 3];
 						unsigned result = BLEND_FUNC(base, top[2 - c]);
-						if (CLAMP)
+						if constexpr (CLAMP)
 						{
 							pDestBuf[4 * i + c] = (uint8)FMath::Min(255u, result);
 						}
@@ -302,13 +350,22 @@ namespace mu
 					}
 					pDestBuf[4 * i + 3] = pBaseBuf[4 * i + 3];
 				}
-			});
+			};
+
+			if (NumBatches == 1)
+			{
+				ProcessBatch(0);
+			}
+			else if (NumBatches > 1)
+			{
+				ParallelFor(NumBatches, ProcessBatch);
+			}
 
 			break;
 		}
 		case EImageFormat::IF_L_UBYTE:
 		{
-			ParallelFor(NumBatches, 
+			auto ProcessBatch =
 				[ pBaseBuf, pDestBuf, top, PixelCount, NumBatchElems ] (uint32 BatchId)
 			{
 				const int32 BatchBegin = BatchId * NumBatchElems;
@@ -328,8 +385,16 @@ namespace mu
 						pDestBuf[i] = (uint8)result;
 					}
 				}
-			});
+			};
 
+			if (NumBatches == 1)
+			{
+				ProcessBatch(0);
+			}
+			else if (NumBatches > 1)
+			{
+				ParallelFor(NumBatches, ProcessBatch);
+			}
 			break;
 		}
 
@@ -660,32 +725,41 @@ namespace mu
 		constexpr int32 NumBatchElems = 4096*2;
 		const int32 NumBatches = FMath::DivideAndRoundUp(PixelCount, NumBatchElems);
 
-		ParallelFor(NumBatches,
-			[
-				pBaseBuf, top, BaseChannelOffset, ColOffset, PixelCount, NumBatchElems
-			] (uint32 BatchId)
+		auto ProcessBatch =
+		[
+			pBaseBuf, top, BaseChannelOffset, ColOffset, PixelCount, NumBatchElems
+		] (uint32 BatchId)
+		{
+			const int32 BatchBegin = BatchId * NumBatchElems;
+			const int32 BatchEnd = FMath::Min(BatchBegin + NumBatchElems, PixelCount);
+			
+			for (int32 i = BatchBegin; i < BatchEnd; ++i)
 			{
-				const int32 BatchBegin = BatchId * NumBatchElems;
-				const int32 BatchEnd = FMath::Min(BatchBegin + NumBatchElems, PixelCount);
-				
-				for (int32 i = BatchBegin; i < BatchEnd; ++i)
+				for (int c = 0; c < CHANNELS_TO_BLEND; ++c)
 				{
-					for (int c = 0; c < CHANNELS_TO_BLEND; ++c)
+					uint32 base = pBaseBuf[BASE_CHANNEL_STRIDE * i + c];
+					uint32 blended = top[c + ColOffset];
+					uint32 result = BLEND_FUNC(base, blended);
+					if (CLAMP)
 					{
-						uint32 base = pBaseBuf[BASE_CHANNEL_STRIDE * i + c];
-						uint32 blended = top[c + ColOffset];
-						uint32 result = BLEND_FUNC(base, blended);
-						if (CLAMP)
-						{
-							pBaseBuf[BASE_CHANNEL_STRIDE * i + c] = (uint8)FMath::Min(255u, result);
-						}
-						else
-						{
-							pBaseBuf[BASE_CHANNEL_STRIDE * i + c] = (uint8)result;
-						}
+						pBaseBuf[BASE_CHANNEL_STRIDE * i + c] = (uint8)FMath::Min(255u, result);
+					}
+					else
+					{
+						pBaseBuf[BASE_CHANNEL_STRIDE * i + c] = (uint8)result;
 					}
 				}
-			});
+			}
+		};
+
+		if (NumBatches == 1)
+		{
+			ProcessBatch(0);
+		}
+		else if (NumBatches > 1)
+		{
+			ParallelFor(NumBatches, ProcessBatch);
+		}
 	}
 
 
@@ -843,32 +917,41 @@ namespace mu
 		constexpr int32 NumBatchElems = 4096*2;
 		const int32 NumBatches = FMath::DivideAndRoundUp(PixelCount, NumBatchElems);
 
-		ParallelFor(NumBatches,
-			[
-				pBaseBuf, pBlendedBuf, PixelCount, NumBatchElems
-			] (int32 BatchId)
-			{
-				const int32 BatchBegin = BatchId * NumBatchElems;
-				const int32 BatchEnd = FMath::Min(BatchBegin + NumBatchElems, PixelCount);
+		auto ProcessBatch =
+		[
+			pBaseBuf, pBlendedBuf, PixelCount, NumBatchElems
+		] (int32 BatchId)
+		{
+			const int32 BatchBegin = BatchId * NumBatchElems;
+			const int32 BatchEnd = FMath::Min(BatchBegin + NumBatchElems, PixelCount);
 
-				for (int32 i = BatchBegin; i < BatchEnd; ++i)
+			for (int32 i = BatchBegin; i < BatchEnd; ++i)
+			{
+				for (int c = 0; c < CHANNELS_TO_BLEND; ++c)
 				{
-					for (int c = 0; c < CHANNELS_TO_BLEND; ++c)
+					unsigned base = pBaseBuf[BASE_CHANNEL_STRIDE * i + c];
+					unsigned blended = pBlendedBuf[BLENDED_CHANNEL_STRIDE * i + c];
+					unsigned result = BLEND_FUNC(base, blended);
+					if (CLAMP)
 					{
-						unsigned base = pBaseBuf[BASE_CHANNEL_STRIDE * i + c];
-						unsigned blended = pBlendedBuf[BLENDED_CHANNEL_STRIDE * i + c];
-						unsigned result = BLEND_FUNC(base, blended);
-						if (CLAMP)
-						{
-							pBaseBuf[BASE_CHANNEL_STRIDE * i + c] = (uint8)FMath::Min(255u, result);
-						}
-						else
-						{
-							pBaseBuf[BASE_CHANNEL_STRIDE * i + c] = (uint8)result;
-						}
+						pBaseBuf[BASE_CHANNEL_STRIDE * i + c] = (uint8)FMath::Min(255u, result);
+					}
+					else
+					{
+						pBaseBuf[BASE_CHANNEL_STRIDE * i + c] = (uint8)result;
 					}
 				}
-			});
+			}
+		};
+
+		if (NumBatches == 1)
+		{
+			ProcessBatch(0);
+		}
+		else if (NumBatches > 1)
+		{
+			ParallelFor(NumBatches, ProcessBatch);
+		}
 	}
 	
 	//---------------------------------------------------------------------------------------------
@@ -911,39 +994,48 @@ namespace mu
 		constexpr int32 NumBatchElems = 4096*2;
 		const int32 NumBatches = FMath::DivideAndRoundUp(PixelCount, NumBatchElems);
 		
-		ParallelFor(NumBatches,
-			[
-				pBaseBuf, pBlendedBuf, pDestBuf, UnblendedChannels, PixelCount, NumBatchElems
-			] (int32 BatchId)
-			{
-				const int32 BatchBegin = BatchId * NumBatchElems;
-				const int32 BatchEnd = FMath::Min(BatchBegin + NumBatchElems, PixelCount);
+		auto ProcessBatch =
+		[
+			pBaseBuf, pBlendedBuf, pDestBuf, UnblendedChannels, PixelCount, NumBatchElems
+		] (int32 BatchId)
+		{
+			const int32 BatchBegin = BatchId * NumBatchElems;
+			const int32 BatchEnd = FMath::Min(BatchBegin + NumBatchElems, PixelCount);
 
-				for (int32 i = BatchBegin; i < BatchEnd; ++i)
+			for (int32 i = BatchBegin; i < BatchEnd; ++i)
+			{
+				for (int c = 0; c < CHANNELS_TO_BLEND; ++c)
 				{
-					for (int c = 0; c < CHANNELS_TO_BLEND; ++c)
+					unsigned base = pBaseBuf[BASE_CHANNEL_STRIDE * i + c];
+					unsigned blended = pBlendedBuf[BLENDED_CHANNEL_STRIDE * i + c];
+					unsigned result = BLEND_FUNC(base, blended);
+					if constexpr (CLAMP)
 					{
-						unsigned base = pBaseBuf[BASE_CHANNEL_STRIDE * i + c];
-						unsigned blended = pBlendedBuf[BLENDED_CHANNEL_STRIDE * i + c];
-						unsigned result = BLEND_FUNC(base, blended);
-						if constexpr (CLAMP)
-						{
-							pDestBuf[BASE_CHANNEL_STRIDE * i + c] = (uint8)FMath::Min(255u, result);
-						}
-						else
-						{
-							pDestBuf[BASE_CHANNEL_STRIDE * i + c] = (uint8)result;
-						}
+						pDestBuf[BASE_CHANNEL_STRIDE * i + c] = (uint8)FMath::Min(255u, result);
 					}
-
-					// Copy the unblended channels
-					// \TODO: unnecessary when doing it in-place?
-					for (int32 c = 0; c < UnblendedChannels; ++c)
+					else
 					{
-						pDestBuf[BASE_CHANNEL_STRIDE * i + CHANNELS_TO_BLEND + c] = pBaseBuf[BASE_CHANNEL_STRIDE * i + CHANNELS_TO_BLEND + c];
+						pDestBuf[BASE_CHANNEL_STRIDE * i + c] = (uint8)result;
 					}
 				}
-			});
+
+				// Copy the unblended channels
+				// \TODO: unnecessary when doing it in-place?
+				for (int32 c = 0; c < UnblendedChannels; ++c)
+				{
+					pDestBuf[BASE_CHANNEL_STRIDE * i + CHANNELS_TO_BLEND + c] = pBaseBuf[BASE_CHANNEL_STRIDE * i + CHANNELS_TO_BLEND + c];
+				}
+			}
+		};
+
+		if (NumBatches == 1)
+		{
+			ProcessBatch(0);
+		}
+		else if (NumBatches > 1)
+		{
+			ParallelFor(NumBatches, ProcessBatch);
+		}
 	}
 
 
@@ -1115,39 +1207,48 @@ namespace mu
 			constexpr int32 NumBatchElems = 4096*2;
 			const int32 NumBatches = FMath::DivideAndRoundUp(PixelCount, NumBatchElems);
 
-			ParallelFor(NumBatches,
-				[
-					pBaseBuf, pBlendedBuf, pMaskBuf, pDestBuf, UnblendedChannels, PixelCount, NumBatchElems
-				] (uint32 BatchId)
-				{
-					const int32 BatchBegin = BatchId * NumBatchElems;
-					const int32 BatchEnd = FMath::Min(BatchBegin + NumBatchElems, PixelCount);
+			auto ProcessBatch = 
+			[
+				pBaseBuf, pBlendedBuf, pMaskBuf, pDestBuf, UnblendedChannels, PixelCount, NumBatchElems
+			] (uint32 BatchId)
+			{
+				const int32 BatchBegin = BatchId * NumBatchElems;
+				const int32 BatchEnd = FMath::Min(BatchBegin + NumBatchElems, PixelCount);
 
-					for (int32 i = BatchBegin; i < BatchEnd; ++i)
+				for (int32 i = BatchBegin; i < BatchEnd; ++i)
+				{
+					unsigned mask = pMaskBuf[i];
+					for (int c = 0; c < CHANNELS_TO_BLEND; ++c)
 					{
-						unsigned mask = pMaskBuf[i];
-						for (int c = 0; c < CHANNELS_TO_BLEND; ++c)
+						uint32 base = pBaseBuf[BASE_CHANNEL_STRIDE * i + c];
+						uint32 blended = pBlendedBuf[BLENDED_CHANNEL_STRIDE * i + c];
+						uint32 result = BLEND_FUNC_MASKED(base, blended, mask);
+						if constexpr (CLAMP)
 						{
-							uint32 base = pBaseBuf[BASE_CHANNEL_STRIDE * i + c];
-							uint32 blended = pBlendedBuf[BLENDED_CHANNEL_STRIDE * i + c];
-							uint32 result = BLEND_FUNC_MASKED(base, blended, mask);
-							if constexpr (CLAMP)
-							{
-								pDestBuf[BASE_CHANNEL_STRIDE * i + c] = (uint8)FMath::Min(255u, result);
-							}
-							else
-							{
-								pDestBuf[BASE_CHANNEL_STRIDE * i + c] = (uint8)result;
-							}
+							pDestBuf[BASE_CHANNEL_STRIDE * i + c] = (uint8)FMath::Min(255u, result);
 						}
-						// Copy the unblended channels
-						// \TODO: unnecessary when doing it in-place?
-						for (int c = 0; c < UnblendedChannels; ++c)
+						else
 						{
-							pDestBuf[BASE_CHANNEL_STRIDE * i + CHANNELS_TO_BLEND + c] = pBaseBuf[BASE_CHANNEL_STRIDE * i + CHANNELS_TO_BLEND + c];
+							pDestBuf[BASE_CHANNEL_STRIDE * i + c] = (uint8)result;
 						}
 					}
-				});
+					// Copy the unblended channels
+					// \TODO: unnecessary when doing it in-place?
+					for (int c = 0; c < UnblendedChannels; ++c)
+					{
+						pDestBuf[BASE_CHANNEL_STRIDE * i + CHANNELS_TO_BLEND + c] = pBaseBuf[BASE_CHANNEL_STRIDE * i + CHANNELS_TO_BLEND + c];
+					}
+				}
+			};
+
+			if (NumBatches == 1)
+			{
+				ProcessBatch(0);
+			}
+			else if (NumBatches > 1)
+			{
+				ParallelFor(NumBatches, ProcessBatch);
+			}
 		}
 		else if (maskFormat == EImageFormat::IF_L_UBYTE_RLE)
 		{
@@ -1336,7 +1437,7 @@ namespace mu
 			constexpr int32 NumBatchElems = 4096*2;
 			const int32 NumBatches = FMath::DivideAndRoundUp(PixelCount, NumBatchElems);
 
-			ParallelFor(NumBatches,
+			auto ProcessBatch = 
 			[
 				pBaseBuf, pBlendedBuf, pDestBuf, pMaskBuf, UnblendedChannels, PixelCount, NumBatchElems
 			] (int32 BatchId)
@@ -1368,14 +1469,23 @@ namespace mu
 						pDestBuf[BASE_CHANNEL_STRIDE * i + CHANNELS_TO_BLEND + c] = pBaseBuf[BASE_CHANNEL_STRIDE * i + CHANNELS_TO_BLEND + c];
 					}
 				}
-			});
+			};
+
+			if (NumBatches == 1)
+			{
+				ProcessBatch(0);
+			}
+			else if (NumBatches > 1)
+			{
+				ParallelFor(NumBatches, ProcessBatch);
+			}
 		}
 		else
 		{
 			constexpr int32 NumBatchElems = 4096*2;
 			const int32 NumBatches = FMath::DivideAndRoundUp(PixelCount, NumBatchElems);
 
-			ParallelFor(NumBatches,
+			auto ProcessBatch = 
 			[
 				pBaseBuf, pBlendedBuf, pDestBuf, UnblendedChannels, PixelCount, NumBatchElems
 			] (uint32 BatchId)
@@ -1406,7 +1516,16 @@ namespace mu
 						pDestBuf[BASE_CHANNEL_STRIDE * i + CHANNELS_TO_BLEND + c] = pBaseBuf[BASE_CHANNEL_STRIDE * i + CHANNELS_TO_BLEND + c];
 					}
 				}
-			});
+			};
+
+			if (NumBatches == 1)
+			{
+				ProcessBatch(0);
+			}
+			else if (NumBatches > 1)
+			{
+				ParallelFor(NumBatches, ProcessBatch);
+			}
 		}
 	}
 
@@ -1621,7 +1740,7 @@ namespace mu
 		constexpr int32 NumBatchElems = 4096*2;
 		const int32 NumBatches = FMath::DivideAndRoundUp(PixelCount, NumBatchElems);
 
-		ParallelFor(NumBatches,
+		auto ProcessBatch = 
 		[
 			pBaseBuf, pBlendedBuf, BlendAlphaSourceChannel, PixelCount, NumBatchElems
 		] (uint32 BatchId)
@@ -1665,17 +1784,27 @@ namespace mu
 					}
 				}
 			}
-		});
+		};
+
+		if (NumBatches == 1)
+		{
+			ProcessBatch(0);
+		}
+		else if (NumBatches > 1)
+		{
+			ParallelFor(NumBatches, ProcessBatch);
+		}
 	}
 
 	//---------------------------------------------------------------------------------------------
-	template< VectorRegister4Int (*RGB_FUNC_MASKED)(VectorRegister4Int, VectorRegister4Int, VectorRegister4Int),
+	template< VectorRegister4Int (*RGB_FUNC_MASKED)(const VectorRegister4Int&, const VectorRegister4Int&, const VectorRegister4Int&),
 		int32 (*A_FUNC)(int32, int32),
 		bool CLAMP >
 	inline void BufferLayerCompositeVector(
 		Image* pBase,
 		const Image* pBlend,
-		bool bOnlyFirstLOD)
+		bool bOnlyFirstLOD,
+		uint8 BlendAlphaSourceChannel)
 	{
 		check(pBase->GetFormat() == EImageFormat::IF_RGBA_UBYTE);
 		check(pBlend->GetFormat() == EImageFormat::IF_RGBA_UBYTE);
@@ -1701,9 +1830,9 @@ namespace mu
 		constexpr int32 NumBatchElems = 4096*2;
 		const int32 NumBatches = FMath::DivideAndRoundUp(PixelCount, NumBatchElems);
 
-		ParallelFor(NumBatches,
+		auto ProcessBatch = 
 		[
-			pBaseBuf, pBlendedBuf, PixelCount, NumBatchElems
+			pBaseBuf, pBlendedBuf, PixelCount, NumBatchElems, BlendAlphaSourceChannel
 		] (uint32 BatchId)
 		{
 
@@ -1712,17 +1841,17 @@ namespace mu
 			for (int32 i = BatchBegin; i < BatchEnd; ++i)
 			{
 				// TODO: Optimize this (SIMD?)
-				const int32 BaseAlpha = pBaseBuf[4 * i + 3];
-				const int32 BlendedAlpha = pBlendedBuf[4 * i + 3];
+				const int32 BaseAlpha = pBaseBuf[4 * i + BlendAlphaSourceChannel];
+				const int32 BlendedAlpha = pBlendedBuf[4 * i + BlendAlphaSourceChannel];
 
 				const VectorRegister4Int Mask = VectorIntSet1(BlendedAlpha);
-				const VectorRegister4Int Blended = MakeVectorRegisterInt(pBlendedBuf[4 * i + 0], pBlendedBuf[4 * i + 1], pBlendedBuf[4 * i + 2], 0);
-				const VectorRegister4Int Base = MakeVectorRegisterInt(pBaseBuf[4 * i + 0], pBaseBuf[4 * i + 1], pBaseBuf[4 * i + 2], 0);
+				const VectorRegister4Int Blended = MakeVectorRegisterInt(pBlendedBuf[4 * i + 0], pBlendedBuf[4 * i + 1], pBlendedBuf[4 * i + 2], pBlendedBuf[4 * i + 3]);
+				const VectorRegister4Int Base = MakeVectorRegisterInt(pBaseBuf[4 * i + 0], pBaseBuf[4 * i + 1], pBaseBuf[4 * i + 2], pBaseBuf[4 * i + 3]);
 
 				VectorRegister4Int Result = RGB_FUNC_MASKED(Base, Blended, Mask);
 				if constexpr (CLAMP)
 				{
-					Result = VectorIntMin(VectorIntSet1(255), Result);
+					Result = VectorIntMin(MakeVectorRegisterIntConstant(255, 255, 255, 255), Result);
 				}
 
 				int32 AlphaResult = A_FUNC(BaseAlpha, BlendedAlpha);
@@ -1731,19 +1860,27 @@ namespace mu
 					AlphaResult = FMath::Min(255, AlphaResult);
 				}
 
-				struct alignas(alignof(VectorRegister4Int)) AlignedOutputRegister4Int
-				{
-					int32 Data[4];
-				} RegisterOutput;
+				alignas(VectorRegister4Int) int32 IndexableRegister[4];
 
-				VectorIntStoreAligned(Result, &RegisterOutput);
+				VectorIntStoreAligned(Result, &IndexableRegister);
 
-				pBaseBuf[4 * i + 0] = static_cast<uint8>(RegisterOutput.Data[0]);
-				pBaseBuf[4 * i + 1] = static_cast<uint8>(RegisterOutput.Data[1]);
-				pBaseBuf[4 * i + 2] = static_cast<uint8>(RegisterOutput.Data[2]);
-				pBaseBuf[4 * i + 3] = static_cast<uint8>(AlphaResult);
+				pBaseBuf[4 * i + 0] = static_cast<uint8>(IndexableRegister[0]);
+				pBaseBuf[4 * i + 1] = static_cast<uint8>(IndexableRegister[1]);
+				pBaseBuf[4 * i + 2] = static_cast<uint8>(IndexableRegister[2]);
+				pBaseBuf[4 * i + 3] = static_cast<uint8>(IndexableRegister[3]);
+
+				pBaseBuf[4 * i + BlendAlphaSourceChannel] = static_cast<uint8>(AlphaResult);
 			}
-		});
+		};
+
+		if (NumBatches == 1)
+		{
+			ProcessBatch(0);
+		}
+		else if (NumBatches > 1)
+		{
+			ParallelFor(NumBatches, ProcessBatch);
+		}
 	}
 
 	//---------------------------------------------------------------------------------------------
diff --git a/Engine/Plugins/Experimental/Mutable/Source/MutableRuntime/Private/MuR/OpImageSwizzle.h b/Engine/Plugins/Experimental/Mutable/Source/MutableRuntime/Private/MuR/OpImageSwizzle.h
--- a/Engine/Plugins/Experimental/Mutable/Source/MutableRuntime/Private/MuR/OpImageSwizzle.h
+++ b/Engine/Plugins/Experimental/Mutable/Source/MutableRuntime/Private/MuR/OpImageSwizzle.h
@@ -69,8 +69,6 @@ namespace mu
 			check(false);
 		}
 
-		constexpr int PixelCountConcurrencyThreshold = 0xff;
-
 		for (int i = 0; i < NumDestChannels; ++i)
 		{
 			uint8* pDestBuf = pDest->GetData() + i;
@@ -96,23 +94,30 @@ namespace mu
 			{
 				const uint8* pSourceBuf = pSources[i]->GetData() + channels[i];
 
-
 				switch (pSources[i]->GetFormat())
 				{
 				case EImageFormat::IF_L_UBYTE:
 					if (channels[i] < 1)
 					{
-						//for (size_t p = 0; p < pixelCount; ++p)
-						ParallelFor(NumBatches,[pDestBuf, pSourceBuf, NumDestChannels, PixelCount, NumBatchElems ](int32 BatchId)
+						auto ProcessBatch = [pDestBuf, pSourceBuf, NumDestChannels, PixelCount, NumBatchElems ](int32 BatchId)
+						{
+							const int32 BatchBegin = BatchId * NumBatchElems;
+							const int32 BatchEnd = FMath::Min(BatchBegin + NumBatchElems, PixelCount);
+
+							for (int32 p = BatchBegin; p < BatchEnd; ++p)
 							{
-								const int32 BatchBegin = BatchId * NumBatchElems;
-								const int32 BatchEnd = FMath::Min(BatchBegin + NumBatchElems, PixelCount);
+								pDestBuf[p * NumDestChannels] = pSourceBuf[p];
+							}
+						};
 
-								for (int32 p = BatchBegin; p < BatchEnd; ++p)
-								{
-									pDestBuf[p * NumDestChannels] = pSourceBuf[p];
-								}
-							});
+						if (NumBatches == 1)
+						{
+							ProcessBatch(0);
+						}
+						else if (NumBatches > 1)
+						{
+							ParallelFor(NumBatches, ProcessBatch);
+						}
 
 						filled = true;
 					}
@@ -121,34 +126,53 @@ namespace mu
 				case EImageFormat::IF_RGB_UBYTE:
 					if (channels[i] < 3)
 					{
+						auto ProcessBatch = [pDestBuf, pSourceBuf, NumDestChannels, PixelCount, NumBatchElems](int32 BatchId)
+						{
+							const int32 BatchBegin = BatchId * NumBatchElems;
+							const int32 BatchEnd = FMath::Min(BatchBegin + NumBatchElems, PixelCount);
 
-						ParallelFor(NumBatches, [pDestBuf, pSourceBuf, NumDestChannels, PixelCount, NumBatchElems](int32 BatchId)
+							for (int32 p = BatchBegin; p < BatchEnd; ++p)
 							{
-								const int32 BatchBegin = BatchId * NumBatchElems;
-								const int32 BatchEnd = FMath::Min(BatchBegin + NumBatchElems, PixelCount);
-
-								for (int32 p = BatchBegin; p < BatchEnd; ++p)
-								{
-									pDestBuf[p * NumDestChannels] = pSourceBuf[p * 3];
-								}
-							});
+								pDestBuf[p * NumDestChannels] = pSourceBuf[p * 3];
+							}
+						};
+
+						if (NumBatches == 1)
+						{
+							ProcessBatch(0);
+						}
+						else if (NumBatches > 1)
+						{
+							ParallelFor(NumBatches, ProcessBatch);
+						}
+
 						filled = true;
 					}
 					break;
 
 				case EImageFormat::IF_RGBA_UBYTE:
 					if (channels[i] < 4)
-					{
-						ParallelFor(NumBatches, [pDestBuf, pSourceBuf, NumDestChannels, PixelCount, NumBatchElems](int32 BatchId)
+					{	
+						auto ProcessBatch = [pDestBuf, pSourceBuf, NumDestChannels, PixelCount, NumBatchElems](int32 BatchId)
+						{
+							const int32 BatchBegin = BatchId * NumBatchElems;
+							const int32 BatchEnd = FMath::Min(BatchBegin + NumBatchElems, PixelCount);
+
+							for (int32 p = BatchBegin; p < BatchEnd; ++p)
 							{
-								const int32 BatchBegin = BatchId * NumBatchElems;
-								const int32 BatchEnd = FMath::Min(BatchBegin + NumBatchElems, PixelCount);
-
-								for (int32 p = BatchBegin; p < BatchEnd; ++p)
-								{
-									pDestBuf[p * NumDestChannels] = pSourceBuf[p * 4];
-								}
-							});
+								pDestBuf[p * NumDestChannels] = pSourceBuf[p * 4];
+							}
+						};
+
+						if (NumBatches == 1)
+						{
+							ProcessBatch(0);
+						}
+						else if (NumBatches > 1)
+						{
+							ParallelFor(NumBatches, ProcessBatch);
+						}
+
 						filled = true;
 					}
 					break;
@@ -164,15 +188,25 @@ namespace mu
 					}
 					if (channels[i] < 4)
 					{
-						ParallelFor(NumBatches, [pDestBuf, pSourceBuf, NumDestChannels, PixelCount, NumBatchElems](int32 BatchId)
+						auto ProcessBatch = [pDestBuf, pSourceBuf, NumDestChannels, PixelCount, NumBatchElems](int32 BatchId)
+						{
+							const int32 BatchBegin = BatchId * NumBatchElems;
+							const int32 BatchEnd = FMath::Min(BatchBegin + NumBatchElems, PixelCount);
+							for (int32 p = BatchBegin; p < BatchEnd; ++p)
 							{
-								const int32 BatchBegin = BatchId * NumBatchElems;
-								const int32 BatchEnd = FMath::Min(BatchBegin + NumBatchElems, PixelCount);
-								for (int32 p = BatchBegin; p < BatchEnd; ++p)
-								{
-									pDestBuf[p * NumDestChannels] = pSourceBuf[p * 4];
-								}
-							});
+								pDestBuf[p * NumDestChannels] = pSourceBuf[p * 4];
+							}
+						};
+
+						if (NumBatches == 1)
+						{
+							ProcessBatch(0);
+						}
+						else if (NumBatches > 1)
+						{
+							ParallelFor(NumBatches, ProcessBatch);
+						}
+
 						filled = true;
 					}
 					break;
@@ -185,15 +219,24 @@ namespace mu
 			if (!filled)
 			{
 				// Source not set. Clear to 0
-				ParallelFor(NumBatches, [pDestBuf, NumDestChannels, PixelCount, NumBatchElems](int32 BatchId)
+				auto ProcessBatch = [pDestBuf, NumDestChannels, PixelCount, NumBatchElems](int32 BatchId)
+				{
+					const int32 BatchBegin = BatchId * NumBatchElems;
+					const int32 BatchEnd = FMath::Min(BatchBegin + NumBatchElems, PixelCount);
+					for (int32 p = BatchBegin; p < BatchEnd; ++p)
 					{
-						const int32 BatchBegin = BatchId * NumBatchElems;
-						const int32 BatchEnd = FMath::Min(BatchBegin + NumBatchElems, PixelCount);
-						for (int32 p = BatchBegin; p < BatchEnd; ++p)
-						{
-							pDestBuf[p * NumDestChannels] = 0;
-						}
-					});
+						pDestBuf[p * NumDestChannels] = 0;
+					}
+				};
+
+				if (NumBatches == 1)
+				{
+					ProcessBatch(0);
+				}
+				else if (NumBatches > 1)
+				{
+					ParallelFor(NumBatches, ProcessBatch);
+				}
 			}
 		}
 
