#rnx

diff --git a/Engine/Plugins/Experimental/Mutable/Source/CustomizableObject/Private/MuCO/CustomizableObjectSystem.cpp b/Engine/Plugins/Experimental/Mutable/Source/CustomizableObject/Private/MuCO/CustomizableObjectSystem.cpp
--- a/Engine/Plugins/Experimental/Mutable/Source/CustomizableObject/Private/MuCO/CustomizableObjectSystem.cpp
+++ b/Engine/Plugins/Experimental/Mutable/Source/CustomizableObject/Private/MuCO/CustomizableObjectSystem.cpp
@@ -683,7 +683,7 @@ static FAutoConsoleVariableRef CVarEnableMutableLiveUpdate(
 	ECVF_Default);
 
 
-int32 FCustomizableObjectSystemPrivate::EnableReuseInstanceTextures = 0;
+int32 FCustomizableObjectSystemPrivate::EnableReuseInstanceTextures = 1;
 
 static FAutoConsoleVariableRef CVarEnableMutableReuseInstanceTextures(
 	TEXT("mutable.EnableReuseInstanceTextures"), FCustomizableObjectSystemPrivate::EnableReuseInstanceTextures,
