[Mutable] Force a minimum UVDensity to improve texture streaming on small meshes.

#rnx

diff --git a/Engine/Plugins/Experimental/Mutable/Source/CustomizableObject/Private/MuCO/CustomizableObjectInstance.cpp b/Engine/Plugins/Experimental/Mutable/Source/CustomizableObject/Private/MuCO/CustomizableObjectInstance.cpp
--- a/Engine/Plugins/Experimental/Mutable/Source/CustomizableObject/Private/MuCO/CustomizableObjectInstance.cpp
+++ b/Engine/Plugins/Experimental/Mutable/Source/CustomizableObject/Private/MuCO/CustomizableObjectInstance.cpp
@@ -1386,6 +1386,7 @@ UPhysicsAsset* UCustomizableInstancePrivateData::GetOrBuildPhysicsAsset(
 void SetMeshUVChannelDensity(FMeshUVChannelInfo& UVChannelInfo, float Density = 0.f)
 {
 	Density = Density > 0.f ? Density : 150.f;
+	Density = FMath::Max(100.0f, Density);
 
 	UVChannelInfo.bInitialized = true;
 	UVChannelInfo.bOverrideDensities = false;
