[Mutable] Add CVar to tweak the MinUVDensity to set on generated meshes.
Cvar: Mutable.MinUVChannelDensity

#rnx

#localization none
#tests PIE

diff --git a/Engine/Plugins/Experimental/Mutable/Source/CustomizableObject/Private/MuCO/CustomizableObjectInstance.cpp b/Engine/Plugins/Experimental/Mutable/Source/CustomizableObject/Private/MuCO/CustomizableObjectInstance.cpp
--- a/Engine/Plugins/Experimental/Mutable/Source/CustomizableObject/Private/MuCO/CustomizableObjectInstance.cpp
+++ b/Engine/Plugins/Experimental/Mutable/Source/CustomizableObject/Private/MuCO/CustomizableObjectInstance.cpp
@@ -1383,10 +1383,17 @@ UPhysicsAsset* UCustomizableInstancePrivateData::GetOrBuildPhysicsAsset(
 }
 
 
+static float MutableMeshesMinUVChannelDensity = 100.f;
+FAutoConsoleVariableRef CVarMutableMeshesMinUVChannelDensity(
+	TEXT("Mutable.MinUVChannelDensity"),
+	MutableMeshesMinUVChannelDensity,
+	TEXT("Min UV density to set on generated meshes. This value will influence the requested texture mip to stream in. Higher values will result in higher quality mips being streamed in earlier."));
+
+
 void SetMeshUVChannelDensity(FMeshUVChannelInfo& UVChannelInfo, float Density = 0.f)
 {
 	Density = Density > 0.f ? Density : 150.f;
-	Density = FMath::Max(100.0f, Density);
+	Density = FMath::Max(MutableMeshesMinUVChannelDensity, Density);
 
 	UVChannelInfo.bInitialized = true;
 	UVChannelInfo.bOverrideDensities = false;
