[Mutable] Solve Instance Parameters Multirange indexing issue.

#rnx

diff --git a/Engine/Plugins/Experimental/Mutable/Source/CustomizableObject/Private/MuCO/CustomizableObjectInstanceDescriptor.cpp b/Engine/Plugins/Experimental/Mutable/Source/CustomizableObject/Private/MuCO/CustomizableObjectInstanceDescriptor.cpp
--- a/Engine/Plugins/Experimental/Mutable/Source/CustomizableObject/Private/MuCO/CustomizableObjectInstanceDescriptor.cpp
+++ b/Engine/Plugins/Experimental/Mutable/Source/CustomizableObject/Private/MuCO/CustomizableObjectInstanceDescriptor.cpp
@@ -769,9 +769,9 @@ void FCustomizableObjectInstanceDescriptor::ReloadParameters()
 						const mu::RangeIndexPtr RangeValueIdxPtr = MutableParameters->GetValueIndex(ParamIndex, ValueIndex);
 						const int32 RangeIndex = RangeValueIdxPtr->GetPosition(0);
 
-						if (!Param.ParameterRangeValueNames.IsValidIndex(ValueIndex))
+						if (!Param.ParameterRangeValueNames.IsValidIndex(RangeIndex))
 						{
-							Param.ParameterRangeValueNames.AddDefaulted(ValueIndex + 1 - Param.ParameterRangeValueNames.Num());
+							Param.ParameterRangeValueNames.AddDefaulted(RangeIndex + 1 - Param.ParameterRangeValueNames.Num());
 						}
 
 						const int32 Value = MutableParameters->GetIntValue(ParamIndex, RangeValueIdxPtr);
@@ -824,9 +824,9 @@ void FCustomizableObjectInstanceDescriptor::ReloadParameters()
 						mu::RangeIndexPtr RangeValueIdxPtr = MutableParameters->GetValueIndex(ParamIndex, ValueIndex);
 						int32 RangeIndex = RangeValueIdxPtr->GetPosition(0);
 
-						if (!Param.ParameterRangeValues.IsValidIndex(ValueIndex))
+						if (!Param.ParameterRangeValues.IsValidIndex(RangeIndex))
 						{
-							Param.ParameterRangeValues.AddDefaulted(ValueIndex + 1 - Param.ParameterRangeValues.Num());
+							Param.ParameterRangeValues.AddDefaulted(RangeIndex + 1 - Param.ParameterRangeValues.Num());
 						}
 
 						Param.ParameterRangeValues[RangeIndex] = MutableParameters->GetFloatValue(ParamIndex, RangeValueIdxPtr);
@@ -937,9 +937,9 @@ void FCustomizableObjectInstanceDescriptor::ReloadParameters()
 						mu::RangeIndexPtr RangeValueIdxPtr = MutableParameters->GetValueIndex(ParamIndex, ValueIndex);
 						int32 RangeIndex = RangeValueIdxPtr->GetPosition(0);
 
-						if (!Param.RangeValues.IsValidIndex(ValueIndex))
+						if (!Param.RangeValues.IsValidIndex(RangeIndex))
 						{
-							Param.RangeValues.AddDefaulted(ValueIndex + 1 - Param.RangeValues.Num());
+							Param.RangeValues.AddDefaulted(RangeIndex + 1 - Param.RangeValues.Num());
 						}
 
 						GetProjector(Param.RangeValues[RangeIndex], RangeValueIdxPtr);
@@ -988,10 +988,10 @@ void FCustomizableObjectInstanceDescriptor::ReloadParameters()
 						mu::RangeIndexPtr RangeValueIdxPtr = MutableParameters->GetValueIndex(ParamIndex, ValueIndex);
 						int32 RangeIndex = RangeValueIdxPtr->GetPosition(0);
 
-						if (!Param.ParameterRangeValues.IsValidIndex(ValueIndex))
+						if (!Param.ParameterRangeValues.IsValidIndex(RangeIndex))
 						{
 							const int32 PreviousNum = Param.ParameterRangeValues.Num();
-							Param.ParameterRangeValues.AddUninitialized(ValueIndex + 1 - Param.ParameterRangeValues.Num());
+							Param.ParameterRangeValues.AddUninitialized(RangeIndex + 1 - Param.ParameterRangeValues.Num());
 
 							for (int32 Index = PreviousNum; Index < Param.ParameterRangeValues.Num(); ++Index)
 							{
