[Mutable] Created a new DestroyLiveUpdateInstance to allow freeing memory ASAP.

#rnx

diff --git a/Engine/Plugins/Experimental/Mutable/Source/CustomizableObject/Classes/MuCO/CustomizableObjectInstance.h b/Engine/Plugins/Experimental/Mutable/Source/CustomizableObject/Classes/MuCO/CustomizableObjectInstance.h
--- a/Engine/Plugins/Experimental/Mutable/Source/CustomizableObject/Classes/MuCO/CustomizableObjectInstance.h
+++ b/Engine/Plugins/Experimental/Mutable/Source/CustomizableObject/Classes/MuCO/CustomizableObjectInstance.h
@@ -274,6 +274,9 @@ public:
 	// Copy parameters from input instance
 	void CopyParametersFromInstance(UCustomizableObjectInstance* Instance);
 
+	/** Immediately destroy the Mutable Core Live Update Instance attached to this (if exists). */
+	void DestroyLiveUpdateInstance();
+	
 	// Releases all the mutable resources this instance holds, should only be called when it is not going to be used any more.
 	void ReleaseMutableResources(bool bCalledFromBeginDestroy);
 
diff --git a/Engine/Plugins/Experimental/Mutable/Source/CustomizableObject/Private/MuCO/CustomizableObjectInstance.cpp b/Engine/Plugins/Experimental/Mutable/Source/CustomizableObject/Private/MuCO/CustomizableObjectInstance.cpp
--- a/Engine/Plugins/Experimental/Mutable/Source/CustomizableObject/Private/MuCO/CustomizableObjectInstance.cpp
+++ b/Engine/Plugins/Experimental/Mutable/Source/CustomizableObject/Private/MuCO/CustomizableObjectInstance.cpp
@@ -310,6 +310,22 @@ void UCustomizableObjectInstance::BeginDestroy()
 	BeginDestroyNativeDelegate.Broadcast(this);
 
 	// Release the Live Instance ID if there it hadn't been released before
+	DestroyLiveUpdateInstance();
+
+	if (PrivateData && PrivateData->StreamingHandle)
+	{
+		PrivateData->StreamingHandle->CancelHandle();
+		PrivateData->StreamingHandle = nullptr;
+	}
+	
+	ReleaseMutableResources(true);
+
+	Super::BeginDestroy();
+}
+
+
+void UCustomizableObjectInstance::DestroyLiveUpdateInstance()
+{
 	if (PrivateData && PrivateData->LiveUpdateModeInstanceID)
 	{
 		// If FCustomizableObjectSystemPrivate::SSystem is nullptr it means it has already been destroyed, no point in registering an instanceID release
@@ -321,16 +337,6 @@ void UCustomizableObjectInstance::BeginDestroy()
 			PrivateData->LiveUpdateModeInstanceID = 0;
 		}
 	}
-
-	if (PrivateData && PrivateData->StreamingHandle)
-	{
-		PrivateData->StreamingHandle->CancelHandle();
-		PrivateData->StreamingHandle = nullptr;
-	}
-	
-	Super::BeginDestroy();
-
-	ReleaseMutableResources(true);
 }
 
 
