[mutable] Protected an access to a potentially uninitialized shared pointer in the Mutable progressive mip streaming system. So far it has not caused any issues but worth doing just in case.
#rnx
#localization none
#tests PIE gameplay and Frontend

diff --git a/Engine/Plugins/Experimental/Mutable/Source/CustomizableObject/Private/MuCO/CustomizableObjectMipDataProvider.cpp b/Engine/Plugins/Experimental/Mutable/Source/CustomizableObject/Private/MuCO/CustomizableObjectMipDataProvider.cpp
--- a/Engine/Plugins/Experimental/Mutable/Source/CustomizableObject/Private/MuCO/CustomizableObjectMipDataProvider.cpp
+++ b/Engine/Plugins/Experimental/Mutable/Source/CustomizableObject/Private/MuCO/CustomizableObjectMipDataProvider.cpp
@@ -63,72 +63,75 @@ namespace impl
 		check(OperationData->UpdateContext->Model);
 		check(OperationData->UpdateContext->Parameters.get());
 
-		mu::SystemPtr System = OperationData->UpdateContext->System;
-		const TSharedPtr<mu::Model,ESPMode::ThreadSafe> Model = OperationData->UpdateContext->Model;
+		if (OperationData.IsValid())
+		{
+			mu::SystemPtr System = OperationData->UpdateContext->System;
+			const TSharedPtr<mu::Model, ESPMode::ThreadSafe> Model = OperationData->UpdateContext->Model;
 
-		// For now, we are forcing the recreation of mutable-side instances with every update.
-		mu::Instance::ID InstanceID = System->NewInstance(Model);
-		UE_LOG(LogMutable, Verbose, TEXT("Creating Mutable instance with id [%d] for a single UpdateImage"), InstanceID)
+			// For now, we are forcing the recreation of mutable-side instances with every update.
+			mu::Instance::ID InstanceID = System->NewInstance(Model);
+			UE_LOG(LogMutable, Verbose, TEXT("Creating Mutable instance with id [%d] for a single UpdateImage"), InstanceID)
 
-		const mu::Instance* Instance = nullptr;
+				const mu::Instance* Instance = nullptr;
 
-		// Main instance generation step
-		{
-			// LOD mask, set to all ones to build all LODs
-			uint32 LODMask = 0xFFFFFFFF;
+			// Main instance generation step
+			{
+				// LOD mask, set to all ones to build all LODs
+				uint32 LODMask = 0xFFFFFFFF;
 
-			Instance = System->BeginUpdate(InstanceID, OperationData->UpdateContext->Parameters, OperationData->UpdateContext->State, LODMask);
-			check(Instance);
-		}
+				Instance = System->BeginUpdate(InstanceID, OperationData->UpdateContext->Parameters, OperationData->UpdateContext->State, LODMask);
+				check(Instance);
+			}
 
 
-		// Generate the required image
-		{
-			MUTABLE_CPUPROFILER_SCOPE(RequestedImage);
+			// Generate the required image
+			{
+				MUTABLE_CPUPROFILER_SCOPE(RequestedImage);
 
-			const FMutableImageReference& ImageRef = OperationData->RequestedImage;
+				const FMutableImageReference& ImageRef = OperationData->RequestedImage;
 
-			int32 SurfaceIndex = Instance->FindSurfaceById(ImageRef.LOD, ImageRef.Component, ImageRef.SurfaceId);
-			check(SurfaceIndex>=0);
+				int32 SurfaceIndex = Instance->FindSurfaceById(ImageRef.LOD, ImageRef.Component, ImageRef.SurfaceId);
+				check(SurfaceIndex >= 0);
 
-			// This ID may be different than the ID obtained the first time the image was generated, because the mutable
-			// runtime cannot remember all the resources it has built, and only remembers a fixed amount.
-			mu::RESOURCE_ID MipImageID = Instance->GetImageId(ImageRef.LOD, ImageRef.Component, SurfaceIndex, ImageRef.Image);
+				// This ID may be different than the ID obtained the first time the image was generated, because the mutable
+				// runtime cannot remember all the resources it has built, and only remembers a fixed amount.
+				mu::RESOURCE_ID MipImageID = Instance->GetImageId(ImageRef.LOD, ImageRef.Component, SurfaceIndex, ImageRef.Image);
 
-			// TODO: Why do we need to do this again? The full size should be stored in the initial image creation.
-			mu::FImageDesc ImageDesc;
-			mu::ImagePtrConst Image;
-			{
-				MUTABLE_CPUPROFILER_SCOPE(GetImage);
+				// TODO: Why do we need to do this again? The full size should be stored in the initial image creation.
+				mu::FImageDesc ImageDesc;
+				mu::ImagePtrConst Image;
+				{
+					MUTABLE_CPUPROFILER_SCOPE(GetImage);
 
-				Image = System->GetImage(InstanceID, MipImageID, OperationData->MipsToSkip, ImageRef.LOD);
-			}
+					Image = System->GetImage(InstanceID, MipImageID, OperationData->MipsToSkip, ImageRef.LOD);
+				}
 
-			check(Image);
+				check(Image);
 
-			OperationData->Result = Image;
-		}
-
-		// End update
-		{
-			MUTABLE_CPUPROFILER_SCOPE(EndUpdate);
-			System->EndUpdate(InstanceID);
-			System->ReleaseInstance(InstanceID);
-		}
+				OperationData->Result = Image;
+			}
 
-		{
-			// The request could be cancelled in parallel from CancelCounterSafely and its value be changed
-			// between reading it and actually running Decrement() and RescheduleCallback(), so lock
-			FScopeLock Lock(&OperationData->CounterTaskLock);
+			// End update
+			{
+				MUTABLE_CPUPROFILER_SCOPE(EndUpdate);
+				System->EndUpdate(InstanceID);
+				System->ReleaseInstance(InstanceID);
+			}
 
-			if (OperationData->Counter) // If the request has been cancelled the counter will be null
 			{
-				// Make the FMutableTextureMipDataProvider continue
-				OperationData->Counter->Decrement();
+				// The request could be cancelled in parallel from CancelCounterSafely and its value be changed
+				// between reading it and actually running Decrement() and RescheduleCallback(), so lock
+				FScopeLock Lock(&OperationData->CounterTaskLock);
 
-				if (OperationData->Counter->GetValue() == 0)
+				if (OperationData->Counter) // If the request has been cancelled the counter will be null
 				{
-					OperationData->RescheduleCallback();
+					// Make the FMutableTextureMipDataProvider continue
+					OperationData->Counter->Decrement();
+
+					if (OperationData->Counter->GetValue() == 0)
+					{
+						OperationData->RescheduleCallback();
+					}
 				}
 			}
 		}
@@ -312,12 +315,15 @@ void FMutableTextureMipDataProvider::AbortPollMips()
 
 void FMutableTextureMipDataProvider::CancelCounterSafely()
 {
-	// The Counter could be read in parallel from Task_Mutable_UpdateImage, so lock
-	FScopeLock Lock(&OperationData->CounterTaskLock);
-
-	if (OperationData->Counter)
+	if (OperationData.IsValid())
 	{
-		OperationData->Counter->Set(0);
-		OperationData->Counter = nullptr;
+		// The Counter could be read in parallel from Task_Mutable_UpdateImage, so lock
+		FScopeLock Lock(&OperationData->CounterTaskLock);
+
+		if (OperationData->Counter)
+		{
+			OperationData->Counter->Set(0);
+			OperationData->Counter = nullptr;
+		}
 	}
 }
