[Mutable] Add proper black initialization for ASTC_RGB_4x4 and ASTC_RGBA_4x4.

#rnx

diff --git a/Engine/Plugins/Experimental/Mutable/Source/MutableRuntime/Private/MuR/Image.cpp b/Engine/Plugins/Experimental/Mutable/Source/MutableRuntime/Private/MuR/Image.cpp
--- a/Engine/Plugins/Experimental/Mutable/Source/MutableRuntime/Private/MuR/Image.cpp
+++ b/Engine/Plugins/Experimental/Mutable/Source/MutableRuntime/Private/MuR/Image.cpp
@@ -10,6 +10,15 @@
 #include "MuR/OpImagePixelFormat.h"
 #include "MuR/OpImageResize.h"
 
+bool bDisableCompressedImageBlackBlockInit = false;
+
+static FAutoConsoleVariableRef CVarDisableCompressedImageBlackBlockInit(
+	TEXT("mutable.DisableCompressedImageBlackBlockInit"),
+	bDisableCompressedImageBlackBlockInit,
+	TEXT("A value of 1 disables mutable compressed black block initialization"),
+	ECVF_Default);
+
+
 namespace mu
 {
 
@@ -41,9 +50,9 @@ namespace mu
 
         FImageFormatData( 1, 1, 4, 4 ),		// IF_BGRA_UBYTE
 
-        FImageFormatData( 4, 4, 16, 3 ),	// IF_ASTC_4x4_RGB_LDR
-        FImageFormatData( 4, 4, 16, 4 ),	// IF_ASTC_4x4_RGBA_LDR
-        FImageFormatData( 4, 4, 16, 2 ),	// IF_ASTC_4x4_RG_LDR
+		FImageFormatData(4, 4, 16, 3, {252, 253, 255, 255, 255, 255, 255, 255, 0, 0, 0, 0, 0, 0, 255, 255}), // IF_ASTC_4x4_RGB_LDR
+		FImageFormatData(4, 4, 16, 4, {252, 253, 255, 255, 255, 255, 255, 255, 0, 0, 0, 0, 0, 0, 0, 0}),     // IF_ASTC_4x4_RGBA_LDR
+        FImageFormatData( 4, 4, 16, 2 ),	// IF_ASTC_4x4_RG_LDR // TODO: check black block for RG.
 		
 		FImageFormatData(8, 8, 16, 3),		// IF_ASTC_8x8_RGB_LDR,
 		FImageFormatData(8, 8, 16, 4),		// IF_ASTC_8x8_RGBA_LDR,
@@ -90,7 +99,33 @@ namespace mu
 			int32 DataSize = CalculateDataSize();
 			if (Init == EInitializationType::Black)
 			{
-				m_data.SetNumZeroed(DataSize);
+				if (bDisableCompressedImageBlackBlockInit)
+				{
+					m_data.SetNumZeroed(DataSize);
+				}
+				else
+				{
+					check(fdata.m_bytesPerBlock <= FImageFormatData::MAX_BYTES_PER_BLOCK);
+					
+					constexpr uint8 ZeroedBlock[FImageFormatData::MAX_BYTES_PER_BLOCK] = { 0 };
+
+					const SIZE_T BlockSizeSanitized = FMath::Min<SIZE_T>(FImageFormatData::MAX_BYTES_PER_BLOCK, fdata.m_bytesPerBlock);
+					const bool bIsFormatBlackBlockZeroed = FMemory::Memcmp(fdata.BlackBlock, ZeroedBlock, BlockSizeSanitized) == 0;
+
+					if (bIsFormatBlackBlockZeroed)
+					{
+						m_data.SetNumZeroed(DataSize);
+					}
+					else
+					{
+						m_data.SetNumUninitialized(DataSize);
+
+						for (int32 BlockDataOffset = 0; BlockDataOffset < DataSize; BlockDataOffset += fdata.m_bytesPerBlock)
+						{
+							FMemory::Memcpy(m_data.GetData() + BlockDataOffset, fdata.BlackBlock, fdata.m_bytesPerBlock);
+						}
+					}
+				}
 			}
 			else
 			{
diff --git a/Engine/Plugins/Experimental/Mutable/Source/MutableRuntime/Private/MuR/ImagePrivate.h b/Engine/Plugins/Experimental/Mutable/Source/MutableRuntime/Private/MuR/ImagePrivate.h
--- a/Engine/Plugins/Experimental/Mutable/Source/MutableRuntime/Private/MuR/ImagePrivate.h
+++ b/Engine/Plugins/Experimental/Mutable/Source/MutableRuntime/Private/MuR/ImagePrivate.h
@@ -7,6 +7,7 @@
 #include "MuR/SerialisationPrivate.h"
 #include "MuR/MutableMath.h"
 
+#include <initializer_list>
 
 namespace mu
 {
@@ -18,6 +19,8 @@ namespace mu
 
     struct FImageFormatData
 	{
+		static constexpr SIZE_T MAX_BYTES_PER_BLOCK = 16;
+
 		FImageFormatData
 			(
 				unsigned pixelsPerBlockX = 0,
@@ -29,7 +32,21 @@ namespace mu
 			m_pixelsPerBlockX = (uint8)pixelsPerBlockX;
 			m_pixelsPerBlockY = (uint8)pixelsPerBlockY;
 			m_bytesPerBlock = (uint16)bytesPerBlock;
-			m_channels = (uint16)channels;
+			m_channels = (uint16)channels;	
+		}
+
+		FImageFormatData
+			(
+				unsigned pixelsPerBlockX,
+				unsigned pixelsPerBlockY,
+				unsigned bytesPerBlock,
+				unsigned channels,
+				std::initializer_list<uint8> BlackBlockInit
+			)
+			: FImageFormatData(pixelsPerBlockX, pixelsPerBlockY, bytesPerBlock, channels)
+		{
+			check(BlackBlockInit.size() <= MAX_BYTES_PER_BLOCK);
+			FMemory::Memcpy(BlackBlock, BlackBlockInit.begin(), FMath::Min<SIZE_T>(MAX_BYTES_PER_BLOCK, BlackBlockInit.size()));
 		}
 
 		//! For block based formats, size of the block size. For uncompressed formats it will
@@ -42,6 +59,9 @@ namespace mu
 
 		//! Channels in every pixel of the image.
         uint16 m_channels;
+
+		//! Representation of a black block of the image.
+		uint8 BlackBlock[MAX_BYTES_PER_BLOCK] = { 0 };
 	};
 
 
