[Mutable] Add a CVar to limit the texture size on non streamed textures.
Cvar: Mutable.MaxTextureSizeToGenerate

#rnx

#localization none
#tests PIE, WindowsClient and PS5 Client

diff --git a/Engine/Plugins/Experimental/Mutable/Source/CustomizableObject/Private/MuCO/CustomizableObjectInstance.cpp b/Engine/Plugins/Experimental/Mutable/Source/CustomizableObject/Private/MuCO/CustomizableObjectInstance.cpp
--- a/Engine/Plugins/Experimental/Mutable/Source/CustomizableObject/Private/MuCO/CustomizableObjectInstance.cpp
+++ b/Engine/Plugins/Experimental/Mutable/Source/CustomizableObject/Private/MuCO/CustomizableObjectInstance.cpp
@@ -2499,6 +2499,12 @@ void CopyTextureProperties(UTexture2D* Texture, const UTexture2D* SourceTexture)
 	Texture->AddressY = SourceTexture->AddressY;
 }
 
+static int32 MaxTextureSizeToGenerate = 0;
+FAutoConsoleVariableRef CVarMaxTextureSizeToGenerate(
+	TEXT("Mutable.MaxTextureSizeToGenerate"),
+	MaxTextureSizeToGenerate,
+	TEXT("Max texture size on non streamed textures. Mip 0 will be the first mip with max size equal or less than MaxTextureSizeToGenerate."
+		"If a texture doesn't have small enough mips, mip 0 will be the last mip available."));
 
 // The memory allocated in the function and pointed by the returned pointer is owned by the caller and must be freed. 
 // If assigned to a UTexture2D, it will be freed by that UTexture2D
@@ -2536,6 +2542,28 @@ FTexturePlatformData* UCustomizableInstancePrivateData::MutableCreateImagePlatfo
 		check(MipsToSkip >= 0);
 	}
 
+	// Reduce final texture size if we surpass the max size we can generate.
+	if (MaxTextureSizeToGenerate > 0)
+	{
+		// Skip mips only if texture streaming is disabled 
+		const bool bIsStreamingEnabled = MipsToSkip > 0;
+
+		// Skip mips if the texture surpasses a certain size
+		if (MaxSize > MaxTextureSizeToGenerate && !bIsStreamingEnabled && OnlyLOD < 0)
+		{
+			// Skip mips until MaxSize is equal or less than MaxTextureSizeToGenerate or there aren't more mips to skip
+			while (MaxSize > MaxTextureSizeToGenerate && FirstLOD < (FullLODCount - 1))
+			{
+				MaxSize = MaxSize >> 1;
+				FirstLOD++;
+			}
+
+			// Update SizeX and SizeY
+			SizeX = SizeX >> FirstLOD;
+			SizeY = SizeY >> FirstLOD;
+		}
+	}
+
 	if (MutableImage->GetLODCount() == 1)
 	{
 		MipsToSkip = 0;
@@ -2547,8 +2575,17 @@ FTexturePlatformData* UCustomizableInstancePrivateData::MutableCreateImagePlatfo
 	
 	mu::EImageFormat MutableFormat = MutableImage->GetFormat();
 
-	int32 MaxPossibleSize = int32(FMath::Pow(2.f, float(FullLODCount - 1)));
-	
+	int32 MaxPossibleSize = 0;
+		
+	if (MaxTextureSizeToGenerate > 0)
+	{
+		MaxPossibleSize = int32(FMath::Pow(2.f, float(FullLODCount - FirstLOD - 1)));
+	}
+	else
+	{
+		MaxPossibleSize = int32(FMath::Pow(2.f, float(FullLODCount - 1)));
+	}
+
 	// This could happen with non-power-of-two images.
 	//check(SizeX == MaxPossibleSize || SizeY == MaxPossibleSize || FullLODCount == 1);
 	if (!(SizeX == MaxPossibleSize || SizeY == MaxPossibleSize || FullLODCount == 1))
@@ -2728,7 +2765,14 @@ FTexturePlatformData* UCustomizableInstancePrivateData::MutableCreateImagePlatfo
 		}
 	}
 
-	check( FullLODCount==1 || OnlyLOD >= 0 || (BulkDataCount == MutableImage->GetLODCount()));
+	if (MaxTextureSizeToGenerate > 0)
+	{
+		check(FullLODCount == 1 || OnlyLOD >= 0 || (BulkDataCount == (MutableImage->GetLODCount() - FirstLOD)));
+	}
+	else
+	{
+		check(FullLODCount == 1 || OnlyLOD >= 0 || (BulkDataCount == MutableImage->GetLODCount()));
+	}
 #endif
 
 	return PlatformData;
