[mutable] Update the support for the new task graph interface but don't enable it.
Pull fixes for the missing event graph call to subsequents and also for the extra null pointer checks on operation data.
#rnx

diff --git a/Engine/Plugins/Experimental/Mutable/Source/CustomizableObject/Private/MuCO/CustomizableObjectMipDataProvider.cpp b/Engine/Plugins/Experimental/Mutable/Source/CustomizableObject/Private/MuCO/CustomizableObjectMipDataProvider.cpp
--- a/Engine/Plugins/Experimental/Mutable/Source/CustomizableObject/Private/MuCO/CustomizableObjectMipDataProvider.cpp
+++ b/Engine/Plugins/Experimental/Mutable/Source/CustomizableObject/Private/MuCO/CustomizableObjectMipDataProvider.cpp
@@ -248,7 +248,14 @@ bool FMutableTextureMipDataProvider::PollMips(const FTextureUpdateSyncOptions& S
 	}
 #endif
 
-	if (!bRequestAborted && UpdateImageMutableTaskEvent)
+	if (!bRequestAborted 
+		&& 
+#ifdef MUTABLE_USE_NEW_TASKGRAPH
+		UpdateImageMutableTaskEvent.IsCompleted()
+#else
+		UpdateImageMutableTaskEvent
+#endif
+		)
 	{
 		if (OperationData && OperationData->Result && OperationData->Levels.Num())
 		{
diff --git a/Engine/Plugins/Experimental/Mutable/Source/CustomizableObject/Private/MuCO/CustomizableObjectSystem.cpp b/Engine/Plugins/Experimental/Mutable/Source/CustomizableObject/Private/MuCO/CustomizableObjectSystem.cpp
--- a/Engine/Plugins/Experimental/Mutable/Source/CustomizableObject/Private/MuCO/CustomizableObjectSystem.cpp
+++ b/Engine/Plugins/Experimental/Mutable/Source/CustomizableObject/Private/MuCO/CustomizableObjectSystem.cpp
@@ -2018,7 +2018,11 @@ namespace impl
 
 		// Task: Mutable GetImages
 		//-------------------------------------------------------------
+#ifdef MUTABLE_USE_NEW_TASKGRAPH
+		UE::Tasks::FTask Mutable_GetImagesTask;
+#else
 		FGraphEventRef Mutable_GetImagesTask;
+#endif
 		{
 			// Task inputs
 			check(CustomizableObject->GetPrivate() != nullptr);
@@ -2078,7 +2082,11 @@ namespace impl
 		check(CurrentOperationData);
 		CurrentOperationData->InstanceID = IDToRelease;
 
+#ifdef MUTABLE_USE_NEW_TASKGRAPH
+		UE::Tasks::FTask Mutable_GetMeshTask;
+#else
 		FGraphEventRef Mutable_GetMeshTask;
+#endif
 		{
 			// Task inputs
 			TSharedPtr<FMutableOperation> CurrentMutableOperation = SystemPrivateData->CurrentMutableOperation;
@@ -2264,7 +2272,11 @@ namespace impl
 			CurrentOperationData->RequestedLODs = Operation->InstanceDescriptorRuntimeHash.GetRequestedLODs();
 		}
 
+#ifdef MUTABLE_USE_NEW_TASKGRAPH
+		UE::Tasks::FTask Mutable_GetMeshTask;
+#else
 		FGraphEventRef Mutable_GetMeshTask;
+#endif
 		{
 			// Task inputs
 			TWeakObjectPtr<UCustomizableObjectInstance> CustomizableObjectInstancePtr = CandidateInstance;
diff --git a/Engine/Plugins/Experimental/Mutable/Source/CustomizableObject/Private/MuCO/CustomizableObjectSystemPrivate.h b/Engine/Plugins/Experimental/Mutable/Source/CustomizableObject/Private/MuCO/CustomizableObjectSystemPrivate.h
--- a/Engine/Plugins/Experimental/Mutable/Source/CustomizableObject/Private/MuCO/CustomizableObjectSystemPrivate.h
+++ b/Engine/Plugins/Experimental/Mutable/Source/CustomizableObject/Private/MuCO/CustomizableObjectSystemPrivate.h
@@ -236,7 +236,12 @@ struct FMutableTask
 
 	/** From the traditional event graph system, still used to wait for async loads. */
 	FGraphEventRef GraphDependency0;
+	
+#ifdef MUTABLE_USE_NEW_TASKGRAPH
+	UE::Tasks::FTask GraphDependency1;
+#else
 	FGraphEventRef GraphDependency1;
+#endif
 
 	/** Check if all the dependencies of this task have been completed. */
 	inline bool AreDependenciesComplete() const
@@ -246,7 +251,7 @@ struct FMutableTask
 			&&
 			(!GraphDependency0 || GraphDependency0->IsComplete())
 			&&
-			(!GraphDependency1 || GraphDependency1->IsComplete());
+			(GraphDependency1.IsCompleted());
 #else
 		return 
 			(!GraphDependency0 || GraphDependency0->IsComplete())
@@ -262,7 +267,12 @@ struct FMutableTask
 		Dependency = {};
 #endif
 		GraphDependency0 = nullptr;
+
+#ifdef MUTABLE_USE_NEW_TASKGRAPH
+		GraphDependency1 = {};
+#else
 		GraphDependency1 = nullptr;
+#endif
 	}
 };
 
@@ -608,6 +618,7 @@ public:
 
 		return LastMutableTask;
 	}
+#endif
 
 	/** This should only be used when shutting down. */
 	inline void WaitForMutableTasks()
@@ -648,14 +659,18 @@ public:
 	{
 		check(IsInGameThread());
 
+#ifdef MUTABLE_USE_NEW_TASKGRAPH
+		UE::Tasks::Launch(DebugName, MoveTemp(TaskBody), Priority);
+#else
 		FGraphEventRef Task = FFunctionGraphTask::CreateAndDispatchWhenReady(
 			MoveTemp(TaskBody),
 			TStatId{},
 			nullptr,
 			ENamedThreads::AnyThread);
 		Task->SetDebugName(TEXT("Mutable_Anythread"));
-	}
 #endif
+	}
+
 
 	//
 	//TSharedPtr<FMutableOperationData> CurrentOperationData;
diff --git a/Engine/Plugins/Experimental/Mutable/Source/CustomizableObject/Private/MuCO/UnrealMutableImageProvider.cpp b/Engine/Plugins/Experimental/Mutable/Source/CustomizableObject/Private/MuCO/UnrealMutableImageProvider.cpp
--- a/Engine/Plugins/Experimental/Mutable/Source/CustomizableObject/Private/MuCO/UnrealMutableImageProvider.cpp
+++ b/Engine/Plugins/Experimental/Mutable/Source/CustomizableObject/Private/MuCO/UnrealMutableImageProvider.cpp
@@ -115,7 +115,11 @@ mu::EImageFormat GetMutablePixelFormat(EPixelFormat InTextureFormat)
 
 
 //-------------------------------------------------------------------------------------------------
-TTuple<FGraphEventRef, TFunction<void()>> FUnrealMutableImageProvider::GetImageAsync(mu::EXTERNAL_IMAGE_ID id, uint8 MipmapsToSkip, TFunction<void (mu::Ptr<mu::Image>)>& ResultCallback)
+#ifdef MUTABLE_USE_NEW_TASKGRAPH
+	TTuple<UE::Tasks::FTask, TFunction<void()>> FUnrealMutableImageProvider::GetImageAsync(mu::EXTERNAL_IMAGE_ID id, uint8 MipmapsToSkip, TFunction<void(mu::Ptr<mu::Image>)>& ResultCallback)
+#else
+	TTuple<FGraphEventRef, TFunction<void()>> FUnrealMutableImageProvider::GetImageAsync(mu::EXTERNAL_IMAGE_ID id, uint8 MipmapsToSkip, TFunction<void(mu::Ptr<mu::Image>)>& ResultCallback)
+#endif
 {
 	// Thread: worker
 	MUTABLE_CPUPROFILER_SCOPE(FUnrealMutableImageProvider::GetImage);
@@ -133,6 +137,15 @@ TTuple<FGraphEventRef, TFunction<void()>> FUnrealMutableImageProvider::GetImageA
 	mu::EImageFormat MutImageFormat = mu::EImageFormat::IF_NONE;
 	int32 MutImageDataSize = 0;
 
+#ifdef MUTABLE_USE_NEW_TASKGRAPH
+	auto TrivialReturn = []() -> TTuple<UE::Tasks::FTask, TFunction<void()>>
+	{
+		UE::Tasks::FTaskEvent CompletionEvent(TEXT("GetImageAsyncCompleted"));
+		CompletionEvent.Trigger();
+
+		return MakeTuple(CompletionEvent, []() -> void {});
+	};
+#else
 	auto TrivialReturn = []() -> TTuple<FGraphEventRef, TFunction<void()>>
 	{
 		FGraphEventRef CompletionEvent = FGraphEvent::CreateGraphEvent();
@@ -140,6 +153,7 @@ TTuple<FGraphEventRef, TFunction<void()>> FUnrealMutableImageProvider::GetImageA
 
 		return MakeTuple(CompletionEvent, []() -> void {});
 	};
+#endif
 
 	{
 		FScopeLock Lock(&ExternalImagesLock);
@@ -222,7 +236,12 @@ TTuple<FGraphEventRef, TFunction<void()>> FUnrealMutableImageProvider::GetImageA
 			// Create a streaming request if the data is not loaded or copy the mip data
 			if (!BulkData.IsBulkDataLoaded())
 			{
+#ifdef MUTABLE_USE_NEW_TASKGRAPH
+				UE::Tasks::FTaskEvent IORequestCompletionEvent(TEXT("Mutable_IORequestCompletionEvent"));
+#else
 				FGraphEventRef IORequestCompletionEvent = FGraphEvent::CreateGraphEvent();
+#endif
+
 
 				TFunction<void(bool, IBulkDataIORequest*)> IOCallback =
 					[
@@ -237,10 +256,15 @@ TTuple<FGraphEventRef, TFunction<void()>> FUnrealMutableImageProvider::GetImageA
 				{
 					ON_SCOPE_EXIT
 					{
+#ifdef MUTABLE_USE_NEW_TASKGRAPH
+						UE::Tasks::FTaskEvent EventCopy = IORequestCompletionEvent;
+						EventCopy.Trigger();
+#else
 						if (IORequestCompletionEvent.IsValid())
 						{
 							IORequestCompletionEvent->DispatchSubsequents();
 						}
+#endif
 					};
 					
 					// Should we do someting different than returning a dummy image if cancelled?
diff --git a/Engine/Plugins/Experimental/Mutable/Source/CustomizableObject/Private/MuCO/UnrealMutableImageProvider.h b/Engine/Plugins/Experimental/Mutable/Source/CustomizableObject/Private/MuCO/UnrealMutableImageProvider.h
--- a/Engine/Plugins/Experimental/Mutable/Source/CustomizableObject/Private/MuCO/UnrealMutableImageProvider.h
+++ b/Engine/Plugins/Experimental/Mutable/Source/CustomizableObject/Private/MuCO/UnrealMutableImageProvider.h
@@ -6,6 +6,11 @@
 #include "MuR/System.h"
 #include "UObject/GCObject.h"
 
+#ifdef MUTABLE_USE_NEW_TASKGRAPH
+#include "Tasks/Task.h"
+#else
+#endif
+
 class UTexture2D;
 
 
@@ -17,7 +22,11 @@ public:
 
 	// mu::ImageParameterGenerator interface
 	// Thread: worker
-	virtual TTuple<FGraphEventRef, TFunction<void()>> GetImageAsync(mu::EXTERNAL_IMAGE_ID Id, uint8 MipmapsToSkip, TFunction<void (mu::Ptr<mu::Image>)>& ResultCallback) override;
+#ifdef MUTABLE_USE_NEW_TASKGRAPH
+	virtual TTuple<UE::Tasks::FTask, TFunction<void()>> GetImageAsync(mu::EXTERNAL_IMAGE_ID Id, uint8 MipmapsToSkip, TFunction<void(mu::Ptr<mu::Image>)>& ResultCallback) override;
+#else
+	virtual TTuple<FGraphEventRef, TFunction<void()>> GetImageAsync(mu::EXTERNAL_IMAGE_ID Id, uint8 MipmapsToSkip, TFunction<void(mu::Ptr<mu::Image>)>& ResultCallback) override;
+#endif
 
 	virtual mu::FImageDesc GetImageDesc(mu::EXTERNAL_IMAGE_ID Id, uint8 MipmapsToSkip) override;
 
diff --git a/Engine/Plugins/Experimental/Mutable/Source/CustomizableObject/Public/MuCO/CustomizableObjectMipDataProvider.h b/Engine/Plugins/Experimental/Mutable/Source/CustomizableObject/Public/MuCO/CustomizableObjectMipDataProvider.h
--- a/Engine/Plugins/Experimental/Mutable/Source/CustomizableObject/Public/MuCO/CustomizableObjectMipDataProvider.h
+++ b/Engine/Plugins/Experimental/Mutable/Source/CustomizableObject/Public/MuCO/CustomizableObjectMipDataProvider.h
@@ -8,6 +8,9 @@
 #include "MuR/Image.h"
 #include "Streaming/TextureMipDataProvider.h"
 
+ // For the MUTABLE_USE_NEW_TASKGRAPH define
+#include "MuR/System.h"
+
 #include "CustomizableObjectMipDataProvider.generated.h"
 
 enum EPixelFormat : uint8;
@@ -97,7 +100,11 @@ public:
 	bool bRequestAborted = false;
 
 	TSharedPtr<FMutableImageOperationData> OperationData;
+#ifdef MUTABLE_USE_NEW_TASKGRAPH
+	UE::Tasks::FTask UpdateImageMutableTaskEvent;
+#else
 	FGraphEventRef UpdateImageMutableTaskEvent;
+#endif
 };
 
 
diff --git a/Engine/Plugins/Experimental/Mutable/Source/CustomizableObjectEditor/Private/MuCOE/SMutableCodeViewer.cpp b/Engine/Plugins/Experimental/Mutable/Source/CustomizableObjectEditor/Private/MuCOE/SMutableCodeViewer.cpp
--- a/Engine/Plugins/Experimental/Mutable/Source/CustomizableObjectEditor/Private/MuCOE/SMutableCodeViewer.cpp
+++ b/Engine/Plugins/Experimental/Mutable/Source/CustomizableObjectEditor/Private/MuCOE/SMutableCodeViewer.cpp
@@ -2241,7 +2241,11 @@ namespace
 			mu::FImageDesc(mu::FImageSize(1024, 1024), mu::EImageFormat::IF_RGBA_UBYTE, 1);
 		
 	public:
-		TTuple<FGraphEventRef, TFunction<void()>> GetImageAsync(mu::EXTERNAL_IMAGE_ID id, uint8 MipmapsToSkip, TFunction<void (mu::Ptr<mu::Image>)>& ResultCallback) override
+#ifdef MUTABLE_USE_NEW_TASKGRAPH
+		TTuple<UE::Tasks::FTask, TFunction<void()>> GetImageAsync(mu::EXTERNAL_IMAGE_ID id, uint8 MipmapsToSkip, TFunction<void(mu::Ptr<mu::Image>)>& ResultCallback) override
+#else
+		TTuple<FGraphEventRef, TFunction<void()>> GetImageAsync(mu::EXTERNAL_IMAGE_ID id, uint8 MipmapsToSkip, TFunction<void(mu::Ptr<mu::Image>)>& ResultCallback) override
+#endif
 		{
 			MUTABLE_CPUPROFILER_SCOPE(TestImageProvider_GetImage);
 
@@ -2272,8 +2276,14 @@ namespace
 			}
 
 			ResultCallback(Image);
+
+#ifdef MUTABLE_USE_NEW_TASKGRAPH
+			UE::Tasks::FTaskEvent CompletionEvent(TEXT("TestImageProvider_GetImageAsunc_Completed"));
+			CompletionEvent.Trigger();
+#else
 			FGraphEventRef CompletionEvent = FGraphEvent::CreateGraphEvent();
 			CompletionEvent->DispatchSubsequents();
+#endif
 
 			return MakeTuple(CompletionEvent, []() -> void {});
 		}
diff --git a/Engine/Plugins/Experimental/Mutable/Source/MutableRuntime/Private/MuR/CodeRunner.cpp b/Engine/Plugins/Experimental/Mutable/Source/MutableRuntime/Private/MuR/CodeRunner.cpp
--- a/Engine/Plugins/Experimental/Mutable/Source/MutableRuntime/Private/MuR/CodeRunner.cpp
+++ b/Engine/Plugins/Experimental/Mutable/Source/MutableRuntime/Private/MuR/CodeRunner.cpp
@@ -113,7 +113,11 @@ namespace mu
 
 
     //---------------------------------------------------------------------------------------------
+#ifdef MUTABLE_USE_NEW_TASKGRAPH
+	TTuple<UE::Tasks::FTask, TFunction<void()>> CodeRunner::LoadExternalImageAsync(EXTERNAL_IMAGE_ID Id, uint8 MipmapsToSkip, TFunction<void(Ptr<Image>)>& ResultCallback)
+#else
 	TTuple<FGraphEventRef, TFunction<void()>> CodeRunner::LoadExternalImageAsync(EXTERNAL_IMAGE_ID Id, uint8 MipmapsToSkip, TFunction<void(Ptr<Image>)>& ResultCallback)
+#endif
     {
 		MUTABLE_CPUPROFILER_SCOPE(LoadExternalImageAsync);
 
@@ -134,8 +138,13 @@ namespace mu
 		}
 
 		// Not needed as it should never reach this point, but added for correctness.
+#ifdef MUTABLE_USE_NEW_TASKGRAPH
+		UE::Tasks::FTaskEvent CompletionEvent(TEXT("LoadExternalImageAsyncCompletion"));
+		CompletionEvent.Trigger();
+#else
 		FGraphEventRef CompletionEvent = FGraphEvent::CreateGraphEvent();
 		CompletionEvent->DispatchSubsequents();
+#endif
 
 		return MakeTuple(CompletionEvent, []() -> void {});
 	}
diff --git a/Engine/Plugins/Experimental/Mutable/Source/MutableRuntime/Private/MuR/CodeRunner.h b/Engine/Plugins/Experimental/Mutable/Source/MutableRuntime/Private/MuR/CodeRunner.h
--- a/Engine/Plugins/Experimental/Mutable/Source/MutableRuntime/Private/MuR/CodeRunner.h
+++ b/Engine/Plugins/Experimental/Mutable/Source/MutableRuntime/Private/MuR/CodeRunner.h
@@ -106,7 +106,11 @@ namespace  mu
 
 		//! Load an external image asynchronously, retuns an event to wait for complition and a cleanup function 
 		//! that must be called once the event has completed.
-        TTuple<FGraphEventRef, TFunction<void()>> LoadExternalImageAsync(EXTERNAL_IMAGE_ID Id, uint8 MipmapsToSkip, TFunction<void (Ptr<Image>)>& ResultCallback);
+#ifdef MUTABLE_USE_NEW_TASKGRAPH
+		TTuple<UE::Tasks::FTask, TFunction<void()>> LoadExternalImageAsync(EXTERNAL_IMAGE_ID Id, uint8 MipmapsToSkip, TFunction<void(Ptr<Image>)>& ResultCallback);
+#else
+		TTuple<FGraphEventRef, TFunction<void()>> LoadExternalImageAsync(EXTERNAL_IMAGE_ID Id, uint8 MipmapsToSkip, TFunction<void(Ptr<Image>)>& ResultCallback);
+#endif
  	    mu::FImageDesc GetExternalImageDesc(EXTERNAL_IMAGE_ID Id, uint8 MipmapsToSkip);
 
     protected:
diff --git a/Engine/Plugins/Experimental/Mutable/Source/MutableRuntime/Public/MuR/System.h b/Engine/Plugins/Experimental/Mutable/Source/MutableRuntime/Public/MuR/System.h
--- a/Engine/Plugins/Experimental/Mutable/Source/MutableRuntime/Public/MuR/System.h
+++ b/Engine/Plugins/Experimental/Mutable/Source/MutableRuntime/Public/MuR/System.h
@@ -17,6 +17,11 @@
 // This is currently broken in Switch and maybe other consoles, for some unknown reason.
 //#define MUTABLE_USE_NEW_TASKGRAPH
 
+#ifdef MUTABLE_USE_NEW_TASKGRAPH
+#include "Tasks/Task.h"
+#else
+#endif
+
 
 namespace mu
 {
@@ -72,7 +77,11 @@ namespace mu
         virtual ~ImageParameterGenerator() = default;
 
         //! Returns the completion event and a cleanup function that must be called once event is completed.
-        virtual TTuple<FGraphEventRef, TFunction<void()>> GetImageAsync(EXTERNAL_IMAGE_ID Id, uint8 MipmapsToSkip, TFunction<void (Ptr<Image>)>& ResultCallback) = 0;
+#ifdef MUTABLE_USE_NEW_TASKGRAPH
+		virtual TTuple<UE::Tasks::FTask, TFunction<void()>> GetImageAsync(EXTERNAL_IMAGE_ID Id, uint8 MipmapsToSkip, TFunction<void(Ptr<Image>)>& ResultCallback) = 0;
+#else
+		virtual TTuple<FGraphEventRef, TFunction<void()>> GetImageAsync(EXTERNAL_IMAGE_ID Id, uint8 MipmapsToSkip, TFunction<void(Ptr<Image>)>& ResultCallback) = 0;
+#endif
 
         virtual mu::FImageDesc GetImageDesc(EXTERNAL_IMAGE_ID Id, uint8 MipmapsToSkip) = 0;
     };
