[mutable] Realign a structure and add padding to avoid uninitialized bytes raising warning in memory checkers and causing indeterministic binary data.
#rnx

diff --git a/Engine/Plugins/Experimental/Mutable/Source/CustomizableObject/Classes/MuCO/CustomizableObject.h b/Engine/Plugins/Experimental/Mutable/Source/CustomizableObject/Classes/MuCO/CustomizableObject.h
--- a/Engine/Plugins/Experimental/Mutable/Source/CustomizableObject/Classes/MuCO/CustomizableObject.h
+++ b/Engine/Plugins/Experimental/Mutable/Source/CustomizableObject/Classes/MuCO/CustomizableObject.h
@@ -1203,7 +1203,7 @@ private:
 	// This is a manual version number for the binary blobs in this asset.
 	// Increasing it invalidates all the previously compiled models.
 	// Warning: If while merging code both versions have changed, take the highest+1.
-	static const int32 CurrentSupportedVersion = 369;
+	static const int32 CurrentSupportedVersion = 370;
 
 public:
 
diff --git a/Engine/Plugins/Experimental/Mutable/Source/MutableRuntime/Private/MuR/ModelPrivate.h b/Engine/Plugins/Experimental/Mutable/Source/MutableRuntime/Private/MuR/ModelPrivate.h
--- a/Engine/Plugins/Experimental/Mutable/Source/MutableRuntime/Private/MuR/ModelPrivate.h
+++ b/Engine/Plugins/Experimental/Mutable/Source/MutableRuntime/Private/MuR/ModelPrivate.h
@@ -57,9 +57,10 @@ namespace mu
 	struct FImageLODRange
 	{
 		int32 FirstIndex = 0;
-		int32 LODCount = 0;
 		uint16 ImageSizeX = 0;
 		uint16 ImageSizeY = 0;
+		uint16 _Padding = 0;
+		uint8 LODCount = 0;
 		EImageFormat ImageFormat = EImageFormat::IF_NONE;
 	};
 	MUTABLE_DEFINE_POD_SERIALISABLE(FImageLODRange);
@@ -394,7 +395,13 @@ namespace mu
 				}
 			}
 
-			int32 ImageIndex = m_constantImages.Add({ FirstLODIndexIndex, MipsToStore, pImage->GetSizeX(), pImage->GetSizeY(), pImage->GetFormat() });
+			FImageLODRange LODRange;
+			LODRange.FirstIndex = FirstLODIndexIndex;
+			LODRange.LODCount = MipsToStore;
+			LODRange.ImageFormat = pImage->GetFormat();
+			LODRange.ImageSizeX = pImage->GetSizeX();
+			LODRange.ImageSizeY = pImage->GetSizeY();
+			int32 ImageIndex = m_constantImages.Add(LODRange);
 			return ImageIndex;
 		}
 
diff --git a/Engine/Plugins/Experimental/Mutable/Source/MutableRuntime/Public/MuR/Model.h b/Engine/Plugins/Experimental/Mutable/Source/MutableRuntime/Public/MuR/Model.h
--- a/Engine/Plugins/Experimental/Mutable/Source/MutableRuntime/Public/MuR/Model.h
+++ b/Engine/Plugins/Experimental/Mutable/Source/MutableRuntime/Public/MuR/Model.h
@@ -13,7 +13,7 @@
 //! data structures. Compiled models are not necessarily compatible when the runtime is updated,
 //! so this version number can be used externally to verify this. It is not used internally, and
 //! serializing models from different versions than this runtime will probably result in a crash.
-#define MUTABLE_COMPILED_MODEL_CODE_VERSION		uint32( 56 )
+#define MUTABLE_COMPILED_MODEL_CODE_VERSION		uint32( 57 )
 #define MUTABLE_PARAMETERS_VERSION              uint32( 1 )
 
 
