[mutable] Various optimization in image operations.
- Rewrote some to ease the use of SIMD by the compiler.
- Calculate relevant rows of an image generated in ImageProject to use as hint for some other operations (ImageLayer)
- Image grow operation border size is now made into a compile-time constant easier to track.
#rnx

diff --git a/Engine/Plugins/Experimental/Mutable/Source/CustomizableObjectEditor/Private/MuCOE/SMutableCodeViewer.cpp b/Engine/Plugins/Experimental/Mutable/Source/CustomizableObjectEditor/Private/MuCOE/SMutableCodeViewer.cpp
--- a/Engine/Plugins/Experimental/Mutable/Source/CustomizableObjectEditor/Private/MuCOE/SMutableCodeViewer.cpp
+++ b/Engine/Plugins/Experimental/Mutable/Source/CustomizableObjectEditor/Private/MuCOE/SMutableCodeViewer.cpp
@@ -2252,24 +2252,31 @@ namespace
 			int32 Size = IMAGE_DESC.m_size[0];
 			Size = FMath::Max(4, Size / (1 << MipmapsToSkip));
 
-			mu::Ptr<mu::Image> Image = 
-					new mu::Image(IMAGE_DESC.m_size[0], IMAGE_DESC.m_size[1], IMAGE_DESC.m_lods, IMAGE_DESC.m_format);
+			mu::Ptr<mu::Image> Image = new mu::Image(
+				Size, Size, IMAGE_DESC.m_lods,
+				IMAGE_DESC.m_format,
+				mu::EInitializationType::NotInitialized);
 
-			// Generate an alphatested circle with an horizontal gradient color.
+			// Generate an alpha-tested circle with an horizontal gradient color.
 			uint8* Data = Image->GetData();
 			int32 CircleRadius = (Size * 2) / 5;
 			int32 CircleRadius2 = CircleRadius * CircleRadius;
 			int32 Color[3] = { 255,128,0 };
 
-			for (int y = 0; y < Size; ++y)
+			int32 LogSize = FMath::CeilLogTwo(Size);
+
+			int32 HalfSize = Size >> 1;
+			for (int32 RadY = -HalfSize; RadY < HalfSize; ++RadY)
 			{
-				for (int x = 0; x < Size; ++x)
+				int32 RadY2 = RadY * RadY;
+				for (int32 x = 0; x < Size; ++x)
 				{
-					float R2 = (x - Size / 2) * (x - Size / 2) + (y - Size / 2) * (y - Size / 2);
-					int32 Opacity = FMath::Clamp( (CircleRadius2-R2)/CircleRadius2 * 512 - 64, 0, 255 );
-					Data[0] = (Color[0] * x) / Size;
-					Data[1] = (Color[1] * x) / Size;
-					Data[2] = (Color[2] * x) / Size;
+					int32 RadX = (x - HalfSize);
+					int32 R2 = RadX * RadX + RadY2;
+					int32 Opacity = FMath::Clamp(((CircleRadius2 - R2) * 512) / CircleRadius2 - 64, 0, 255);
+					Data[0] = uint8((Color[0] * x) >> LogSize);
+					Data[1] = uint8((Color[1] * x) >> LogSize);
+					Data[2] = uint8((Color[2] * x) >> LogSize);
 					Data[3] = uint8(Opacity);
 					Data += 4;
 				}
diff --git a/Engine/Plugins/Experimental/Mutable/Source/MutableRuntime/Private/MuR/CodeRunner.cpp b/Engine/Plugins/Experimental/Mutable/Source/MutableRuntime/Private/MuR/CodeRunner.cpp
--- a/Engine/Plugins/Experimental/Mutable/Source/MutableRuntime/Private/MuR/CodeRunner.cpp
+++ b/Engine/Plugins/Experimental/Mutable/Source/MutableRuntime/Private/MuR/CodeRunner.cpp
@@ -4089,6 +4089,14 @@ namespace mu
 				if (OriginalSourceScale.x() > 0 && OriginalSourceScale.y() > 0)
 				{
 					ImageDisplace(pNew.get(), pSource.get(), pMap.get());
+					
+					// Update the relevancy data of the image for the worst case. 
+					if (pNew->m_flags & Image::EImageFlags::IF_HAS_RELEVANCY_MAP)
+					{
+						// Displace can encode at max MUTABLE_GROW_BORDER_VALUE pixels away according to "border" in MakeGrowMap.
+						pNew->RelevancyMinY = FMath::Max(int32(pNew->RelevancyMinY) - MUTABLE_GROW_BORDER_VALUE, 0);
+						pNew->RelevancyMaxY = FMath::Min(pNew->RelevancyMaxY + MUTABLE_GROW_BORDER_VALUE, pNew->GetSizeY() - 1);
+					}
 
 					if (OriginalSourceScale != pNew->GetSize())
 					{
diff --git a/Engine/Plugins/Experimental/Mutable/Source/MutableRuntime/Private/MuR/CodeRunnerTasked.cpp b/Engine/Plugins/Experimental/Mutable/Source/MutableRuntime/Private/MuR/CodeRunnerTasked.cpp
--- a/Engine/Plugins/Experimental/Mutable/Source/MutableRuntime/Private/MuR/CodeRunnerTasked.cpp
+++ b/Engine/Plugins/Experimental/Mutable/Source/MutableRuntime/Private/MuR/CodeRunnerTasked.cpp
@@ -1173,7 +1173,7 @@ namespace mu
 	class FImageResizeTask : public CodeRunner::FIssuedTask
 	{
 	public:
-		FImageResizeTask(FScheduledOp, FProgramCache&, const OP::ImageResizeArgs&, int32 ImageCompressionQuality);
+		FImageResizeTask(FScheduledOp InOp, FProgramCache& Memory, const Ptr<const Image>& InBase, uint16 InSizeX, uint16 InSizeY, int32 InImageCompressionQuality);
 
 		// FIssuedTask interface
 		void DoWork() override;
@@ -1182,17 +1182,18 @@ namespace mu
 	private:
 		Ptr<const Image> pBase;
 		Ptr<const Image> Result;
-		OP::ImageResizeArgs Args;
 		int32 ImageCompressionQuality;
+		uint16 SizeX, SizeY;
 	};
 
 
 	//---------------------------------------------------------------------------------------------
-	FImageResizeTask::FImageResizeTask(FScheduledOp InOp, FProgramCache& Memory, const OP::ImageResizeArgs& InArgs, int32 InImageCompressionQuality)
+	FImageResizeTask::FImageResizeTask(FScheduledOp InOp, FProgramCache& Memory, const Ptr<const Image>& InBase, uint16 InSizeX, uint16 InSizeY, int32 InImageCompressionQuality)
 	{
 		Op = InOp;
-		Args = InArgs;
-		pBase = Memory.GetImage(FCacheAddress(InArgs.source, InOp));
+		SizeX = InSizeX;
+		SizeY = InSizeY;
+		pBase = InBase;
 		ImageCompressionQuality = InImageCompressionQuality;
 	}
 
@@ -1207,11 +1208,7 @@ namespace mu
 			return;
 		}
 
-		FImageSize destSize = FImageSize
-			(
-				Args.size[0],
-				Args.size[1]
-			);
+		FImageSize destSize = FImageSize( SizeX, SizeY );
 
 		ImagePtr pResult;
 
@@ -2019,7 +2016,13 @@ namespace mu
 			if (item.Stage == 1)
 			{
 				OP::ImageResizeArgs Args = program.GetOpArgs<OP::ImageResizeArgs>(item.At);
-				Issued = MakeShared<FImageResizeTask>(item, GetMemory(), Args, m_pSettings->GetPrivate()->m_imageCompressionQuality);
+				Ptr<const Image> pBase = GetMemory().GetImage(FCacheAddress(Args.source, item));
+
+				// \TODO: Shortcut without creating the task and resolving it if there is no need to resize.
+
+				// \TODO: This operation is ignoring the mips-to-skip operation data
+
+				Issued = MakeShared<FImageResizeTask>(item, GetMemory(), pBase, Args.size[0], Args.size[1], m_pSettings->GetPrivate()->m_imageCompressionQuality);
 			}
 			break;
 		}
diff --git a/Engine/Plugins/Experimental/Mutable/Source/MutableRuntime/Private/MuR/ModelPrivate.h b/Engine/Plugins/Experimental/Mutable/Source/MutableRuntime/Private/MuR/ModelPrivate.h
--- a/Engine/Plugins/Experimental/Mutable/Source/MutableRuntime/Private/MuR/ModelPrivate.h
+++ b/Engine/Plugins/Experimental/Mutable/Source/MutableRuntime/Private/MuR/ModelPrivate.h
@@ -14,6 +14,7 @@
 #include "MuR/OpImageMipmap.h"
 
 #define MUTABLE_MAX_RUNTIME_PARAMETERS_PER_STATE	64
+#define MUTABLE_GROW_BORDER_VALUE					2
 
 namespace mu
 {
diff --git a/Engine/Plugins/Experimental/Mutable/Source/MutableRuntime/Private/MuR/OpImageBlend.h b/Engine/Plugins/Experimental/Mutable/Source/MutableRuntime/Private/MuR/OpImageBlend.h
--- a/Engine/Plugins/Experimental/Mutable/Source/MutableRuntime/Private/MuR/OpImageBlend.h
+++ b/Engine/Plugins/Experimental/Mutable/Source/MutableRuntime/Private/MuR/OpImageBlend.h
@@ -3,7 +3,6 @@
 #pragma once
 
 #include "MuR/ImagePrivate.h"
-#include "MuR/OpImageLayer.h"
 #include "Math/VectorRegister.h"
 
 namespace mu
@@ -22,16 +21,6 @@ namespace mu
 		return masked;
 	}
 
-	//! Blend a subimage on the base using a mask.
-	inline void ImageBlendOnBaseNoAlpha( Image* pBase,
-										 const Image* pMask,
-										 const Image* pBlended,
-										 const box< vec2<int> >& rect )
-	{
-		ImageLayerOnBaseNoAlpha< BlendChannel, false>
-				( pBase, pMask, pBlended, rect );
-	}
-
 
 	//---------------------------------------------------------------------------------------------
 	inline uint32 ScreenChannelMasked(uint32 base, uint32 blended, uint32 mask)
diff --git a/Engine/Plugins/Experimental/Mutable/Source/MutableRuntime/Private/MuR/OpImageLayer.h b/Engine/Plugins/Experimental/Mutable/Source/MutableRuntime/Private/MuR/OpImageLayer.h
--- a/Engine/Plugins/Experimental/Mutable/Source/MutableRuntime/Private/MuR/OpImageLayer.h
+++ b/Engine/Plugins/Experimental/Mutable/Source/MutableRuntime/Private/MuR/OpImageLayer.h
@@ -7,6 +7,8 @@
 #include "Async/ParallelFor.h"
 #include "Templates/UnrealTemplate.h"
 
+#include "MuR/OpImageBlend.h"
+
 namespace mu
 {
 	inline bool IsAnyComponentLargerThan1( FVector4f v )
@@ -1724,19 +1726,32 @@ namespace mu
 		uint8* pBaseBuf = pBase->GetData();
 		const uint8* pBlendedBuf = pBlend->GetData();
 
-		// The base determines the number of lods to process.
-		int32 LODCount = bOnlyFirstLOD ? 1 : pBase->GetLODCount();
-
 		int32 PixelCount = 0;
-		if (bOnlyFirstLOD)
+		uint16 SizeX = pBase->GetSizeX();
+
+		if (pBlend->m_flags & Image::IF_HAS_RELEVANCY_MAP 
+			&& 
+			(bOnlyFirstLOD || pBase->GetLODCount()==1 ) )
 		{
-			PixelCount = pBase->GetSizeX() * pBase->GetSizeY();
+			int32 NumRelevantRows = pBlend->RelevancyMaxY - pBlend->RelevancyMinY + 1;
+			PixelCount = SizeX * NumRelevantRows;
+			int32 DataOffset = pBlend->RelevancyMinY * SizeX * 4;
+			pBaseBuf += DataOffset;
+			pBlendedBuf += DataOffset;
 		}
 		else
-		{
-			PixelCount = pBase->CalculatePixelCount();
+		{ 
+			if (bOnlyFirstLOD)
+			{
+				PixelCount = SizeX * pBase->GetSizeY();
+			}
+			else
+			{
+				PixelCount = pBase->CalculatePixelCount();
+			}
 		}
 
+		// \TODO: Ensure batches operate on 32Kb aligned buffers?
 		constexpr int32 NumBatchElems = 4096*2;
 		const int32 NumBatches = FMath::DivideAndRoundUp(PixelCount, NumBatchElems);
 
@@ -1796,6 +1811,91 @@ namespace mu
 		}
 	}
 
+
+
+	//---------------------------------------------------------------------------------------------
+	template<>
+	inline void BufferLayerComposite<BlendChannelMasked, LightenChannel, false>
+	(
+		Image* pBase,
+		const Image* pBlend,
+		bool bOnlyFirstLOD,
+		uint8 BlendAlphaSourceChannel)
+	{
+		check(pBase->GetFormat() == EImageFormat::IF_RGBA_UBYTE);
+		check(pBlend->GetFormat() == EImageFormat::IF_RGBA_UBYTE);
+		check(pBase->GetSizeX() == pBlend->GetSizeX() && pBase->GetSizeY() == pBlend->GetSizeY());
+		check(bOnlyFirstLOD || pBase->GetLODCount() <= pBlend->GetLODCount());
+		check(BlendAlphaSourceChannel==3);
+
+		uint32* pBaseBuf = reinterpret_cast<uint32*>(pBase->GetData());
+		const uint32* pBlendedBuf = reinterpret_cast<const uint32*>(pBlend->GetData());
+
+		int32 PixelCount = 0;
+		uint16 SizeX = pBase->GetSizeX();
+
+		if (pBlend->m_flags & Image::IF_HAS_RELEVANCY_MAP
+			&&
+			(bOnlyFirstLOD || pBase->GetLODCount() == 1))
+		{
+			int32 NumRelevantRows = pBlend->RelevancyMaxY - pBlend->RelevancyMinY + 1;
+			PixelCount = SizeX * NumRelevantRows;
+			int32 DataOffset = pBlend->RelevancyMinY * SizeX;
+			pBaseBuf += DataOffset;
+			pBlendedBuf += DataOffset;
+		}
+		else
+		{
+			if (bOnlyFirstLOD)
+			{
+				PixelCount = SizeX * pBase->GetSizeY();
+			}
+			else
+			{
+				PixelCount = pBase->CalculatePixelCount();
+			}
+		}
+
+		// \TODO: Ensure batches operate on 32Kb aligned buffers?
+		constexpr int32 NumBatchElems = 4096 * 2;
+		const int32 NumBatches = FMath::DivideAndRoundUp(PixelCount, NumBatchElems);
+
+		auto ProcessBatch =
+			[
+				pBaseBuf, pBlendedBuf, PixelCount, NumBatchElems
+			] (uint32 BatchId)
+		{
+			const int32 BatchBegin = BatchId * NumBatchElems;
+			const int32 BatchEnd = FMath::Min(BatchBegin + NumBatchElems, PixelCount);
+
+			for (int32 i = BatchBegin; i < BatchEnd; ++i)
+			{
+				// TODO: Optimize this (SIMD?)
+				uint32 FullBase = pBaseBuf[i];
+				uint32 FullBlended = pBlendedBuf[i];
+				uint32 mask = (FullBlended & 0xff000000) >> 24;
+
+				uint32 FullResult = 0;
+				FullResult |= BlendChannelMasked((FullBase >>  0) & 0xff, (FullBlended >>  0) & 0xff, mask) << 0;
+				FullResult |= BlendChannelMasked((FullBase >>  8) & 0xff, (FullBlended >>  8) & 0xff, mask) << 8;
+				FullResult |= BlendChannelMasked((FullBase >> 16) & 0xff, (FullBlended >> 16) & 0xff, mask) << 16;
+				FullResult |= LightenChannel	((FullBase >> 24) & 0xff, (FullBlended >> 24) & 0xff) << 24;
+
+				pBaseBuf[i] = FullResult;
+			}
+		};
+
+		if (NumBatches == 1)
+		{
+			ProcessBatch(0);
+		}
+		else if (NumBatches > 1)
+		{
+			ParallelFor(NumBatches, ProcessBatch);
+		}
+	}
+
+
 	//---------------------------------------------------------------------------------------------
 	template< VectorRegister4Int (*RGB_FUNC_MASKED)(const VectorRegister4Int&, const VectorRegister4Int&, const VectorRegister4Int&),
 		int32 (*A_FUNC)(int32, int32),
@@ -2558,4 +2658,15 @@ namespace mu
 		}
 	}
 
+
+	//! Blend a subimage on the base using a mask.
+	inline void ImageBlendOnBaseNoAlpha(Image* pBase,
+		const Image* pMask,
+		const Image* pBlended,
+		const box< vec2<int> >& rect)
+	{
+		ImageLayerOnBaseNoAlpha< BlendChannel, false>
+			(pBase, pMask, pBlended, rect);
+	}
+
 }
diff --git a/Engine/Plugins/Experimental/Mutable/Source/MutableRuntime/Private/MuR/OpImageProject.cpp b/Engine/Plugins/Experimental/Mutable/Source/MutableRuntime/Private/MuR/OpImageProject.cpp
--- a/Engine/Plugins/Experimental/Mutable/Source/MutableRuntime/Private/MuR/OpImageProject.cpp
+++ b/Engine/Plugins/Experimental/Mutable/Source/MutableRuntime/Private/MuR/OpImageProject.cpp
@@ -912,8 +912,8 @@ void ImageRasterProjected_Optimised( const Mesh* pMesh,
 	EImageFormat format = pImage->GetFormat();
     int pixelSize = GetImageFormatData( format ).m_bytesPerBlock;
 
-    int sizeX = pImage->GetSizeX();
-    int sizeY = pImage->GetSizeY();
+    int32 sizeX = pImage->GetSizeX();
+    int32 sizeY = pImage->GetSizeY();
 
     // Get the vertices
     int vertexCount = pMesh->GetVertexCount();
@@ -969,6 +969,37 @@ void ImageRasterProjected_Optimised( const Mesh* pMesh,
 			}
 		});
 
+	// Update the relevancy map
+	{
+		MUTABLE_CPUPROFILER_SCOPE(ImageRasterProjected_Optimised_Relevancy);
+
+		float MinY = float(sizeY) - 1.0f;
+		float MaxY = 0.0f;
+		for (int f = 0; f < faceCount; ++f)
+		{
+			int i0 = pIndices[f * 3 + 0];
+			int i1 = pIndices[f * 3 + 1];
+			int i2 = pIndices[f * 3 + 2];
+
+			// A bit ugly, probably can be improved if integrated in above loops and made more precise
+			// inside the pixel processor to account for masked out pixels?
+			if (!scratch->culledVertex[i0] ||
+				!scratch->culledVertex[i1] ||
+				!scratch->culledVertex[i2])
+			{
+				MinY = FMath::Min(MinY, scratch->vertices[i0].y);
+				MinY = FMath::Min(MinY, scratch->vertices[i1].y);
+				MinY = FMath::Min(MinY, scratch->vertices[i2].y);
+				MaxY = FMath::Max(MaxY, scratch->vertices[i0].y);
+				MaxY = FMath::Max(MaxY, scratch->vertices[i1].y);
+				MaxY = FMath::Max(MaxY, scratch->vertices[i2].y);
+			}
+		}
+
+		pImage->m_flags |= Image::IF_HAS_RELEVANCY_MAP;
+		pImage->RelevancyMinY = uint16(FMath::FloorToFloat(MinY));
+		pImage->RelevancyMaxY = uint16(FMath::Min( int32(FMath::CeilToFloat(MaxY)), sizeY-1 ));
+	}
 }
 
 
diff --git a/Engine/Plugins/Experimental/Mutable/Source/MutableRuntime/Private/MuR/OpImageResize.h b/Engine/Plugins/Experimental/Mutable/Source/MutableRuntime/Private/MuR/OpImageResize.h
--- a/Engine/Plugins/Experimental/Mutable/Source/MutableRuntime/Private/MuR/OpImageResize.h
+++ b/Engine/Plugins/Experimental/Mutable/Source/MutableRuntime/Private/MuR/OpImageResize.h
@@ -33,18 +33,27 @@ namespace mu
             const uint8_t* pBaseBuf = pBase->GetData();
             uint8_t* pDestBuf = pDest->GetData();
 
-            // Simple nearest pixel implementation
-            for ( int y=0; y<destSize[1]; ++y )
-            {
-                int sy = (y*baseSize[1])/destSize[1];
+			// Simple nearest pixel implementation
+			uint32 dx_16 = (uint32(baseSize[0]) << 16) / destSize[0];
+			uint32 dy_16 = (uint32(baseSize[1]) << 16) / destSize[1];
+			uint32 sy_16 = 0;
+			for (uint16 y = 0; y < destSize[1]; ++y)
+			{
+				uint32 sy = sy_16 >> 16;
+				const uint8_t* pRowBuf = pBase->GetData() + baseSize[0] * sy;
 
-                for ( int x=0; x<destSize[0]; ++x )
-                {
-                    int sx = (x*baseSize[0])/destSize[0];
-                    *pDestBuf = pBaseBuf[sy*baseSize[0]+sx];
-                    pDestBuf++;
-                }
-            }
+				uint32 sx_16 = 0;
+				for (uint16 x = 0; x < destSize[0]; ++x)
+				{
+					uint32 sx = sx_16 >> 16;
+					*pDestBuf = pRowBuf[sx];
+					pDestBuf++;
+
+					sx_16 += dx_16;
+				}
+
+				sy_16 += dy_16;
+			}
 
             break;
         }
@@ -330,6 +339,89 @@ namespace mu
 		ParallelFor(sizeY, ProcessLine);
     }
 
+	inline uint32 AverageChannel(uint32 a, uint32 b)
+	{
+		uint32 r = (a + b) >> 1;
+		return r;
+	}
+
+	template<>
+	inline void ImageMinifyX_Exact<4,2>(Image* pDest, const Image* pBase)
+	{
+		int baseSizeX = pBase->GetSizeX();
+		int destSizeX = pDest->GetSizeX();
+		int sizeY = pBase->GetSizeY();
+
+		const uint8* pBaseBuf = pBase->GetData();
+		uint8* pDestBuf = pDest->GetData();
+
+		int32 TotalBasePixels = baseSizeX * sizeY;
+		constexpr int32 BasePixelsPerBatch = 4096 * 2;
+		int32 NumBatches = FMath::DivideAndRoundUp(TotalBasePixels, BasePixelsPerBatch);
+
+		// Linear filtering
+		const auto& ProcessBatchUnaligned = 
+			[ pDestBuf, pBaseBuf, baseSizeX, destSizeX, BasePixelsPerBatch, TotalBasePixels ]
+		(int32 BatchIndex)
+		{
+			const uint8* pBatchBaseBuf = pBaseBuf + BatchIndex * BasePixelsPerBatch * 4;
+			uint8* pBatchDestBuf = pDestBuf + BatchIndex * BasePixelsPerBatch * 4/2;
+
+			int32 NumBasePixels = FMath::Min(BasePixelsPerBatch, TotalBasePixels-BatchIndex* BasePixelsPerBatch );
+
+			uint16 r[4];
+			for (int x = 0; x < NumBasePixels; x+=2)
+			{
+				r[0] = pBatchBaseBuf[0] + pBatchBaseBuf[0 + 4];
+				r[1] = pBatchBaseBuf[1] + pBatchBaseBuf[1 + 4];
+				r[2] = pBatchBaseBuf[2] + pBatchBaseBuf[2 + 4];
+				r[3] = pBatchBaseBuf[3] + pBatchBaseBuf[3 + 4];
+
+				pBatchDestBuf[0] = (uint8)(r[0] >> 1);
+				pBatchDestBuf[1] = (uint8)(r[1] >> 1);
+				pBatchDestBuf[2] = (uint8)(r[2] >> 1);
+				pBatchDestBuf[3] = (uint8)(r[3] >> 1);
+
+				pBatchBaseBuf += 4 * 2;
+				pBatchDestBuf += 4;
+			}
+		};
+
+		const auto& ProcessBatchAligned =
+			[pDestBuf, pBaseBuf, baseSizeX, destSizeX, BasePixelsPerBatch, TotalBasePixels]
+		(int32 BatchIndex)
+		{
+			const uint32* pBatchBaseBuf = reinterpret_cast<const uint32*>(pBaseBuf) + BatchIndex * BasePixelsPerBatch;
+			uint32* pBatchDestBuf = reinterpret_cast<uint32*>(pDestBuf) + BatchIndex * (BasePixelsPerBatch>>1);
+
+			int32 NumBasePixels = FMath::Min(BasePixelsPerBatch, TotalBasePixels - BatchIndex * BasePixelsPerBatch);
+
+			for (int p=0; p<NumBasePixels;++p)
+			{
+				uint32 FullSource0 = pBatchBaseBuf[p*2+0];
+				uint32 FullSource1 = pBatchBaseBuf[p*2+1];
+
+				uint32 FullResult = 0;
+
+				FullResult |= AverageChannel((FullSource0 >>  0) & 0xff, (FullSource1 >>  0) & 0xff)  <<  0;
+				FullResult |= AverageChannel((FullSource0 >>  8) & 0xff, (FullSource1 >>  8) & 0xff)  <<  8;
+				FullResult |= AverageChannel((FullSource0 >> 16) & 0xff, (FullSource1 >> 16) & 0xff)  << 16;
+				FullResult |= AverageChannel((FullSource0 >> 24) & 0xff, (FullSource1 >> 24) & 0xff)  << 24;
+
+				pBatchDestBuf[p] = FullResult;
+			}
+		};
+
+//		if ( (SIZE_T(pBaseBuf) & 0x3) || (SIZE_T(pDestBuf) & 0x3) )
+		{
+			ParallelFor(NumBatches, ProcessBatchUnaligned);
+		}
+		/*else
+		{
+			ParallelFor(NumBatches, ProcessBatchAligned);
+		}*/
+	}
+
 
     //---------------------------------------------------------------------------------------------
     //! Image minify X version hub.
@@ -399,8 +491,6 @@ namespace mu
     }
 
 
-
-
     //---------------------------------------------------------------------------------------------
     //!
     //---------------------------------------------------------------------------------------------
diff --git a/Engine/Plugins/Experimental/Mutable/Source/MutableRuntime/Public/MuR/Image.h b/Engine/Plugins/Experimental/Mutable/Source/MutableRuntime/Public/MuR/Image.h
--- a/Engine/Plugins/Experimental/Mutable/Source/MutableRuntime/Public/MuR/Image.h
+++ b/Engine/Plugins/Experimental/Mutable/Source/MutableRuntime/Public/MuR/Image.h
@@ -257,13 +257,20 @@ namespace mu
 			IF_IS_PLAIN_COLOUR = 1 << 1,
 
 			// If this is set, the image shouldn't be scaled: it's contents is resoultion-dependent.
-			IF_CANNOT_BE_SCALED = 1 << 2
+			IF_CANNOT_BE_SCALED = 1 << 2,
+
+			// If this is set, the image has an updated relevancy map. This flag is not persisent.
+			IF_HAS_RELEVANCY_MAP = 1 << 3,
 		} EImageFlags;
 
 		//! Persistent flags with some image properties. The meaning will depend of every
 		//! context.
 		mutable uint8 m_flags = 0;
 
+		/** Non-persistent relevancy map. */
+		uint16 RelevancyMinY = 0;
+		uint16 RelevancyMaxY = 0;
+
 		//! Pixel data
 		TArray<uint8> m_data;
 
@@ -284,6 +291,8 @@ namespace mu
 			pResult->m_format = m_format;
 			pResult->m_lods = m_lods;
 			pResult->m_flags = m_flags;
+			pResult->RelevancyMinY = RelevancyMinY;
+			pResult->RelevancyMaxY = RelevancyMaxY;
 			pResult->m_data = m_data;
 
 			return pResult;
@@ -300,7 +309,10 @@ namespace mu
 			arch << m_lods;
 			arch << (uint8)m_format;
 			arch << m_data;
-			arch << m_flags;
+
+			// Remove non-persistent flags.
+			uint8 flags = m_flags & ~IF_HAS_RELEVANCY_MAP;
+			arch << flags;
 		}
 
 		//!
diff --git a/Engine/Plugins/Experimental/Mutable/Source/MutableTools/Private/MuT/CodeGenerator_Image.cpp b/Engine/Plugins/Experimental/Mutable/Source/MutableTools/Private/MuT/CodeGenerator_Image.cpp
--- a/Engine/Plugins/Experimental/Mutable/Source/MutableTools/Private/MuT/CodeGenerator_Image.cpp
+++ b/Engine/Plugins/Experimental/Mutable/Source/MutableTools/Private/MuT/CodeGenerator_Image.cpp
@@ -1615,7 +1615,7 @@ namespace mu
 
         Ptr<ASTOpImageMakeGrowMap> MakeGrowMapOp = new ASTOpImageMakeGrowMap();
 		MakeGrowMapOp->Mask = rasterop;
-		MakeGrowMapOp->Border = 2;
+		MakeGrowMapOp->Border = MUTABLE_GROW_BORDER_VALUE;
 		
 		// If we want to be able to generate progressive mips efficiently, we need mipmaps for the "displacement map".
 		if (m_compilerOptions->m_optimisationOptions.bEnableProgressiveImages)
