[mutable] Fixed float pins of table nodes. Now they are double properties instead of float properties.

#rnx

diff --git a/Engine/Plugins/Experimental/Mutable/Source/CustomizableObjectEditor/Private/MuCOE/GenerateMutableSource/GenerateMutableSourceTable.cpp b/Engine/Plugins/Experimental/Mutable/Source/CustomizableObjectEditor/Private/MuCOE/GenerateMutableSource/GenerateMutableSourceTable.cpp
--- a/Engine/Plugins/Experimental/Mutable/Source/CustomizableObjectEditor/Private/MuCOE/GenerateMutableSource/GenerateMutableSourceTable.cpp
+++ b/Engine/Plugins/Experimental/Mutable/Source/CustomizableObjectEditor/Private/MuCOE/GenerateMutableSource/GenerateMutableSourceTable.cpp
@@ -605,20 +605,19 @@ void FillTableColumn(const UCustomizableObjectNodeTable* TableNode,	mu::TablePtr
 		MutableTable->SetCell(CurrentColumn, RowIdx, Value);
 	}
 
-	//TODO
-	//else if (const FNumericProperty* DoubleNumProperty = CastField<FDoubleProperty>(Property))
-	//{
-	//	CurrentColumn = MutableTable->FindColumn(TCHAR_TO_ANSI(*ColumnName));
-	//
-	//	if (CurrentColumn == INDEX_NONE)
-	//	{
-	//		CurrentColumn = MutableTable->AddColumn(TCHAR_TO_ANSI(*ColumnName), mu::TABLE_COLUMN_TYPE::TCT_SCALAR);
-	//	}
-	//
-	//	// Setting cell value
-	//	float Value = DoubleNumProperty->GetFloatingPointPropertyValue(CellData);
-	//	MutableTable->SetCell(CurrentColumn, RowIdx, Value);
-	//}
+	else if (const FNumericProperty* DoubleNumProperty = CastField<FDoubleProperty>(Property))
+	{
+		CurrentColumn = MutableTable->FindColumn(TCHAR_TO_ANSI(*ColumnName));
+	
+		if (CurrentColumn == INDEX_NONE)
+		{
+			CurrentColumn = MutableTable->AddColumn(TCHAR_TO_ANSI(*ColumnName), mu::TABLE_COLUMN_TYPE::TCT_SCALAR);
+		}
+	
+		// Setting cell value
+		float Value = DoubleNumProperty->GetFloatingPointPropertyValue(CellData);
+		MutableTable->SetCell(CurrentColumn, RowIdx, Value);
+	}
 }
 
 
diff --git a/Engine/Plugins/Experimental/Mutable/Source/CustomizableObjectEditor/Private/MuCOE/Nodes/CustomizableObjectNodeTable.cpp b/Engine/Plugins/Experimental/Mutable/Source/CustomizableObjectEditor/Private/MuCOE/Nodes/CustomizableObjectNodeTable.cpp
--- a/Engine/Plugins/Experimental/Mutable/Source/CustomizableObjectEditor/Private/MuCOE/Nodes/CustomizableObjectNodeTable.cpp
+++ b/Engine/Plugins/Experimental/Mutable/Source/CustomizableObjectEditor/Private/MuCOE/Nodes/CustomizableObjectNodeTable.cpp
@@ -202,7 +202,12 @@ void UCustomizableObjectNodeTable::AllocateDefaultPins(UCustomizableObjectNodeRe
 			}
 		}
 
-		else if (const FNumericProperty* NumProperty = CastField<FFloatProperty>(ColumnProperty))
+		else if (const FNumericProperty* NumFloatProperty = CastField<FFloatProperty>(ColumnProperty))
+		{
+			OutPin = CustomCreatePin(EGPD_Output, Schema->PC_Float, FName(*PinName));
+		}
+		
+		else if (const FNumericProperty* NumDoubleProperty = CastField<FDoubleProperty>(ColumnProperty))
 		{
 			OutPin = CustomCreatePin(EGPD_Output, Schema->PC_Float, FName(*PinName));
 		}
@@ -468,7 +473,18 @@ bool UCustomizableObjectNodeTable::IsNodeOutDatedAndNeedsRefresh()
 					NumPins++;
 				}
 			}
-			else if (const FFloatProperty* NumProperty = CastField<FFloatProperty>(ColumnProperty))
+			else if (const FFloatProperty* FloatNumProperty = CastField<FFloatProperty>(ColumnProperty))
+			{
+				FString PinName = DataTableUtils::GetPropertyExportName(ColumnProperty);
+
+				if (CheckPinUpdated(PinName, Schema->PC_Float))
+				{
+					bNeedsUpdate = true;
+				}
+
+				NumPins++;
+			}
+			else if (const FDoubleProperty* DoubleNumProperty = CastField<FDoubleProperty>(ColumnProperty))
 			{
 				FString PinName = DataTableUtils::GetPropertyExportName(ColumnProperty);
 
