[mutable] Fixed the image-resize operation that was ignoring the mips-to-skip argument.
Improved handling of relevancy data in some operations.
Removed old unused code from projections.
#rnx

diff --git a/Engine/Plugins/Experimental/Mutable/Source/MutableRuntime/Private/MuR/CodeRunner.cpp b/Engine/Plugins/Experimental/Mutable/Source/MutableRuntime/Private/MuR/CodeRunner.cpp
--- a/Engine/Plugins/Experimental/Mutable/Source/MutableRuntime/Private/MuR/CodeRunner.cpp
+++ b/Engine/Plugins/Experimental/Mutable/Source/MutableRuntime/Private/MuR/CodeRunner.cpp
@@ -2608,6 +2608,9 @@ namespace mu
 					// Store the base
 					Ptr<Image> New = mu::CloneOrTakeOver(Base.get());
 
+					// Reset relevancy map.
+					New->m_flags &= ~Image::EImageFlags::IF_HAS_RELEVANCY_MAP;
+
 					// This shouldn't happen in optimised models, but it could happen in editors, etc.
 					// \todo: raise a performance warning?
 					EImageFormat BaseFormat = GetUncompressedFormat(Base->GetFormat());
@@ -4089,14 +4092,6 @@ namespace mu
 				if (OriginalSourceScale.x() > 0 && OriginalSourceScale.y() > 0)
 				{
 					ImageDisplace(pNew.get(), pSource.get(), pMap.get());
-					
-					// Update the relevancy data of the image for the worst case. 
-					if (pNew->m_flags & Image::EImageFlags::IF_HAS_RELEVANCY_MAP)
-					{
-						// Displace can encode at max MUTABLE_GROW_BORDER_VALUE pixels away according to "border" in MakeGrowMap.
-						pNew->RelevancyMinY = FMath::Max(int32(pNew->RelevancyMinY) - MUTABLE_GROW_BORDER_VALUE, 0);
-						pNew->RelevancyMaxY = FMath::Min(pNew->RelevancyMaxY + MUTABLE_GROW_BORDER_VALUE, pNew->GetSizeY() - 1);
-					}
 
 					if (OriginalSourceScale != pNew->GetSize())
 					{
diff --git a/Engine/Plugins/Experimental/Mutable/Source/MutableRuntime/Private/MuR/CodeRunnerTasked.cpp b/Engine/Plugins/Experimental/Mutable/Source/MutableRuntime/Private/MuR/CodeRunnerTasked.cpp
--- a/Engine/Plugins/Experimental/Mutable/Source/MutableRuntime/Private/MuR/CodeRunnerTasked.cpp
+++ b/Engine/Plugins/Experimental/Mutable/Source/MutableRuntime/Private/MuR/CodeRunnerTasked.cpp
@@ -473,7 +473,7 @@ namespace mu
 		}
 
 		ImagePtr pNew = mu::CloneOrTakeOver<Image>(m_base.get());
-		//ImagePtr pNew = m_base->Clone();
+
 		check(pNew->GetDataSize() == m_base->GetDataSize());
 
 		bool bValid = pNew->GetSizeX() > 0 && pNew->GetSizeY() > 0;
@@ -648,6 +648,9 @@ namespace mu
 			pNew = ImagePixelFormat(m_imageCompressionQuality, pNew.get(), InitialFormat);
 		}
 
+		// Reset relevancy map.
+		pNew->m_flags &= ~Image::EImageFlags::IF_HAS_RELEVANCY_MAP;
+
 		// Free resources
 		m_base = nullptr;
 		m_blended = nullptr;
@@ -859,6 +862,9 @@ namespace mu
 
 		}
 
+		// Reset relevancy map.
+		pNew->m_flags &= ~Image::EImageFlags::IF_HAS_RELEVANCY_MAP;
+
 		m_base = nullptr;
 		m_mask = nullptr;
 		Result = pNew;
@@ -1210,14 +1216,17 @@ namespace mu
 
 		FImageSize destSize = FImageSize( SizeX, SizeY );
 
-		ImagePtr pResult;
+		// Apply the mips-to-skip to the dest size
+		int32 MipsToSkip = Op.ExecutionOptions;
+		destSize[0] = FMath::Max(destSize[0] >> MipsToSkip, 1);
+		destSize[1] = FMath::Max(destSize[1] >> MipsToSkip, 1);
 
 		if ( destSize[0]!=pBase->GetSizeX()
 			 ||
 			 destSize[1]!=pBase->GetSizeY() )
 		{
-			//pResult = ImageResize( pBase.get(), destSize );
-			pResult = ImageResizeLinear( ImageCompressionQuality, pBase.get(), destSize );
+			ImagePtr Resized = ImageResizeLinear( ImageCompressionQuality, pBase.get(), destSize );
+			Result = Resized;
 
 			// If the source image had mips, generate them as well for the resized image.
 			// This shouldn't happen often since "ResizeLike" should be usually optimised out
@@ -1226,28 +1235,23 @@ namespace mu
 			bool sourceHasMips = pBase->GetLODCount()>1;
 			if (sourceHasMips)
 			{
-				int levelCount = Image::GetMipmapCount( pResult->GetSizeX(), pResult->GetSizeY() );
-				ImagePtr pMipmapped = new Image( pResult->GetSizeX(), pResult->GetSizeY(),
-												 levelCount,
-												 pResult->GetFormat() );
+				int levelCount = Image::GetMipmapCount(Resized->GetSizeX(), Resized->GetSizeY() );
+				ImagePtr Mipmapped = new Image(Resized->GetSizeX(), Resized->GetSizeY(), levelCount, Resized->GetFormat() );
 
 				SCRATCH_IMAGE_MIPMAP scratch;
 				FMipmapGenerationSettings mipSettings{};
 
-				ImageMipmap_PrepareScratch( pMipmapped.get(), pResult.get(), levelCount, &scratch );
+				ImageMipmap_PrepareScratch(Mipmapped.get(), Resized.get(), levelCount, &scratch );
 				ImageMipmap( ImageCompressionQuality,
-							 pMipmapped.get(), pResult.get(), levelCount, &scratch, mipSettings );
-
-				pResult = pMipmapped;
-			}
+							 Mipmapped.get(), Resized.get(), levelCount, &scratch, mipSettings );
 
+				Result = Mipmapped;
+			}			
 		}
 		else
 		{
-			pResult = pBase->Clone();
+			Result = pBase;
 		}
-
-		Result = pResult;
 	}
 
 	//---------------------------------------------------------------------------------------------
diff --git a/Engine/Plugins/Experimental/Mutable/Source/MutableRuntime/Private/MuR/OpImageDisplace.h b/Engine/Plugins/Experimental/Mutable/Source/MutableRuntime/Private/MuR/OpImageDisplace.h
--- a/Engine/Plugins/Experimental/Mutable/Source/MutableRuntime/Private/MuR/OpImageDisplace.h
+++ b/Engine/Plugins/Experimental/Mutable/Source/MutableRuntime/Private/MuR/OpImageDisplace.h
@@ -338,6 +338,14 @@ namespace mu
 		}
 
 
+		// Update the relevancy data of the image for the worst case. 
+		if (pResult->m_flags & Image::EImageFlags::IF_HAS_RELEVANCY_MAP)
+		{
+			// Displace can encode at max MUTABLE_GROW_BORDER_VALUE pixels away according to "border" in MakeGrowMap.
+			pResult->RelevancyMinY = FMath::Max(int32(pResult->RelevancyMinY) - MUTABLE_GROW_BORDER_VALUE, 0);
+			pResult->RelevancyMaxY = FMath::Min(pResult->RelevancyMaxY + MUTABLE_GROW_BORDER_VALUE, pResult->GetSizeY() - 1);
+		}
+
 	}
 
 
diff --git a/Engine/Plugins/Experimental/Mutable/Source/MutableRuntime/Private/MuR/OpImageProject.cpp b/Engine/Plugins/Experimental/Mutable/Source/MutableRuntime/Private/MuR/OpImageProject.cpp
--- a/Engine/Plugins/Experimental/Mutable/Source/MutableRuntime/Private/MuR/OpImageProject.cpp
+++ b/Engine/Plugins/Experimental/Mutable/Source/MutableRuntime/Private/MuR/OpImageProject.cpp
@@ -31,392 +31,6 @@ namespace mu
 namespace
 {
 
-//	class RaycastPixelProcessor
-//	{
-//	public:
-//
-//		RaycastPixelProcessor( const FProjector& projector,
-//							   const Image* pSource,
-//                               const uint8* pTargetData,
-//                               const uint8* pMaskData,
-//							   float fadeStart,
-//							   float fadeEnd )
-//		{
-//			projector.GetDirectionSideUp( m_direction, m_side, m_up );
-//
-//			m_position[0] = projector.position[0];
-//			m_position[1] = projector.position[1];
-//			m_position[2] = projector.position[2];
-//
-//			m_scale[0] = projector.scale[0];
-//            m_scale[1] = projector.scale[1];
-//            m_scale[2] = projector.scale[2];
-//
-//			m_pSource = pSource;
-//
-//            m_fadeStartCos = cosf(fadeStart);
-//            m_fadeEndCos = cosf(fadeEnd);
-//
-//			m_pTargetData = pTargetData;
-//			m_pMaskData = pMaskData;
-//
-//			EImageFormat format = m_pSource->GetFormat();
-//			m_pixelSize = GetImageFormatData( format ).m_bytesPerBlock;
-//
-//		}
-//
-//
-//		//-----------------------------------------------------------------------------------------
-//        inline void ProcessPixel( unsigned char* pBufferPos, float varying[4] ) const
-//		{
-//			float factor = 1;
-//
-//            // depth clamp
-//            if (varying[2]<0.0f || varying[2]>1.0f)
-//            {
-//                factor = 0;
-//            }
-//
-//            if ( factor>0 && m_pMaskData )
-//			{
-//                uint8 maskFactor = m_pMaskData[ (pBufferPos-m_pTargetData) / m_pixelSize ];
-//				factor = float(maskFactor) / 255.0f;
-//			}
-//
-//			if ( factor>0 )
-//			{
-////				vec3<float> normal( varying[2], varying[3], varying[4] );
-////				normal = normalise_approx( normal );
-//
-////				float angleCos = dot( normal, m_direction * -1.0f );
-//                float angleCos = varying[3];
-//
-//                if ( angleCos<m_fadeStartCos && angleCos>m_fadeEndCos )
-//				{
-//                    factor *= ( angleCos - m_fadeEndCos ) / ( m_fadeStartCos - m_fadeEndCos );
-//				}
-//                else if ( angleCos<=m_fadeEndCos )
-//				{
-//					factor = 0;
-//				}
-//			}
-//
-//
-//			if ( factor>0 )
-//			{
-//				// Transform the origin to the projected image texture coordinates
-//				float u = varying[0];
-//				float v = varying[1];
-//
-//				if ( u>=0 && u<1 && v>=0 && v<1 )
-//				{
-//
-//                    const uint8* pPixel = m_pSource->GetData();
-//					// TODO: clamp?
-//					pPixel += ( m_pSource->GetSizeX() * int(m_pSource->GetSizeY() * v)
-//								+ int( m_pSource->GetSizeX() * u ) )
-//							* m_pixelSize;
-//
-//					// Write result
-//					switch ( m_pSource->GetFormat() )
-//					{
-//					case EImageFormat::IF_L_UBYTE:
-//                        pBufferPos[0] = uint8( pPixel[0] * factor );
-//						break;
-//
-//					case EImageFormat::IF_RGB_UBYTE:
-//                        pBufferPos[0] = uint8( pPixel[0] * factor );
-//                        pBufferPos[1] = uint8( pPixel[1] * factor );
-//                        pBufferPos[2] = uint8( pPixel[2] * factor );
-//						break;
-//
-//                    case EImageFormat::IF_BGRA_UBYTE:
-//                    case EImageFormat::IF_RGBA_UBYTE:
-//                        pBufferPos[0] = uint8( pPixel[0] * factor );
-//                        pBufferPos[1] = uint8( pPixel[1] * factor );
-//                        pBufferPos[2] = uint8( pPixel[2] * factor );
-//                        pBufferPos[3] = uint8( pPixel[3] * factor );
-//						break;
-//
-//					default:
-//						// Not implemented
-//						check( false );
-//						break;
-//					}
-//				}
-//			}
-//		}
-//
-//
-//		const Image* m_pSource = nullptr;
-//		FVector3f m_direction;
-//		FVector3f m_up;
-//		FVector3f m_side;
-//		FVector3f m_position;
-//		FVector3f m_scale;
-//
-//        //! Cosine of the fading angle range
-//		float m_fadeStartCos = 0.0f, m_fadeEndCos = 0.0f;
-//
-//        const uint8* m_pTargetData = nullptr;
-//        const uint8* m_pMaskData = nullptr;
-//		int m_pixelSize = 0;
-//
-//	};
-
-
-	//---------------------------------------------------------------------------------------------
-//    template<int PIXEL_SIZE,bool HAS_MASK=true>
-//	class RaycastPixelProcessor_UBYTE
-//	{
-//	public:
-//
-//		RaycastPixelProcessor_UBYTE( const FProjector& projector,
-//									 const Image* pSource,
-//                                     const uint8* pTargetData,
-//                                     const uint8* pMaskData,
-//									 float fadeStart,
-//									 float fadeEnd )
-//		{
-//			projector.GetDirectionSideUp( m_direction, m_side, m_up );
-//
-//			m_position[0] = projector.position[0];
-//			m_position[1] = projector.position[1];
-//			m_position[2] = projector.position[2];
-//
-//			m_scale[0] = projector.scale[0];
-//            m_scale[1] = projector.scale[1];
-//            m_scale[2] = projector.scale[2];
-//
-//			m_pSourceData = pSource->GetData();
-//			m_sourceSizeX = pSource->GetSizeX();
-//			m_sourceSizeY = pSource->GetSizeY();
-//
-//            m_fadeStartCos = cosf(fadeStart);
-//            m_fadeEndCos = cosf(fadeEnd);
-//
-//			m_pTargetData = pTargetData;
-//			m_pMaskData = pMaskData;
-//
-//			check( GetImageFormatData( pSource->GetFormat() ).m_bytesPerBlock
-//							==
-//							PIXEL_SIZE );
-//		}
-//
-//
-//		//-----------------------------------------------------------------------------------------
-//        inline void ProcessPixel( unsigned char* pBufferPos, float varying[4] ) const
-//		{
-//            int32 factor_8 = 256;
-//
-//            // depth clamp
-//            if (varying[2]<0.0f || varying[2]>1.0f)
-//            {
-//                factor_8 = 0;
-//            }
-//
-////				vec3<float> normal( varying[2], varying[3], varying[4] );
-////				normal = normalise_approx( normal );
-//
-////				float angleCos = dot( normal, m_direction * -1.0f );
-//            float angleCos = varying[3];
-//
-//            if (HAS_MASK)
-//            {
-//                if ( factor_8>0 && m_pMaskData )
-//                {
-//                    uint8 maskFactor = m_pMaskData[ (pBufferPos-m_pTargetData) / PIXEL_SIZE ];
-//                    factor_8 = maskFactor;
-//                }
-//
-//                if ( factor_8>0 )
-//                {
-//                    if ( angleCos<m_fadeStartCos && angleCos>m_fadeEndCos )
-//                    {
-//                        factor_8 = ( factor_8
-//                                     *
-//                                     int32( 256.0f * ( angleCos - m_fadeEndCos ) / ( m_fadeStartCos - m_fadeEndCos ) )
-//                                     ) >> 8;
-//                    }
-//                    else if ( angleCos<=m_fadeEndCos )
-//                    {
-//                        factor_8 = 0;
-//                    }
-//                }
-//            }
-//            else
-//            {
-//                if ( angleCos<m_fadeStartCos && angleCos>m_fadeEndCos )
-//                {
-//                    factor_8 = int32( 256.0f * ( angleCos - m_fadeEndCos ) / ( m_fadeStartCos - m_fadeEndCos ) );
-//                }
-//                else if ( angleCos<=m_fadeEndCos )
-//                {
-//                    factor_8 = 0;
-//                }
-//            }
-//
-//
-//			if ( factor_8>0 )
-//			{
-//				float u = varying[0];
-//				float v = varying[1];
-//
-//				if ( u>=0 && u<1 && v>=0 && v<1 )
-//				{
-//
-//                    const uint8* pPixel = m_pSourceData;
-//					// TODO: clamp?
-//					pPixel += ( m_sourceSizeX * int(m_sourceSizeY * v)
-//								+ int( m_sourceSizeX * u ) )
-//							* PIXEL_SIZE;
-//
-//					// Write result
-//					for ( int i=0; i<PIXEL_SIZE; ++i )
-//					{
-//						pBufferPos[i] = uint8( ( pPixel[i] * factor_8 ) >> 8 );
-//					}
-//				}
-//			}
-//		}
-//
-//
-//        const uint8* m_pSourceData;
-//		int m_sourceSizeX, m_sourceSizeY;
-//
-//		FVector3f m_direction;
-//		FVector3f m_up;
-//		FVector3f m_side;
-//		FVector3f m_position;
-//		FVector3f m_scale;
-//
-//        //! Cosine of the fading angle range
-//        float m_fadeStartCos, m_fadeEndCos;
-//
-//        const uint8* m_pTargetData;
-//        const uint8* m_pMaskData;
-//	};
-
-
-
-  //  class RasterProjectedPixelProcessor
-  //  {
-  //  public:
-
-  //      RasterProjectedPixelProcessor( const Image* pSource,
-		//	const uint8* pTargetData,
-		//	const uint8* pMaskData,
-		//	bool bIsRGBFadingEnabled, bool bIsAlphaFadingEnabled,
-		//	float fadeStart, float fadeEnd )
-  //      {
-  //          m_pSource = pSource;
-
-		//	// \TODO: to template parameters?
-		//	m_bIsRGBFadingEnabled = bIsRGBFadingEnabled;
-		//	m_bIsAlphaFadingEnabled = bIsAlphaFadingEnabled;
-
-  //          m_fadeStartCos = cosf(fadeStart);
-  //          m_fadeEndCos = cosf(fadeEnd);
-
-  //          m_pTargetData = pTargetData;
-  //          m_pMaskData = pMaskData;
-
-		//	EImageFormat format = m_pSource->GetFormat();
-  //          m_pixelSize = GetImageFormatData( format ).m_bytesPerBlock;
-
-  //      }
-
-
-  //      //-----------------------------------------------------------------------------------------
-  //      inline void ProcessPixel( unsigned char* pBufferPos, float varying[4] ) const
-  //      {
-  //          float factor = 1;
-
-  //          // depth clamp
-  //          if (varying[2]<0.0f || varying[2]>1.0f)
-  //          {
-  //              factor = 0;
-  //          }
-
-  //          if ( factor>0 && m_pMaskData )
-  //          {
-  //              uint8 maskFactor = m_pMaskData[ (pBufferPos-m_pTargetData) / m_pixelSize ];
-  //              factor = float(maskFactor) / 255.0f;
-  //          }
-
-  //          if ( factor>0 )
-  //          {
-  //              float angleCos = varying[3];
-
-  //              if ( angleCos<m_fadeStartCos && angleCos>m_fadeEndCos )
-  //              {
-  //                  factor *= ( angleCos - m_fadeEndCos ) / ( m_fadeStartCos - m_fadeEndCos );
-  //              }
-  //              else if ( angleCos<=m_fadeEndCos )
-  //              {
-  //                  factor = 0;
-  //              }
-  //          }
-
-
-  //          if ( factor>0 )
-  //          {
-  //              // Transform the origin to the projected image texture coordinates
-  //              float u = varying[0];
-  //              float v = varying[1];
-
-  //              if ( u>=0 && u<1 && v>=0 && v<1 )
-  //              {
-
-  //                  const uint8* pPixel = m_pSource->GetData();
-  //                  // TODO: clamp?
-  //                  pPixel += ( m_pSource->GetSizeX() * int(m_pSource->GetSizeY() * v)
-  //                              + int( m_pSource->GetSizeX() * u ) )
-  //                          * m_pixelSize;
-
-  //                  // Write result
-  //                  switch ( m_pSource->GetFormat() )
-  //                  {
-  //                  case EImageFormat::IF_L_UBYTE:
-  //                      pBufferPos[0] = uint8( pPixel[0] * factor );
-  //                      break;
-
-  //                  case EImageFormat::IF_RGB_UBYTE:
-  //                      pBufferPos[0] = uint8( pPixel[0] * factor );
-  //                      pBufferPos[1] = uint8( pPixel[1] * factor );
-  //                      pBufferPos[2] = uint8( pPixel[2] * factor );
-  //                      break;
-
-  //                  case EImageFormat::IF_BGRA_UBYTE:
-  //                  case EImageFormat::IF_RGBA_UBYTE:
-  //                      pBufferPos[0] = uint8( pPixel[0] * factor );
-  //                      pBufferPos[1] = uint8( pPixel[1] * factor );
-  //                      pBufferPos[2] = uint8( pPixel[2] * factor );
-  //                      pBufferPos[3] = uint8( pPixel[3] * factor );
-  //                      break;
-
-  //                  default:
-  //                      // Not implemented
-  //                      check( false );
-  //                      break;
-  //                  }
-  //              }
-  //          }
-  //      }
-
-
-  //      const Image* m_pSource;
-
-  //      //! Cosine of the fading angle range
-  //      float m_fadeStartCos, m_fadeEndCos;
-		//bool m_bIsRGBFadingEnabled, m_bIsAlphaFadingEnabled;
-
-  //      const uint8* m_pTargetData;
-  //      const uint8* m_pMaskData;
-  //      int m_pixelSize;
-
-  //  };
-
 
     //---------------------------------------------------------------------------------------------
     template<int PIXEL_SIZE,bool HAS_MASK=true>
diff --git a/Engine/Plugins/Experimental/Mutable/Source/MutableRuntime/Private/MuR/OpImageResize.h b/Engine/Plugins/Experimental/Mutable/Source/MutableRuntime/Private/MuR/OpImageResize.h
--- a/Engine/Plugins/Experimental/Mutable/Source/MutableRuntime/Private/MuR/OpImageResize.h
+++ b/Engine/Plugins/Experimental/Mutable/Source/MutableRuntime/Private/MuR/OpImageResize.h
@@ -474,7 +474,12 @@ namespace mu
                 // Optimised case
                 ImageMinifyX_Exact<4,2>( pDest, pBase );
             }
-            else
+            else if (4 * pDest->GetSizeX() == pBase->GetSizeX())
+			{
+				// Optimised case
+				ImageMinifyX_Exact<4, 4>(pDest, pBase);
+			}
+			else
             {
                 // Generic case
                 ImageMinifyX<4>( pDest, pBase );
@@ -885,6 +890,17 @@ namespace mu
             pDest = ImagePixelFormat( imageCompressionQuality, pDest.get(), sourceFormat );
         }
 
+		// Update the relevancy data of the image.
+		if (pBase->m_flags & Image::EImageFlags::IF_HAS_RELEVANCY_MAP)
+		{
+			pDest->m_flags |= Image::EImageFlags::IF_HAS_RELEVANCY_MAP;
+
+			float FactorY = float(destSize[1]) / float(baseSize[1]);
+
+			pDest->RelevancyMinY = uint16( FMath::FloorToFloat(pBase->RelevancyMinY*FactorY) );
+			pDest->RelevancyMaxY = uint16( FMath::Min((int32)FMath::CeilToFloat(pBase->RelevancyMinY * FactorY), pDest->GetSizeY() - 1) );
+		}
+
         return pDest;
     }
 
