Fixed nonunity/pch compile errors

diff --git a/Engine/Plugins/Experimental/Mutable/Source/MutableRuntime/Private/MuR/OpImageDisplace.h b/Engine/Plugins/Experimental/Mutable/Source/MutableRuntime/Private/MuR/OpImageDisplace.h
--- a/Engine/Plugins/Experimental/Mutable/Source/MutableRuntime/Private/MuR/OpImageDisplace.h
+++ b/Engine/Plugins/Experimental/Mutable/Source/MutableRuntime/Private/MuR/OpImageDisplace.h
@@ -3,6 +3,7 @@
 #pragma once
 
 #include "MuR/ImagePrivate.h"
+#include "MuR/ModelPrivate.h"
 #include "MuR/Platform.h"
 #include "MuR/ConvertData.h"
 #include "Async/ParallelFor.h"
