[mutable] Peak memory attenuation for objects using a lot of chained image-layer operations.
#rnx

diff --git a/Engine/Plugins/Experimental/Mutable/Source/MutableRuntime/Private/MuR/CodeRunner.cpp b/Engine/Plugins/Experimental/Mutable/Source/MutableRuntime/Private/MuR/CodeRunner.cpp
--- a/Engine/Plugins/Experimental/Mutable/Source/MutableRuntime/Private/MuR/CodeRunner.cpp
+++ b/Engine/Plugins/Experimental/Mutable/Source/MutableRuntime/Private/MuR/CodeRunner.cpp
@@ -2543,14 +2543,19 @@ namespace mu
 			OP::ImageLayerArgs args = pModel->GetPrivate()->m_program.GetOpArgs<OP::ImageLayerArgs>(item.At);
             switch (item.Stage)
             {
-            case 0:
-                    AddOp( FScheduledOp( item.At, item, 1),
-                           FScheduledOp( args.base, item),
-                           FScheduledOp( args.blended, item),
-                           FScheduledOp( args.mask, item) );
-                break;
+			case 0:
+				AddOp(FScheduledOp(item.At, item, 1),
+					FScheduledOp(args.base, item));
+				break;
 
-            case 1:
+			case 1:
+				// Request the rest of the data.
+				AddOp(FScheduledOp(item.At, item, 2),
+					FScheduledOp(args.blended, item),
+					FScheduledOp(args.mask, item));
+				break;
+
+			case 2:
 				// This has been moved to a task. It should have been intercepted in IssueOp.
 				check(false);
 				break;
diff --git a/Engine/Plugins/Experimental/Mutable/Source/MutableRuntime/Private/MuR/CodeRunnerTasked.cpp b/Engine/Plugins/Experimental/Mutable/Source/MutableRuntime/Private/MuR/CodeRunnerTasked.cpp
--- a/Engine/Plugins/Experimental/Mutable/Source/MutableRuntime/Private/MuR/CodeRunnerTasked.cpp
+++ b/Engine/Plugins/Experimental/Mutable/Source/MutableRuntime/Private/MuR/CodeRunnerTasked.cpp
@@ -1163,6 +1163,8 @@ namespace mu
 				Result = ImageSaturate(Source.get(), Factor);
 			}
 		}
+
+		Source = nullptr;
 	}
 
 	//---------------------------------------------------------------------------------------------
@@ -1233,6 +1235,10 @@ namespace mu
 			// during model compilation. The mipmap generation below is not very precise with
 			// the number of mips that are needed and will probably generate too many
 			bool sourceHasMips = pBase->GetLODCount()>1;
+
+			// No longer needed.
+			pBase = nullptr;
+
 			if (sourceHasMips)
 			{
 				int levelCount = Image::GetMipmapCount(Resized->GetSizeX(), Resized->GetSizeY() );
@@ -1252,6 +1258,8 @@ namespace mu
 		{
 			Result = pBase;
 		}
+
+		pBase = nullptr;
 	}
 
 	//---------------------------------------------------------------------------------------------
@@ -1315,6 +1323,10 @@ namespace mu
 		// during model compilation. The mipmap generation below is not very precise with
 		// the number of mips that are needed and will probably generate too many
 		bool sourceHasMips = pBase->GetLODCount()>1;
+
+		// We don't need this anymore
+		pBase = nullptr;
+
 		if (sourceHasMips)
 		{
 			int levelCount = Image::GetMipmapCount(pResult->GetSizeX(), pResult->GetSizeY());
@@ -1382,6 +1394,8 @@ namespace mu
 		{
 			Result = ImageInvert(Source.get());
 		}
+
+		Source = nullptr;
 	}
 
 	//---------------------------------------------------------------------------------------------
@@ -1967,7 +1981,7 @@ namespace mu
 
 		case OP_TYPE::IM_LAYER:
 		{
-			if (item.Stage == 1)
+			if (item.Stage == 2)
 			{
 				OP::ImageLayerArgs args = program.GetOpArgs<OP::ImageLayerArgs>(item.At);
 				Issued = MakeShared<FImageLayerTask>(item, GetMemory(), args, m_pSettings->GetPrivate()->m_imageCompressionQuality);
