[mutable] Prevent crash when missing skeleton in optimize.
#rnx

diff --git a/Engine/Plugins/Experimental/Mutable/Source/MutableRuntime/Private/MuR/OpMeshOptimizeSkinning.h b/Engine/Plugins/Experimental/Mutable/Source/MutableRuntime/Private/MuR/OpMeshOptimizeSkinning.h
--- a/Engine/Plugins/Experimental/Mutable/Source/MutableRuntime/Private/MuR/OpMeshOptimizeSkinning.h
+++ b/Engine/Plugins/Experimental/Mutable/Source/MutableRuntime/Private/MuR/OpMeshOptimizeSkinning.h
@@ -27,6 +27,12 @@ namespace mu
 		TArray<uint16> RemapedBoneIndices;
 
 		mu::SkeletonPtrConst Skeleton = InMesh->GetSkeleton();
+
+		if (!Skeleton)
+		{
+			return nullptr;
+		}
+
 		RemapedBoneIndices.AddZeroed(Skeleton->GetBoneCount());
 
 		int32 SkinnedBonesCount = 0;
