Fix unnecessary private include from parent directory

#rnx

diff --git a/Engine/Plugins/Experimental/Mutable/Source/CustomizableObjectEditor/Private/MuCOE/AssetDefinition_CustomizableObject.cpp b/Engine/Plugins/Experimental/Mutable/Source/CustomizableObjectEditor/Private/MuCOE/AssetDefinition_CustomizableObject.cpp
--- a/Engine/Plugins/Experimental/Mutable/Source/CustomizableObjectEditor/Private/MuCOE/AssetDefinition_CustomizableObject.cpp
+++ b/Engine/Plugins/Experimental/Mutable/Source/CustomizableObjectEditor/Private/MuCOE/AssetDefinition_CustomizableObject.cpp
@@ -2,19 +2,19 @@
 
 #include "MuCOE/AssetDefinition_CustomizableObject.h"
 
+#include "AssetRegistry/AssetRegistryModule.h"
 #include "AssetToolsModule.h"
 #include "ContentBrowserMenuContexts.h"
 #include "ContentBrowserModule.h"
 #include "Editor.h"
 #include "IContentBrowserSingleton.h"
-#include "ToolMenus.h"
-#include "AssetRegistry/AssetRegistryModule.h"
-#include "CustomizableObject/Classes/MuCO/CustomizableObject.h"
 #include "Misc/MessageDialog.h"
+#include "MuCO/CustomizableObject.h"
 #include "MuCO/CustomizableObjectInstance.h"
 #include "MuCO/CustomizableObjectSystem.h"
 #include "MuCOE/CustomizableObjectEditorModule.h"
 #include "Subsystems/AssetEditorSubsystem.h"
+#include "ToolMenus.h"
 
 
 #define LOCTEXT_NAMESPACE "CustomizableObjectEditor"
