[Mutable] Fix mesh variation mesh data combination.

- The dummy mesh data from non default options was combined with the default. The default should be the only one inspected.
- This caused some mesh reshape test to give false positives.

#rnx

diff --git a/Engine/Plugins/Experimental/Mutable/Source/CustomizableObjectEditor/Private/MuCOE/GenerateMutableSource/GenerateMutableSourceMesh.cpp b/Engine/Plugins/Experimental/Mutable/Source/CustomizableObjectEditor/Private/MuCOE/GenerateMutableSource/GenerateMutableSourceMesh.cpp
--- a/Engine/Plugins/Experimental/Mutable/Source/CustomizableObjectEditor/Private/MuCOE/GenerateMutableSource/GenerateMutableSourceMesh.cpp
+++ b/Engine/Plugins/Experimental/Mutable/Source/CustomizableObjectEditor/Private/MuCOE/GenerateMutableSource/GenerateMutableSourceMesh.cpp
@@ -3049,7 +3049,7 @@ mu::NodeMeshPtr GenerateMutableSourceMesh(const UEdGraphPin * Pin,
 				FMutableGraphMeshGenerationData DummyMeshData;
 				mu::NodeMeshPtr ChildNode = GenerateMutableSourceMesh(ConnectedPin, GenerationContext, DummyMeshData);
 				MeshNode->SetVariationMesh(VariationIndex, ChildNode.get());
-				MeshData.Combine(DummyMeshData);
+				//MeshData.Combine(DummyMeshData);
 			}
 		}
 	}
