[Mutable] Disable mesh reshape warning as it may generate false positives

- Reenable mesh variation mesh properties combination. Having this disabled could cause problems when computing the most generic mesh properties.

#rnx

diff --git a/Engine/Plugins/Experimental/Mutable/Source/CustomizableObjectEditor/Private/MuCOE/GenerateMutableSource/GenerateMutableSourceMesh.cpp b/Engine/Plugins/Experimental/Mutable/Source/CustomizableObjectEditor/Private/MuCOE/GenerateMutableSource/GenerateMutableSourceMesh.cpp
--- a/Engine/Plugins/Experimental/Mutable/Source/CustomizableObjectEditor/Private/MuCOE/GenerateMutableSource/GenerateMutableSourceMesh.cpp
+++ b/Engine/Plugins/Experimental/Mutable/Source/CustomizableObjectEditor/Private/MuCOE/GenerateMutableSource/GenerateMutableSourceMesh.cpp
@@ -3046,10 +3046,10 @@ mu::NodeMeshPtr GenerateMutableSourceMesh(const UEdGraphPin * Pin,
 			MeshNode->SetVariationTag(VariationIndex, TCHAR_TO_ANSI(*TypedNodeMeshVar->Variations[VariationIndex].Tag));
 			if (const UEdGraphPin* ConnectedPin = FollowInputPin(*VariationPin))
 			{
-				FMutableGraphMeshGenerationData DummyMeshData;
-				mu::NodeMeshPtr ChildNode = GenerateMutableSourceMesh(ConnectedPin, GenerationContext, DummyMeshData);
+				FMutableGraphMeshGenerationData VariationMeshData;
+				mu::NodeMeshPtr ChildNode = GenerateMutableSourceMesh(ConnectedPin, GenerationContext, VariationMeshData);
 				MeshNode->SetVariationMesh(VariationIndex, ChildNode.get());
-				//MeshData.Combine(DummyMeshData);
+				MeshData.Combine(VariationMeshData);
 			}
 		}
 	}
@@ -3230,13 +3230,23 @@ mu::NodeMeshPtr GenerateMutableSourceMesh(const UEdGraphPin * Pin,
 			}
 		}
 
-		// If any of the shape pins is not set, don't warn about it.
-		if ( BaseShapeTriangleCount != PinNotSetValue && TargetShapeTriangleCount != PinNotSetValue )
-		{ 
-			if ( BaseShapeTriangleCount != TargetShapeTriangleCount || BaseShapeTriangleCount == -1 || TargetShapeTriangleCount == -1)
+
+		// There is cases where it is not possible to determine if the test passes or not, e.g., mesh variations or switches.
+		// Until now if there were the possibility of two meshes not being compatible the warning was raised. This is not ideal
+		// as there are legitimate cases were the meshes will match but we cannot be sure they will. For now disable the warning.
+		
+		constexpr bool bDissableMeshReshapeWarning = true;
+
+		if (!bDissableMeshReshapeWarning)
+		{
+			// If any of the shape pins is not set, don't warn about it.
+			if (BaseShapeTriangleCount != PinNotSetValue && TargetShapeTriangleCount != PinNotSetValue)
 			{
-				GenerationContext.Compiler->CompilerLog(LOCTEXT("ReshapeMeshShapeIncompatible", 
-					"Base and Target Shapes might not be compatible. Don't have the same number of triangles."), Node, EMessageSeverity::Warning);
+				if (BaseShapeTriangleCount != TargetShapeTriangleCount || BaseShapeTriangleCount == -1 || TargetShapeTriangleCount == -1)
+				{
+					GenerationContext.Compiler->CompilerLog(LOCTEXT("ReshapeMeshShapeIncompatible",
+						"Base and Target Shapes might not be compatible. Don't have the same number of triangles."), Node, EMessageSeverity::Warning);
+				}
 			}
 		}
 
