[mutable] Fix bug with wrong code execution when double-executing an external image request.
#rnx

diff --git a/Engine/Plugins/Experimental/Mutable/Source/MutableRuntime/Private/MuR/CodeRunnerTasked.cpp b/Engine/Plugins/Experimental/Mutable/Source/MutableRuntime/Private/MuR/CodeRunnerTasked.cpp
--- a/Engine/Plugins/Experimental/Mutable/Source/MutableRuntime/Private/MuR/CodeRunnerTasked.cpp
+++ b/Engine/Plugins/Experimental/Mutable/Source/MutableRuntime/Private/MuR/CodeRunnerTasked.cpp
@@ -1826,7 +1826,6 @@ namespace mu
 		virtual void Complete(CodeRunner* Runner) override;
 		
 	private:
-		FScheduledOp Item;
 		uint8 MipmapsToSkip;
 		EXTERNAL_IMAGE_ID Id;
 
@@ -1839,7 +1838,7 @@ namespace mu
 	//---------------------------------------------------------------------------------------------
 	FImageExternalLoadTask::FImageExternalLoadTask(FScheduledOp InItem,  uint8 InMipmapsToSkip, EXTERNAL_IMAGE_ID InId)
 	{
-		Item = InItem;
+		Op = InItem;
 		MipmapsToSkip = InMipmapsToSkip;
 		Id = InId;
 	}
@@ -1873,7 +1872,7 @@ namespace mu
 			Invoke(ExternalCleanUpFunc);
 		}
 
-		Runner->GetMemory().SetImage(Item, Result);
+		Runner->GetMemory().SetImage(Op, Result);
 	}	
 
 	
