[mutable] Cleanup task issueing interface to prevent wrong data initialization in the future.
#rnx

diff --git a/Engine/Plugins/Experimental/Mutable/Source/MutableRuntime/Private/MuR/CodeRunner.cpp b/Engine/Plugins/Experimental/Mutable/Source/MutableRuntime/Private/MuR/CodeRunner.cpp
--- a/Engine/Plugins/Experimental/Mutable/Source/MutableRuntime/Private/MuR/CodeRunner.cpp
+++ b/Engine/Plugins/Experimental/Mutable/Source/MutableRuntime/Private/MuR/CodeRunner.cpp
@@ -172,9 +172,7 @@ namespace mu
 
 	
     //---------------------------------------------------------------------------------------------
-    void CodeRunner::RunCode_Conditional( FScheduledOp& item,
-                                          const Model* pModel
-                                          )
+    void CodeRunner::RunCode_Conditional( const FScheduledOp& item, const Model* pModel )
     {
         MUTABLE_CPUPROFILER_SCOPE(RunCode_Conditional);
 
@@ -245,9 +243,7 @@ namespace mu
 
 
 	//---------------------------------------------------------------------------------------------
-	void CodeRunner::RunCode_Switch(FScheduledOp& item,
-		const Model* pModel
-	)
+	void CodeRunner::RunCode_Switch(const FScheduledOp& item, const Model* pModel )
 	{
 		OP_TYPE type = pModel->GetPrivate()->m_program.GetOpType(item.At);
 
@@ -360,10 +356,7 @@ namespace mu
 
 
     //---------------------------------------------------------------------------------------------
-    void CodeRunner::RunCode_Instance( FScheduledOp& item,
-                                  const Model* pModel,
-                                  uint32 lodMask
-                                  )
+    void CodeRunner::RunCode_Instance(const FScheduledOp& item, const Model* pModel, uint32 lodMask )
     {
         MUTABLE_CPUPROFILER_SCOPE(RunCode_Instance);
 
@@ -713,10 +706,7 @@ namespace mu
 
 
     //---------------------------------------------------------------------------------------------
-    void CodeRunner::RunCode_InstanceAddResource( FScheduledOp& item,
-                                  const Model* pModel,
-                                  const Parameters* pParams
-                                  )
+    void CodeRunner::RunCode_InstanceAddResource(const FScheduledOp& item, const Model* pModel, const Parameters* pParams )
     {
 		MUTABLE_CPUPROFILER_SCOPE(RunCode_InstanceAddResource);
 
@@ -823,7 +813,7 @@ namespace mu
 
 
     //---------------------------------------------------------------------------------------------
-    void CodeRunner::RunCode_ConstantResource( FScheduledOp& item, const Model* pModel )
+    void CodeRunner::RunCode_ConstantResource(const FScheduledOp& item, const Model* pModel )
     {
 		MUTABLE_CPUPROFILER_SCOPE(RunCode_Constant);
 
@@ -896,7 +886,7 @@ namespace mu
     }
 
     //---------------------------------------------------------------------------------------------
-    void CodeRunner::RunCode_Mesh( FScheduledOp& item, const Model* pModel )
+    void CodeRunner::RunCode_Mesh(const FScheduledOp& item, const Model* pModel )
     {
 		MUTABLE_CPUPROFILER_SCOPE(RunCode_Mesh);
 
@@ -2503,10 +2493,7 @@ namespace mu
     }
 
     //---------------------------------------------------------------------------------------------
-    void CodeRunner::RunCode_Image( FScheduledOp& item,
-                                    const Parameters* pParams,
-                                    const Model* pModel
-                                    )
+    void CodeRunner::RunCode_Image(const FScheduledOp& item, const Parameters* pParams, const Model* pModel )
     {
 		MUTABLE_CPUPROFILER_SCOPE(RunCode_Image);
 
@@ -4219,10 +4206,7 @@ namespace mu
 
 
     //---------------------------------------------------------------------------------------------
-    void CodeRunner::RunCode_Bool( FScheduledOp& item,
-                                   const Parameters* pParams,
-                                   const Model* pModel
-                                   )
+    void CodeRunner::RunCode_Bool(const FScheduledOp& item, const Parameters* pParams, const Model* pModel )
     {
         MUTABLE_CPUPROFILER_SCOPE(RunCode_Bool);
 
@@ -4463,10 +4447,7 @@ namespace mu
 
 
     //---------------------------------------------------------------------------------------------
-    void CodeRunner::RunCode_Int( FScheduledOp& item,
-                                  const Parameters* pParams,
-                                  const Model* pModel
-                              )
+    void CodeRunner::RunCode_Int(const FScheduledOp& item, const Parameters* pParams, const Model* pModel )
     {
         MUTABLE_CPUPROFILER_SCOPE(RunCode_Int);
 
@@ -4520,10 +4501,7 @@ namespace mu
 
 
     //---------------------------------------------------------------------------------------------
-    void CodeRunner::RunCode_Scalar( FScheduledOp& item,
-                                     const Parameters* pParams,
-                                     const Model* pModel
-                                     )
+    void CodeRunner::RunCode_Scalar(const FScheduledOp& item, const Parameters* pParams, const Model* pModel )
     {
         MUTABLE_CPUPROFILER_SCOPE(RunCode_Scalar);
 
@@ -4638,7 +4616,7 @@ namespace mu
 
 
     //---------------------------------------------------------------------------------------------
-    void CodeRunner::RunCode_String( FScheduledOp& item, const Parameters* pParams, const Model* pModel )
+    void CodeRunner::RunCode_String(const FScheduledOp& item, const Parameters* pParams, const Model* pModel )
     {
         MUTABLE_CPUPROFILER_SCOPE(RunCode_String );
 
@@ -4674,10 +4652,7 @@ namespace mu
 
 
     //---------------------------------------------------------------------------------------------
-    void CodeRunner::RunCode_Colour( FScheduledOp& item,
-                                     const Parameters* pParams,
-                                     const Model* pModel
-                                     )
+    void CodeRunner::RunCode_Colour(const FScheduledOp& item, const Parameters* pParams, const Model* pModel )
     {
 		MUTABLE_CPUPROFILER_SCOPE(RunCode_Colour);
 
@@ -4990,10 +4965,7 @@ namespace mu
 
 
     //---------------------------------------------------------------------------------------------
-    void CodeRunner::RunCode_Projector( FScheduledOp& item,
-                                const Parameters* pParams,
-                                const Model* pModel
-                              )
+    void CodeRunner::RunCode_Projector(const FScheduledOp& item, const Parameters* pParams, const Model* pModel )
     {
         MUTABLE_CPUPROFILER_SCOPE(RunCode_Projector);
 
@@ -5034,9 +5006,7 @@ namespace mu
 
 
     //---------------------------------------------------------------------------------------------
-    void CodeRunner::RunCode_Layout( FScheduledOp& item,
-                                     const Model* pModel
-                                     )
+    void CodeRunner::RunCode_Layout(const FScheduledOp& item, const Model* pModel )
     {
         //MUTABLE_CPUPROFILER_SCOPE(RunCode_Layout);
 
@@ -5213,10 +5183,7 @@ namespace mu
 
 
     //---------------------------------------------------------------------------------------------
-    void CodeRunner::RunCode( FScheduledOp& item,
-                              const Parameters* pParams,
-                              const Model* pModel,
-                              uint32 lodMask)
+    void CodeRunner::RunCode( const FScheduledOp& item, const Parameters* pParams, const Model* pModel, uint32 lodMask)
     {
 		//UE_LOG( LogMutableCore, Log, TEXT("Running :%5d , %d "), item.At, item.Stage );
 		check( item.Type == FScheduledOp::EType::Full );
@@ -5315,11 +5282,7 @@ namespace mu
 
 
 	//---------------------------------------------------------------------------------------------
-	void CodeRunner::RunCodeImageDesc(FScheduledOp& item,
-		const Parameters* pParams,
-		const Model* pModel, 
-		uint32 lodMask
-	)
+	void CodeRunner::RunCodeImageDesc(const FScheduledOp& item, const Parameters* pParams, const Model* pModel,  uint32 lodMask )
 	{
 		MUTABLE_CPUPROFILER_SCOPE(RunCodeImageDesc);
 
diff --git a/Engine/Plugins/Experimental/Mutable/Source/MutableRuntime/Private/MuR/CodeRunner.h b/Engine/Plugins/Experimental/Mutable/Source/MutableRuntime/Private/MuR/CodeRunner.h
--- a/Engine/Plugins/Experimental/Mutable/Source/MutableRuntime/Private/MuR/CodeRunner.h
+++ b/Engine/Plugins/Experimental/Mutable/Source/MutableRuntime/Private/MuR/CodeRunner.h
@@ -83,24 +83,24 @@ namespace  mu
 
         Ptr<RangeIndex> BuildCurrentOpRangeIndex( const FScheduledOp&, const Parameters*, const Model*, int32 ParameterIndex );
 
-        void RunCode( FScheduledOp&, const Parameters*, const Model*, uint32 LodMask);
-
-        void RunCode_Conditional( FScheduledOp&, const Model* );
-        void RunCode_Switch( FScheduledOp&, const Model* );
-        void RunCode_Instance( FScheduledOp&, const Model*, uint32 LodMask );
-        void RunCode_InstanceAddResource( FScheduledOp&, const Model*, const Parameters* );
-        void RunCode_ConstantResource( FScheduledOp&, const Model* );
-        void RunCode_Mesh( FScheduledOp&, const Model* );
-        void RunCode_Image( FScheduledOp&, const Parameters*, const Model* );
-        void RunCode_Layout( FScheduledOp&, const Model* );
-        void RunCode_Bool( FScheduledOp&, const Parameters*, const Model* );
-        void RunCode_Int( FScheduledOp&, const Parameters*, const Model* );
-        void RunCode_Scalar( FScheduledOp&, const Parameters*, const Model* );        
-        void RunCode_String( FScheduledOp&, const Parameters*, const Model* );
-        void RunCode_Colour( FScheduledOp&, const Parameters*, const Model* );
-        void RunCode_Projector( FScheduledOp&, const Parameters*, const Model* );
-
-		void RunCodeImageDesc(FScheduledOp&, const Parameters*, const Model*, uint32 LodMask);
+        void RunCode( const FScheduledOp&, const Parameters*, const Model*, uint32 LodMask);
+
+        void RunCode_Conditional(const FScheduledOp&, const Model* );
+        void RunCode_Switch(const FScheduledOp&, const Model* );
+        void RunCode_Instance(const FScheduledOp&, const Model*, uint32 LodMask );
+        void RunCode_InstanceAddResource(const FScheduledOp&, const Model*, const Parameters* );
+        void RunCode_ConstantResource(const FScheduledOp&, const Model* );
+        void RunCode_Mesh(const FScheduledOp&, const Model* );
+        void RunCode_Image(const FScheduledOp&, const Parameters*, const Model* );
+        void RunCode_Layout(const FScheduledOp&, const Model* );
+        void RunCode_Bool(const FScheduledOp&, const Parameters*, const Model* );
+        void RunCode_Int(const FScheduledOp&, const Parameters*, const Model* );
+        void RunCode_Scalar(const FScheduledOp&, const Parameters*, const Model* );
+        void RunCode_String(const FScheduledOp&, const Parameters*, const Model* );
+        void RunCode_Colour(const FScheduledOp&, const Parameters*, const Model* );
+        void RunCode_Projector(const FScheduledOp&, const Parameters*, const Model* );
+
+		void RunCodeImageDesc(const FScheduledOp&, const Parameters*, const Model*, uint32 LodMask);
 
     public:
 
@@ -150,7 +150,7 @@ namespace  mu
 		class FIssuedTask
 		{
 		public:
-			FScheduledOp Op;
+			const FScheduledOp Op;
 
 #ifdef MUTABLE_USE_NEW_TASKGRAPH
 			UE::Tasks::FTask Event = {};
@@ -158,6 +158,7 @@ namespace  mu
 			FGraphEventRef Event = nullptr;
 #endif
 
+			FIssuedTask(const FScheduledOp& InOp) : Op(InOp) {}
 			virtual ~FIssuedTask() {}
 			virtual bool Prepare(CodeRunner*, const TSharedPtr<const Model>&, bool& bOutFailed) { bOutFailed = false; return true; }
 			virtual void DoWork() {}
@@ -176,9 +177,9 @@ namespace  mu
 		class FLoadMeshRomTask : public CodeRunner::FIssuedTask
 		{
 		public:
-			FLoadMeshRomTask(FScheduledOp InOp, int32 InRomIndex)
+			FLoadMeshRomTask( const FScheduledOp& InOp, int32 InRomIndex)
+				: FIssuedTask(InOp)
 			{
-				Op = InOp;
 				RomIndex = InRomIndex;
 			}
 
@@ -194,9 +195,9 @@ namespace  mu
 		class FLoadImageRomsTask : public CodeRunner::FIssuedTask
 		{
 		public:
-			FLoadImageRomsTask(FScheduledOp InOp, int32 InLODIndexIndex, int32 InLODIndexCount )
+			FLoadImageRomsTask(const FScheduledOp& InOp, int32 InLODIndexIndex, int32 InLODIndexCount )
+				: FIssuedTask(InOp)
 			{
-				Op = InOp;
 				LODIndexIndex = InLODIndexIndex;
 				LODIndexCount = InLODIndexCount;
 			}
diff --git a/Engine/Plugins/Experimental/Mutable/Source/MutableRuntime/Private/MuR/CodeRunnerTasked.cpp b/Engine/Plugins/Experimental/Mutable/Source/MutableRuntime/Private/MuR/CodeRunnerTasked.cpp
--- a/Engine/Plugins/Experimental/Mutable/Source/MutableRuntime/Private/MuR/CodeRunnerTasked.cpp
+++ b/Engine/Plugins/Experimental/Mutable/Source/MutableRuntime/Private/MuR/CodeRunnerTasked.cpp
@@ -411,7 +411,7 @@ namespace mu
 	class FImageLayerTask : public CodeRunner::FIssuedTask
 	{
 	public:
-		FImageLayerTask(FScheduledOp, FProgramCache&, const OP::ImageLayerArgs&, int imageCompressionQuality );
+		FImageLayerTask(const FScheduledOp&, FProgramCache&, const OP::ImageLayerArgs&, int imageCompressionQuality );
 
 		// FIssuedTask interface
 		void DoWork() override;
@@ -428,9 +428,9 @@ namespace mu
 
 
 	//---------------------------------------------------------------------------------------------
-	FImageLayerTask::FImageLayerTask(FScheduledOp InOp, FProgramCache& Memory, const OP::ImageLayerArgs& InArgs, int imageCompressionQuality)
+	FImageLayerTask::FImageLayerTask(const FScheduledOp& InOp, FProgramCache& Memory, const OP::ImageLayerArgs& InArgs, int imageCompressionQuality)
+		: FIssuedTask(InOp)
 	{
-		Op = InOp;
 		m_imageCompressionQuality = imageCompressionQuality;
 		Args = InArgs;
 		m_base = Memory.GetImage({ Args.base, InOp.ExecutionIndex, InOp.ExecutionOptions });
@@ -673,7 +673,7 @@ namespace mu
 	class FImageLayerColourTask : public CodeRunner::FIssuedTask
 	{
 	public:
-		FImageLayerColourTask(FScheduledOp, FProgramCache&, const OP::ImageLayerColourArgs&, int imageCompressionQuality);
+		FImageLayerColourTask(const FScheduledOp&, FProgramCache&, const OP::ImageLayerColourArgs&, int imageCompressionQuality);
 
 		// FIssuedTask interface
 		void DoWork() override;
@@ -690,9 +690,9 @@ namespace mu
 
 
 	//---------------------------------------------------------------------------------------------
-	FImageLayerColourTask::FImageLayerColourTask(FScheduledOp InOp, FProgramCache& Memory, const OP::ImageLayerColourArgs& InArgs, int imageCompressionQuality)
+	FImageLayerColourTask::FImageLayerColourTask(const FScheduledOp& InOp, FProgramCache& Memory, const OP::ImageLayerColourArgs& InArgs, int imageCompressionQuality)
+		: FIssuedTask(InOp)
 	{
-		Op = InOp;
 		m_imageCompressionQuality = imageCompressionQuality;
 		Args = InArgs;
 		m_base = Memory.GetImage({ Args.base, InOp.ExecutionIndex, InOp.ExecutionOptions });
@@ -885,7 +885,7 @@ namespace mu
 	class FImagePixelFormatTask : public CodeRunner::FIssuedTask
 	{
 	public:
-		FImagePixelFormatTask(FScheduledOp, FProgramCache&, const OP::ImagePixelFormatArgs&, int imageCompressionQuality);
+		FImagePixelFormatTask(const FScheduledOp&, FProgramCache&, const OP::ImagePixelFormatArgs&, int imageCompressionQuality);
 
 		// FIssuedTask interface
 		void DoWork() override;
@@ -900,9 +900,9 @@ namespace mu
 
 
 	//---------------------------------------------------------------------------------------------
-	FImagePixelFormatTask::FImagePixelFormatTask(FScheduledOp InOp, FProgramCache& Memory, const OP::ImagePixelFormatArgs& InArgs, int imageCompressionQuality)
+	FImagePixelFormatTask::FImagePixelFormatTask(const FScheduledOp& InOp, FProgramCache& Memory, const OP::ImagePixelFormatArgs& InArgs, int imageCompressionQuality)
+		: FIssuedTask(InOp)
 	{
-		Op = InOp;
 		m_imageCompressionQuality = imageCompressionQuality;
 		Args = InArgs;
 		m_base = Memory.GetImage({ Args.source, InOp.ExecutionIndex, InOp.ExecutionOptions });
@@ -963,7 +963,7 @@ namespace mu
 	class FImageMipmapTask : public CodeRunner::FIssuedTask
 	{
 	public:
-		FImageMipmapTask(FScheduledOp, FProgramCache&, const OP::ImageMipmapArgs&, int imageCompressionQuality);
+		FImageMipmapTask(const FScheduledOp&, FProgramCache&, const OP::ImageMipmapArgs&, int imageCompressionQuality);
 
 		// FIssuedTask interface
 		void DoWork() override;
@@ -978,9 +978,9 @@ namespace mu
 
 
 	//---------------------------------------------------------------------------------------------
-	FImageMipmapTask::FImageMipmapTask(FScheduledOp InOp, FProgramCache& Memory, const OP::ImageMipmapArgs& InArgs, int imageCompressionQuality)
+	FImageMipmapTask::FImageMipmapTask(const FScheduledOp& InOp, FProgramCache& Memory, const OP::ImageMipmapArgs& InArgs, int imageCompressionQuality)
+		: FIssuedTask(InOp)
 	{
-		Op = InOp;
 		m_imageCompressionQuality = imageCompressionQuality;
 		Args = InArgs;
 		m_base = Memory.GetImage({ Args.source, InOp.ExecutionIndex, InOp.ExecutionOptions });
@@ -1044,7 +1044,7 @@ namespace mu
 	class FImageSwizzleTask : public CodeRunner::FIssuedTask
 	{
 	public:
-		FImageSwizzleTask(FScheduledOp, FProgramCache&, const OP::ImageSwizzleArgs&);
+		FImageSwizzleTask(const FScheduledOp&, FProgramCache&, const OP::ImageSwizzleArgs&);
 
 		// FIssuedTask interface
 		void DoWork() override;
@@ -1058,9 +1058,9 @@ namespace mu
 
 
 	//---------------------------------------------------------------------------------------------
-	FImageSwizzleTask::FImageSwizzleTask(FScheduledOp InOp, FProgramCache& Memory, const OP::ImageSwizzleArgs& InArgs)
+	FImageSwizzleTask::FImageSwizzleTask(const FScheduledOp& InOp, FProgramCache& Memory, const OP::ImageSwizzleArgs& InArgs)
+		: FIssuedTask(InOp)
 	{
-		Op = InOp;
 		Args = InArgs;
 		for (int i = 0; i < MUTABLE_OP_MAX_SWIZZLE_CHANNELS; ++i)
 		{
@@ -1113,7 +1113,7 @@ namespace mu
 	class FImageSaturateTask : public CodeRunner::FIssuedTask
 	{
 	public:
-		FImageSaturateTask(FScheduledOp, FProgramCache&, const OP::ImageSaturateArgs&);
+		FImageSaturateTask(const FScheduledOp&, FProgramCache&, const OP::ImageSaturateArgs&);
 
 		// FIssuedTask interface
 		void DoWork() override;
@@ -1127,9 +1127,9 @@ namespace mu
 
 
 	//---------------------------------------------------------------------------------------------
-	FImageSaturateTask::FImageSaturateTask(FScheduledOp InOp, FProgramCache& Memory, const OP::ImageSaturateArgs& InArgs)
+	FImageSaturateTask::FImageSaturateTask(const FScheduledOp& InOp, FProgramCache& Memory, const OP::ImageSaturateArgs& InArgs)
+		: FIssuedTask(InOp)
 	{
-		Op = InOp;
 		Source = Memory.GetImage(FCacheAddress(InArgs.base, InOp));
 		Factor = Memory.GetScalar(FScheduledOp::FromOpAndOptions(InArgs.factor, InOp, 0));
 	}
@@ -1181,7 +1181,7 @@ namespace mu
 	class FImageResizeTask : public CodeRunner::FIssuedTask
 	{
 	public:
-		FImageResizeTask(FScheduledOp InOp, FProgramCache& Memory, const Ptr<const Image>& InBase, uint16 InSizeX, uint16 InSizeY, int32 InImageCompressionQuality);
+		FImageResizeTask(const FScheduledOp& InOp, FProgramCache& Memory, const Ptr<const Image>& InBase, uint16 InSizeX, uint16 InSizeY, int32 InImageCompressionQuality);
 
 		// FIssuedTask interface
 		void DoWork() override;
@@ -1196,9 +1196,9 @@ namespace mu
 
 
 	//---------------------------------------------------------------------------------------------
-	FImageResizeTask::FImageResizeTask(FScheduledOp InOp, FProgramCache& Memory, const Ptr<const Image>& InBase, uint16 InSizeX, uint16 InSizeY, int32 InImageCompressionQuality)
+	FImageResizeTask::FImageResizeTask(const FScheduledOp& InOp, FProgramCache& Memory, const Ptr<const Image>& InBase, uint16 InSizeX, uint16 InSizeY, int32 InImageCompressionQuality)
+		: FIssuedTask(InOp)
 	{
-		Op = InOp;
 		SizeX = InSizeX;
 		SizeY = InSizeY;
 		pBase = InBase;
@@ -1276,7 +1276,7 @@ namespace mu
 	class FImageResizeRelTask : public CodeRunner::FIssuedTask
 	{
 	public:
-		FImageResizeRelTask(FScheduledOp, FProgramCache&, const OP::ImageResizeRelArgs&, int32 ImageCompressionQuality);
+		FImageResizeRelTask(const FScheduledOp&, FProgramCache&, const OP::ImageResizeRelArgs&, int32 ImageCompressionQuality);
 
 		// FIssuedTask interface
 		void DoWork() override;
@@ -1291,9 +1291,9 @@ namespace mu
 
 
 	//---------------------------------------------------------------------------------------------
-	FImageResizeRelTask::FImageResizeRelTask(FScheduledOp InOp, FProgramCache& Memory, const OP::ImageResizeRelArgs& InArgs, int32 InImageCompressionQuality)
+	FImageResizeRelTask::FImageResizeRelTask(const FScheduledOp& InOp, FProgramCache& Memory, const OP::ImageResizeRelArgs& InArgs, int32 InImageCompressionQuality)
+		: FIssuedTask(InOp)
 	{
-		Op = InOp;
 		Args = InArgs;
 		pBase = Memory.GetImage(FCacheAddress(InArgs.source, InOp));
 		ImageCompressionQuality = InImageCompressionQuality;
@@ -1360,7 +1360,7 @@ namespace mu
 	class FImageInvertTask : public CodeRunner::FIssuedTask
 	{
 	public:
-		FImageInvertTask(FScheduledOp, FProgramCache&, const OP::ImageInvertArgs&);
+		FImageInvertTask(const FScheduledOp&, FProgramCache&, const OP::ImageInvertArgs&);
 
 		// FIssuedTask interface
 		void DoWork() override;
@@ -1373,9 +1373,9 @@ namespace mu
 
 
 	//---------------------------------------------------------------------------------------------
-	FImageInvertTask::FImageInvertTask(FScheduledOp InOp, FProgramCache& Memory, const OP::ImageInvertArgs& InArgs)
+	FImageInvertTask::FImageInvertTask(const FScheduledOp& InOp, FProgramCache& Memory, const OP::ImageInvertArgs& InArgs)
+		: FIssuedTask(InOp)
 	{
-		Op = InOp;
 		Source = Memory.GetImage(FCacheAddress(InArgs.base, InOp));
 	}
 
@@ -1412,7 +1412,7 @@ namespace mu
 	class FImageComposeTask : public CodeRunner::FIssuedTask
 	{
 	public:
-		FImageComposeTask(FScheduledOp, FProgramCache&, const OP::ImageComposeArgs&, int imageCompressionQuality, const Ptr<const Layout>& );
+		FImageComposeTask(const FScheduledOp&, FProgramCache&, const OP::ImageComposeArgs&, int imageCompressionQuality, const Ptr<const Layout>& );
 
 		// FIssuedTask interface
 		void DoWork() override;
@@ -1430,9 +1430,9 @@ namespace mu
 
 
 	//---------------------------------------------------------------------------------------------
-	FImageComposeTask::FImageComposeTask(FScheduledOp InOp, FProgramCache& Memory, const OP::ImageComposeArgs& InArgs, int imageCompressionQuality, const Ptr<const Layout>& InLayout)
+	FImageComposeTask::FImageComposeTask(const FScheduledOp& InOp, FProgramCache& Memory, const OP::ImageComposeArgs& InArgs, int imageCompressionQuality, const Ptr<const Layout>& InLayout)
+		: FIssuedTask(InOp)
 	{
-		Op = InOp;
 		m_imageCompressionQuality = imageCompressionQuality;
 		Args = InArgs;
 		m_layout = InLayout;
@@ -1819,7 +1819,7 @@ namespace mu
 	class FImageExternalLoadTask : public CodeRunner::FIssuedTask
 	{
 	public:
-		FImageExternalLoadTask(FScheduledOp InItem, uint8 InMipmapsToSkip, EXTERNAL_IMAGE_ID InId);
+		FImageExternalLoadTask(const FScheduledOp& InItem, uint8 InMipmapsToSkip, EXTERNAL_IMAGE_ID InId);
 
 		// FIssuedTask interface
 		virtual bool Prepare(CodeRunner*, const TSharedPtr<const Model>&, bool& bOutFailed) override;
@@ -1836,9 +1836,9 @@ namespace mu
 
 
 	//---------------------------------------------------------------------------------------------
-	FImageExternalLoadTask::FImageExternalLoadTask(FScheduledOp InItem,  uint8 InMipmapsToSkip, EXTERNAL_IMAGE_ID InId)
+	FImageExternalLoadTask::FImageExternalLoadTask(const FScheduledOp& InOp,  uint8 InMipmapsToSkip, EXTERNAL_IMAGE_ID InId)
+		: FIssuedTask(InOp)
 	{
-		Op = InItem;
 		MipmapsToSkip = InMipmapsToSkip;
 		Id = InId;
 	}
