#rnx

diff --git a/Engine/Plugins/Experimental/Mutable/Source/CustomizableObject/Classes/MuCO/CustomizableObjectSystem.h b/Engine/Plugins/Experimental/Mutable/Source/CustomizableObject/Classes/MuCO/CustomizableObjectSystem.h
--- a/Engine/Plugins/Experimental/Mutable/Source/CustomizableObject/Classes/MuCO/CustomizableObjectSystem.h
+++ b/Engine/Plugins/Experimental/Mutable/Source/CustomizableObject/Classes/MuCO/CustomizableObjectSystem.h
@@ -258,8 +258,10 @@ public:
 	bool IsSupport16BitBoneIndexEnabled() const;
 
 	bool IsProgressiveMipStreamingEnabled() const;
+	void SetProgressiveMipStreamingEnabled(bool bIsEnabled);
 
 	bool IsOnlyGenerateRequestedLODsEnabled() const;
+	void SetOnlyGenerateRequestedLODsEnabled(bool bIsEnabled);
 
 	void AddPendingReleaseSkeletalMesh( USkeletalMesh* SkeletalMesh );
 
diff --git a/Engine/Plugins/Experimental/Mutable/Source/CustomizableObject/Private/MuCO/CustomizableObjectSystem.cpp b/Engine/Plugins/Experimental/Mutable/Source/CustomizableObject/Private/MuCO/CustomizableObjectSystem.cpp
--- a/Engine/Plugins/Experimental/Mutable/Source/CustomizableObject/Private/MuCO/CustomizableObjectSystem.cpp
+++ b/Engine/Plugins/Experimental/Mutable/Source/CustomizableObject/Private/MuCO/CustomizableObjectSystem.cpp
@@ -2753,6 +2753,13 @@ bool UCustomizableObjectSystem::IsProgressiveMipStreamingEnabled() const
 }
 
 
+void UCustomizableObjectSystem::SetProgressiveMipStreamingEnabled(bool bIsEnabled)
+{
+	check(Private != nullptr);
+	Private->EnableMutableProgressiveMipStreaming = bIsEnabled ? 1 : 0;
+}
+
+
 bool UCustomizableObjectSystem::IsOnlyGenerateRequestedLODsEnabled() const
 {
 	check(Private != nullptr);
@@ -2760,6 +2767,13 @@ bool UCustomizableObjectSystem::IsOnlyGenerateRequestedLODsEnabled() const
 }
 
 
+void UCustomizableObjectSystem::SetOnlyGenerateRequestedLODsEnabled(bool bIsEnabled)
+{
+	check(Private != nullptr);
+	Private->EnableOnlyGenerateRequestedLODs = bIsEnabled ? 1 : 0;
+}
+
+
 void UCustomizableObjectSystem::AddUncompiledCOWarning(UCustomizableObject* InObject)
 {
 	if (!InObject)
diff --git a/Engine/Plugins/Experimental/Mutable/Source/CustomizableObjectEditor/Private/MuCOE/CustomizableObjectEditorViewportClient.cpp b/Engine/Plugins/Experimental/Mutable/Source/CustomizableObjectEditor/Private/MuCOE/CustomizableObjectEditorViewportClient.cpp
--- a/Engine/Plugins/Experimental/Mutable/Source/CustomizableObjectEditor/Private/MuCOE/CustomizableObjectEditorViewportClient.cpp
+++ b/Engine/Plugins/Experimental/Mutable/Source/CustomizableObjectEditor/Private/MuCOE/CustomizableObjectEditorViewportClient.cpp
@@ -24,6 +24,7 @@
 #include "Misc/MessageDialog.h"
 #include "Misc/PackageName.h"
 #include "MuCO/CustomizableObjectInstance.h"
+#include "MuCO/CustomizableObjectSystem.h"
 #include "MuCO/CustomizableObjectMipDataProvider.h"
 #include "MuCO/UnrealBakeHelpers.h"
 #include "MuCOE/CustomizableObjectBakeHelpers.h"
@@ -2235,6 +2236,12 @@ void RemoveRestrictedChars(FString& String)
 
 //-------------------------------------------------------------------------------------------------
 void FCustomizableObjectEditorViewportClient::BakeInstance()
+{
+	BakeInstance(nullptr);
+}
+
+//-------------------------------------------------------------------------------------------------
+void FCustomizableObjectEditorViewportClient::BakeInstance(UCustomizableObjectInstance* InInstance)
 {
 	if (!AssetRegistryLoaded)
 	{
@@ -2247,11 +2254,14 @@ void FCustomizableObjectEditorViewportClient::BakeInstance()
 		return;
 	}
 
+	UCustomizableObjectInstance* Instance = InInstance ? InInstance : CustomizableObjectEditorPtr.Pin()->GetPreviewInstance();
+
 	bool bHasSkeletalMesh = false;
-	int32 NumComponents = SkeletalMeshComponents.Num();
+	int32 NumComponents = Instance->SkeletalMeshes.Num();
+
 	for (int32 ComponentIndex = 0; ComponentIndex < NumComponents; ++ComponentIndex)
 	{
-		if (SkeletalMeshComponents[ComponentIndex].IsValid() && SkeletalMeshComponents[ComponentIndex]->GetSkinnedAsset())
+		if (Instance->SkeletalMeshes.IsValidIndex(ComponentIndex) && Instance->SkeletalMeshes[ComponentIndex])
 		{
 			bHasSkeletalMesh = true;
 		}
@@ -2262,7 +2272,25 @@ void FCustomizableObjectEditorViewportClient::BakeInstance()
 		return;
 	}
 
-	UCustomizableObjectInstance* Instance = CustomizableObjectEditorPtr.Pin()->GetPreviewInstance();
+	UCustomizableObjectSystem* System = UCustomizableObjectSystem::GetInstance();
+	check(System);
+
+	if (System->IsProgressiveMipStreamingEnabled())
+	{
+		// The instance in the editor viewport does not have high quality mips in the platform data because streaming is enabled.
+		// Disable streaming and retry with a newly generated temp instance.
+		System->SetProgressiveMipStreamingEnabled(false);
+		// Disable requested LOD generation as it will prevent the new instance from having all the LODs
+		System->SetOnlyGenerateRequestedLODsEnabled(false);
+
+		BakeTempInstance = Instance->Clone();
+		BakeTempInstance->SkeletalMeshes.Empty();
+		BakeTempInstance->UpdatedNativeDelegate.AddSP(this, &FCustomizableObjectEditorViewportClient::BakeInstance);
+		BakeTempInstance->UpdateSkeletalMeshAsync(true, true);
+
+		return;
+	}
+
 	FString ObjectName = Instance->GetCustomizableObject()->GetName();
 	FText DefaultFileName = FText::Format(LOCTEXT("DefaultFileNameForBakeInstance", "{0}"), FText::AsCultureInvariant(ObjectName));
 	bool bExportAllResources = false;
@@ -2320,8 +2348,8 @@ void FCustomizableObjectEditorViewportClient::BakeInstance()
 
 			for (int32 ComponentIndex = 0; ComponentIndex < NumComponents; ++ComponentIndex)
 			{
-				USkeletalMesh* Mesh = SkeletalMeshComponents.Num() && SkeletalMeshComponents[ComponentIndex].IsValid() ?
-					Cast<USkeletalMesh>(SkeletalMeshComponents[ComponentIndex]->GetSkinnedAsset()) : nullptr;
+				USkeletalMesh* Mesh = Instance->SkeletalMeshes.IsValidIndex(ComponentIndex) && Instance->SkeletalMeshes[ComponentIndex] ?
+					Cast<USkeletalMesh>(Instance->SkeletalMeshes[ComponentIndex]) : nullptr;
 				
 				if (!Mesh)
 				{
@@ -2664,7 +2692,7 @@ void FCustomizableObjectEditorViewportClient::BakeInstance()
 					}
 				}
 
-				// Make sure source data is present in the mesh before we duplciate:
+				// Make sure source data is present in the mesh before we duplicate:
 				FUnrealBakeHelpers::BakeHelper_RegenerateImportedModel(Mesh);
 
 				// Skeletal Mesh
@@ -2704,6 +2732,14 @@ void FCustomizableObjectEditorViewportClient::BakeInstance()
 			}
 		}
 	}
+
+	if (InInstance)
+	{
+		// Reenable Mutable texture streaming and requested LOD generation as they had been disabled to bake the textures
+		System->SetProgressiveMipStreamingEnabled(true);
+		System->SetOnlyGenerateRequestedLODsEnabled(true);
+		BakeTempInstance = nullptr;
+	}
 }
 
 
@@ -3008,6 +3044,11 @@ void FCustomizableObjectEditorViewportClient::AddReferencedObjects(FReferenceCol
 	Collector.AddReferencedObject(ClipMorphMaterial);
 	Collector.AddReferencedObject(TransparentPlaneMaterialXY);
 	Collector.AddReferencedObject(AnimationBeingPlayed);
+
+	if (BakeTempInstance)
+	{
+		Collector.AddReferencedObject(BakeTempInstance);
+	}
 }
 
 
diff --git a/Engine/Plugins/Experimental/Mutable/Source/CustomizableObjectEditor/Public/MuCOE/CustomizableObjectEditorViewportClient.h b/Engine/Plugins/Experimental/Mutable/Source/CustomizableObjectEditor/Public/MuCOE/CustomizableObjectEditorViewportClient.h
--- a/Engine/Plugins/Experimental/Mutable/Source/CustomizableObjectEditor/Public/MuCOE/CustomizableObjectEditorViewportClient.h
+++ b/Engine/Plugins/Experimental/Mutable/Source/CustomizableObjectEditor/Public/MuCOE/CustomizableObjectEditorViewportClient.h
@@ -148,7 +148,7 @@ public:
 
 
 /** Viewport Client for the preview viewport */
-class FCustomizableObjectEditorViewportClient : public FEditorViewportClient
+class FCustomizableObjectEditorViewportClient : public FEditorViewportClient, public TSharedFromThis<FCustomizableObjectEditorViewportClient>
 {
 public:
 	FCustomizableObjectEditorViewportClient(TWeakPtr<class ICustomizableObjectInstanceEditor> InCustomizableObjectEditor, FPreviewScene* InPreviewScene);
@@ -196,8 +196,9 @@ public:
 	 */
 	void DrawUVs(FViewport* InViewport, FCanvas* InCanvas, int32 InTextYPos, const FString& MaterialName);
 
-	// Callback for baking the current instance.
+	// Callback for baking the current instance. If nullptr is passed, the instance in the editor will be taken.
 	void BakeInstance();
+	void BakeInstance(class UCustomizableObjectInstance* Instance);
 	
 	// Returns false if the resource has already been duplicated, otherwise, returns true and an unique ResourceName.
 	bool GetUniqueResourceName(UObject* InResource, FString& InOutResourceName, TArray<UObject*>& InCachedObjects, TArray<FString>& InCachedObjectNames);
@@ -515,6 +516,10 @@ private:
 
 	/** Flag to control the bones visibility in the viewport */
 	bool bShowBones;
+
+	// Temp Instance used in the bake process if a new instance is needed because mutable texture streaming is enabled so the viewport 
+	// instance does not have the high quality mips in the texture's platform data
+	UCustomizableObjectInstance* BakeTempInstance = nullptr;
 };
 
 #if UE_ENABLE_INCLUDE_ORDER_DEPRECATED_IN_5_2
