Temporal reduction on warning severity to unblock work

#rnx

diff --git a/Engine/Plugins/Experimental/Mutable/Source/CustomizableObjectEditor/Private/MuCOE/GenerateMutableSource/GenerateMutableSourceSurface.cpp b/Engine/Plugins/Experimental/Mutable/Source/CustomizableObjectEditor/Private/MuCOE/GenerateMutableSource/GenerateMutableSourceSurface.cpp
--- a/Engine/Plugins/Experimental/Mutable/Source/CustomizableObjectEditor/Private/MuCOE/GenerateMutableSource/GenerateMutableSourceSurface.cpp
+++ b/Engine/Plugins/Experimental/Mutable/Source/CustomizableObjectEditor/Private/MuCOE/GenerateMutableSource/GenerateMutableSourceSurface.cpp
@@ -75,7 +75,7 @@ void LayoutOperationUVNormalizedWarning(FMutableGraphGenerationContext& Generati
 					if (!bNormalized)
 					{
 						FText Text = FText::Format(LOCTEXT("UVNotNormalized", "UV from mesh {0} not normalized. Required to perform texture layout operations."), FText::FromString(*MeshData.Mesh->GetName()));
-						GenerationContext.Compiler->CompilerLog(Text, OperationNode);
+						GenerationContext.Compiler->CompilerLog(Text, OperationNode, EMessageSeverity::Type::Info);
 					}
 				}
 			}
