[Mutable] Remove unused class.

#rnx

diff --git a/Engine/Plugins/Experimental/Mutable/Source/CustomizableObject/Classes/MuCO/CustomizableObjectInstance.h b/Engine/Plugins/Experimental/Mutable/Source/CustomizableObject/Classes/MuCO/CustomizableObjectInstance.h
--- a/Engine/Plugins/Experimental/Mutable/Source/CustomizableObject/Classes/MuCO/CustomizableObjectInstance.h
+++ b/Engine/Plugins/Experimental/Mutable/Source/CustomizableObject/Classes/MuCO/CustomizableObjectInstance.h
@@ -27,33 +27,6 @@ struct FPropertyChangedEvent;
 #define MUTABLE_VERTEXBUFFER_TEXCOORDS	2
 
 
-USTRUCT()
-struct FCustomizedMaterialTexture2D
-{
-	GENERATED_USTRUCT_BODY()
-	
-	UPROPERTY(Category=CustomizedMaterialTexture2D, VisibleAnywhere)
-	FName Name;
-	
-	UPROPERTY(Category=CustomizedMaterialTexture2D, EditAnywhere) // Replaced EditInline with EditAnywhere for 4.10
-	TObjectPtr<UTexture2D> Texture;
-
-
-	FCustomizedMaterialTexture2D()
-		: Texture(nullptr)
-	{
-	}
-
-	inline friend FArchive& operator<<( FArchive& Ar, FCustomizedMaterialTexture2D& Data )
-	{
-		Ar << Data.Name;
-		Ar << Data.Texture;
-		return Ar;
-	}
-	
-};
-
-
 // Current state of a projector associated with a parameter.
 namespace EProjectorState
 {
