#rnx

diff --git a/Engine/Plugins/Experimental/Mutable/Source/CustomizableObject/Private/MuCO/UnrealMutableImageProvider.cpp b/Engine/Plugins/Experimental/Mutable/Source/CustomizableObject/Private/MuCO/UnrealMutableImageProvider.cpp
--- a/Engine/Plugins/Experimental/Mutable/Source/CustomizableObject/Private/MuCO/UnrealMutableImageProvider.cpp
+++ b/Engine/Plugins/Experimental/Mutable/Source/CustomizableObject/Private/MuCO/UnrealMutableImageProvider.cpp
@@ -183,6 +183,18 @@ mu::EImageFormat GetMutablePixelFormat(EPixelFormat InTextureFormat)
 			// not being modified by the game thread at the moment and the texture cannot be GCed because of the AddReferencedObjects
 			// in the FUnrealMutableImageProvider
 
+#if WITH_EDITOR
+			if (!TextureToLoad->IsAsyncCacheComplete())
+			{
+				// If the UE texture is being compiled/processed in a background process, it's not safe to access
+				FString TextureName;
+				TextureToLoad->GetName(TextureName);
+				UE_LOG(LogMutable, Warning, TEXT("Failed to get the external texture [%s] because it's still being async cached. Wait until all background processed finish and retry again."), *TextureName);
+				ResultCallback(CreateDummy());
+				return Invoke(TrivialReturn);
+			}
+#endif
+
 			uint8 MipIndex = MipmapsToSkip < TextureToLoad->GetPlatformData()->Mips.Num() ? MipmapsToSkip : TextureToLoad->GetPlatformData()->Mips.Num() - 1;
 
 			// Mips in the mip tail are inlined and can't be streamed, find the smallest mip available.
