#rnx

diff --git a/Engine/Plugins/Experimental/Mutable/Source/CustomizableObjectEditor/Private/MuCOE/GenerateMutableSource/GenerateMutableSourceTable.cpp b/Engine/Plugins/Experimental/Mutable/Source/CustomizableObjectEditor/Private/MuCOE/GenerateMutableSource/GenerateMutableSourceTable.cpp
--- a/Engine/Plugins/Experimental/Mutable/Source/CustomizableObjectEditor/Private/MuCOE/GenerateMutableSource/GenerateMutableSourceTable.cpp
+++ b/Engine/Plugins/Experimental/Mutable/Source/CustomizableObjectEditor/Private/MuCOE/GenerateMutableSource/GenerateMutableSourceTable.cpp
@@ -81,7 +81,7 @@ void FillTableColumn(const UCustomizableObjectNodeTable* TableNode,	mu::TablePtr
 										{
 											TSoftClassPtr<UAnimInstance> AnimInstance(SoftClassProperty->GetPropertyValue(AnimBPData).ToSoftObjectPath());
 
-											if (AnimInstance)
+											if (!AnimInstance.IsNull())
 											{
 												GenerationContext.AnimBPAssetsMap.Add(AnimInstance.ToString(), AnimInstance);
 
