[Mutable] Fix Mutable Core high memory ussage.

#rnx

diff --git a/Engine/Plugins/Experimental/Mutable/Source/MutableRuntime/Private/MuR/CodeRunner.h b/Engine/Plugins/Experimental/Mutable/Source/MutableRuntime/Private/MuR/CodeRunner.h
--- a/Engine/Plugins/Experimental/Mutable/Source/MutableRuntime/Private/MuR/CodeRunner.h
+++ b/Engine/Plugins/Experimental/Mutable/Source/MutableRuntime/Private/MuR/CodeRunner.h
@@ -347,7 +347,10 @@ namespace  mu
 				}
 			}
 
-			m_pSystem->m_memory->IncreaseHitCount(at);
+			if (dep.Type == FScheduledOp::EType::Full)
+			{
+				m_pSystem->m_memory->IncreaseHitCount(at);				
+			}
 		}
 
 		/** Try to create a concurrent task for the given op. Return null if not possible. */
