[mutable] Fixes for typos and other spelling anomalies.
- Found when working on the tagged jira task.

diff --git a/Engine/Plugins/Experimental/Mutable/Source/MutableRuntime/Private/MuR/ModelPrivate.h b/Engine/Plugins/Experimental/Mutable/Source/MutableRuntime/Private/MuR/ModelPrivate.h
--- a/Engine/Plugins/Experimental/Mutable/Source/MutableRuntime/Private/MuR/ModelPrivate.h
+++ b/Engine/Plugins/Experimental/Mutable/Source/MutableRuntime/Private/MuR/ModelPrivate.h
@@ -216,7 +216,7 @@ namespace mu
         //! The value stored here is the default value.
 		TArray<FParameterDesc> m_parameters;
 
-        //! Ranges for interation of the model operations.
+        //! Ranges for iteration of the model operations.
 		TArray<FRangeDesc> m_ranges;
 
         //! List of parameter lists. These are used in several places, like storing the
diff --git a/Engine/Plugins/Experimental/Mutable/Source/MutableRuntime/Public/MuR/Parameters.h b/Engine/Plugins/Experimental/Mutable/Source/MutableRuntime/Public/MuR/Parameters.h
--- a/Engine/Plugins/Experimental/Mutable/Source/MutableRuntime/Public/MuR/Parameters.h
+++ b/Engine/Plugins/Experimental/Mutable/Source/MutableRuntime/Public/MuR/Parameters.h
@@ -183,7 +183,7 @@ namespace mu
         static ParametersPtr StaticUnserialise( InputArchive& arch );
 
         //! Portable serialisation, which tries to keep values when:
-        //! - new parametes have been added
+        //! - new parameters have been added
         //! - parameters have been removed
         //! - parameter types have changed
         //! - parameter names have changed, and the application is using UIDs
@@ -240,7 +240,7 @@ namespace mu
 		//! Return the value of a boolean parameter.
 		//! \pre The parameter specified by index is a T_BOOL.
 		//! \param index Index of the parameter from 0 to GetCount()-1
-        //! \param pos Only for multidimensional parametres: relevant position to get in the ranges
+        //! \param pos Only for multidimensional parameters: relevant position to get in the ranges
         bool GetBoolValue( int index,
                            const Ptr<const RangeIndex>& pos=nullptr ) const;
 
@@ -258,25 +258,25 @@ namespace mu
 		int GetIntPossibleValueCount( int paramIndex ) const;
 
 		//! Get the value of one of the possible values for this integer.
-		//! The paramIndex is in the reange of 0 to GetIntPossibleValueCount()-1
+		//! The paramIndex is in the range of 0 to GetIntPossibleValueCount()-1
 		int GetIntPossibleValue( int paramIndex, int valueIndex ) const;
 
         //! Get the name of one of the possible values for this integer.
-        //! The paramIndex is in the reange of 0 to GetIntPossibleValueCount()-1
+        //! The paramIndex is in the range of 0 to GetIntPossibleValueCount()-1
         const char* GetIntPossibleValueName( int paramIndex, int valueIndex ) const;
 
         //! Get the index of the value of one of the possible values for this integer.
-        //! The paramIndex is in the reange of 0 to GetIntPossibleValueCount()-1
+        //! The paramIndex is in the range of 0 to GetIntPossibleValueCount()-1
 		int GetIntValueIndex(int paramIndex, const char* valueName) const;
 
 		//! Get the index of the value of one of the possible values for this integer.
-		//! The paramIndex is in the reange of 0 to GetIntPossibleValueCount()-1
+		//! The paramIndex is in the range of 0 to GetIntPossibleValueCount()-1
 		int GetIntValueIndex(int paramIndex, int32 Value) const;
 
 		//! Return the value of a integer parameter.
 		//! \pre The parameter specified by index is a T_INT.
 		//! \param index Index of the parameter from 0 to GetCount()-1
-        //! \param pos Only for multidimensional parametres: relevant position to get in the ranges
+        //! \param pos Only for multidimensional parameters: relevant position to get in the ranges
         int GetIntValue( int index,
                          const Ptr<const RangeIndex>& pos=nullptr ) const;
 
@@ -284,21 +284,21 @@ namespace mu
 		//! \param index Index of the parameter from 0 to GetCount()-1
 		//! \param value new value of the parameter. It must be in the possible values for this
 		//! parameter (see GetIntPossibleValue), or the method will leave the value unchanged.
-        //! \param pos Only for multidimensional parametres: relevant position to set in the ranges
+        //! \param pos Only for multidimensional parameters: relevant position to set in the ranges
         void SetIntValue( int index, int value,
                           const Ptr<const RangeIndex>& pos=nullptr );
 
 		//! Return the value of a float parameter.
 		//! \pre The parameter specified by index is a T_FLOAT.
 		//! \param index Index of the parameter from 0 to GetCount()-1
-        //! \param pos Only for multidimensional parametres: relevant position to get in the ranges
+        //! \param pos Only for multidimensional parameters: relevant position to get in the ranges
         float GetFloatValue( int index,
                              const Ptr<const RangeIndex>& pos=nullptr ) const;
 
 		//! If the parameter is of the float type, set its value.
 		//! \param index Index of the parameter from 0 to GetCount()-1
 		//! \param value new value of the parameter
-        //! \param pos Only for multidimensional parametres: relevant position to set in the ranges
+        //! \param pos Only for multidimensional parameters: relevant position to set in the ranges
         void SetFloatValue( int index, float value,
                             const Ptr<const RangeIndex>& pos=nullptr );
 
@@ -306,14 +306,14 @@ namespace mu
 		//! \pre The parameter specified by index is a T_FLOAT.
         //! \param index Index of the parameter from 0 to GetCount()-1
         //! \param pR,pG,pB Pointers to values where every resulting colour channel will be stored
-        //! \param pos Only for multidimensional parametres: relevant position to get in the ranges
+        //! \param pos Only for multidimensional parameters: relevant position to get in the ranges
         void GetColourValue( int index, float* pR, float* pG, float* pB,
                              const Ptr<const RangeIndex>& pos=nullptr ) const;
 
 		//! If the parameter is of the colour type, set its value.
 		//! \param index Index of the parameter from 0 to GetCount()-1
         //! \param r,g,b new value of the parameter
-        //! \param pos Only for multidimensional parametres: relevant position to set in the ranges
+        //! \param pos Only for multidimensional parameters: relevant position to set in the ranges
         void SetColourValue( int index, float r, float g, float b,
                              const Ptr<const RangeIndex>& pos=nullptr );
 
@@ -321,13 +321,13 @@ namespace mu
 		//! a linear transform in column-major.
 		//! \pre The parameter specified by index is a T_PROJECTOR.
         //! \param ParameterIndex Index of the parameter from 0 to GetCount()-1
-        //! \param OutPosX Pointer to where the object-space position coordinates of the projector will be stored.
+        //! \param OutPos Pointer to where the object-space position coordinates of the projector will be stored.
         //! \param OutDir Pointer to where the object-space direction vector of the projector will be stored.
         //! \param OutUp Pointer to where the object-space vertically up direction vector
         //!         of the projector will be stored. This controls the "roll" angle of the
         //!         projector.
         //! \param OutScale Pointer to the projector-space scaling of the projector.
-        //! \param RangePosition Only for multidimensional parametres: relevant position to get in the ranges
+        //! \param RangePosition Only for multidimensional parameters: relevant position to get in the ranges
         void GetProjectorValue( int ParameterIndex,
                                 PROJECTOR_TYPE* OutProjectionType,
 								FVector3f* OutPos,
@@ -339,13 +339,13 @@ namespace mu
 
 		//! If the parameter is of the projector type, set its value.
 		//! \param ParameterIndex Index of the parameter from 0 to GetCount()-1
-        //! \param posX,posY,posZ Object-space position coordinates of the projector.
-        //! \param dirX,dirY,dirZ Object-space direction vector of the projector.
-        //! \param upX,upY,upZ Object-space vertically up direction vector of the projector.
-        //! \param scaleX, scaleY Projector-space scaling of the projector.
+        //! \param pos Object-space position coordinates of the projector.
+        //! \param dir Object-space direction vector of the projector.
+        //! \param up Object-space vertically up direction vector of the projector.
+        //! \param scale Projector-space scaling of the projector.
         //! \param projectionAngle [only for Cylindrical projectors], the angle in radians of the
         //! projection area on the cylinder surface.
-        //! \param RangePosition Only for multidimensional parametres: relevant position to set in the ranges
+        //! \param RangePosition Only for multidimensional parameters: relevant position to set in the ranges
         void SetProjectorValue( int ParameterIndex,
 								const FVector3f& pos,
 								const FVector3f& dir,
@@ -357,26 +357,26 @@ namespace mu
         //! Return the value of an image parameter.
         //! \pre The parameter specified by index is a T_IMAGE.
         //! \param index Index of the parameter from 0 to GetCount()-1
-		//! \param pos Only for multidimensional parametres: relevant position to set in the ranges
+		//! \param pos Only for multidimensional parameters: relevant position to set in the ranges
 	    //! \return The externalId specified when setting the image value (\see SetImageValue)
         EXTERNAL_IMAGE_ID GetImageValue( int index, const Ptr<const RangeIndex>& pos = nullptr) const;
 
         //! If the parameter is of the image type, set its value.
         //! \param index Index of the parameter from 0 to GetCount()-1
         //! \param externalId Application-specific id used to identify this image during replication.
-		//! \param pos Only for multidimensional parametres: relevant position to set in the ranges
+		//! \param pos Only for multidimensional parameters: relevant position to set in the ranges
 		void SetImageValue( int index, EXTERNAL_IMAGE_ID externalId, const Ptr<const RangeIndex>& pos = nullptr);
 
         //! Return the value of a float parameter.
         //! \pre The parameter specified by index is a T_FLOAT.
         //! \param index Index of the parameter from 0 to GetCount()-1
-        //! \param pos Only for multidimensional parametres: relevant position to get in the ranges
+        //! \param pos Only for multidimensional parameters: relevant position to get in the ranges
         const char* GetStringValue( int index, const Ptr<const RangeIndex>& pos = nullptr ) const;
 
         //! If the parameter is of the float type, set its value.
         //! \param index Index of the parameter from 0 to GetCount()-1
         //! \param value new value of the parameter
-        //! \param pos Only for multidimensional parametres: relevant position to set in the ranges
+        //! \param pos Only for multidimensional parameters: relevant position to set in the ranges
         void SetStringValue( int index, const char* value, const Ptr<const RangeIndex>& pos = nullptr );
 
         //! Get the number of additional description images defined in the parameter.
@@ -385,7 +385,7 @@ namespace mu
 		int GetAdditionalImageCount( int index ) const;
 
         //! Utility method to compare the values of a specific parameter with the values of another
-        //! Parameters object. It returns false type or values are different.
+        //! Parameters object. It returns false if type or values are different.
         bool HasSameValue( int thisParamIndex,
                            const ParametersPtrConst& other,
                            int otherParamIndex ) const;
