[mutable] Ensure live-update cache is used.
Fix bug with mutable code container IsValid check.
Added insights counters for some CodeRunner metrics.
Removed older core stats and profile metrics.
#rnx

diff --git a/Engine/Plugins/Experimental/Mutable/Source/CustomizableObject/Classes/MuCO/CustomizableObjectSystem.h b/Engine/Plugins/Experimental/Mutable/Source/CustomizableObject/Classes/MuCO/CustomizableObjectSystem.h
--- a/Engine/Plugins/Experimental/Mutable/Source/CustomizableObject/Classes/MuCO/CustomizableObjectSystem.h
+++ b/Engine/Plugins/Experimental/Mutable/Source/CustomizableObject/Classes/MuCO/CustomizableObjectSystem.h
@@ -65,11 +65,6 @@ DECLARE_DWORD_ACCUMULATOR_STAT_EXTERN(TEXT("Avrg Instance Build Time"), STAT_Mut
 DECLARE_DWORD_ACCUMULATOR_STAT_EXTERN(TEXT("Streaming Ops"), STAT_MutableStreamingOps, STATGROUP_Mutable, );
 DECLARE_DWORD_ACCUMULATOR_STAT_EXTERN(TEXT("Streaming Cache"), STAT_MutableStreamingCache, STATGROUP_Mutable, );
 
-// These stats are provided by the mutable runtime
-DECLARE_DWORD_ACCUMULATOR_STAT_EXTERN(TEXT("Profile-LiveInstanceCount"), STAT_MutableProfile_LiveInstanceCount, STATGROUP_Mutable, );
-DECLARE_DWORD_ACCUMULATOR_STAT_EXTERN(TEXT("Profile-StreamingCacheBytes"), STAT_MutableProfile_StreamingCacheBytes, STATGROUP_Mutable, );
-DECLARE_DWORD_ACCUMULATOR_STAT_EXTERN(TEXT("Profile-InstanceUpdateCount"), STAT_MutableProfile_InstanceUpdateCount, STATGROUP_Mutable, );
-
 
 #if WITH_EDITOR
 
diff --git a/Engine/Plugins/Experimental/Mutable/Source/CustomizableObject/Private/MuCO/CustomizableObjectSystem.cpp b/Engine/Plugins/Experimental/Mutable/Source/CustomizableObject/Private/MuCO/CustomizableObjectSystem.cpp
--- a/Engine/Plugins/Experimental/Mutable/Source/CustomizableObject/Private/MuCO/CustomizableObjectSystem.cpp
+++ b/Engine/Plugins/Experimental/Mutable/Source/CustomizableObject/Private/MuCO/CustomizableObjectSystem.cpp
@@ -71,11 +71,6 @@ DEFINE_STAT(STAT_MutableInstanceBuildTimeAvrg);
 DEFINE_STAT(STAT_MutableStreamingOps);
 DEFINE_STAT(STAT_MutableStreamingCache);
 
-// These stats are provided by the mutable runtime
-DEFINE_STAT(STAT_MutableProfile_LiveInstanceCount);
-DEFINE_STAT(STAT_MutableProfile_StreamingCacheBytes);
-DEFINE_STAT(STAT_MutableProfile_InstanceUpdateCount);
-
 DECLARE_CYCLE_STAT(TEXT("MutablePendingRelease Time"), STAT_MutablePendingRelease, STATGROUP_Game);
 DECLARE_CYCLE_STAT(TEXT("MutableTask"), STAT_MutableTask, STATGROUP_Game);
 
@@ -1163,11 +1158,13 @@ namespace impl
 
 			TArray<bool> Relevant;
 			Relevant.SetNumZeroed(NumParameters);
-			MutableSystem->GetParameterRelevancy(MutableModel, MutableParameters, Relevant.GetData());
+			// This prevents live update cache reusal
+			//MutableSystem->GetParameterRelevancy(MutableModel, MutableParameters, Relevant.GetData());
 
 			for (int32 ParamIndex = 0; ParamIndex < NumParameters; ++ParamIndex)
 			{
-				if (Relevant[ParamIndex])
+				// This prevents live update cache reusal
+				//if (Relevant[ParamIndex])
 				{
 					OperationData->RelevantParametersInProgress.Add(ParamIndex);
 				}
@@ -1763,20 +1760,6 @@ namespace impl
 			SET_DWORD_STAT(STAT_MutableInstanceBuildTime, deltaMs);
 			SET_DWORD_STAT(STAT_MutableInstanceBuildTimeAvrg, CustomizableObjectSystemPrivateData->TotalBuildMs / CustomizableObjectSystemPrivateData->TotalBuiltInstances);
 			SET_DWORD_STAT(STAT_MutableStreamingCache, streamingCache);
-
-			mu::System* MutableSystem = CustomizableObjectSystemPrivateData->MutableSystem.get();
-			if (MutableSystem)
-			{
-				SET_DWORD_STAT(STAT_MutableProfile_LiveInstanceCount, MutableSystem->GetProfileMetric(mu::System::ProfileMetric::LiveInstanceCount));
-				SET_DWORD_STAT(STAT_MutableProfile_StreamingCacheBytes, MutableSystem->GetProfileMetric(mu::System::ProfileMetric::StreamingCacheBytes));
-				SET_DWORD_STAT(STAT_MutableProfile_InstanceUpdateCount, MutableSystem->GetProfileMetric(mu::System::ProfileMetric::InstanceUpdateCount));
-			}
-			else
-			{
-				SET_DWORD_STAT(STAT_MutableProfile_LiveInstanceCount, 0);
-				SET_DWORD_STAT(STAT_MutableProfile_StreamingCacheBytes, 0);
-				SET_DWORD_STAT(STAT_MutableProfile_InstanceUpdateCount, 0);
-			}
 		}
 
 		// End Update
diff --git a/Engine/Plugins/Experimental/Mutable/Source/CustomizableObjectEditor/Private/MuCOE/SMutableCodeViewer.cpp b/Engine/Plugins/Experimental/Mutable/Source/CustomizableObjectEditor/Private/MuCOE/SMutableCodeViewer.cpp
--- a/Engine/Plugins/Experimental/Mutable/Source/CustomizableObjectEditor/Private/MuCOE/SMutableCodeViewer.cpp
+++ b/Engine/Plugins/Experimental/Mutable/Source/CustomizableObjectEditor/Private/MuCOE/SMutableCodeViewer.cpp
@@ -123,7 +123,13 @@ public:
 			// See if the operation type accepts additional information in the label
 			switch ( Type )
 			{
+			case mu::OP_TYPE::BO_PARAMETER:
 			case mu::OP_TYPE::NU_PARAMETER:
+			case mu::OP_TYPE::SC_PARAMETER:
+			case mu::OP_TYPE::CO_PARAMETER:
+			case mu::OP_TYPE::PR_PARAMETER:
+			case mu::OP_TYPE::IM_PARAMETER:
+			case mu::OP_TYPE::ST_PARAMETER:
 			{
 				mu::OP::ParameterArgs Args = Program.GetOpArgs<mu::OP::ParameterArgs>(RowItem->MutableOperation);
 				OpName += TEXT(" ");
diff --git a/Engine/Plugins/Experimental/Mutable/Source/MutableRuntime/Private/MuR/CodeRunner.h b/Engine/Plugins/Experimental/Mutable/Source/MutableRuntime/Private/MuR/CodeRunner.h
--- a/Engine/Plugins/Experimental/Mutable/Source/MutableRuntime/Private/MuR/CodeRunner.h
+++ b/Engine/Plugins/Experimental/Mutable/Source/MutableRuntime/Private/MuR/CodeRunner.h
@@ -349,7 +349,7 @@ namespace  mu
 
 			if (dep.Type == FScheduledOp::EType::Full)
 			{
-				m_pSystem->m_memory->IncreaseHitCount(at);				
+				m_pSystem->m_memory->IncreaseHitCount(at);
 			}
 		}
 
@@ -358,6 +358,9 @@ namespace  mu
 
 		/** */
 		void CompleteRomLoadOp(FRomLoadOp& o);
+		
+		/** Update debug stats. */
+		void UpdateTraces();
 
     };
 
diff --git a/Engine/Plugins/Experimental/Mutable/Source/MutableRuntime/Private/MuR/CodeRunnerTasked.cpp b/Engine/Plugins/Experimental/Mutable/Source/MutableRuntime/Private/MuR/CodeRunnerTasked.cpp
--- a/Engine/Plugins/Experimental/Mutable/Source/MutableRuntime/Private/MuR/CodeRunnerTasked.cpp
+++ b/Engine/Plugins/Experimental/Mutable/Source/MutableRuntime/Private/MuR/CodeRunnerTasked.cpp
@@ -11,6 +11,7 @@
 #include "Math/UnrealMathSSE.h"
 #include "Misc/AssertionMacros.h"
 #include "MuR/CodeRunner.h"
+#include "MuR/System.h"
 #include "MuR/Image.h"
 #include "MuR/ImagePrivate.h"
 #include "MuR/Layout.h"
@@ -42,6 +43,7 @@
 #include "Templates/SharedPointer.h"
 #include "Templates/Tuple.h"
 #include "Trace/Detail/Channel.h"
+#include "ProfilingDebugging/CountersTrace.h"
 
 
 DECLARE_CYCLE_STAT(TEXT("MutableCoreTask"), STAT_MutableCoreTask, STATGROUP_Game);
@@ -101,8 +103,91 @@ namespace mu
 	}
 
 
+	//---------------------------------------------------------------------------------------------
+	TRACE_DECLARE_INT_COUNTER(MutableRuntime_OpenTask, TEXT("MutableRuntime/OpenTask"));
+	TRACE_DECLARE_INT_COUNTER(MutableRuntime_ClosedTasks, TEXT("MutableRuntime/ClosedTasks"));
+	TRACE_DECLARE_INT_COUNTER(MutableRuntime_IssuedTasks, TEXT("MutableRuntime/IssuedTasks"));
+	TRACE_DECLARE_INT_COUNTER(MutableRuntime_PCache_Keep, TEXT("MutableRuntime/PCache/Keep"));
+	TRACE_DECLARE_INT_COUNTER(MutableRuntime_PCache_Burn, TEXT("MutableRuntime/PCache/Burn"));
+	TRACE_DECLARE_INT_COUNTER(MutableRuntime_PCache_Update, TEXT("MutableRuntime/PCache/UpdateCount"));
+	TRACE_DECLARE_INT_COUNTER(MutableRuntime_PCache_UpdateSize, TEXT("MutableRuntime/PCache/UpdateSize"));
+
+	void CodeRunner::UpdateTraces()
+	{
+#if UE_MUTABLE_ENABLE_SLOW_TRACES && !UE_BUILD_SHIPPING && !UE_BUILD_TEST
+
+		//// Code runner status
+		//TRACE_COUNTER_SET(MutableRuntime_OpenTask, OpenTasks.Num());
+		//TRACE_COUNTER_SET(MutableRuntime_ClosedTasks, ClosedTasks.Num());
+		//TRACE_COUNTER_SET(MutableRuntime_IssuedTasks, IssuedTasks.Num());
+
+		//// Program cache status
+		//{
+		//	CodeContainer< TPair<int32, Ptr<const RefCounted>> >::iterator It = GetMemory().m_resources.begin();
+		//	int32 NumKeep = 0;
+		//	int32 NumBurn = 0;
+		//	for (; It.IsValid(); ++It)
+		//	{
+		//		int32 State = (*It).Key;
+		//		if (State == 1)
+		//		{
+		//			++NumKeep;
+		//		}
+		//		else if (State == 2)
+		//		{
+		//			++NumBurn;
+		//		}
+		//	}
+		//	TRACE_COUNTER_SET(MutableRuntime_PCache_Keep, NumKeep);
+		//	TRACE_COUNTER_SET(MutableRuntime_PCache_Burn, NumBurn);
+		//}
+
+		//// Program cache types
+		//{
+		//	int32 NumUpdateCacheResources = 0;
+		//	int32 UpdateCacheResourcesSize = 0;
+		//	CodeContainer<int32>::iterator It = GetMemory().m_opHitCount.begin();
+		//	for (; It.IsValid(); ++It)
+		//	{
+		//		int32 Count = *It;
+		//		if ( (Count > 0x0fffff)
+		//			&& 
+		//			GetMemory().IsValid(It.get_address()))
+		//		{
+		//			++NumUpdateCacheResources;
+
+		//			OP_TYPE OpType = m_pModel->GetPrivate()->m_program.GetOpType(It.get_address().At);
+		//			DATATYPE DataType = GetOpDataType(OpType);
+		//			switch (DataType)
+		//			{
+		//			case DATATYPE::DT_MESH:
+		//			{
+		//				const Mesh* Resource = reinterpret_cast<const Mesh*>(GetMemory().m_resources[It.get_address()].Value.get());
+		//				UpdateCacheResourcesSize += Resource->GetDataSize();
+		//				break;
+		//			}
+		//			case DATATYPE::DT_IMAGE:
+		//			{
+		//				const Image* Resource = reinterpret_cast<const Image*>(GetMemory().m_resources[It.get_address()].Value.get());
+		//				UpdateCacheResourcesSize += Resource->GetDataSize();
+		//				break;
+		//			}
+		//			default:
+		//				check(false);
+		//				break;
+		//			}
+		//		}
+		//	}
+		//	TRACE_COUNTER_SET(MutableRuntime_PCache_Update, NumUpdateCacheResources);
+		//	TRACE_COUNTER_SET(MutableRuntime_PCache_UpdateSize, UpdateCacheResourcesSize);
+		//}
+
+#endif
+	}
+
+
     //---------------------------------------------------------------------------------------------
-    void CodeRunner::Run ()
+    void CodeRunner::Run()
     {
 		bUnrecoverableError = false;
 
@@ -116,6 +201,7 @@ namespace mu
 
         while( !OpenTasks.IsEmpty() || !ClosedTasks.IsEmpty() || !IssuedTasks.IsEmpty())
         {
+			UpdateTraces();
 			// Debug: log the amount of tasks that we'd be able to run concurrently:
 			//{
 			//	int32 ClosedReady = ClosedTasks.Num();
@@ -160,6 +246,7 @@ namespace mu
 				}
 			}
 
+			UpdateTraces();
 
 			while (!OpenTasks.IsEmpty())
 			{
@@ -261,6 +348,8 @@ namespace mu
 					}
 				}
 
+				UpdateTraces();
+
 				if (bProfile)
 				{
 					++NumRunOps;
@@ -275,6 +364,8 @@ namespace mu
 				{
 					CompleteRomLoadOp(o);
 				}
+
+				UpdateTraces();
 			}
 
 			// Look for a closed task with dependencies satisfied and move them to open.
@@ -305,6 +396,8 @@ namespace mu
 				{
 					++Index;
 				}
+
+				UpdateTraces();
 			}
 
 			// Debug: Did we dead-lock?
@@ -360,6 +453,7 @@ namespace mu
 					//check(false);
 				}
 
+				UpdateTraces();
 			}
 		}
 
diff --git a/Engine/Plugins/Experimental/Mutable/Source/MutableRuntime/Private/MuR/OpImageResize.h b/Engine/Plugins/Experimental/Mutable/Source/MutableRuntime/Private/MuR/OpImageResize.h
--- a/Engine/Plugins/Experimental/Mutable/Source/MutableRuntime/Private/MuR/OpImageResize.h
+++ b/Engine/Plugins/Experimental/Mutable/Source/MutableRuntime/Private/MuR/OpImageResize.h
@@ -830,7 +830,12 @@ namespace mu
     inline ImagePtr ImageResizeLinear( int imageCompressionQuality, const Image* pBasePtr,
                                        FImageSize destSize )
     {
-		MUTABLE_CPUPROFILER_SCOPE(ImageResizeLinear)
+		MUTABLE_CPUPROFILER_SCOPE(ImageResizeLinear);
+		
+		if(pBasePtr->GetSize()==destSize)
+		{
+			return pBasePtr->Clone();
+		}
 
 		check(!(pBasePtr->m_flags & Image::IF_CANNOT_BE_SCALED));
 
diff --git a/Engine/Plugins/Experimental/Mutable/Source/MutableRuntime/Private/MuR/System.cpp b/Engine/Plugins/Experimental/Mutable/Source/MutableRuntime/Private/MuR/System.cpp
--- a/Engine/Plugins/Experimental/Mutable/Source/MutableRuntime/Private/MuR/System.cpp
+++ b/Engine/Plugins/Experimental/Mutable/Source/MutableRuntime/Private/MuR/System.cpp
@@ -31,6 +31,7 @@
 #include "Templates/SharedPointer.h"
 #include "Templates/Tuple.h"
 #include "Trace/Detail/Channel.h"
+#include "ProfilingDebugging/CountersTrace.h"
 
 
 namespace mu
@@ -107,11 +108,6 @@ namespace mu
         m_pStreamInterface = nullptr;
         m_pImageParameterGenerator = nullptr;
         m_maxMemory = 0;
-
-        for(uint8 p=0; p< uint8(ProfileMetric::_Count); ++p)
-        {
-            m_profileMetrics[p]=0;
-        }
 	
 		m_modelCache.m_romBudget = pSettings->GetPrivate()->m_streamingCacheBytes;
 	}
@@ -181,17 +177,8 @@ namespace mu
         // TODO
     }
 
-
-    //---------------------------------------------------------------------------------------------
-    uint64 System::GetProfileMetric( ProfileMetric m ) const
-    {
-        size_t metric = size_t(m);
-
-        if (metric>=size_t(ProfileMetric::_Count)) return 0;
-
-        return m_pD->m_profileMetrics[metric];
-    }
-
+	TRACE_DECLARE_INT_COUNTER(MutableRuntime_LiveInstances, TEXT("MutableRuntime/LiveInstances"));
+	TRACE_DECLARE_INT_COUNTER(MutableRuntime_Updates, TEXT("MutableRuntime/Updates"));
 
     //---------------------------------------------------------------------------------------------
     Instance::ID System::NewInstance( const TSharedPtr<const Model>& pModel )
@@ -207,7 +194,7 @@ namespace mu
 		instanceData.m_memory = MakeShared<FProgramCache>();
 		m_pD->m_liveInstances.Add(instanceData);
 
-		m_pD->m_profileMetrics[size_t(System::ProfileMetric::LiveInstanceCount)]++;
+		TRACE_COUNTER_SET(MutableRuntime_LiveInstances, m_pD->m_liveInstances.Num());
 
 		return instanceData.m_instanceID;
 	}
@@ -221,6 +208,7 @@ namespace mu
     {
 		LLM_SCOPE_BYNAME(TEXT("MutableRuntime"));
 		MUTABLE_CPUPROFILER_SCOPE(SystemBeginUpdate);
+		TRACE_COUNTER_INCREMENT(MutableRuntime_Updates);
 
 		if (!pParams)
 		{
@@ -275,14 +263,15 @@ namespace mu
 
 		InstancePtrConst pResult = pLiveInstance->m_memory->GetInstance(FCacheAddress(rootAt, 0, 0));
 
+		// Debug check to see if we managed the op-hit-counts correctly
+		pLiveInstance->m_memory->CheckHitCountsCleared();
+
 		pLiveInstance->m_pInstance = pResult;
 		if (pResult)
 		{
 			pResult->GetPrivate()->m_id = pLiveInstance->m_instanceID;
 		}
 
-		m_pD->m_profileMetrics[size_t(System::ProfileMetric::InstanceUpdateCount)]++;
-
 		m_pD->m_memory = nullptr;
 
 		return pResult.get();
@@ -311,7 +300,7 @@ namespace mu
 					pLiveInstance->m_pOldParameters.get(),
 					res.m_rootAddress, MipsToSkip, InImageLOD);
 
-				// We always need ot return something valid.
+				// We always need to return something valid.
 				if (!pResult)
 				{
 					pResult = new mu::Image(16, 16, 1, EImageFormat::IF_RGBA_UBYTE);
@@ -433,14 +422,22 @@ namespace mu
 
 		// Reduce the cache until it fits the limit.
 		uint64 totalMemory = m_pD->m_modelCache.EnsureCacheBelowBudget(0, [](const Model*, int) {return false;});
-		m_pD->m_profileMetrics[size_t(System::ProfileMetric::StreamingCacheBytes)] = totalMemory;
-
 
 		Private::FLiveInstance* pLiveInstance = m_pD->FindLiveInstance(instanceID);
 		if (pLiveInstance)
 		{
 			pLiveInstance->m_pInstance = nullptr;
-			pLiveInstance->m_memory->ClearCacheLayer1();
+			
+			// Debug check to see if we managed the op-hit-counts correctly
+			pLiveInstance->m_memory->CheckHitCountsCleared();
+
+			// We don't want to clear the cache layer 1 because it contains data that can be useful for a 
+			// future update (same states, just runtime parameters changed).
+			//pLiveInstance->m_memory->ClearCacheLayer1();
+
+			// We need to clear the layer 0 cache, because it contains data that is only valid for the current 
+			// parameter values (unless it is data marked as state cache)
+			pLiveInstance->m_memory->ClearCacheLayer0();
 		}
 	}
 
@@ -456,7 +453,9 @@ namespace mu
 			{
 				return (Instance.m_instanceID == instanceID);
 			});
-		m_pD->m_profileMetrics[int32(System::ProfileMetric::LiveInstanceCount)]-=Removed;
+
+		TRACE_COUNTER_SET(MutableRuntime_LiveInstances, m_pD->m_liveInstances.Num());
+
 	}
 
 
@@ -626,6 +625,9 @@ namespace mu
 		m_memory = MakeShared<FProgramCache>();
 		m_memory->Init(pModel->GetPrivate()->m_program.m_opAddress.Num());
 
+		// Remove previously results cached from previous builds.
+		m_memory->ClearCacheLayer0();
+
 		PrepareCache(pModel.Get(), -1);
 	}
 
@@ -744,13 +746,16 @@ namespace mu
 		mu::OP_TYPE opType = pModel->GetPrivate()->m_program.GetOpType(at);
 		if (GetOpDataType(opType) == DT_IMAGE)
 		{
-			m_memory->ClearCacheLayer0();
 			RunCode(pModel, pParams, at, System::AllLODs, uint8(MipsToSkip), InImageLOD);
 			if (bUnrecoverableError)
 			{
 				return nullptr;
 			}
 			ImagePtrConst Result = m_memory->GetImage(FCacheAddress(at, 0, MipsToSkip));
+
+			// Debug check to see if we managed the op-hit-counts correctly
+			m_memory->CheckHitCountsCleared();
+
 			return Result;
 		}
 
@@ -766,13 +771,16 @@ namespace mu
 		mu::OP_TYPE opType = pModel->GetPrivate()->m_program.GetOpType(at);
 		if (GetOpDataType(opType) == DT_MESH)
 		{
-			m_memory->ClearCacheLayer0();
 			RunCode(pModel, pParams, at);
 			if (bUnrecoverableError)
 			{
 				return nullptr;
 			}
 			MeshPtrConst pResult = m_memory->GetMesh(FCacheAddress(at, 0, 0));
+
+			// Debug check to see if we managed the op-hit-counts correctly
+			m_memory->CheckHitCountsCleared();
+
 			return pResult;
 		}
 
@@ -883,6 +891,8 @@ namespace mu
         return m_cachePerModel.Last();
     }
 
+	/** */
+	TRACE_DECLARE_INT_COUNTER(MutableRuntime_StreamingBytes, TEXT("MutableRuntime/StreamingBytes"));
 
     //---------------------------------------------------------------------------------------------
     uint64 FModelCache::EnsureCacheBelowBudget( uint64 additionalMemory,
@@ -955,7 +965,10 @@ namespace mu
                 {
                     // None found
                     finished = true;
-                    UE_LOG(LogMutableCore,Log, TEXT("EnsureCacheBelowBudget failed to free all the necessary memory."));
+
+					// If we reached this it means we need to use more memory for streaming data that was given to mutable.
+					// Try to continue anyway.
+                    //UE_LOG(LogMutableCore,Log, TEXT("EnsureCacheBelowBudget failed to free all the necessary memory."));
                 }
                 else
                 {
@@ -968,6 +981,8 @@ namespace mu
             }
         }
 
+		TRACE_COUNTER_SET(MutableRuntime_StreamingBytes, totalMemory);
+
         return totalMemory;
     }
 
diff --git a/Engine/Plugins/Experimental/Mutable/Source/MutableRuntime/Private/MuR/SystemPrivate.h b/Engine/Plugins/Experimental/Mutable/Source/MutableRuntime/Private/MuR/SystemPrivate.h
--- a/Engine/Plugins/Experimental/Mutable/Source/MutableRuntime/Private/MuR/SystemPrivate.h
+++ b/Engine/Plugins/Experimental/Mutable/Source/MutableRuntime/Private/MuR/SystemPrivate.h
@@ -11,6 +11,7 @@
 #include "MuR/MeshPrivate.h"
 #include "MuR/InstancePrivate.h"
 #include "MuR/ParametersPrivate.h"
+#include "MuR/MutableTrace.h"
 
 
 namespace mu
@@ -385,7 +386,7 @@ namespace mu
 
 			inline bool IsValid() const
 			{
-				return bool(it0) && bool(it1);
+				return bool(it0) || bool(it1);
 			}
 		};
 
@@ -407,27 +408,25 @@ namespace mu
 	/** Interface for storage of data while Mutable code is being executed. */
 	class FProgramCache
 	{
-	private:
+	public:
 
 		TArray< ExecutionIndex, TInlineAllocator<4> > m_usedRangeIndices;
 
-		// first value of the pair:
-		// 0 : value not valid (not set).
-		// 1 : valid, not worth freeing for memory
-		// 2 : valid, worth freeing
+		/** Cached resources while the program is executing. 
+		* first value of the pair:
+		* 0 : value not valid (not set).
+		* 1 : valid, not worth freeing for memory
+		* 2 : valid, worth freeing
+		*/
 		CodeContainer< TPair<int32, Ptr<const RefCounted>> > m_resources;
 
-		// TODO
-	public:
-
 		//! Addressed with OP::ADDRESS. It is true if value for an image desc is valid.
 		TArray<bool> m_descCache;
 
-		//! The number of times an operation will be run for the current build operation.
-		//! \todo: this could be optimised by merging in other CodeContainers here.
-		CodeContainer<int> m_opHitCount;
-
-	public:
+		/** The number operation stages waiting for the output of a specific operation.
+		 * \todo: this could be optimised by merging in other CodeContainers here.
+		 */
+		CodeContainer<int32> m_opHitCount;
 
 
 		inline const ExecutionIndex& GetRageIndex(uint32_t i)
@@ -469,7 +468,8 @@ namespace mu
 
 		void Init(size_t size)
 		{
-			m_resources.clear();
+			// This prevents live update cache reusal
+			//m_resources.clear();
 			m_resources.resize(size);
 			m_opHitCount.resize(size);
 		}
@@ -512,6 +512,43 @@ namespace mu
 		}
 
 
+		/** Remove all intermediate data (big and small) from the memory except for the one that has been explicitely
+		* marked as state cache.
+		*/
+		void CheckHitCountsCleared()
+		{
+#if !(UE_BUILD_SHIPPING || UE_BUILD_TEST)
+			//MUTABLE_CPUPROFILER_SCOPE(CheckHitCountsCleared);
+
+			//int32 IncorrectCount = 0;
+			//CodeContainer<int>::iterator it = m_opHitCount.begin();
+			//for (; it.IsValid(); ++it)
+			//{
+			//	int32 Count = *it;
+			//	if (Count>0 && Count < 3000000)
+			//	{
+			//		// We don't manage the hitcounts of small types that don't use much memory
+			//		if (m_resources[it.get_address()].Key==2)
+			//		{
+			//			// Op hitcount should have reached 0, otherwise, it means we requested an operation but never read 
+			//			// the result.
+			//			++IncorrectCount;
+			//		}
+			//	}
+			//}
+
+			//if (IncorrectCount > 0)
+			//{
+			//	UE_LOG(LogMutableCore, Log, TEXT("The op-hit-count didn't hit 0 for %5d operations. This may mean that too much memory is cached."), IncorrectCount);
+			//	//check(false);
+			//}
+#endif
+		}
+
+
+		/** Remove all intermediate data (big and small) from the memory except for the one that has been explicitely
+		* marked as state cache.
+		*/
 		void ClearCacheLayer0()
 		{
 			MUTABLE_CPUPROFILER_SCOPE(ClearLayer0);
@@ -522,22 +559,29 @@ namespace mu
 				int32 Count = *it;
 				if (Count < 3000000)
 				{
-					SetUnused(it.get_address());
+					// SetUnused only clears if the data size is relevant (mesh or image) but we need to 
+					// clear it in all cases here, because it may have become invalid because of parameter changes.
+					// SetUnused(it.get_address());
+					m_resources[it.get_address()].Value = nullptr;
+					m_resources[it.get_address()].Key = 0;
 					*it = 0;
 				}
 			}
 		}
 
 
+		/** Remove all intermediate data (big and small) from the memory including the one that has been explicitely
+		* marked as state cache.
+		*/
 		void ClearCacheLayer1()
 		{
 			MUTABLE_CPUPROFILER_SCOPE(ClearLayer1);
 
-			CodeContainer<int>::iterator it = m_opHitCount.begin();
+			CodeContainer< TPair<int32, Ptr<const RefCounted>> >::iterator it = m_resources.begin();
 			for (; it.IsValid(); ++it)
 			{
-				m_resources[it.get_address()].Value = nullptr;
-				m_resources[it.get_address()].Key = 0;
+				(*it).Value = nullptr;
+				(*it).Key = 0;
 			}
 
 			m_descCache.SetNum(0);
@@ -775,11 +819,18 @@ namespace mu
 
 		inline void IncreaseHitCount(FCacheAddress at)
 		{
-			m_opHitCount[at] += 1;
+			// Don't count hits for instruction 0, which is always null. It is usually already
+			// check that At is not 0, and then it is not requested, generating a stray non-zero count
+			// at its position.
+			if (at.At)
+			{
+				m_opHitCount[at] += 1;
+			}
 		}
 
 		inline void SetForceCached(OP::ADDRESS at)
 		{
+			// \TODO: Review the ,0,0
 			m_opHitCount[{at, 0, 0}] = 0xffffff;
 		}
 
@@ -1001,10 +1052,6 @@ namespace mu
 			const Ptr<const Parameters>& pParams,
 			uint64& OutUpdatedParameters );
 
-        // Profile metrics
-        std::atomic<uint64> m_profileMetrics[ size_t(ProfileMetric::_Count) ];
-
-
 	public:
 
 		void RunCode(const TSharedPtr<const Model>& pModel,
diff --git a/Engine/Plugins/Experimental/Mutable/Source/MutableRuntime/Public/MuR/System.h b/Engine/Plugins/Experimental/Mutable/Source/MutableRuntime/Public/MuR/System.h
--- a/Engine/Plugins/Experimental/Mutable/Source/MutableRuntime/Public/MuR/System.h
+++ b/Engine/Plugins/Experimental/Mutable/Source/MutableRuntime/Public/MuR/System.h
@@ -13,10 +13,14 @@
 #include "MuR/Types.h"
 #include "Templates/Tuple.h"
 
-// This define will use the newer task graph interface to manage mutable concurrency. 
-// This is currently broken in Switch and maybe other consoles, for some unknown reason.
+/** This define will use the newer task graph interface to manage mutable concurrency. It has not been fully tested. */
 //#define MUTABLE_USE_NEW_TASKGRAPH
 
+/** If set to 1, this enables some expensive Unreal Insights traces, but can lead to 5x slower mutable operation. 
+* Other cheaper traces are enabled at all times.
+*/
+#define UE_MUTABLE_ENABLE_SLOW_TRACES	0
+
 #ifdef MUTABLE_USE_NEW_TASKGRAPH
 #include "Tasks/Task.h"
 #else
@@ -46,6 +50,7 @@ namespace mu
         //! A memory allocation failed
         OutOfMemory,
 
+
         //! A file was missing or the data was corrupted.
         StreamingError,
 
@@ -195,34 +200,6 @@ namespace mu
         //! level.
         void ClearCaches();
 
-        //-----------------------------------------------------------------------------------------
-        // Debugging and profiling
-        //-----------------------------------------------------------------------------------------
-
-        //! Types of metrics that can be queried
-        enum class ProfileMetric : uint8
-        {
-            //! No metric
-            None,
-
-            //! Number of model instances currently being created or updated.
-            LiveInstanceCount,
-
-            //! Number of bytes currently loaded for streaming data.
-            StreamingCacheBytes,
-
-            //! Number of update operations performed since the system was created.
-            InstanceUpdateCount,
-
-            //! Utility value that doesn't really represent a metric.
-            _Count
-        };
-
-        //! Get the values of several profile metrics. This can be called any time in a valid
-        //! system object from any thread.
-        //! These values can change any time.
-        uint64 GetProfileMetric( ProfileMetric ) const;
-
 		//-----------------------------------------------------------------------------------------
 		// Interface pattern
 		//-----------------------------------------------------------------------------------------
