[mutable] Added option to avoid the log of unassigned vertices warnings for some LODs.

#rnx

diff --git a/Engine/Plugins/Experimental/Mutable/Source/CustomizableObjectEditor/Classes/MuCOE/CustomizableObjectLayout.h b/Engine/Plugins/Experimental/Mutable/Source/CustomizableObjectEditor/Classes/MuCOE/CustomizableObjectLayout.h
--- a/Engine/Plugins/Experimental/Mutable/Source/CustomizableObjectEditor/Classes/MuCOE/CustomizableObjectLayout.h
+++ b/Engine/Plugins/Experimental/Mutable/Source/CustomizableObjectEditor/Classes/MuCOE/CustomizableObjectLayout.h
@@ -49,11 +49,15 @@ public:
 
 	UCustomizableObjectLayout();
 
+	// Sets the layout parameters
 	void SetLayout(UObject* InMesh, int32 LODIndex, int32 MatIndex, int32 UVIndex);
 	void SetPackingStrategy(ECustomizableObjectTextureLayoutPackingStrategy Strategy);
 	void SetGridSize(FIntPoint Size);
 	void SetMaxGridSize(FIntPoint Size);
 	void SetLayoutName(FString Name);
+	void SetIgnoreVertexLayoutWarnings(bool bValue);
+	void SetIgnoreWarningsLOD(int32 LODValue);
+
 
 	int32 GetLOD() const { return LOD; }
 	int32 GetMaterial() const { return Material; }
@@ -63,6 +67,8 @@ public:
 	FIntPoint GetGridSize() const { return GridSize; }
 	FIntPoint GetMaxGridSize() const { return MaxGridSize; }
 	ECustomizableObjectTextureLayoutPackingStrategy GetPackingStrategy() const { return PackingStrategy; }
+	bool GetIgnoreVertexLayoutWarnings() const { return bIgnoreUnassignedVertexWarning; };
+	int32 GetFirstLODToIgnoreWarnings() { return FirstLODToIgnore; };
 
 	void GetUVChannel(TArray<FVector2f>& UVs, int32 UVChannelIndex = 0) const;
 
@@ -103,4 +109,12 @@ private:
 	UPROPERTY()
 	FString LayoutName;
 
+	/* If true, vertex warning messages will be ignored */
+	UPROPERTY()
+	bool bIgnoreUnassignedVertexWarning = false;
+
+	/* First LOD from which unassigned vertices warning will be ignored */
+	UPROPERTY()
+	int32 FirstLODToIgnore = 0;
+
 };
diff --git a/Engine/Plugins/Experimental/Mutable/Source/CustomizableObjectEditor/Private/MuCOE/CustomizableObjectLayout.cpp b/Engine/Plugins/Experimental/Mutable/Source/CustomizableObjectEditor/Private/MuCOE/CustomizableObjectLayout.cpp
--- a/Engine/Plugins/Experimental/Mutable/Source/CustomizableObjectEditor/Private/MuCOE/CustomizableObjectLayout.cpp
+++ b/Engine/Plugins/Experimental/Mutable/Source/CustomizableObjectEditor/Private/MuCOE/CustomizableObjectLayout.cpp
@@ -167,4 +167,16 @@ int32 UCustomizableObjectLayout::FindBlock(const FGuid& InId) const
 	return -1;
 }
 
+
+void UCustomizableObjectLayout::SetIgnoreVertexLayoutWarnings(bool bValue)
+{
+	bIgnoreUnassignedVertexWarning = bValue;
+}
+
+
+void UCustomizableObjectLayout::SetIgnoreWarningsLOD(int32 LODValue)
+{
+	FirstLODToIgnore = LODValue;
+}
+
 #undef LOCTEXT_NAMESPACE
diff --git a/Engine/Plugins/Experimental/Mutable/Source/CustomizableObjectEditor/Private/MuCOE/GenerateMutableSource/GenerateMutableSourceLayout.cpp b/Engine/Plugins/Experimental/Mutable/Source/CustomizableObjectEditor/Private/MuCOE/GenerateMutableSource/GenerateMutableSourceLayout.cpp
--- a/Engine/Plugins/Experimental/Mutable/Source/CustomizableObjectEditor/Private/MuCOE/GenerateMutableSource/GenerateMutableSourceLayout.cpp
+++ b/Engine/Plugins/Experimental/Mutable/Source/CustomizableObjectEditor/Private/MuCOE/GenerateMutableSource/GenerateMutableSourceLayout.cpp
@@ -58,6 +58,7 @@ mu::NodeLayoutPtr GenerateMutableSourceLayout(const UEdGraphPin * Pin, FMutableG
 		LayoutNode->SetGridSize(TypedNodeBlocks->Layout->GetGridSize().X, TypedNodeBlocks->Layout->GetGridSize().Y);
 		LayoutNode->SetMaxGridSize(TypedNodeBlocks->Layout->GetMaxGridSize().X, TypedNodeBlocks->Layout->GetMaxGridSize().Y);
 		LayoutNode->SetBlockCount(TypedNodeBlocks->Layout->Blocks.Num() ? TypedNodeBlocks->Layout->Blocks.Num() : 1);
+		LayoutNode->SetIgnoreWarningsLOD(TypedNodeBlocks->Layout->GetIgnoreVertexLayoutWarnings() ? TypedNodeBlocks->Layout->GetFirstLODToIgnoreWarnings() : -1);
 
 		mu::EPackStrategy strategy = mu::EPackStrategy::RESIZABLE_LAYOUT;
 
diff --git a/Engine/Plugins/Experimental/Mutable/Source/CustomizableObjectEditor/Private/MuCOE/Nodes/CustomizableObjectNodeLayoutBlocksDetails.cpp b/Engine/Plugins/Experimental/Mutable/Source/CustomizableObjectEditor/Private/MuCOE/Nodes/CustomizableObjectNodeLayoutBlocksDetails.cpp
--- a/Engine/Plugins/Experimental/Mutable/Source/CustomizableObjectEditor/Private/MuCOE/Nodes/CustomizableObjectNodeLayoutBlocksDetails.cpp
+++ b/Engine/Plugins/Experimental/Mutable/Source/CustomizableObjectEditor/Private/MuCOE/Nodes/CustomizableObjectNodeLayoutBlocksDetails.cpp
@@ -8,6 +8,8 @@
 #include "MuCOE/CustomizableObjectLayout.h"
 #include "MuCOE/Nodes/CustomizableObjectNodeLayoutBlocks.h"
 #include "MuCOE/SCustomizableObjectNodeLayoutBlocksEditor.h"
+#include "PropertyCustomizationHelpers.h"
+#include "Widgets/Input/SSpinBox.h"
 #include "Widgets/Layout/SBox.h"
 #include "Widgets/Text/STextBlock.h"
 
@@ -32,12 +34,90 @@ void FCustomizableObjectNodeLayoutBlocksDetails::CustomizeDetails( IDetailLayout
 		Node = Cast<UCustomizableObjectNodeLayoutBlocks>(DetailsView->GetSelectedObjects()[0].Get());
 	}
 
-	IDetailCategoryBuilder& BlocksCategory = DetailBuilder.EditCategory( "Blocks" );
+	IDetailCategoryBuilder& CustomizableObjectCategory = DetailBuilder.EditCategory( "LayoutOptions" );
+	IDetailCategoryBuilder& BlocksCategory = DetailBuilder.EditCategory( "LayoutEditor" );
 
 	if (Node)
 	{
 		LayoutBlocksEditor = SNew(SCustomizableObjectNodeLayoutBlocksEditor);
 
+		CustomizableObjectCategory.AddCustomRow(LOCTEXT("BlocksDetails_LayoutOptions", "LayoutOptions"))
+		[
+			SNew(SVerticalBox)
+			+SVerticalBox::Slot()
+			.Padding(0.0f,5.0f,0.0f,2.5f)
+			.AutoHeight()
+			[
+				SNew(SHorizontalBox)
+				+SHorizontalBox::Slot()
+				.HAlign(HAlign_Left)
+				.VAlign(VAlign_Center)
+				.FillWidth(0.75f)
+				[
+					SNew(STextBlock)
+					.Text( LOCTEXT("BlocksDetails_LayoutOptions_IgnoreLodsCheckBox", "Ignore Unassigned Vertices Warning:"))
+					.ToolTipText(LOCTEXT("BlocksDetails_LayoutOptions_IgnoreLodsCheckBox_Tooltip",
+						"If true, warning message \"Source mesh has vertices not assigned to any layout block\" will be ignored."
+						"\n Note:"
+						"\n This warning can appear when a CO has more than one LOD using the same Layout Block node and these LODs have been generated using the automatic LOD generation."
+						"\n (At high LODs, some vertices may have been displaced from their original position which means they could have been displaced outside their layout blocks.)"
+						"\n Ignoring these warnings can cause some visual artifacts that may or may not be visually important at higher LODs."))
+					.Font(IDetailLayoutBuilder::GetDetailFont())
+				]
+				
+				+SHorizontalBox::Slot()
+				.HAlign(HAlign_Left)
+				.FillWidth(0.25f)
+				.Padding(5.0f, 0.0f, 0.0f, 0.0f)
+				[
+					SNew(SCheckBox)
+					.IsChecked(Node->Layout ? (Node->Layout->GetIgnoreVertexLayoutWarnings() ? ECheckBoxState::Checked : ECheckBoxState::Unchecked) : ECheckBoxState::Unchecked)
+					.OnCheckStateChanged(this, &FCustomizableObjectNodeLayoutBlocksDetails::OnIgnoreErrorsCheckStateChanged)
+				]
+			]
+			
+			+ SVerticalBox::Slot()
+			.Padding(0.0f, 2.5f, 0.0f, 5.0f)
+			.AutoHeight()
+			[
+				SNew(SHorizontalBox)
+				+SHorizontalBox::Slot()
+				.HAlign(HAlign_Left)
+				.VAlign(VAlign_Center)
+				.FillWidth(0.75f)
+				[
+					SAssignNew(LODSelectorTextWidget, STextBlock)
+					.Text( LOCTEXT("BlocksDetails_LayoutOptions_IgnoreLod", "First LOD to ignore:"))
+					.ToolTipText(LOCTEXT("BlocksDetails_LayoutOptions_IgnoreLod_Tooltip", "LOD from which vertex warning messages will be ignored."))
+					.Font(IDetailLayoutBuilder::GetDetailFont())
+					.IsEnabled(Node->Layout ? (Node->Layout->GetIgnoreVertexLayoutWarnings() ? true : false) : false)
+				]
+				
+				+ SHorizontalBox::Slot()
+				.HAlign(HAlign_Left)
+				.FillWidth(0.25f)
+				.Padding(5.0f,0.0f,0.0f,0.0f)
+				[
+					SAssignNew(LODSelectorWidget, SSpinBox<int32>)
+					.Value_Lambda([=]()->int32
+					{
+						if (Node->Layout)
+						{
+							return Node->Layout->GetFirstLODToIgnoreWarnings();
+						}
+
+						return 0;
+					})
+					.IsEnabled(Node->Layout ? (Node->Layout->GetIgnoreVertexLayoutWarnings() ? true : false ) : false)
+					.OnValueChanged(this, &FCustomizableObjectNodeLayoutBlocksDetails::OnLODBoxValueChanged)
+					.MinValue(0)
+					.Delta(1)
+					.AlwaysUsesDeltaSnap(true)
+					.MinDesiredWidth(40.0f)
+				]
+			]
+		];
+
 		BlocksCategory.AddCustomRow( LOCTEXT("BlocksDetails_BlockInstructions", "BlockInstructions") )
 		[
 			SNew(SBox)
@@ -118,4 +198,27 @@ void FCustomizableObjectNodeLayoutBlocksDetails::OnGridComboBoxSelectionChanged(
 	*/
 }
 
+
+void FCustomizableObjectNodeLayoutBlocksDetails::OnIgnoreErrorsCheckStateChanged(ECheckBoxState State)
+{
+	if (Node->Layout)
+	{
+		bool bStateBool = State == ECheckBoxState::Checked;
+		Node->Layout->SetIgnoreVertexLayoutWarnings(bStateBool);
+
+		LODSelectorWidget->SetEnabled(bStateBool);
+		LODSelectorTextWidget->SetEnabled(bStateBool);
+	}
+}
+
+
+void FCustomizableObjectNodeLayoutBlocksDetails::OnLODBoxValueChanged(int32 Value)
+{
+	if (Node->Layout)
+	{
+		Node->Layout->SetIgnoreWarningsLOD(Value);
+	}
+}
+
+
 #undef LOCTEXT_NAMESPACE
diff --git a/Engine/Plugins/Experimental/Mutable/Source/CustomizableObjectEditor/Private/MuCOE/Nodes/CustomizableObjectNodeLayoutBlocksDetails.h b/Engine/Plugins/Experimental/Mutable/Source/CustomizableObjectEditor/Private/MuCOE/Nodes/CustomizableObjectNodeLayoutBlocksDetails.h
--- a/Engine/Plugins/Experimental/Mutable/Source/CustomizableObjectEditor/Private/MuCOE/Nodes/CustomizableObjectNodeLayoutBlocksDetails.h
+++ b/Engine/Plugins/Experimental/Mutable/Source/CustomizableObjectEditor/Private/MuCOE/Nodes/CustomizableObjectNodeLayoutBlocksDetails.h
@@ -28,10 +28,16 @@ private:
 	// Layout block editor widget
 	TSharedPtr<class SCustomizableObjectNodeLayoutBlocksEditor> LayoutBlocksEditor;
 
+	// Widget to select at which LOD layout vertex warnings will start to be ignored
+	TSharedPtr<SWidget> LODSelectorWidget;
+	TSharedPtr<STextBlock> LODSelectorTextWidget;
+
 	TArray< TSharedPtr<FString> > GridComboOptions;
 
 	void OnGridComboBoxSelectionChanged(TSharedPtr<FString> Selection, ESelectInfo::Type SelectInfo);
 	void OnBlockChanged( int BlockIndex, FIntRect Block );
+	void OnIgnoreErrorsCheckStateChanged(ECheckBoxState State);
+	void OnLODBoxValueChanged(int32 Value);
 
 	FIntPoint GetGridSize() const;
 	TArray<FIntRect> GetBlocks() const;
diff --git a/Engine/Plugins/Experimental/Mutable/Source/MutableRuntime/Private/MuR/Layout.cpp b/Engine/Plugins/Experimental/Mutable/Source/MutableRuntime/Private/MuR/Layout.cpp
--- a/Engine/Plugins/Experimental/Mutable/Source/MutableRuntime/Private/MuR/Layout.cpp
+++ b/Engine/Plugins/Experimental/Mutable/Source/MutableRuntime/Private/MuR/Layout.cpp
@@ -42,6 +42,7 @@ namespace mu {
 		pResult->m_maxsize = m_maxsize;
 		pResult->m_blocks = m_blocks;
 		pResult->m_strategy = m_strategy;
+		pResult->FirstLODToIgnoreWarnings = FirstLODToIgnoreWarnings;
 		return pResult;
 	}
 
@@ -52,7 +53,9 @@ namespace mu {
 		return  (m_size == o.m_size) &&
 			(m_maxsize == o.m_maxsize) &&
 			(m_blocks == o.m_blocks) &&
-			(m_strategy == o.m_strategy);
+			(m_strategy == o.m_strategy) &&
+			// maybe this is not needed
+			(FirstLODToIgnoreWarnings == o.FirstLODToIgnoreWarnings);;
 	}
 
 
@@ -218,5 +221,19 @@ namespace mu {
 		}
 		return false;
 	}
+
+
+	//---------------------------------------------------------------------------------------------
+	void Layout::SetIgnoreLODWarnings(int32 LOD)
+	{
+		FirstLODToIgnoreWarnings = LOD;
+	}
+
+
+	//---------------------------------------------------------------------------------------------
+	int32 Layout::GetIgnoreLODWarnings()
+	{
+		return FirstLODToIgnoreWarnings;
+	}
 }
 
diff --git a/Engine/Plugins/Experimental/Mutable/Source/MutableRuntime/Public/MuR/Layout.h b/Engine/Plugins/Experimental/Mutable/Source/MutableRuntime/Public/MuR/Layout.h
--- a/Engine/Plugins/Experimental/Mutable/Source/MutableRuntime/Public/MuR/Layout.h
+++ b/Engine/Plugins/Experimental/Mutable/Source/MutableRuntime/Public/MuR/Layout.h
@@ -125,6 +125,12 @@ namespace mu
 		//! Set the texture layout packing strategy
 		EPackStrategy GetLayoutPackingStrategy() const;
 
+		//! Set at which LOD the unassigned vertices warnings will star to be ignored
+		void SetIgnoreLODWarnings(int32 LOD);
+
+		//! Get the LOD where the unassigned vertices warnings starts to be ignored
+		int32 GetIgnoreLODWarnings();
+
 	protected:
 
 		//! Forbidden. Manage with the Ptr<> template.
@@ -202,10 +208,13 @@ namespace mu
 		//! Packing strategy
 		EPackStrategy m_strategy = EPackStrategy::RESIZABLE_LAYOUT;
 
+		int32 FirstLODToIgnoreWarnings;
+
+
 		//!
 		void Serialise(OutputArchive& arch) const
 		{
-			uint32 ver = 3;
+			uint32 ver = 4;
 			arch << ver;
 
 			arch << m_size;
@@ -213,6 +222,7 @@ namespace mu
 
 			arch << m_maxsize;
 			arch << uint32(m_strategy);
+			arch << FirstLODToIgnoreWarnings;
 		}
 
 		//!
@@ -220,7 +230,7 @@ namespace mu
 		{
 			uint32 ver;
 			arch >> ver;
-			check(ver == 3);
+			check(ver <= 4);
 
 			arch >> m_size;
 
@@ -230,6 +240,11 @@ namespace mu
 			uint32 temp;
 			arch >> temp;
 			m_strategy = EPackStrategy(temp);
+
+			if (ver >= 4)
+			{
+				arch >> FirstLODToIgnoreWarnings;
+			}
 		}
 
 		//!
diff --git a/Engine/Plugins/Experimental/Mutable/Source/MutableTools/Private/MuT/CodeGenerator_Mesh.cpp b/Engine/Plugins/Experimental/Mutable/Source/MutableTools/Private/MuT/CodeGenerator_Mesh.cpp
--- a/Engine/Plugins/Experimental/Mutable/Source/MutableTools/Private/MuT/CodeGenerator_Mesh.cpp
+++ b/Engine/Plugins/Experimental/Mutable/Source/MutableTools/Private/MuT/CodeGenerator_Mesh.cpp
@@ -228,7 +228,7 @@ class Node;
         }
 
         int outside = currentLayoutMesh->GetVertexBuffers().GetElementCount() - inside;
-        if (outside>0)
+        if (outside>0 && (pLayout->FirstLODToIgnoreWarnings == -1 || m_currentParents.Last().m_lod < pLayout->FirstLODToIgnoreWarnings))
         {
             char buf[256];
             mutable_snprintf
diff --git a/Engine/Plugins/Experimental/Mutable/Source/MutableTools/Private/MuT/NodeLayout.cpp b/Engine/Plugins/Experimental/Mutable/Source/MutableTools/Private/MuT/NodeLayout.cpp
--- a/Engine/Plugins/Experimental/Mutable/Source/MutableTools/Private/MuT/NodeLayout.cpp
+++ b/Engine/Plugins/Experimental/Mutable/Source/MutableTools/Private/MuT/NodeLayout.cpp
@@ -402,6 +402,19 @@ namespace mu
 
 		return layout;
 	}
+
+	//---------------------------------------------------------------------------------------------
+	void NodeLayoutBlocks::SetIgnoreWarningsLOD(int32 LOD)
+	{
+		m_pD->m_pLayout->SetIgnoreLODWarnings(LOD);
+	}
+
+
+	//---------------------------------------------------------------------------------------------
+	int32 NodeLayoutBlocks::GetIgnoreWarningsLOD()
+	{
+		return m_pD->m_pLayout->GetIgnoreLODWarnings();
+	}
 }
 
 
diff --git a/Engine/Plugins/Experimental/Mutable/Source/MutableTools/Public/MuT/NodeLayout.h b/Engine/Plugins/Experimental/Mutable/Source/MutableTools/Public/MuT/NodeLayout.h
--- a/Engine/Plugins/Experimental/Mutable/Source/MutableTools/Public/MuT/NodeLayout.h
+++ b/Engine/Plugins/Experimental/Mutable/Source/MutableTools/Public/MuT/NodeLayout.h
@@ -149,6 +149,12 @@ namespace mu
 		//! Generate the blocks of a layout using the UV of the meshes
 		static NodeLayoutBlocksPtr GenerateLayoutBlocks(const MeshPtr pMesh, int layoutIndex, int gridSizeX, int gridSizeY);
 
+		//! Set at which LOD the unassigned vertices warnings will star to be ignored
+		void SetIgnoreWarningsLOD(int32 LOD);
+
+		//! Get the LOD where the unassigned vertices warnings starts to be ignored
+		int32 GetIgnoreWarningsLOD();
+
 
 		//-----------------------------------------------------------------------------------------
 		// Interface pattern
