[mutable] Made the selection/deselection of elements of the MutableCodeViewer and ConstantsWidget less confusing.
- Now both widgets share one unique selected element (functionally speaking) which improves usability.

diff --git a/Engine/Plugins/Experimental/Mutable/Source/CustomizableObjectEditor/Private/MuCOE/SMutableCodeViewer.cpp b/Engine/Plugins/Experimental/Mutable/Source/CustomizableObjectEditor/Private/MuCOE/SMutableCodeViewer.cpp
--- a/Engine/Plugins/Experimental/Mutable/Source/CustomizableObjectEditor/Private/MuCOE/SMutableCodeViewer.cpp
+++ b/Engine/Plugins/Experimental/Mutable/Source/CustomizableObjectEditor/Private/MuCOE/SMutableCodeViewer.cpp
@@ -409,6 +409,11 @@ FString SMutableCodeViewer::GetReferencerName() const
 	return TEXT("SMutableCodeViewer");
 }
 
+void SMutableCodeViewer::ClearSelectedTreeRow() const
+{
+	check(TreeView);
+	TreeView->ClearSelection();
+}
 
 void SMutableCodeViewer::SetCurrentModel(const TSharedPtr<mu::Model, ESPMode::ThreadSafe>& InMutableModel)
 {
@@ -1542,6 +1547,9 @@ void SMutableCodeViewer::OnSelectionChanged(TSharedPtr<FMutableCodeTreeElement>
 		return;
 	}
 
+	// Clear all selected items in the constant resources widget
+	ConstantsWidget->ClearSelectedConstantItems();
+	
 	// Find the duplicates for the selected tree element element and highlight them
 	if (InNode)
 	{
diff --git a/Engine/Plugins/Experimental/Mutable/Source/CustomizableObjectEditor/Private/MuCOE/SMutableCodeViewer.h b/Engine/Plugins/Experimental/Mutable/Source/CustomizableObjectEditor/Private/MuCOE/SMutableCodeViewer.h
--- a/Engine/Plugins/Experimental/Mutable/Source/CustomizableObjectEditor/Private/MuCOE/SMutableCodeViewer.h
+++ b/Engine/Plugins/Experimental/Mutable/Source/CustomizableObjectEditor/Private/MuCOE/SMutableCodeViewer.h
@@ -64,6 +64,9 @@ public:
 	virtual void AddReferencedObjects(FReferenceCollector& Collector) override;
 	virtual FString GetReferencerName() const override;
 
+	/** Clear the selected operation of the tree view.*/
+	void ClearSelectedTreeRow() const;
+	
 private:
 
 	/** The Mutable Model that we are showing. */
diff --git a/Engine/Plugins/Experimental/Mutable/Source/CustomizableObjectEditor/Private/MuCOE/SMutableConstantsWidget.cpp b/Engine/Plugins/Experimental/Mutable/Source/CustomizableObjectEditor/Private/MuCOE/SMutableConstantsWidget.cpp
--- a/Engine/Plugins/Experimental/Mutable/Source/CustomizableObjectEditor/Private/MuCOE/SMutableConstantsWidget.cpp
+++ b/Engine/Plugins/Experimental/Mutable/Source/CustomizableObjectEditor/Private/MuCOE/SMutableConstantsWidget.cpp
@@ -434,7 +434,7 @@ void SMutableConstantsWidget::Construct(const FArguments& InArgs,const mu::FProg
 			.AreaTitle(this,&SMutableConstantsWidget::OnDrawStringsAreaTitle)
 			.BodyContent()
 			[
-				SNew(STileView<TSharedPtr<FMutableConstantStringElement>>)
+				SAssignNew(ConstantStringsSlate,STileView<TSharedPtr<FMutableConstantStringElement>>)
 				.OnSelectionChanged(this,&SMutableConstantsWidget::OnSelectedStringChanged)
 				.ListItemsSource(&ConstantStringElements)
 				.ItemHeight(ProxyEntryHeight)
@@ -518,7 +518,7 @@ void SMutableConstantsWidget::Construct(const FArguments& InArgs,const mu::FProg
 			.AreaTitle(this,&SMutableConstantsWidget::OnDrawMeshesAreaTitle)
 			.BodyContent()
 			[
-				SNew(STileView<TSharedPtr<FMutableConstantMeshElement>>)
+				SAssignNew(ConstantMeshesSlate,STileView<TSharedPtr<FMutableConstantMeshElement>>)
 				.OnSelectionChanged(this,&SMutableConstantsWidget::OnSelectedMeshChanged)
 				.ListItemsSource(&ConstantMeshElements)
 				.ItemHeight(ProxyEntryHeight)
@@ -542,7 +542,7 @@ void SMutableConstantsWidget::Construct(const FArguments& InArgs,const mu::FProg
 			.AreaTitle(this,&SMutableConstantsWidget::OnDrawLayoutsAreaTitle)
 			.BodyContent()
 			[
-				SNew(STileView<TSharedPtr<FMutableConstantLayoutElement>>)
+				SAssignNew(ConstantLayoutsSlate,STileView<TSharedPtr<FMutableConstantLayoutElement>>)
 				.OnSelectionChanged(this,&SMutableConstantsWidget::OnSelectedLayoutChanged)
 				.ListItemsSource(&ConstantLayoutElements)
 				.ItemHeight(ProxyEntryHeight)
@@ -566,7 +566,7 @@ void SMutableConstantsWidget::Construct(const FArguments& InArgs,const mu::FProg
 			.AreaTitle(this,&SMutableConstantsWidget::OnDrawProjectorsAreaTitle)
 			.BodyContent()
 			[
-				SNew(STileView<TSharedPtr<FMutableConstantProjectorElement>>)
+				SAssignNew(ConstantProjectorsSlate,STileView<TSharedPtr<FMutableConstantProjectorElement>>)
 				.OnSelectionChanged(this,&SMutableConstantsWidget::OnSelectedProjectorChanged)
 				.ListItemsSource(&ConstantProjectorElements)
 				.ItemHeight(ProxyEntryHeight)
@@ -590,7 +590,7 @@ void SMutableConstantsWidget::Construct(const FArguments& InArgs,const mu::FProg
 			.AreaTitle(this,&SMutableConstantsWidget::OnDrawMatricesAreaTitle)
 			.BodyContent()
 			[
-				SNew(STileView<TSharedPtr<FMutableConstantMatrixElement>>)
+				SAssignNew(ConstantMatricesSlate,STileView<TSharedPtr<FMutableConstantMatrixElement>>)
 				.OnSelectionChanged(this,&SMutableConstantsWidget::OnSelectedMatrixChanged)
 				.ListItemsSource(&ConstantMatrixElements)
 				.ItemHeight(ProxyEntryHeight)
@@ -614,7 +614,7 @@ void SMutableConstantsWidget::Construct(const FArguments& InArgs,const mu::FProg
 			.AreaTitle(this,&SMutableConstantsWidget::OnDrawShapesAreaTitle)
 			.BodyContent()
 			[
-				SNew(STileView<TSharedPtr<FMutableConstantShapeElement>>)
+				SAssignNew(ConstantShapesSlate,STileView<TSharedPtr<FMutableConstantShapeElement>>)
 				.OnSelectionChanged(this,&SMutableConstantsWidget::OnSelectedShapeChanged)
 				.ListItemsSource(&ConstantShapeElements)
 				.ItemHeight(ProxyEntryHeight)
@@ -638,7 +638,7 @@ void SMutableConstantsWidget::Construct(const FArguments& InArgs,const mu::FProg
 			.AreaTitle(this,&SMutableConstantsWidget::OnDrawCurvesAreaTitle)
 			.BodyContent()
 			[
-				SNew(STileView<TSharedPtr<FMutableConstantCurveElement>>)
+				SAssignNew(ConstantCurvesSlate,STileView<TSharedPtr<FMutableConstantCurveElement>>)
 				.OnSelectionChanged(this,&SMutableConstantsWidget::OnSelectedCurveChanged)
 				.ListItemsSource(&ConstantCurveElements)
 				.ItemHeight(ProxyEntryHeight)
@@ -662,7 +662,7 @@ void SMutableConstantsWidget::Construct(const FArguments& InArgs,const mu::FProg
 			.AreaTitle(this,&SMutableConstantsWidget::OnDrawSkeletonsAreaTitle)
 			.BodyContent()
 			[
-				SNew(STileView<TSharedPtr<FMutableConstantSkeletonElement>>)
+				SAssignNew(ConstantSkeletonsSlate,STileView<TSharedPtr<FMutableConstantSkeletonElement>>)
 				.OnSelectionChanged(this,&SMutableConstantsWidget::OnSelectedSkeletonChanged)
 				.ListItemsSource(&ConstantSkeletonElements)
 				.ItemHeight(ProxyEntryHeight)
@@ -1351,16 +1351,63 @@ void SMutableConstantsWidget::LoadConstantCurves()
 #pragma region Previewer invocation methods
 
 
+void SMutableConstantsWidget::ClearSelectedConstantItems(mu::DATATYPE ExceptionDataType /* = mu::DATATYPE::DT_NONE */) const
+{
+	if (ExceptionDataType != mu::DATATYPE::DT_MESH)
+	{
+		ConstantMeshesSlate->ClearSelection();
+	}
+	if (ExceptionDataType != mu::DATATYPE::DT_STRING)
+	{
+		ConstantStringsSlate->ClearSelection();
+	}
+	if (ExceptionDataType != mu::DATATYPE::DT_LAYOUT)
+	{
+		ConstantLayoutsSlate->ClearSelection();
+	}
+	if (ExceptionDataType != mu::DATATYPE::DT_PROJECTOR)
+	{
+		ConstantProjectorsSlate->ClearSelection();
+	}
+	if (ExceptionDataType != mu::DATATYPE::DT_MATRIX)
+	{
+		ConstantMatricesSlate->ClearSelection();
+	}
+	if (ExceptionDataType != mu::DATATYPE::DT_SHAPE)
+	{
+		ConstantShapesSlate->ClearSelection();
+	}
+	if (ExceptionDataType != mu::DATATYPE::DT_CURVE)
+	{
+		ConstantCurvesSlate->ClearSelection();
+	}
+	if (ExceptionDataType != mu::DATATYPE::DT_SKELETON)
+	{
+		ConstantSkeletonsSlate->ClearSelection();
+	}
+	if (ExceptionDataType != mu::DATATYPE::DT_IMAGE)
+	{
+		ConstantImagesListView->ClearSelection();
+	}
+
+	// todo: if you add more slates for the types update this to make sure they behave as the others
+}
+
 void SMutableConstantsWidget::OnSelectedStringChanged(
 	TSharedPtr<FMutableConstantStringElement> MutableConstantStringElement,
 	ESelectInfo::Type SelectionType) const
 {
 	if (MutableConstantStringElement)
 	{
+		static constexpr mu::DATATYPE SlateDataType = mu::DATATYPE::DT_STRING;
+
+		// Clear the selected CodeViewer row and all other constant viewer slates that are not the type provided.
+		MutableCodeViewerPtr->ClearSelectedTreeRow();
+		ClearSelectedConstantItems(SlateDataType);
+		
 		// Ask the Code viewer to present the element held on your element on the previewer window
 		MutableCodeViewerPtr->PreviewMutableString(MutableConstantStringElement->MutableString);
-
-		MutableCodeViewerPtr->CacheAddressesRelatedWithConstantResource(mu::DATATYPE::DT_STRING,MutableConstantStringElement->IndexOnSourceVector);
+		MutableCodeViewerPtr->CacheAddressesRelatedWithConstantResource(SlateDataType,MutableConstantStringElement->IndexOnSourceVector);
 	}
 }
 
@@ -1370,9 +1417,13 @@ void SMutableConstantsWidget::OnSelectedImageChanged(
 {
 	if (MutableConstantImageElement)
 	{
-		MutableCodeViewerPtr->PreviewMutableImage(MutableConstantImageElement->ImagePtr);
+		static constexpr mu::DATATYPE SlateDataType = mu::DATATYPE::DT_IMAGE;
 
-		MutableCodeViewerPtr->CacheAddressesRelatedWithConstantResource(mu::DATATYPE::DT_IMAGE,MutableConstantImageElement->IndexOnSourceVector);
+		MutableCodeViewerPtr->ClearSelectedTreeRow();
+		ClearSelectedConstantItems(SlateDataType);
+		
+		MutableCodeViewerPtr->PreviewMutableImage(MutableConstantImageElement->ImagePtr);
+		MutableCodeViewerPtr->CacheAddressesRelatedWithConstantResource(SlateDataType,MutableConstantImageElement->IndexOnSourceVector);
 	}
 }
 
@@ -1382,9 +1433,13 @@ void SMutableConstantsWidget::OnSelectedMeshChanged(
 {
 	if (MutableConstantMeshElement)
 	{
-		MutableCodeViewerPtr->PreviewMutableMesh(MutableConstantMeshElement->MeshPtr);
+		static constexpr mu::DATATYPE SlateDataType = mu::DATATYPE::DT_MESH;
+		
+		MutableCodeViewerPtr->ClearSelectedTreeRow();
+		ClearSelectedConstantItems(SlateDataType);
 
-		MutableCodeViewerPtr->CacheAddressesRelatedWithConstantResource(mu::DATATYPE::DT_MESH,MutableConstantMeshElement->IndexOnSourceVector);
+		MutableCodeViewerPtr->PreviewMutableMesh(MutableConstantMeshElement->MeshPtr);
+		MutableCodeViewerPtr->CacheAddressesRelatedWithConstantResource(SlateDataType,MutableConstantMeshElement->IndexOnSourceVector);
 	}
 }
 
@@ -1394,9 +1449,13 @@ void SMutableConstantsWidget::OnSelectedLayoutChanged(
 {
 	if (MutableConstantLayoutElement)
 	{
-		MutableCodeViewerPtr->PreviewMutableLayout(MutableConstantLayoutElement->Layout);
+		static constexpr mu::DATATYPE SlateDataType = mu::DATATYPE::DT_LAYOUT;
+		
+		MutableCodeViewerPtr->ClearSelectedTreeRow();
+		ClearSelectedConstantItems(SlateDataType);
 
-		MutableCodeViewerPtr->CacheAddressesRelatedWithConstantResource(mu::DATATYPE::DT_LAYOUT,MutableConstantLayoutElement->IndexOnSourceVector);
+		MutableCodeViewerPtr->PreviewMutableLayout(MutableConstantLayoutElement->Layout);
+		MutableCodeViewerPtr->CacheAddressesRelatedWithConstantResource(SlateDataType,MutableConstantLayoutElement->IndexOnSourceVector);
 	}
 }
 
@@ -1406,9 +1465,13 @@ void SMutableConstantsWidget::OnSelectedProjectorChanged(
 {
 	if (MutableConstantProjectorElement)
 	{
-		MutableCodeViewerPtr->PreviewMutableProjector(MutableConstantProjectorElement->Projector);
+		static constexpr mu::DATATYPE SlateDataType = mu::DATATYPE::DT_PROJECTOR;
+
+		MutableCodeViewerPtr->ClearSelectedTreeRow();
+		ClearSelectedConstantItems(SlateDataType);
 
-		MutableCodeViewerPtr->CacheAddressesRelatedWithConstantResource(mu::DATATYPE::DT_PROJECTOR,MutableConstantProjectorElement->IndexOnSourceVector);
+		MutableCodeViewerPtr->PreviewMutableProjector(MutableConstantProjectorElement->Projector);
+		MutableCodeViewerPtr->CacheAddressesRelatedWithConstantResource(SlateDataType,MutableConstantProjectorElement->IndexOnSourceVector);
 	}
 }
 
@@ -1418,9 +1481,13 @@ void SMutableConstantsWidget::OnSelectedMatrixChanged(
 {
 	if (MutableConstantMatrixElement)
 	{
-		MutableCodeViewerPtr->PreviewMutableMatrix(MutableConstantMatrixElement->Matrix);
+		static constexpr mu::DATATYPE SlateDataType = mu::DATATYPE::DT_MATRIX;
+
+		MutableCodeViewerPtr->ClearSelectedTreeRow();
+		ClearSelectedConstantItems(SlateDataType);
 
-		MutableCodeViewerPtr->CacheAddressesRelatedWithConstantResource(mu::DATATYPE::DT_MATRIX,MutableConstantMatrixElement->IndexOnSourceVector);
+		MutableCodeViewerPtr->PreviewMutableMatrix(MutableConstantMatrixElement->Matrix);
+		MutableCodeViewerPtr->CacheAddressesRelatedWithConstantResource(SlateDataType,MutableConstantMatrixElement->IndexOnSourceVector);
 	}
 }
 
@@ -1430,9 +1497,13 @@ void SMutableConstantsWidget::OnSelectedShapeChanged(
 {
 	if (MutableConstantShapeElement)
 	{
-		MutableCodeViewerPtr->PreviewMutableShape(MutableConstantShapeElement->Shape);
+		static constexpr mu::DATATYPE SlateDataType = mu::DATATYPE::DT_SHAPE;
+
+		MutableCodeViewerPtr->ClearSelectedTreeRow();
+		ClearSelectedConstantItems(SlateDataType);
 
-		MutableCodeViewerPtr->CacheAddressesRelatedWithConstantResource(mu::DATATYPE::DT_SHAPE,MutableConstantShapeElement->IndexOnSourceVector);
+		MutableCodeViewerPtr->PreviewMutableShape(MutableConstantShapeElement->Shape);
+		MutableCodeViewerPtr->CacheAddressesRelatedWithConstantResource(SlateDataType,MutableConstantShapeElement->IndexOnSourceVector);
 	}
 }
 
@@ -1442,9 +1513,13 @@ void SMutableConstantsWidget::OnSelectedCurveChanged(
 {
 	if (MutableConstantCurveElement)
 	{
-		MutableCodeViewerPtr->PreviewMutableCurve(MutableConstantCurveElement->Curve);
+		static constexpr mu::DATATYPE SlateDataType = mu::DATATYPE::DT_CURVE;
 
-		MutableCodeViewerPtr->CacheAddressesRelatedWithConstantResource(mu::DATATYPE::DT_CURVE,MutableConstantCurveElement->IndexOnSourceVector);
+		MutableCodeViewerPtr->ClearSelectedTreeRow();
+		ClearSelectedConstantItems(SlateDataType);
+		
+		MutableCodeViewerPtr->PreviewMutableCurve(MutableConstantCurveElement->Curve);
+		MutableCodeViewerPtr->CacheAddressesRelatedWithConstantResource(SlateDataType,MutableConstantCurveElement->IndexOnSourceVector);
 	}
 }
 
@@ -1454,9 +1529,13 @@ void SMutableConstantsWidget::OnSelectedSkeletonChanged(
 {
 	if (MutableConstantSkeletonElement)
 	{
-		MutableCodeViewerPtr->PreviewMutableSkeleton(MutableConstantSkeletonElement->Skeleton);
+		static constexpr mu::DATATYPE SlateDataType = mu::DATATYPE::DT_SKELETON;
+
+		MutableCodeViewerPtr->ClearSelectedTreeRow();
+		ClearSelectedConstantItems(SlateDataType);
 
-		MutableCodeViewerPtr->CacheAddressesRelatedWithConstantResource(mu::DATATYPE::DT_SKELETON,MutableConstantSkeletonElement->IndexOnSourceVector);
+		MutableCodeViewerPtr->PreviewMutableSkeleton(MutableConstantSkeletonElement->Skeleton);
+		MutableCodeViewerPtr->CacheAddressesRelatedWithConstantResource(SlateDataType,MutableConstantSkeletonElement->IndexOnSourceVector);
 	}
 }
 
diff --git a/Engine/Plugins/Experimental/Mutable/Source/CustomizableObjectEditor/Private/MuCOE/SMutableConstantsWidget.h b/Engine/Plugins/Experimental/Mutable/Source/CustomizableObjectEditor/Private/MuCOE/SMutableConstantsWidget.h
--- a/Engine/Plugins/Experimental/Mutable/Source/CustomizableObjectEditor/Private/MuCOE/SMutableConstantsWidget.h
+++ b/Engine/Plugins/Experimental/Mutable/Source/CustomizableObjectEditor/Private/MuCOE/SMutableConstantsWidget.h
@@ -4,7 +4,9 @@
 
 #include "MuR/Image.h"
 #include "MuR/Mesh.h"
+#include "MuR/Operations.h"
 #include "Widgets/Views/SHeaderRow.h"
+#include "Widgets/Views/STileView.h"
 
 template <typename ItemType> class SListView;
 
@@ -117,25 +119,45 @@ public:
 	 * @param  InMutableCodeViewerPtr - Pointer to the MutableCodeViewer tasked with the previewing of the constant values
 	 */
 	void Construct(const FArguments& InArgs, const mu::FProgram*  InMutableProgramPtr ,  TSharedPtr<SMutableCodeViewer> InMutableCodeViewerPtr);
-	
+
+	/** It clears the selected element for all constant view slates (SListView and STileView) except the one that shows the data of the provided type
+	 * @param ExceptionDataType The type of data of the slate container that we want to skip when clearing the selected elements.
+	 * EX : We want to clear all selected elements but not the ones from the constant images. Provide as argument "mu::DATATYPE::DT_IMAGE"
+	 * to clear all but the selected item/s of the Image type.
+	 */
+	void ClearSelectedConstantItems(mu::DATATYPE ExceptionDataType = mu::DATATYPE::DT_NONE) const;
+
 private:
 
 	/** Mutable object containing the constants data */
 	const mu::FProgram* MutableProgramPtr = nullptr;
 
 	/** Slate capable of accessing the previewer object */
-	 TSharedPtr<SMutableCodeViewer> MutableCodeViewerPtr = nullptr;
+	TSharedPtr<SMutableCodeViewer> MutableCodeViewerPtr = nullptr;
 	
 	/**
 	 *Sets the back end for the operation of this widget. Each time this is done the ui backend gets updated
 	 * @param InProgram - Mutable program object holding all the constants data
 	 */
 	void SetProgram(const mu::FProgram* InProgram);
+
+	/*
+	 * Pointers to all slates showing the constants
+	 */
+	
+	TSharedPtr<STileView<TSharedPtr<FMutableConstantMeshElement>>> ConstantMeshesSlate;
+	TSharedPtr<STileView<TSharedPtr<FMutableConstantStringElement>>> ConstantStringsSlate;
+	TSharedPtr<STileView<TSharedPtr<FMutableConstantLayoutElement>>> ConstantLayoutsSlate;
+	TSharedPtr<STileView<TSharedPtr<FMutableConstantProjectorElement>>> ConstantProjectorsSlate;
+	TSharedPtr<STileView<TSharedPtr<FMutableConstantMatrixElement>>> ConstantMatricesSlate;
+	TSharedPtr<STileView<TSharedPtr<FMutableConstantShapeElement>>> ConstantShapesSlate;
+	TSharedPtr<STileView<TSharedPtr<FMutableConstantCurveElement>>> ConstantCurvesSlate;
+	TSharedPtr<STileView<TSharedPtr<FMutableConstantSkeletonElement>>> ConstantSkeletonsSlate;
 	
 	/*
 	* Data backend for the lists of constants
 	*/
-	
+
 	TArray<TSharedPtr<FMutableConstantImageElement>> ConstantImageElements;
 	TArray<TSharedPtr<FMutableConstantMeshElement>> ConstantMeshElements;
 	TArray<TSharedPtr<FMutableConstantStringElement>> ConstantStringElements;
@@ -166,7 +188,7 @@ private:
 	 */
 	
 	void OnSelectedStringChanged(TSharedPtr<FMutableConstantStringElement> MutableConstantStringElement, ESelectInfo::Type SelectionType) const;
-	void OnSelectedImageChanged(TSharedPtr<FMutableConstantImageElement> MutableConstantImageElement, ESelectInfo::Type SelectionType)const;
+	void OnSelectedImageChanged(TSharedPtr<FMutableConstantImageElement> MutableConstantImageElement, ESelectInfo::Type SelectionType) const;
 	void OnSelectedMeshChanged(TSharedPtr<FMutableConstantMeshElement> MutableConstantMeshElement, ESelectInfo::Type SelectionType) const;
 	void OnSelectedLayoutChanged(TSharedPtr<FMutableConstantLayoutElement> MutableConstantLayoutElement, ESelectInfo::Type SelectionType) const;
 	void OnSelectedProjectorChanged(TSharedPtr<FMutableConstantProjectorElement> MutableConstantProjectorElement, ESelectInfo::Type SelectionType) const;
