[mutable] Fixed missing headers.

#rnx

diff --git a/Engine/Plugins/Experimental/Mutable/Source/CustomizableObjectEditor/Private/MuCOE/Nodes/CustomizableObjectNodeLayoutBlocksDetails.cpp b/Engine/Plugins/Experimental/Mutable/Source/CustomizableObjectEditor/Private/MuCOE/Nodes/CustomizableObjectNodeLayoutBlocksDetails.cpp
--- a/Engine/Plugins/Experimental/Mutable/Source/CustomizableObjectEditor/Private/MuCOE/Nodes/CustomizableObjectNodeLayoutBlocksDetails.cpp
+++ b/Engine/Plugins/Experimental/Mutable/Source/CustomizableObjectEditor/Private/MuCOE/Nodes/CustomizableObjectNodeLayoutBlocksDetails.cpp
@@ -9,8 +9,11 @@
 #include "MuCOE/Nodes/CustomizableObjectNodeLayoutBlocks.h"
 #include "MuCOE/SCustomizableObjectNodeLayoutBlocksEditor.h"
 #include "PropertyCustomizationHelpers.h"
+#include "Styling/SlateTypes.h"
+#include "Widgets/Input/SCheckBox.h"
 #include "Widgets/Input/SSpinBox.h"
 #include "Widgets/Layout/SBox.h"
+#include "Widgets/SWidget.h"
 #include "Widgets/Text/STextBlock.h"
 
 class FString;
diff --git a/Engine/Plugins/Experimental/Mutable/Source/CustomizableObjectEditor/Private/MuCOE/Nodes/CustomizableObjectNodeLayoutBlocksDetails.h b/Engine/Plugins/Experimental/Mutable/Source/CustomizableObjectEditor/Private/MuCOE/Nodes/CustomizableObjectNodeLayoutBlocksDetails.h
--- a/Engine/Plugins/Experimental/Mutable/Source/CustomizableObjectEditor/Private/MuCOE/Nodes/CustomizableObjectNodeLayoutBlocksDetails.h
+++ b/Engine/Plugins/Experimental/Mutable/Source/CustomizableObjectEditor/Private/MuCOE/Nodes/CustomizableObjectNodeLayoutBlocksDetails.h
@@ -5,11 +5,13 @@
 
 #include "IDetailCustomization.h"
 
+enum class ECheckBoxState : uint8;
 namespace ESelectInfo { enum Type : int; }
 
 class FString;
 class IDetailLayoutBuilder;
-
+class STextBlock;
+class SWidget;
 
 class FCustomizableObjectNodeLayoutBlocksDetails : public IDetailCustomization
 {
