[mutable] Remove custom memory allocator override boilerplate.
#rnx

diff --git a/Engine/Plugins/Experimental/Mutable/Source/MutableRuntime/Private/MuR/Memory.cpp b/Engine/Plugins/Experimental/Mutable/Source/MutableRuntime/Private/MuR/Memory.cpp
--- a/Engine/Plugins/Experimental/Mutable/Source/MutableRuntime/Private/MuR/Memory.cpp
+++ b/Engine/Plugins/Experimental/Mutable/Source/MutableRuntime/Private/MuR/Memory.cpp
@@ -2,173 +2,23 @@
 
 #include "MuR/BlockCompression/Miro/Miro.h"
 #include "MuR/MutableMemory.h"
-#include "HAL/MallocAnsi.h"
 
 namespace mu
 {
-
-
-//! Static pointers to the user provided malloc and free methods.
-//! If they are null (default) the versions on the standard library will be used.
-	static void* (*s_custom_malloc)(size_t, uint32) = nullptr;
-	static void(*s_custom_free)(void*) = nullptr;
-
-
 	static int s_initialized = 0;
 	static int s_finalized = 0;
 
-
-	//---------------------------------------------------------------------------------------------
-	//! Call the system malloc or the user-provided replacement.
-	//---------------------------------------------------------------------------------------------
-	inline void* lowerlevel_malloc(size_t bytes)
-	{
-		void* pMem = 0;
-		if (s_custom_malloc)
-		{
-			pMem = s_custom_malloc(bytes, 1);
-		}
-		else
-		{
-			pMem = malloc(bytes);
-		}
-		return pMem;
-	}
-
-	inline void* lowerlevel_malloc_aligned(size_t bytes, uint32 alignment)
-	{
-		void* pMem = 0;
-		if (s_custom_malloc)
-		{
-			pMem = s_custom_malloc(bytes, alignment);
-		}
-		else
-		{
-			size_t padding = bytes % alignment;
-			bytes += padding ? alignment - padding : 0;
-			pMem = AnsiMalloc(bytes, alignment);
-		}
-		return pMem;
-	}
-
-	//-------------------------------------------------------------------------------------------------
-	//! Call the system free or the user-provided replacement.
-	//-------------------------------------------------------------------------------------------------
-	inline void lowerlevel_free(void* ptr)
-	{
-		if (s_custom_free)
-		{
-			s_custom_free(ptr);
-		}
-		else
-		{
-			free(ptr);
-		}
-	}
-
-
-	//-------------------------------------------------------------------------------------------------
-	inline void lowerlevel_free_aligned(void* ptr)
-	{
-		if (s_custom_free)
-		{
-			s_custom_free(ptr);
-		}
-		else
-		{
-			AnsiFree(ptr);
-		}
-	}
-
-
-	//-------------------------------------------------------------------------------------------------
-	//! Memory management functions to be used inside the library. No other memory allocation is
-	//! allowed.
-	//-------------------------------------------------------------------------------------------------
-	void* mutable_malloc(size_t size)
-	{
-		if (!size)
-		{
-			return nullptr;
-		}
-
-		void* pMem = nullptr;
-
-		pMem = lowerlevel_malloc(size);
-
-		return pMem;
-	}
-
-
-	void* mutable_malloc_aligned(size_t size, uint32 alignment)
-	{
-		if (!size)
-		{
-			return nullptr;
-		}
-
-		void* pMem = nullptr;
-
-		pMem = lowerlevel_malloc_aligned(size, alignment);
-
-		return pMem;
-	}
-
-
-	void mutable_free(void* ptr)
-	{
-		if (!ptr)
-		{
-			return;
-		}
-
-		lowerlevel_free(ptr);
-	}
-
-
-	void mutable_free(void* ptr, size_t size)
-	{
-		if (!ptr || !size)
-		{
-			return;
-		}
-
-		lowerlevel_free(ptr);
-	}
-
-
-	void mutable_free_aligned(void* ptr, size_t size)
-	{
-		if (!size || !ptr)
-		{
-			return;
-		}
-
-		lowerlevel_free_aligned(ptr);
-	}
-
-
-	//-------------------------------------------------------------------------------------------------
-	void Initialize
-	(
-		void* (*customMalloc)(size_t, uint32),
-		void (*customFree)(void*)
-	)
+	void Initialize()
 	{
 		if (!s_initialized)
 		{
 			s_initialized = 1;
 			s_finalized = 0;
 
-			s_custom_malloc = customMalloc;
-			s_custom_free = customFree;
-
 			miro::initialize();
 		}
 	}
 
-
-	//-------------------------------------------------------------------------------------------------
 	void Finalize()
 	{
 		if (s_initialized && !s_finalized)
@@ -177,10 +27,7 @@ namespace mu
 
 			s_finalized = 1;
 			s_initialized = 0;
-			s_custom_malloc = nullptr;
-			s_custom_free = nullptr;
 		}
 	}
 
-
 }
diff --git a/Engine/Plugins/Experimental/Mutable/Source/MutableRuntime/Private/MuR/MutableRuntimeModule.cpp b/Engine/Plugins/Experimental/Mutable/Source/MutableRuntime/Private/MuR/MutableRuntimeModule.cpp
--- a/Engine/Plugins/Experimental/Mutable/Source/MutableRuntime/Private/MuR/MutableRuntimeModule.cpp
+++ b/Engine/Plugins/Experimental/Mutable/Source/MutableRuntime/Private/MuR/MutableRuntimeModule.cpp
@@ -2,7 +2,6 @@
 
 #include "MuR/MutableRuntimeModule.h"
 
-#include "HAL/PlatformMath.h"
 #include "HAL/UnrealMemory.h"
 #include "Logging/LogMacros.h"
 #include "Misc/AssertionMacros.h"
@@ -14,80 +13,10 @@ IMPLEMENT_MODULE(FMutableRuntimeModule, MutableRuntime);
 
 DEFINE_LOG_CATEGORY(LogMutableCore);
 
-#if USE_STAT_MUTABLE_MEMORY
-DECLARE_MEMORY_STAT(TEXT("Mutable Memory Used"), STAT_MemoryMutableTotalAllocationSize, STATGROUP_Memory);
-#endif
-
-
-// Uncomment the following line to have allocation/deallocation messages printed to the output
-//#define MUTABLE_PRINT_ALLOCATIONS
-
-#ifdef MUTABLE_PRINT_ALLOCATIONS
-
-static int64 NumAllocations = 0;
-
-static void* CustomMalloc(std::size_t Size_t, uint32_t alignment)
-{
-	NumAllocations++;
-	UE_LOG(LogTemp, Warning, TEXT("NumAllocations=%d, Allocated %d bytes"), NumAllocations, Size_t);
-
-	return malloc(Size_t);
-}
-
-
-static void CustomFree(void* mem)
-{
-	NumAllocations--;
-	UE_LOG(LogTemp, Warning, TEXT("NumAllocations=%d, Deallocated"), NumAllocations);
-
-	free(mem);
-}
-
-#else //MUTABLE_PRINT_ALLOCATIONS
-
-
-
-static void* CustomMalloc(std::size_t Size_t, uint32_t alignment)
-{
-	void* Result = FMemory::Malloc(Size_t, alignment);
-
-#if USE_STAT_MUTABLE_MEMORY
-	SIZE_T ActualSize = FMemory::GetAllocSize(Result);
-	INC_DWORD_STAT_BY(STAT_MemoryMutableTotalAllocationSize, ActualSize);
-#endif
-	// This should always be fatal: mutable cannot handle out-of-memory errors yet.
-	if (Size_t > 0 && !Result)
-	{
-		LowLevelFatalError(TEXT("Not enough memory for mutable request %ul."), Size_t);
-		check(false);
-	}
-
-	return Result;
-}
-
-
-static void CustomFree(void* mem)
-{
-#if USE_STAT_MUTABLE_MEMORY
-	DEC_DWORD_STAT_BY(STAT_MemoryMutableTotalAllocationSize, FMemory::GetAllocSize(mem));
-#endif
-	return FMemory::Free(mem);
-}
-
-#endif //MUTABLE_PRINT_ALLOCATIONS
-
 
 void FMutableRuntimeModule::StartupModule()
 {
-	void* (*customMalloc)(size_t, uint32_t) = nullptr;
-	void(*customFree)(void*) = nullptr;
-
-#if USE_UNREAL_ALLOC_IN_MUTABLE
-	customMalloc = CustomMalloc;
-	customFree = CustomFree;
-#endif
-
-	mu::Initialize(customMalloc, customFree);
+	mu::Initialize();
 }
 
 
diff --git a/Engine/Plugins/Experimental/Mutable/Source/MutableRuntime/Private/MuR/OpMeshRemoveChart.h b/Engine/Plugins/Experimental/Mutable/Source/MutableRuntime/Private/MuR/OpMeshRemoveChart.h
--- a/Engine/Plugins/Experimental/Mutable/Source/MutableRuntime/Private/MuR/OpMeshRemoveChart.h
+++ b/Engine/Plugins/Experimental/Mutable/Source/MutableRuntime/Private/MuR/OpMeshRemoveChart.h
@@ -96,7 +96,7 @@ namespace mu
 
         // For each source vertex, true if it is removed.
         int resultVertexCount = pResult->GetVertexCount();
-        uint8_t* removedVertices = (uint8_t*)mutable_malloc_aligned( resultVertexCount, 16 );
+		uint8* removedVertices = reinterpret_cast<uint8*>(FMemory::Malloc(resultVertexCount, 16));
 		FMemory::Memzero(removedVertices,resultVertexCount);
         {
 			TArray<ID_INTERVAL> intervals;
@@ -224,7 +224,7 @@ namespace mu
             pResult->GetIndexBuffers().SetElementCount( faceCount*3-(int)removedIndices );
         }
 
-        mutable_free_aligned( removedVertices, resultVertexCount );
+        FMemory::Free( removedVertices );
 
         // Rebuild the vertex buffers
         for ( int b=0; b<pResult->GetVertexBuffers().GetBufferCount(); ++b )
diff --git a/Engine/Plugins/Experimental/Mutable/Source/MutableRuntime/Public/MuR/MutableMemory.h b/Engine/Plugins/Experimental/Mutable/Source/MutableRuntime/Public/MuR/MutableMemory.h
--- a/Engine/Plugins/Experimental/Mutable/Source/MutableRuntime/Public/MuR/MutableMemory.h
+++ b/Engine/Plugins/Experimental/Mutable/Source/MutableRuntime/Public/MuR/MutableMemory.h
@@ -4,7 +4,6 @@
 
 #include "CoreTypes.h"
 
-#include <cstring>	// for size_t in some platforms
 #include <new>
 
 namespace mu
@@ -24,24 +23,13 @@ namespace mu
 	//! \param customFree Function pointer to a function taking a pointer obtained reviously with
 	//! 	   customMalloc, that needs to be freed.
 	//! \ingroup runtime
-	MUTABLERUNTIME_API extern void Initialize
-		(
-            void* (*customMalloc)( size_t, uint32 ),
-            void (*customFree)( void* )
-		);
+	MUTABLERUNTIME_API extern void Initialize();
 
     //! \brief Shutdown everything related to mutable.
     //! This should be called when no other mutable objects will be used, created or destroyed,
     //! and all threads involving mutable tasks have been terminated.
 	MUTABLERUNTIME_API extern void Finalize();
 
-	//! Memory management functions to be used inside the library. No other memory allocation is
-	//! allowed.
-	MUTABLERUNTIME_API extern void* mutable_malloc(size_t size);
-	MUTABLERUNTIME_API extern void* mutable_malloc_aligned(size_t size, uint32 alignment);
-	MUTABLERUNTIME_API extern void mutable_free(void* ptr);
-	MUTABLERUNTIME_API extern void mutable_free(void* ptr, size_t size);
-	MUTABLERUNTIME_API extern void mutable_free_aligned(void* ptr, size_t size);
 
     //! \brief %Base class from which all library classes inherit to control memory allocation.
 	//! For internal use only.
@@ -50,15 +38,15 @@ namespace mu
 	{
 	public:
 
-		inline void* operator new(size_t s)
-		{
-			return mutable_malloc(s);
-		}
+		//inline void* operator new(size_t s)
+		//{
+		//	return mutable_malloc(s);
+		//}
 
-		inline void operator delete(void* p, size_t size)
-		{
-			mutable_free(p, size);
-		}
+		//inline void operator delete(void* p, size_t size)
+		//{
+		//	mutable_free(p, size);
+		//}
     };
 
 }
diff --git a/Engine/Plugins/Experimental/Mutable/Source/MutableTools/Private/MuT/Streams.cpp b/Engine/Plugins/Experimental/Mutable/Source/MutableTools/Private/MuT/Streams.cpp
--- a/Engine/Plugins/Experimental/Mutable/Source/MutableTools/Private/MuT/Streams.cpp
+++ b/Engine/Plugins/Experimental/Mutable/Source/MutableTools/Private/MuT/Streams.cpp
@@ -28,7 +28,7 @@ namespace mu
 		m_pD->m_pos = 0;
         m_pD->m_readBytes = 0;
 
-        m_pD->m_buffer = (uint8_t*)mutable_malloc_aligned(MUTABLE_STREAM_BUFFER_SIZE,16);
+        m_pD->m_buffer = reinterpret_cast<uint8*>(FMemory::Malloc(MUTABLE_STREAM_BUFFER_SIZE,16));
 	}
 
 
@@ -40,7 +40,7 @@ namespace mu
             fclose( m_pD->m_pFile );
         }
 
-        mutable_free_aligned( m_pD->m_buffer, MUTABLE_STREAM_BUFFER_SIZE );
+		FMemory::Free( m_pD->m_buffer );
         m_pD->m_buffer = nullptr;
 
         delete m_pD;
@@ -124,7 +124,7 @@ namespace mu
 		m_pD = new Private();
 		m_pD->m_pFile = mutable_fopen( strFile, "wb" );
 		m_pD->m_pos = 0;
-        m_pD->m_buffer = (uint8_t*)mutable_malloc_aligned(MUTABLE_STREAM_BUFFER_SIZE,16);
+		m_pD->m_buffer = reinterpret_cast<uint8*>(FMemory::Malloc(MUTABLE_STREAM_BUFFER_SIZE, 16));
 
 		check( m_pD->m_pFile );
 	}
@@ -182,7 +182,7 @@ namespace mu
 
 		fclose( m_pD->m_pFile );
 
-        mutable_free_aligned( m_pD->m_buffer, MUTABLE_STREAM_BUFFER_SIZE );
+        FMemory::Free( m_pD->m_buffer );
         m_pD->m_buffer = nullptr;
 
         delete m_pD;
