[Mutable] Fix Preview Instance Properties search bar losing the search text filter.

#rnx

diff --git a/Engine/Plugins/Experimental/Mutable/Source/CustomizableObject/Classes/MuCO/CustomizableObjectInstance.h b/Engine/Plugins/Experimental/Mutable/Source/CustomizableObject/Classes/MuCO/CustomizableObjectInstance.h
--- a/Engine/Plugins/Experimental/Mutable/Source/CustomizableObject/Classes/MuCO/CustomizableObjectInstance.h
+++ b/Engine/Plugins/Experimental/Mutable/Source/CustomizableObject/Classes/MuCO/CustomizableObjectInstance.h
@@ -677,8 +677,13 @@ private:
 	/** Textures used in the Texture Parameters. */
 	UPROPERTY(EditAnywhere, Category = TextureParameter)
 	TArray<TObjectPtr<UTexture2D>> TextureParameterDeclarations;
+
+public:
+	/** Preview Instance Properties search box filter. Saved here to avoid losing the text during UI refreshes. */
+	FText ParametersSearchFilter;
 #endif
-	
+
+private:	
 	// Deprecated properties
 	
 	UPROPERTY()
diff --git a/Engine/Plugins/Experimental/Mutable/Source/CustomizableObjectEditor/Private/MuCOE/SCustomizableInstanceProperties.cpp b/Engine/Plugins/Experimental/Mutable/Source/CustomizableObjectEditor/Private/MuCOE/SCustomizableInstanceProperties.cpp
--- a/Engine/Plugins/Experimental/Mutable/Source/CustomizableObjectEditor/Private/MuCOE/SCustomizableInstanceProperties.cpp
+++ b/Engine/Plugins/Experimental/Mutable/Source/CustomizableObjectEditor/Private/MuCOE/SCustomizableInstanceProperties.cpp
@@ -104,13 +104,6 @@ void SCustomizableInstanceProperties::Construct(const FArguments& InArgs)
 			}
 		}
 
-		FText SearchBoxText = FText::GetEmpty();
-
-		if (SB_Properties.IsValid())
-		{
-			SearchBoxText = SB_Properties->GetText();
-		}
-
 		this->ChildSlot
 		[
 			SNew(SVerticalBox)
@@ -312,9 +305,9 @@ void SCustomizableInstanceProperties::Construct(const FArguments& InArgs)
 
 				.Padding(2.0f)
 				[
-					SAssignNew(SB_Properties, SSearchBox)
+					SNew(SSearchBox)
 					.HintText(LOCTEXT("SearchHint", "Search Properties"))
-					.InitialText(SearchBoxText)
+					.InitialText(CustomInstance->ParametersSearchFilter)
 					.OnTextChanged(this, &SCustomizableInstanceProperties::OnFilterTextChanged)
 				]
 			]
@@ -528,7 +521,7 @@ void SCustomizableInstanceProperties::AddParameter(int32 ParamIndexInObject)
 		|| (ParamName.EndsWith(FMultilayerProjector::OPACITY_PARAMETER_POSTFIX) && CustomInstance->IsParamMultidimensional(ParamIndexInObject))
 		|| (ParamName.EndsWith(FMultilayerProjector::POSE_PARAMETER_POSTFIX));
 	
-	FString Name = SB_Properties->GetText().ToString();
+	const FString Name = CustomInstance->ParametersSearchFilter.ToString();
 
 	//If SearchBox text is not empty
 	if (!Name.IsEmpty())
@@ -1677,6 +1670,7 @@ void SCustomizableInstanceProperties::OnTextureParameterComboBoxSelectionChanged
 
 void SCustomizableInstanceProperties::OnFilterTextChanged(const FText & InFilterText)
 {
+	CustomInstance->ParametersSearchFilter = InFilterText;
 	ResetParamBox();
 }
 
