[Mutable] Remove Group Projector Node Image Pin.

#rnx

diff --git a/Engine/Plugins/Experimental/Mutable/Source/CustomizableObject/Private/MuCO/CustomizableObjectInstanceDescriptor.cpp b/Engine/Plugins/Experimental/Mutable/Source/CustomizableObject/Private/MuCO/CustomizableObjectInstanceDescriptor.cpp
--- a/Engine/Plugins/Experimental/Mutable/Source/CustomizableObject/Private/MuCO/CustomizableObjectInstanceDescriptor.cpp
+++ b/Engine/Plugins/Experimental/Mutable/Source/CustomizableObject/Private/MuCO/CustomizableObjectInstanceDescriptor.cpp
@@ -2179,17 +2179,14 @@ void FCustomizableObjectInstanceDescriptor::SetRandomValues()
 
 bool FCustomizableObjectInstanceDescriptor::CreateMultiLayerProjector(const FName& ProjectorParamName)
 {
-	if (const FMultilayerProjector* Result = MultilayerProjectors.Find(ProjectorParamName))
+	if (!FMultilayerProjector::AreDescriptorParametersValid(*this, ProjectorParamName.ToString()))
 	{
-		checkCode(Result->CheckDescriptorParameters(*this));
+		ensureAlwaysMsgf(false, TEXT("%s"), *FMultilayerProjector::DESCRIPTOR_PARAMETERS_INVALID);
+		return false;
 	}
-	else
+
+	if (!MultilayerProjectors.Contains(ProjectorParamName))
 	{
-		if (!FMultilayerProjector::AreDescriptorParametersValid(*this, ProjectorParamName))
-		{
-			return false;
-		}
-		
 		const FMultilayerProjector MultilayerProjector(ProjectorParamName);
 		MultilayerProjectors.Add(ProjectorParamName, MultilayerProjector);
 	}
diff --git a/Engine/Plugins/Experimental/Mutable/Source/CustomizableObject/Private/MuCO/MultilayerProjector.cpp b/Engine/Plugins/Experimental/Mutable/Source/CustomizableObject/Private/MuCO/MultilayerProjector.cpp
--- a/Engine/Plugins/Experimental/Mutable/Source/CustomizableObject/Private/MuCO/MultilayerProjector.cpp
+++ b/Engine/Plugins/Experimental/Mutable/Source/CustomizableObject/Private/MuCO/MultilayerProjector.cpp
@@ -8,12 +8,20 @@
 #include UE_INLINE_GENERATED_CPP_BY_NAME(MultilayerProjector)
 
 
+const FString FMultilayerProjector::DESCRIPTOR_PARAMETERS_INVALID = TEXT("Descriptor Parameters Invalid.");
+
+
 void FMultilayerProjectorLayer::Read(const FCustomizableObjectInstanceDescriptor& Descriptor, const FMultilayerProjector& MultilayerProjector, const int32 Index)
 {
-	checkCode(MultilayerProjector.CheckDescriptorParameters(Descriptor));
 	check(Index >= 0 && Index < MultilayerProjector.NumLayers(Descriptor)); // Layer out of range.
 	
 	const FString ParamName = MultilayerProjector.ParamName.ToString();
+
+	if (!FMultilayerProjector::AreDescriptorParametersValid(Descriptor, ParamName))
+	{
+		ensureAlwaysMsgf(false, TEXT("%s"), *FMultilayerProjector::DESCRIPTOR_PARAMETERS_INVALID);
+		return;
+	}
 	
 	{
 		const int32 ProjectorParamIndex = Descriptor.FindProjectorParameterNameIndex(ParamName);
@@ -25,10 +33,8 @@ void FMultilayerProjectorLayer::Read(const FCustomizableObjectInstanceDescriptor
 		Angle = Projector.Angle;
 	}
 
-	// Image Parameter can not exist due to the Group Projector Node providing the image through its own Image Pin.
-	if (const int32 ImageParamIndex = Descriptor.FindIntParameterNameIndex(ParamName + FMultilayerProjector::IMAGE_PARAMETER_POSTFIX);
-		ImageParamIndex >= 0)
 	{
+		const int32 ImageParamIndex = Descriptor.FindIntParameterNameIndex(ParamName + FMultilayerProjector::IMAGE_PARAMETER_POSTFIX);
 		Image = Descriptor.GetIntParameters()[ImageParamIndex].ParameterRangeValueNames[Index];
 	}
 
@@ -41,11 +47,16 @@ void FMultilayerProjectorLayer::Read(const FCustomizableObjectInstanceDescriptor
 
 void FMultilayerProjectorLayer::Write(FCustomizableObjectInstanceDescriptor& Descriptor, const FMultilayerProjector& MultilayerProjector, const int32 Index) const
 {
-	checkCode(MultilayerProjector.CheckDescriptorParameters(Descriptor));
 	check(Index >= 0 && Index < MultilayerProjector.NumLayers(Descriptor)); // Layer out of range.
 
 	const FString ParamName = MultilayerProjector.ParamName.ToString();
 
+	if (!FMultilayerProjector::AreDescriptorParametersValid(Descriptor, ParamName))
+	{
+		ensureAlwaysMsgf(false, TEXT("%s"), *FMultilayerProjector::DESCRIPTOR_PARAMETERS_INVALID);
+		return;
+	}
+
 	{
 		const int32 ProjectorParamIndex = Descriptor.FindProjectorParameterNameIndex(ParamName);
 		FCustomizableObjectProjector& Projector = Descriptor.GetProjectorParameters()[ProjectorParamIndex].RangeValues[Index];
@@ -56,10 +67,8 @@ void FMultilayerProjectorLayer::Write(FCustomizableObjectInstanceDescriptor& Des
 		Projector.Angle = Angle;
 	}
 
-	// Image Parameter can not exist due to the Group Projector Node providing the image through its own Image Pin.
-	if (const int32 ImageParamIndex = Descriptor.FindIntParameterNameIndex(ParamName + FMultilayerProjector::IMAGE_PARAMETER_POSTFIX);
-		ImageParamIndex >= 0)
 	{
+		const int32 ImageParamIndex = Descriptor.FindIntParameterNameIndex(ParamName + FMultilayerProjector::IMAGE_PARAMETER_POSTFIX);
 		Descriptor.GetIntParameters()[ImageParamIndex].ParameterRangeValueNames[Index] = Image;
 	}
 
@@ -115,32 +124,40 @@ int32 FMultilayerProjector::NumLayers(const FCustomizableObjectInstanceDescripto
 
 void FMultilayerProjector::CreateLayer(FCustomizableObjectInstanceDescriptor& Descriptor, const int32 Index) const
 {
-	checkCode(CheckDescriptorParameters(Descriptor));
 	check(Index >= 0 && Index <= NumLayers(Descriptor)); // Layer is non-contiguous or out of range.
 
+	const FString ParamNameString = ParamName.ToString();
+
+	if (!AreDescriptorParametersValid(Descriptor, ParamNameString))
+	{
+		ensureAlwaysMsgf(false, TEXT("%s"), *FMultilayerProjector::DESCRIPTOR_PARAMETERS_INVALID);
+		return;
+	}
+
 	const UCustomizableObject* Object = Descriptor.GetCustomizableObject(); 
 	check(Object);
 	
 	// Num Layers.
 	{
-        const FString NumLayersParamName = ParamName.ToString() + NUM_LAYERS_PARAMETER_POSTFIX;
+        const FString NumLayersParamName = ParamNameString + NUM_LAYERS_PARAMETER_POSTFIX;
 		const int32 FloatParameterIndex = Descriptor.FindFloatParameterNameIndex(NumLayersParamName);
 
         Descriptor.GetFloatParameters()[FloatParameterIndex].ParameterValue += 1;
     }
+	
 	// Projector Range.
 	{
-		const int32 ProjectorParameterIndex = Descriptor.FindProjectorParameterNameIndex(ParamName.ToString());
+		const int32 ProjectorParameterIndex = Descriptor.FindProjectorParameterNameIndex(ParamNameString);
 		
 		FCustomizableObjectProjectorParameterValue& ProjectorParameter = Descriptor.GetProjectorParameters()[ProjectorParameterIndex];
-		const FCustomizableObjectProjector Projector = Descriptor.GetProjectorDefaultValue(Object->FindParameter(ParamName.ToString()));
+		const FCustomizableObjectProjector Projector = Descriptor.GetProjectorDefaultValue(Object->FindParameter(ParamNameString));
 		ProjectorParameter.RangeValues.Insert(Projector, Index);
 	}
 	
-	// Selected Image Range. Image Parameter can not exist due to the Group Projector Node providing the image through its own Image Pin.
-	if (const int32 IntParameterIndex = Descriptor.FindIntParameterNameIndex(ParamName.ToString() + FMultilayerProjector::IMAGE_PARAMETER_POSTFIX);
-		IntParameterIndex >= 0)
+	// Selected Image Range.
 	{
+		const int32 IntParameterIndex = Descriptor.FindIntParameterNameIndex(ParamNameString + FMultilayerProjector::IMAGE_PARAMETER_POSTFIX);
+
 		FCustomizableObjectIntParameterValue& IntParameter = Descriptor.GetIntParameters()[IntParameterIndex];
 		const int32 ParamIndexInObject = Object->FindParameter(IntParameter.ParameterName);
 
@@ -150,7 +167,7 @@ void FMultilayerProjector::CreateLayer(FCustomizableObjectInstanceDescriptor& De
 	
 	// Opacity Range.
 	{
-		const int32 FloatParameterIndex = Descriptor.FindFloatParameterNameIndex(ParamName.ToString() + FMultilayerProjector::OPACITY_PARAMETER_POSTFIX);
+		const int32 FloatParameterIndex = Descriptor.FindFloatParameterNameIndex(ParamNameString + FMultilayerProjector::OPACITY_PARAMETER_POSTFIX);
 
 		FCustomizableObjectFloatParameterValue& FloatParameter = Descriptor.GetFloatParameters()[FloatParameterIndex];
 		FloatParameter.ParameterRangeValues.Insert(0.5, Index); // TODO: Define the default float in the editor instead of [0.5f], like it's currently defined for GetProjectorDefaultValue()
@@ -160,12 +177,19 @@ void FMultilayerProjector::CreateLayer(FCustomizableObjectInstanceDescriptor& De
 
 void FMultilayerProjector::RemoveLayerAt(FCustomizableObjectInstanceDescriptor& Descriptor, const int32 Index) const
 {
-	checkCode(CheckDescriptorParameters(Descriptor));
 	check(Index >= 0 && Index < NumLayers(Descriptor)); // Layer out of range.
+
+	const FString ParamNameString = ParamName.ToString();
+
+	if (!AreDescriptorParametersValid(Descriptor, ParamName.ToString()))
+	{
+		ensureAlwaysMsgf(false, TEXT("%s"), *FMultilayerProjector::DESCRIPTOR_PARAMETERS_INVALID);
+		return;
+	}
 	
 	// Num Layers.
 	{
-		const FString NumLayersParamName = ParamName.ToString() + NUM_LAYERS_PARAMETER_POSTFIX;
+		const FString NumLayersParamName = ParamNameString + NUM_LAYERS_PARAMETER_POSTFIX;
         const int32 FloatParameterIndex = Descriptor.FindFloatParameterNameIndex(NumLayersParamName);
 		
     	Descriptor.GetFloatParameters()[FloatParameterIndex].ParameterValue -= 1;
@@ -173,23 +197,23 @@ void FMultilayerProjector::RemoveLayerAt(FCustomizableObjectInstanceDescriptor&
     
 	// Projector Range.
 	{
-		const int32 ProjectorParameterIndex = Descriptor.FindProjectorParameterNameIndex(ParamName.ToString());
+		const int32 ProjectorParameterIndex = Descriptor.FindProjectorParameterNameIndex(ParamNameString);
 		
 		FCustomizableObjectProjectorParameterValue& ProjectorParameter = Descriptor.GetProjectorParameters()[ProjectorParameterIndex];
 		ProjectorParameter.RangeValues.RemoveAt(Index);
 	}
 	
-	// Selected Image Range. Image Parameter can not exist due to the Group Projector Node providing the image through its own Image Pin.
-	if (const int32 IntParameterIndex = Descriptor.FindIntParameterNameIndex(ParamName.ToString() + FMultilayerProjector::IMAGE_PARAMETER_POSTFIX);
-		IntParameterIndex >= 0)
-	{		
+	// Selected Image Range.
+	{
+		const int32 IntParameterIndex = Descriptor.FindIntParameterNameIndex(ParamNameString + IMAGE_PARAMETER_POSTFIX);
+		
 		FCustomizableObjectIntParameterValue& IntParameter = Descriptor.GetIntParameters()[IntParameterIndex];
 		IntParameter.ParameterRangeValueNames.RemoveAt(Index);
 	}
 	
 	// Opacity Range.
 	{
-		const int32 FloatParameterIndex = Descriptor.FindFloatParameterNameIndex(ParamName.ToString() + OPACITY_PARAMETER_POSTFIX);
+		const int32 FloatParameterIndex = Descriptor.FindFloatParameterNameIndex(ParamNameString + OPACITY_PARAMETER_POSTFIX);
 		
 		FCustomizableObjectFloatParameterValue& FloatParameter = Descriptor.GetFloatParameters()[FloatParameterIndex];
 		FloatParameter.ParameterRangeValues.RemoveAt(Index);
@@ -439,65 +463,37 @@ void FMultilayerProjector::UpdateMappingVirtualLayerDisabled(const FName& Id, co
 }
 
 
-void FMultilayerProjector::CheckDescriptorParameters(const FCustomizableObjectInstanceDescriptor& Descriptor) const
+bool FMultilayerProjector::AreDescriptorParametersValid(const FCustomizableObjectInstanceDescriptor& Descriptor, const FString& ParamName)
 {
-	const FString ParamNameString = ParamName.ToString();
-
 	// Num layers.
 	{
-		const FString NumLayersParamName = ParamNameString + NUM_LAYERS_PARAMETER_POSTFIX;
+		const FString NumLayersParamName = ParamName + NUM_LAYERS_PARAMETER_POSTFIX;
 		const int32 FloatParameterIndex = Descriptor.FindFloatParameterNameIndex(NumLayersParamName);
-		check(FloatParameterIndex >= 0); // Descriptor Parameter does not exist.
-	}
-    
-	// Projector.
-	{
-		const int32 ProjectorParameterIndex = Descriptor.FindProjectorParameterNameIndex(ParamNameString);
-		check(ProjectorParameterIndex >= 0) // Descriptor Parameter does not exist.
-	}
-	
-	// Selected Image. No longer checked since the image can be provided by its own Node Group Projector pin.
-	
-	// Opacity.
-	{
-		const int32 FloatParameterIndex = Descriptor.FindFloatParameterNameIndex(ParamNameString + OPACITY_PARAMETER_POSTFIX);
-		check(FloatParameterIndex >= 0) // Descriptor Parameter does not exist.
-	}
-}
-
-
-bool FMultilayerProjector::AreDescriptorParametersValid(const FCustomizableObjectInstanceDescriptor& Descriptor, const FName& ParamName)
-{
-	const FString ParamNameString = ParamName.ToString();
-
-	// Num layers.
-	{
-		const FString NumLayersParamName = ParamNameString + NUM_LAYERS_PARAMETER_POSTFIX;
-		const int32 FloatParameterIndex = Descriptor.FindFloatParameterNameIndex(NumLayersParamName);
-		if (FloatParameterIndex < 0)
+		if (FloatParameterIndex != INDEX_NONE)
 		{
 			return false;
 		}
 	}
     
 	// Projector.
+	if (const int32 ProjectorParameterIndex = Descriptor.FindProjectorParameterNameIndex(ParamName);
+		ProjectorParameterIndex != INDEX_NONE)
 	{
-		const int32 ProjectorParameterIndex = Descriptor.FindProjectorParameterNameIndex(ParamNameString);
-		if (ProjectorParameterIndex < 0)
-		{
-			return false;
-		}
+		return false;
 	}
 	
-	// Selected Image. No longer checked since the image can be provided by its own Node Group Projector pin.
-	
+	// Selected Image.
+	if (const int32 IntParameterIndex = Descriptor.FindIntParameterNameIndex(ParamName + IMAGE_PARAMETER_POSTFIX);
+		IntParameterIndex != INDEX_NONE)
+	{
+		return false;
+	}
+		
 	// Opacity.	
+	if (const int32 FloatParameterIndex = Descriptor.FindFloatParameterNameIndex(ParamName + OPACITY_PARAMETER_POSTFIX);
+		FloatParameterIndex != INDEX_NONE)
 	{
-		const int32 FloatParameterIndex = Descriptor.FindFloatParameterNameIndex(ParamNameString + OPACITY_PARAMETER_POSTFIX);
-		if (FloatParameterIndex < 0)
-		{
-			return false;
-		}
+		return false;
 	}
 
 	return true;
@@ -528,3 +524,4 @@ uint32 GetTypeHash(const FMultilayerProjector& Key)
 	
 	return Hash;
 }
+
diff --git a/Engine/Plugins/Experimental/Mutable/Source/CustomizableObject/Public/MuCO/CustomizableObjectCustomVersion.h b/Engine/Plugins/Experimental/Mutable/Source/CustomizableObject/Public/MuCO/CustomizableObjectCustomVersion.h
--- a/Engine/Plugins/Experimental/Mutable/Source/CustomizableObject/Public/MuCO/CustomizableObjectCustomVersion.h
+++ b/Engine/Plugins/Experimental/Mutable/Source/CustomizableObject/Public/MuCO/CustomizableObjectCustomVersion.h
@@ -63,7 +63,9 @@ struct CUSTOMIZABLEOBJECT_API FCustomizableObjectCustomVersion
 		DeformSkeletonOptionsAdded,
 
 		GroupProjectorImagePin,
-		
+
+		GroupProjectorImagePinRemoved,
+
 		// -----<new versions can be added above this line>-------------------------------------------------
 		VersionPlusOne,
 		LatestVersion = VersionPlusOne - 1
diff --git a/Engine/Plugins/Experimental/Mutable/Source/CustomizableObject/Public/MuCO/MultilayerProjector.h b/Engine/Plugins/Experimental/Mutable/Source/CustomizableObject/Public/MuCO/MultilayerProjector.h
--- a/Engine/Plugins/Experimental/Mutable/Source/CustomizableObject/Public/MuCO/MultilayerProjector.h
+++ b/Engine/Plugins/Experimental/Mutable/Source/CustomizableObject/Public/MuCO/MultilayerProjector.h
@@ -96,13 +96,15 @@ struct CUSTOMIZABLEOBJECT_API FMultilayerProjector
 	friend FMultilayerProjectorLayer;
 	friend uint32 GetTypeHash(const FMultilayerProjector& Key);
 
+	// Log texts
+	static const FString DESCRIPTOR_PARAMETERS_INVALID;
 
 	// Parameters encoding
 	static const FString NUM_LAYERS_PARAMETER_POSTFIX;
 	static const FString OPACITY_PARAMETER_POSTFIX;
 	static const FString IMAGE_PARAMETER_POSTFIX;
 	static const FString POSE_PARAMETER_POSTFIX;
-
+	
 	// Constructors
 	FMultilayerProjector() = default;
 
@@ -192,12 +194,9 @@ private:
 	 * @param Index Index which the layer has been disabled.
 	 */
 	void UpdateMappingVirtualLayerDisabled(const FName& Id, int32 Index);
-
-	/** Rise an assert if the Instance Descriptor does not contain the necessary Parameters. */
-	void CheckDescriptorParameters(const FCustomizableObjectInstanceDescriptor& Descriptor) const;
 	
 	/** Return false if the Instance does not contain the necessary Instance Parameters. */
-	static bool AreDescriptorParametersValid(const FCustomizableObjectInstanceDescriptor& Descriptor, const FName& ParamName);
+	static bool AreDescriptorParametersValid(const FCustomizableObjectInstanceDescriptor& Descriptor, const FString& ParamName);
 };
 
 
diff --git a/Engine/Plugins/Experimental/Mutable/Source/CustomizableObjectEditor/Classes/MuCOE/Nodes/CustomizableObjectNodeGroupProjectorParameter.h b/Engine/Plugins/Experimental/Mutable/Source/CustomizableObjectEditor/Classes/MuCOE/Nodes/CustomizableObjectNodeGroupProjectorParameter.h
--- a/Engine/Plugins/Experimental/Mutable/Source/CustomizableObjectEditor/Classes/MuCOE/Nodes/CustomizableObjectNodeGroupProjectorParameter.h
+++ b/Engine/Plugins/Experimental/Mutable/Source/CustomizableObjectEditor/Classes/MuCOE/Nodes/CustomizableObjectNodeGroupProjectorParameter.h
@@ -73,7 +73,7 @@ public:
 	UPROPERTY(EditAnywhere, Category = ProjectorGroup, Meta = (ToolTip = "If true, projection textures will be shared between LODs of the same object, and will save memory. Only use if all the LODs share the same UV layout."))
 	bool bShareProjectionTexturesBetweenLODs = false;
 
-	UPROPERTY(EditAnywhere, Category = ProjectorGroup,  meta = (EditConditionHides))
+	UPROPERTY(EditAnywhere, Category = ProjectorGroup)
 	TArray<FGroupProjectorParameterImage> OptionImages;
 
 	UPROPERTY(EditAnywhere, Category = ProjectorGroup)
@@ -86,7 +86,7 @@ public:
 	/** Table where additional option images besides Option Images are read. The elements in this table have priority
 	* over elements from Option Images in case of duplicity. Use the "Data Table Texture Column Name" property to specify
 	* the name of the column where textures are read in the table. */
-	UPROPERTY(EditAnywhere, Category = ProjectorGroup, meta = (EditConditionHides))
+	UPROPERTY(EditAnywhere, Category = ProjectorGroup)
 	TObjectPtr<UDataTable> OptionImagesDataTable = nullptr;
 
 	UPROPERTY(EditAnywhere, Category = ProjectorGroup)
@@ -106,12 +106,5 @@ public:
 	// UCustomizableObjectNode interface.
 	virtual void BackwardsCompatibleFixup() override;
 	virtual void AllocateDefaultPins(UCustomizableObjectNodeRemapPins* RemapPins);
-
-	// Own interface.
-	UEdGraphPin& GetImagePin() const;
-
-private:
-	UPROPERTY()
-	FEdGraphPinReference ImagePin;
 };
 
diff --git a/Engine/Plugins/Experimental/Mutable/Source/CustomizableObjectEditor/Private/MuCOE/CustomizableObjectEditorModule.cpp b/Engine/Plugins/Experimental/Mutable/Source/CustomizableObjectEditor/Private/MuCOE/CustomizableObjectEditorModule.cpp
--- a/Engine/Plugins/Experimental/Mutable/Source/CustomizableObjectEditor/Private/MuCOE/CustomizableObjectEditorModule.cpp
+++ b/Engine/Plugins/Experimental/Mutable/Source/CustomizableObjectEditor/Private/MuCOE/CustomizableObjectEditorModule.cpp
@@ -50,7 +50,6 @@
 #include "MuCOE/Nodes/CustomizableObjectNodePinViewerDetails.h"
 #include "MuCOE/Nodes/CustomizableObjectNodeProjectorConstant.h"
 #include "MuCOE/Nodes/CustomizableObjectNodeProjectorParameterDetails.h"
-#include "MuCOE/Nodes/CustomizableObjectNodeGroupProjectorParameterDetails.h"
 #include "MuCOE/Nodes/CustomizableObjectNodeRemoveMeshBlocks.h"
 #include "MuCOE/Nodes/CustomizableObjectNodeRemoveMeshBlocksDetails.h"
 #include "MuCOE/Nodes/CustomizableObjectNodeSkeletalMesh.h"
@@ -174,7 +173,6 @@ void FCustomizableObjectEditorModule::StartupModule()
 	RegisterCustomDetails(PropertyModule, UCustomizableObjectNodeObject::StaticClass(), FOnGetDetailCustomizationInstance::CreateStatic(&FCustomizableObjectNodeObjectDetails::MakeInstance));
 	RegisterCustomDetails(PropertyModule, UCustomizableObjectNodeObjectGroup::StaticClass(), FOnGetDetailCustomizationInstance::CreateStatic(&FCustomizableObjectNodeObjectGroupDetails::MakeInstance));
 	RegisterCustomDetails(PropertyModule, UCustomizableObjectNodeProjectorParameter::StaticClass(), FOnGetDetailCustomizationInstance::CreateStatic(&FCustomizableObjectNodeProjectorParameterDetails::MakeInstance));
-	RegisterCustomDetails(PropertyModule, UCustomizableObjectNodeGroupProjectorParameter::StaticClass(), FOnGetDetailCustomizationInstance::CreateStatic(&FCustomizableObjectNodeGroupProjectorParameterDetails::MakeInstance));
 	RegisterCustomDetails(PropertyModule, UCustomizableObjectNodeProjectorConstant::StaticClass(), FOnGetDetailCustomizationInstance::CreateStatic(&FCustomizableObjectNodeProjectorParameterDetails::MakeInstance));
 	RegisterCustomDetails(PropertyModule, UCustomizableObjectNodeMeshMorph::StaticClass(), FOnGetDetailCustomizationInstance::CreateStatic(&FCustomizableObjectNodeMeshMorphDetails::MakeInstance));
 	RegisterCustomDetails(PropertyModule, UCustomizableObjectNodeMeshClipMorph::StaticClass(), FOnGetDetailCustomizationInstance::CreateStatic(&FCustomizableObjectNodeMeshClipMorphDetails::MakeInstance));
diff --git a/Engine/Plugins/Experimental/Mutable/Source/CustomizableObjectEditor/Private/MuCOE/GenerateMutableSource/GenerateMutableSourceGroupProjector.cpp b/Engine/Plugins/Experimental/Mutable/Source/CustomizableObjectEditor/Private/MuCOE/GenerateMutableSource/GenerateMutableSourceGroupProjector.cpp
--- a/Engine/Plugins/Experimental/Mutable/Source/CustomizableObjectEditor/Private/MuCOE/GenerateMutableSource/GenerateMutableSourceGroupProjector.cpp
+++ b/Engine/Plugins/Experimental/Mutable/Source/CustomizableObjectEditor/Private/MuCOE/GenerateMutableSource/GenerateMutableSourceGroupProjector.cpp
@@ -509,68 +509,62 @@ bool GenerateMutableSourceGroupProjector(const UEdGraphPin* Pin, FMutableGraphGe
 					GroupProjectorTempData.PoseBoneDataArray[PoseIndex].ArrayBoneName, GroupProjectorTempData.PoseBoneDataArray[PoseIndex].ArrayTransform);
 			}
 
-			if (UEdGraphPin* ConnectedPin = FollowInputPin(ProjParamNode->GetImagePin()))
-			{
-				GroupProjectorTempData.NodeImagePtr = GenerateMutableSourceImage(ConnectedPin, GenerationContext, 0, NodeRangeFromScalar);
-			}
-			else
+		
+			mu::NodeScalarEnumParameterPtr EnumParameterNode = new mu::NodeScalarEnumParameter;
+			FString NodeEnumParamName = ProjParamNode->ParameterName + FMultilayerProjector::IMAGE_PARAMETER_POSTFIX;
+			EnumParameterNode->SetName(TCHAR_TO_ANSI(*(NodeEnumParamName)));
+			EnumParameterNode->SetUid(TCHAR_TO_ANSI(*(ProjectorParamUid + FString("_SI"))));
+			GenerationContext.AddParameterNameUnique(originalGroup, NodeEnumParamName);
+			EnumParameterNode->SetValueCount(ArrayOptionImage.Num());
+			EnumParameterNode->SetDefaultValueIndex(0);
+			EnumParameterNode->SetRangeCount(1);
+			EnumParameterNode->SetRange(0, NodeRangeFromScalar);
+
+			FParameterUIData ParameterUIData(NodeEnumParamName, ProjParamNode->ParamUIMetadata, EMutableParameterType::Int);
+			ParameterUIData.IntegerParameterGroupType = ECustomizableObjectGroupType::COGT_ONE;
+			ParameterUIData.ParamUIMetadata.ExtraInformation.Add(FString("UseThumbnails"));
+
+			for (int ImageIndex = 0; ImageIndex < ArrayOptionImage.Num(); ++ImageIndex)
 			{
-				mu::NodeScalarEnumParameterPtr EnumParameterNode = new mu::NodeScalarEnumParameter;
-				FString NodeEnumParamName = ProjParamNode->ParameterName + FMultilayerProjector::IMAGE_PARAMETER_POSTFIX;
-				EnumParameterNode->SetName(TCHAR_TO_ANSI(*(NodeEnumParamName)));
-				EnumParameterNode->SetUid(TCHAR_TO_ANSI(*(ProjectorParamUid + FString("_SI"))));
-				GenerationContext.AddParameterNameUnique(originalGroup, NodeEnumParamName);
-				EnumParameterNode->SetValueCount(ArrayOptionImage.Num());
-				EnumParameterNode->SetDefaultValueIndex(0);
-				EnumParameterNode->SetRangeCount(1);
-				EnumParameterNode->SetRange(0, NodeRangeFromScalar);
-
-				FParameterUIData ParameterUIData(NodeEnumParamName, ProjParamNode->ParamUIMetadata, EMutableParameterType::Int);
-				ParameterUIData.IntegerParameterGroupType = ECustomizableObjectGroupType::COGT_ONE;
-				ParameterUIData.ParamUIMetadata.ExtraInformation.Add(FString("UseThumbnails"));
-
-				for (int ImageIndex = 0; ImageIndex < ArrayOptionImage.Num(); ++ImageIndex)
-				{
-					EnumParameterNode->SetValue(ImageIndex, (float)ImageIndex, TCHAR_TO_ANSI(*ArrayOptionImage[ImageIndex].OptionName));
+				EnumParameterNode->SetValue(ImageIndex, (float)ImageIndex, TCHAR_TO_ANSI(*ArrayOptionImage[ImageIndex].OptionName));
 
-					FMutableParamUIMetadata optionMetadata = ParameterUIData.ParamUIMetadata;
-					optionMetadata.UIThumbnail = ArrayOptionImage[ImageIndex].OptionImage;
-					ParameterUIData.ArrayIntegerParameterOption.Add(FIntegerParameterUIData(
-						ArrayOptionImage[ImageIndex].OptionName,
-						optionMetadata));
-				}
+				FMutableParamUIMetadata optionMetadata = ParameterUIData.ParamUIMetadata;
+				optionMetadata.UIThumbnail = ArrayOptionImage[ImageIndex].OptionImage;
+				ParameterUIData.ArrayIntegerParameterOption.Add(FIntegerParameterUIData(
+					ArrayOptionImage[ImageIndex].OptionName,
+					optionMetadata));
+			}
 
-				GenerationContext.ParameterUIDataMap.Add(NodeEnumParamName, ParameterUIData);
+			GenerationContext.ParameterUIDataMap.Add(NodeEnumParamName, ParameterUIData);
 
-				mu::NodeImageSwitchPtr SwitchNode = new mu::NodeImageSwitch;
-				SwitchNode->SetParameter(EnumParameterNode);
-				SwitchNode->SetOptionCount(ArrayOptionImage.Num());
+			mu::NodeImageSwitchPtr SwitchNode = new mu::NodeImageSwitch;
+			SwitchNode->SetParameter(EnumParameterNode);
+			SwitchNode->SetOptionCount(ArrayOptionImage.Num());
 
-				bool bFoundUnlinkedPin = false; 
+			bool bFoundUnlinkedPin = false; 
 
-				for (int SelectorIndex = 0; SelectorIndex < ArrayOptionImage.Num(); ++SelectorIndex)
+			for (int SelectorIndex = 0; SelectorIndex < ArrayOptionImage.Num(); ++SelectorIndex)
+			{
+				if (ArrayOptionImage[SelectorIndex].OptionImage)
 				{
-					if (ArrayOptionImage[SelectorIndex].OptionImage)
-					{
-						mu::NodeImageConstantPtr ImageNode = new mu::NodeImageConstant();
-						GenerationContext.ArrayTextureUnrealToMutableTask.Add(FTextureUnrealToMutableTask(ImageNode, ArrayOptionImage[SelectorIndex].OptionImage, ProjParamNode));
-						SwitchNode->SetOption(SelectorIndex, ImageNode);
-					}
-					else
-					{
-						bFoundUnlinkedPin = true;
-					}
+					mu::NodeImageConstantPtr ImageNode = new mu::NodeImageConstant();
+					GenerationContext.ArrayTextureUnrealToMutableTask.Add(FTextureUnrealToMutableTask(ImageNode, ArrayOptionImage[SelectorIndex].OptionImage, ProjParamNode));
+					SwitchNode->SetOption(SelectorIndex, ImageNode);
 				}
-
-				if (bFoundUnlinkedPin)
+				else
 				{
-					FString msg = FString::Printf(TEXT("The group projection node must have a texture for all the options. Please set a texture for all the options."));
-					GenerationContext.Compiler->CompilerLog(FText::FromString(msg), ProjParamNode);
+					bFoundUnlinkedPin = true;
 				}
+			}
 
-				GroupProjectorTempData.NodeImagePtr = SwitchNode;
+			if (bFoundUnlinkedPin)
+			{
+				FString msg = FString::Printf(TEXT("The group projection node must have a texture for all the options. Please set a texture for all the options."));
+				GenerationContext.Compiler->CompilerLog(FText::FromString(msg), ProjParamNode);
 			}
 
+			GroupProjectorTempData.NodeImagePtr = SwitchNode;
+
 			GenerationContext.ProjectorGroupMap.Add(originalGroup, GroupProjectorTempData);
 		}
 		else
diff --git a/Engine/Plugins/Experimental/Mutable/Source/CustomizableObjectEditor/Private/MuCOE/GenerateMutableSource/GenerateMutableSourceImage.cpp b/Engine/Plugins/Experimental/Mutable/Source/CustomizableObjectEditor/Private/MuCOE/GenerateMutableSource/GenerateMutableSourceImage.cpp
--- a/Engine/Plugins/Experimental/Mutable/Source/CustomizableObjectEditor/Private/MuCOE/GenerateMutableSource/GenerateMutableSourceImage.cpp
+++ b/Engine/Plugins/Experimental/Mutable/Source/CustomizableObjectEditor/Private/MuCOE/GenerateMutableSource/GenerateMutableSourceImage.cpp
@@ -121,7 +121,7 @@ mu::NodeImagePtr ResizeToMaxTextureSize(float MaxTextureSize, const UTexture2D*
 }
 
 
-mu::NodeImagePtr GenerateMutableSourceImage(const UEdGraphPin* Pin, FMutableGraphGenerationContext& GenerationContext, float MaxTextureSize, mu::NodeRangePtr NodeRange)
+mu::NodeImagePtr GenerateMutableSourceImage(const UEdGraphPin* Pin, FMutableGraphGenerationContext& GenerationContext, float MaxTextureSize)
 {
 	check(Pin)
 	RETURN_ON_CYCLE(*Pin, GenerationContext)
@@ -177,12 +177,6 @@ mu::NodeImagePtr GenerateMutableSourceImage(const UEdGraphPin* Pin, FMutableGrap
 
 		TextureNode->SetName(TCHAR_TO_ANSI(*TypedNodeParam->ParameterName));
 		TextureNode->SetUid(TCHAR_TO_ANSI(*GenerationContext.GetNodeIdUnique(Node).ToString()));
-		
-		if (NodeRange)
-		{
-			TextureNode->SetRangeCount(1);
-			TextureNode->SetRange(0, NodeRange);
-		}
 
 		// TODO: Set a default value for the texture parameter?
 		//ColorNode->SetDefaultValue(TypedNodeParam->DefaultValue.R, TypedNodeParam->DefaultValue.G, TypedNodeParam->DefaultValue.B);
diff --git a/Engine/Plugins/Experimental/Mutable/Source/CustomizableObjectEditor/Private/MuCOE/GenerateMutableSource/GenerateMutableSourceImage.h b/Engine/Plugins/Experimental/Mutable/Source/CustomizableObjectEditor/Private/MuCOE/GenerateMutableSource/GenerateMutableSourceImage.h
--- a/Engine/Plugins/Experimental/Mutable/Source/CustomizableObjectEditor/Private/MuCOE/GenerateMutableSource/GenerateMutableSourceImage.h
+++ b/Engine/Plugins/Experimental/Mutable/Source/CustomizableObjectEditor/Private/MuCOE/GenerateMutableSource/GenerateMutableSourceImage.h
@@ -5,7 +5,6 @@
 
 #include "MuR/Image.h"
 #include "MuT/NodeImageConstant.h"
-#include "MuT/NodeRange.h"
 
 class FCustomizableObjectCompiler;
 class UCustomizableObjectNode;
@@ -13,6 +12,7 @@ class UEdGraphPin;
 class UTexture2D;
 struct FMutableGraphGenerationContext;
 
+
 mu::ImagePtr ConvertTextureUnrealToMutable(UTexture2D* Texture, const UCustomizableObjectNode* Node, FCustomizableObjectCompiler* Compiler, bool bIsNormalComposite);
 
 
@@ -20,4 +20,4 @@ mu::NodeImagePtr ResizeToMaxTextureSize(float MaxTextureSize, const UTexture2D*
 
 
 /** Convert a CustomizableObject Source Graph into a mutable source graph. */
-mu::NodeImagePtr GenerateMutableSourceImage(const UEdGraphPin* Pin, FMutableGraphGenerationContext& GenerationContext, float MaxTextureSize, mu::NodeRangePtr NodeRange = nullptr);
+mu::NodeImagePtr GenerateMutableSourceImage(const UEdGraphPin* Pin, FMutableGraphGenerationContext& GenerationContext, float MaxTextureSize);
diff --git a/Engine/Plugins/Experimental/Mutable/Source/CustomizableObjectEditor/Private/MuCOE/Nodes/CustomizableObjectNodeGroupProjectorParameter.cpp b/Engine/Plugins/Experimental/Mutable/Source/CustomizableObjectEditor/Private/MuCOE/Nodes/CustomizableObjectNodeGroupProjectorParameter.cpp
--- a/Engine/Plugins/Experimental/Mutable/Source/CustomizableObjectEditor/Private/MuCOE/Nodes/CustomizableObjectNodeGroupProjectorParameter.cpp
+++ b/Engine/Plugins/Experimental/Mutable/Source/CustomizableObjectEditor/Private/MuCOE/Nodes/CustomizableObjectNodeGroupProjectorParameter.cpp
@@ -120,6 +120,11 @@ void UCustomizableObjectNodeGroupProjectorParameter::BackwardsCompatibleFixup()
 	{
 		ReconstructNode();		
 	}
+
+	if (CustomizableObjectCustomVersion < FCustomizableObjectCustomVersion::GroupProjectorImagePinRemoved)
+	{
+		ReconstructNode();		
+	}
 }
 
 
@@ -128,13 +133,6 @@ void UCustomizableObjectNodeGroupProjectorParameter::AllocateDefaultPins(UCustom
 	const UEdGraphSchema_CustomizableObject* Schema = GetDefault<UEdGraphSchema_CustomizableObject>();
 
 	CustomCreatePin(EGPD_Output, Schema->PC_GroupProjector, TEXT("Value"));	
-	ImagePin = CustomCreatePin(EGPD_Input, Schema->PC_Image, TEXT("Texture"));
-}
-
-
-UEdGraphPin& UCustomizableObjectNodeGroupProjectorParameter::GetImagePin() const
-{
-	return *ImagePin.Get();
 }
 
 
diff --git a/Engine/Plugins/Experimental/Mutable/Source/CustomizableObjectEditor/Private/MuCOE/Nodes/CustomizableObjectNodeGroupProjectorParameterDetails.cpp b/Engine/Plugins/Experimental/Mutable/Source/CustomizableObjectEditor/Private/MuCOE/Nodes/CustomizableObjectNodeGroupProjectorParameterDetails.cpp
deleted file mode 100644
diff --git a/Engine/Plugins/Experimental/Mutable/Source/CustomizableObjectEditor/Private/MuCOE/Nodes/CustomizableObjectNodeGroupProjectorParameterDetails.h b/Engine/Plugins/Experimental/Mutable/Source/CustomizableObjectEditor/Private/MuCOE/Nodes/CustomizableObjectNodeGroupProjectorParameterDetails.h
deleted file mode 100644
diff --git a/Engine/Plugins/Experimental/Mutable/Source/CustomizableObjectEditor/Private/MuCOE/SCustomizableInstanceProperties.cpp b/Engine/Plugins/Experimental/Mutable/Source/CustomizableObjectEditor/Private/MuCOE/SCustomizableInstanceProperties.cpp
--- a/Engine/Plugins/Experimental/Mutable/Source/CustomizableObjectEditor/Private/MuCOE/SCustomizableInstanceProperties.cpp
+++ b/Engine/Plugins/Experimental/Mutable/Source/CustomizableObjectEditor/Private/MuCOE/SCustomizableInstanceProperties.cpp
@@ -1018,30 +1018,34 @@ void SCustomizableInstanceProperties::AddParameter(int32 ParamIndexInObject)
 						];
 				}
 
-
-				FString TextureSwitchEnumParamName = ParamName + FMultilayerProjector::IMAGE_PARAMETER_POSTFIX;
-
-				// Add the necessary ranges. Ranges could be lower if in the previous compilation the Group Projector used the Texture pin.
-				if (int32 ParameterIndex = CustomInstance->FindIntParameterNameIndex(TextureSwitchEnumParamName);
-					ParameterIndex != INDEX_NONE)
-				{
-					const int32 ProjectorRange = ProjectorParameters[ProjectorParamIndex].RangeValues.Num();
-					const int32 ImagesRange = CustomInstance->GetIntParameters()[ParameterIndex].ParameterRangeValueNames.Num();
-					for (int32 Range = 0; Range < ProjectorRange - ImagesRange; ++Range)
-					{
-						CustomInstance->AddValueToIntRange(TextureSwitchEnumParamName);
-					}
-				}
-
+				const FString TextureSwitchEnumParamName = ParamName + FMultilayerProjector::IMAGE_PARAMETER_POSTFIX;
+				
 				for (int32 i = 0; i < ProjectorParameters[ProjectorParamIndex].RangeValues.Num(); ++i)
 				{
-					FString ParamNameWithIndex = ParamName + FString::Printf(TEXT("__%d"), i);
-					
-					FString TextureSwitchEnumParamNameWithRange = TextureSwitchEnumParamName + FString::Printf(TEXT("__%d"), i);
+					const FString ParamNameWithIndex = ParamName + FString::Printf(TEXT("__%d"), i);
+
+					const FString TextureSwitchEnumParamNameWithRange = TextureSwitchEnumParamName + FString::Printf(TEXT("__%d"), i);
+					const int32 TextureSwitchEnumParamIndexInObject = CustomizableObject->FindParameter(TextureSwitchEnumParamName);
+					check(TextureSwitchEnumParamIndexInObject >= 0);
+					int32 NumValues = CustomizableObject->GetIntParameterNumOptions(TextureSwitchEnumParamIndexInObject);
 
 					FString OpacitySliderParamName = ParamName + FMultilayerProjector::OPACITY_PARAMETER_POSTFIX;
 					FString OpacitySliderParamNameWithRange = OpacitySliderParamName + FString::Printf(TEXT("__%d"), i);
 
+					TArray<FString> OptionNamesAttribute;
+					FString Value = GetIntParameterValue(TextureSwitchEnumParamName, i);
+					int32 ValueIndex = 0;
+
+					for (int32 CandidateIndex = 0; CandidateIndex < NumValues; ++CandidateIndex)
+					{
+						FString PossibleValue = CustomizableObject->GetIntParameterAvailableOption(TextureSwitchEnumParamIndexInObject, CandidateIndex);
+						if (PossibleValue == Value)
+						{
+							ValueIndex = CandidateIndex;
+						}
+						OptionNamesAttribute.Add(CustomizableObject->GetIntParameterAvailableOption(TextureSwitchEnumParamIndexInObject, CandidateIndex));
+					}
+					
 					TSharedPtr<SHorizontalBox> SliderBox;
 					TSharedPtr<SSpinBox<float>> Slider;
 					TSharedPtr<SHorizontalBox> PropertiesHB;
@@ -1084,6 +1088,21 @@ void SCustomizableInstanceProperties::AddParameter(int32 ParamIndexInObject)
 									]
 								]
 							]
+
+							+ SHorizontalBox::Slot()
+							.Padding(1, 0)
+							[
+								SNew(SBox)
+								.MinDesiredWidth(120.f)
+								[
+									SNew(SMutableTextSearchBox)
+									.PossibleSuggestions(OptionNamesAttribute)
+									.InitialText(FText::FromString(OptionNamesAttribute[ValueIndex]))
+									.MustMatchPossibleSuggestions(TAttribute<bool>(true))
+									.SuggestionListPlacement(EMenuPlacement::MenuPlacement_ComboBox)
+									.OnTextCommitted(this, &SCustomizableInstanceProperties::OnProjectorTextureParameterComboBoxChanged, TextureSwitchEnumParamNameWithRange)
+								]
+							]
 						]
 					];
 						
@@ -1219,44 +1238,6 @@ void SCustomizableInstanceProperties::AddParameter(int32 ParamIndexInObject)
 					}
 
 					FloatSliders.Add(FSliderData(Slider, OpacitySliderParamName, i, GetFloatParameterValue(OpacitySliderParamName, i)));
-
-					if (CustomizableObject->FindParameter(TextureSwitchEnumParamName) != INDEX_NONE)
-					{
-						int32 TextureSwitchEnumParamIndexInObject = CustomizableObject->FindParameter(TextureSwitchEnumParamName);
-						check(TextureSwitchEnumParamIndexInObject >= 0);
-						int32 NumValues = CustomizableObject->GetIntParameterNumOptions(TextureSwitchEnumParamIndexInObject);
-
-						TArray<FString> OptionNamesAttribute;
-						int32 ValueIndex = 0;
-
-						FString Value = GetIntParameterValue(TextureSwitchEnumParamName, i);
-
-						for (int32 CandidateIndex = 0; CandidateIndex < NumValues; ++CandidateIndex)
-						{
-							FString PossibleValue = CustomizableObject->GetIntParameterAvailableOption(TextureSwitchEnumParamIndexInObject, CandidateIndex);
-							if (PossibleValue == Value)
-							{
-								ValueIndex = CandidateIndex;
-							}
-
-							OptionNamesAttribute.Add(CustomizableObject->GetIntParameterAvailableOption(TextureSwitchEnumParamIndexInObject, CandidateIndex));
-						}
-
-						PropertiesHB->AddSlot()
-							.Padding(1, 0)
-							[
-								SNew(SBox)
-								.MinDesiredWidth(120.f)
-								[
-									SNew(SMutableTextSearchBox)
-									.PossibleSuggestions(OptionNamesAttribute)
-									.InitialText(FText::FromString(OptionNamesAttribute[ValueIndex]))
-									.MustMatchPossibleSuggestions(TAttribute<bool>(true))
-									.SuggestionListPlacement(EMenuPlacement::MenuPlacement_ComboBox)
-									.OnTextCommitted(this, &SCustomizableInstanceProperties::OnProjectorTextureParameterComboBoxChanged, TextureSwitchEnumParamNameWithRange)
-								]
-							];
-					}
 				}
 			}
 		}
@@ -1787,12 +1768,9 @@ FReply SCustomizableInstanceProperties::OnProjectorLayerAdded(FString ParamName)
 	}
 
 	const FString TextureSwitchEnumParamName = ParamName + FMultilayerProjector::IMAGE_PARAMETER_POSTFIX;
-	if (CustomInstance->GetCustomizableObject()->FindParameter(TextureSwitchEnumParamName) != INDEX_NONE)
-	{
-		CustomInstance->AddValueToIntRange(TextureSwitchEnumParamName);
-		check(CustomInstance->FindIntParameterNameIndex(TextureSwitchEnumParamName) != INDEX_NONE);
-		check(NumLayers == CustomInstance->GetIntParameters()[CustomInstance->FindIntParameterNameIndex(TextureSwitchEnumParamName)].ParameterRangeValueNames.Num());
-	}
+	CustomInstance->AddValueToIntRange(TextureSwitchEnumParamName);
+	check(CustomInstance->FindIntParameterNameIndex(TextureSwitchEnumParamName) != INDEX_NONE);
+	check(NumLayers == CustomInstance->GetIntParameters()[CustomInstance->FindIntParameterNameIndex(TextureSwitchEnumParamName)].ParameterRangeValueNames.Num());
 
 	TArray<FCustomizableObjectFloatParameterValue>& FloatParameters = CustomInstance->GetFloatParameters();
 	
@@ -1836,12 +1814,9 @@ FReply SCustomizableInstanceProperties::OnProjectorLayerRemoved(FString ParamNam
 	CustomInstance->RemovedProjectorParameterNameWithIndex = ParamNameWithIndex; // Since the gizmo updates the CustomObjectInstance when unselected, a check is required to avoid updating the removed projector parameter range index
 
 	const FString TextureSwitchEnumParamName = ParamName + FMultilayerProjector::IMAGE_PARAMETER_POSTFIX;
-	if (CustomInstance->GetCustomizableObject()->FindParameter(TextureSwitchEnumParamName) != INDEX_NONE)
-	{
-		CustomInstance->RemoveValueFromIntRange(TextureSwitchEnumParamName, RangeIndex);
-		check(CustomInstance->FindIntParameterNameIndex(TextureSwitchEnumParamName) != INDEX_NONE);
-		check(NumLayers == CustomInstance->GetIntParameters()[CustomInstance->FindIntParameterNameIndex(TextureSwitchEnumParamName)].ParameterRangeValueNames.Num());
-	}
+	CustomInstance->RemoveValueFromIntRange(TextureSwitchEnumParamName, RangeIndex);
+	check(CustomInstance->FindIntParameterNameIndex(TextureSwitchEnumParamName) != INDEX_NONE);
+	check(NumLayers == CustomInstance->GetIntParameters()[CustomInstance->FindIntParameterNameIndex(TextureSwitchEnumParamName)].ParameterRangeValueNames.Num());
 
 	TArray<FCustomizableObjectFloatParameterValue>& FloatParameters = CustomInstance->GetFloatParameters();
 
