[mutable] Mutable Parameters Widget : Added suport for multidimensional parameters
- Currently we support only the first dimension of each range index object.
- Updating the value of an Int or Float parameter that is known to control the amount of values of another one will trigger the reconstruction of the contents of the Parameters widget.

diff --git a/Engine/Plugins/Experimental/Mutable/Source/CustomizableObjectEditor/Private/MuCOE/SMutableParametersWidget.cpp b/Engine/Plugins/Experimental/Mutable/Source/CustomizableObjectEditor/Private/MuCOE/SMutableParametersWidget.cpp
--- a/Engine/Plugins/Experimental/Mutable/Source/CustomizableObjectEditor/Private/MuCOE/SMutableParametersWidget.cpp
+++ b/Engine/Plugins/Experimental/Mutable/Source/CustomizableObjectEditor/Private/MuCOE/SMutableParametersWidget.cpp
@@ -3,6 +3,10 @@
 #include "MuCOE/SMutableParametersWidget.h"
 
 #include "MuCOE/SMutableTextSearchBox.h"
+#include "MuR/ModelPrivate.h"
+#include "MuR/ParametersPrivate.h"
+#include "Templates/SharedPointer.h"
+#include "Widgets/SBoxPanel.h"
 #include "Widgets/Colors/SColorBlock.h"
 #include "Widgets/Colors/SColorPicker.h"
 #include "Widgets/Input/SVectorInputBox.h"
@@ -56,10 +60,11 @@ void SMutableParametersWidget::Tick(const FGeometry& AllottedGeometry, const dou
 	{
 		return;
 	}
-
+	
 	for ( int32 ParamIndex=0; ParamIndex< MutableParameters->GetCount(); ++ParamIndex )
 	{
 		FString ParamName = ANSI_TO_TCHAR(MutableParameters->GetName(ParamIndex));
+
 		TSharedPtr<SHorizontalBox> ParameterBox;
 
 		ParamBox->AddSlot()
@@ -70,6 +75,20 @@ void SMutableParametersWidget::Tick(const FGeometry& AllottedGeometry, const dou
 			[
 				SAssignNew(ParameterBox, SHorizontalBox)
 
+#if UE_BUILD_DEBUG
+				+ SHorizontalBox::Slot()
+				.AutoWidth()
+				.Padding(4,0)
+				.VAlign(VAlign_Center)
+				[
+					SNew(STextBlock)
+					.ColorAndOpacity(FSlateColor(FLinearColor(0.5f,0.5f,0.5f,1.0f)))
+					.Text(FText::FromString(FString::FromInt(ParamIndex)))
+					.Justification(ETextJustify::Left)
+					.AutoWrapText(true)
+				]
+#endif
+
 				+ SHorizontalBox::Slot()
 				.FillWidth(0.66f)
 				.VAlign(VAlign_Center)
@@ -80,172 +99,300 @@ void SMutableParametersWidget::Tick(const FGeometry& AllottedGeometry, const dou
 					.AutoWrapText(true)
 				]
 			];
-
-		switch (MutableParameters->GetType(ParamIndex) )
+		
+		
+		// Use the parameter descriptor to know what parameter value determines the amount of values this one has
+		mu::FParameterDesc& ParameterDescription = MutableParameters->GetPrivate()->m_pModel->GetPrivate()->m_program.
+															m_parameters[ParamIndex];
+		TArray<uint32> Ranges = ParameterDescription.m_ranges;
+		if (Ranges.Num() >= 1)
 		{
-		case mu::PARAMETER_TYPE::T_BOOL:
-			ParameterBox->AddSlot()
-				.Padding(4.0f)
-				.VAlign(VAlign_Center)
-				[
-					SNew(SCheckBox)
-					.Visibility(this, &SMutableParametersWidget::GetParameterVisibility, ParamIndex)
-					.IsChecked(this, &SMutableParametersWidget::GetBoolParameterValue, ParamIndex)
-					.OnCheckStateChanged(this, &SMutableParametersWidget::OnBoolParameterChanged, ParamIndex)
-				];
-			break;
-
-		case mu::PARAMETER_TYPE::T_FLOAT:
-			ParameterBox->AddSlot()
-				.Padding(4.0f)
-				.VAlign(VAlign_Center)
-				[
-					SNew(SNumericEntryBox<float>)
-					.Visibility(this, &SMutableParametersWidget::GetParameterVisibility, ParamIndex)
-					.AllowSpin(true)
-					.MinSliderValue(0.0f)
-					.MaxSliderValue(1.0f)
-					.Value(this, &SMutableParametersWidget::GetFloatParameterValue, ParamIndex)
-					.OnValueChanged(this, &SMutableParametersWidget::OnFloatParameterChanged, ParamIndex)
-				];
-			break;
-
-		case mu::PARAMETER_TYPE::T_COLOUR:
-			ParameterBox->AddSlot()
-				.Padding(4.0f)
-				.VAlign(VAlign_Center)
-				[
-					SNew(SColorBlock)
-					.Color(this, &SMutableParametersWidget::GetColorParameterValue, ParamIndex)
-					.ShowBackgroundForAlpha(false)
-					.AlphaDisplayMode(EColorBlockAlphaDisplayMode::Ignore)
-					.UseSRGB(false)
-					.OnMouseButtonDown(this, &SMutableParametersWidget::OnColorBlockMouseButtonDown, ParamIndex)
-					.Size(FVector2D(10.0f, 10.0f))
-				];
-			break;
+			if (Ranges.Num() > 1)
+			{
+				// todo: Only supports first dimension
+				UE_LOG(LogTemp,Warning,TEXT("Currently the debugger only suports showing the first dimension of each range."))
+			}
 
-		case mu::PARAMETER_TYPE::T_INT:
-		{
-			// If we have a list of options, add a combo box too
-			TSharedPtr<SMutableTextSearchBox> ParamComboBox;
-			int32 ValueCount = MutableParameters->GetIntPossibleValueCount(ParamIndex);
-			if (ValueCount > 0)
+			// This parameter does have more than one value (multivalue)
+			
+			// A range index can have more than one range
+			constexpr int32 Dimension = 0;		// Only operate over the first dimension
+			// Default amount of values to display
+			int32 PositionCount = 1;			
+			
+			// grab only first value (first dimension)
+			const uint32 RangeIdx = Ranges[Dimension];
+			// Get the parameter that controls the amount of values of current parameter
+			const int32 ParameterControllingValueCount = MutableParameters->GetPrivate()->m_pModel->GetPrivate()->m_program.m_ranges[RangeIdx].m_dimensionParameter;
+
+			// If a parameter index has been reported then we know that another parameter's value determines the amount of values we have
+			if (ParameterControllingValueCount >= 0)
 			{
-				FString ToolTipText = FString("None");
-				TArray<FString> OptionNamesAttribute;
-				int Value = MutableParameters->GetIntValue( ParamIndex );
-				int ValueIndex = MutableParameters->GetIntValueIndex( ParamIndex, Value );
-				for (int i = 0; i < ValueCount; ++i)
+				// Get the numerical value of that parameters
+				mu::PARAMETER_TYPE ParameterType = MutableParameters->GetType(ParameterControllingValueCount);
+				if (ParameterType == mu::PARAMETER_TYPE::T_INT)
 				{
-					const char* ValueText = MutableParameters->GetIntPossibleValueName( ParamIndex, i );
-					OptionNamesAttribute.Add( FString(ANSI_TO_TCHAR(ValueText)) );				
+					PositionCount = MutableParameters->GetIntValue(ParameterControllingValueCount);
+				}
+				else if (ParameterType == mu::PARAMETER_TYPE::T_FLOAT)
+				{
+					PositionCount = MutableParameters->GetFloatValue(ParameterControllingValueCount);
+				}
+				else
+				{
+					checkNoEntry();
+					// Unable to get the amount of options of a non numerical value
 				}
 
-				ParameterBox->AddSlot()
-					.Padding(4.0f)
-					.VAlign(VAlign_Center)
+				// Show one value at least.
+				if (PositionCount < 1)
+				{
+					PositionCount = 1;
+				}
+			}
+
+			// If no parameter is controlling the amount of values of this one or if the amount of values is 0 the
+			// amount of values that will therefore be used will be 1.
+
+			// Slate where to place all the values for each of the positions of the current parameter at the target dimension
+			TSharedPtr<SVerticalBox> ParameterValuesCollection = SNew(SVerticalBox);
+			
+			// Iterate over all the values for the parameter and show them. Updating their value is managed by the UI
+			for (int32 Position = 0; Position < PositionCount; Position++)
+			{
+				mu::RangeIndexPtr RangeIndex = MutableParameters->NewRangeIndex(ParamIndex);
+				RangeIndex->SetPosition(Dimension,Position);
+			
+				// Show position data if more than one is set to avoid UI clutter
+				if (PositionCount > 1)
+				{
+					// Add identifying data on top
+					ParameterValuesCollection->AddSlot()
+					.AutoHeight()
+					.HAlign(EHorizontalAlignment::HAlign_Right)
+					.Padding(4,4,4,2)
 					[
-						SAssignNew(ParamComboBox, SMutableTextSearchBox)
-						.Visibility(this, &SMutableParametersWidget::GetParameterVisibility, ParamIndex)
-						.ToolTipText(FText::FromString(ToolTipText))
-						.PossibleSuggestions(OptionNamesAttribute)
-						.InitialText(FText::FromString(OptionNamesAttribute[ValueIndex]))
-						.MustMatchPossibleSuggestions(TAttribute<bool>(true))
-						.SuggestionListPlacement(EMenuPlacement::MenuPlacement_ComboBox)
-						.OnTextCommitted(this, &SMutableParametersWidget::OnIntParameterTextChanged, ParamIndex)
+						SNew(STextBlock)
+						.Text(FText::FromString(FString::Printf(TEXT("Value : %i"), Position )))
+						.Justification(ETextJustify::Right)
+						.AutoWrapText(true)
 					];
-			}
+				}
+			
+				// Generate a new temporary slate container to be later added to the vertical container for all dimensions and positions
+				TSharedPtr<SHorizontalBox> ParameterProxyContainer = SNew(SHorizontalBox);
 
-			ParameterBox->AddSlot()
-				.Padding(4.0f)
-				.VAlign(VAlign_Center)
+				// Process the parameter 
+				GenerateAndAttachParameterSlate(ParamIndex,ParameterProxyContainer,RangeIndex);
+
+				// Add the horizontal slot we have used to place the proxy now onto our vertical box (so it is placed inside a vertical slot)
+				ParameterValuesCollection->AddSlot()
+				.AutoHeight()
 				[
-					SNew(SNumericEntryBox<int32>)
-					.Visibility(this, &SMutableParametersWidget::GetParameterVisibility, ParamIndex)
-					.AllowSpin(true)
-					.MinSliderValue(0)
-					.MaxSliderValue(this, &SMutableParametersWidget::GetIntParameterValueMax, ParamIndex)
-					.Value(this, &SMutableParametersWidget::GetIntParameterValue, ParamIndex)
-					.OnValueChanged(this, &SMutableParametersWidget::OnIntParameterChanged, ParamIndex, ParamComboBox)
+					ParameterProxyContainer.ToSharedRef()
 				];
-			break;
+			}
+		
+			// Add the collection of parameter value objects to the ParameterBox
+			ParameterBox->AddSlot()
+			.FillWidth(1)
+			.Padding(4)
+			.HAlign(HAlign_Fill)
+			[
+				ParameterValuesCollection.ToSharedRef()
+			];
+
 		}
+		else
+		{
+			// The parameter only can have one value (no multivalue) so draw it normally
+			GenerateAndAttachParameterSlate(ParamIndex,ParameterBox,nullptr);
+		}
+		
+	}
+}
 
-		case mu::PARAMETER_TYPE::T_PROJECTOR:
-			ParameterBox->AddSlot()
+
+void SMutableParametersWidget::ScheduleUpdateIfRequired(const int32& InParameterIndex)
+{
+	// Update the ui to display the new amount of values we want to show for the multidimensional values if this parameter's value
+	// controls the amount of values another parameter has.
+	const bool bIsRangeSize = MutableParameters->GetPrivate()->m_pModel->GetPrivate()->m_program.m_ranges.ContainsByPredicate([InParameterIndex](const mu::FRangeDesc& r)
+		{ return r.m_dimensionParameter == InParameterIndex; });
+
+	// Schedule update since we know the value does change the amount of values shown by another parameter (multivalue parameter)
+	if (bIsRangeSize && !bIsPendingUpdate)
+	{
+		bIsPendingUpdate = true;
+	}
+}
+
+void SMutableParametersWidget::GenerateAndAttachParameterSlate(const int32 ParamIndex, TSharedPtr<SHorizontalBox> ParameterHorizontalBox,  mu::RangeIndexPtrConst RangeIndex)
+{
+	switch (MutableParameters->GetType(ParamIndex) )
+	{
+	case mu::PARAMETER_TYPE::T_BOOL:
+		ParameterHorizontalBox->AddSlot()
+			.Padding(4.0f)
+			.VAlign(VAlign_Center)
+			[
+				SNew(SCheckBox)
+				.Visibility(this, &SMutableParametersWidget::GetParameterVisibility, ParamIndex)
+				.IsChecked(this, &SMutableParametersWidget::GetBoolParameterValue, ParamIndex, RangeIndex)
+				.OnCheckStateChanged(this, &SMutableParametersWidget::OnBoolParameterChanged, ParamIndex, RangeIndex)
+			];
+		break;
+
+	case mu::PARAMETER_TYPE::T_FLOAT:
+		ParameterHorizontalBox->AddSlot()
+			.Padding(4.0f)
+			.VAlign(VAlign_Center)
+			[
+				SNew(SNumericEntryBox<float>)
+				.Visibility(this, &SMutableParametersWidget::GetParameterVisibility, ParamIndex)
+				.AllowSpin(true)
+				.MinSliderValue(0.0f)
+				.MaxSliderValue(1.0f)			// TODO: Use actual max possible value instead of hard-coding it
+				.Value(this, &SMutableParametersWidget::GetFloatParameterValue, ParamIndex, RangeIndex)
+				.OnValueChanged(this, &SMutableParametersWidget::OnFloatParameterChanged, ParamIndex, RangeIndex)
+				.OnValueCommitted(this,&SMutableParametersWidget::OnFloatParameterCommitted,ParamIndex)
+			];
+		break;
+
+	case mu::PARAMETER_TYPE::T_COLOUR:
+		ParameterHorizontalBox->AddSlot()
+			.Padding(4.0f)
+			.VAlign(VAlign_Center)
+			[
+				SNew(SColorBlock)
+				.Color(this, &SMutableParametersWidget::GetColorParameterValue, ParamIndex, RangeIndex)
+				.ShowBackgroundForAlpha(false)
+				.AlphaDisplayMode(EColorBlockAlphaDisplayMode::Ignore)
+				.UseSRGB(false)
+				.OnMouseButtonDown(this, &SMutableParametersWidget::OnColorBlockMouseButtonDown, ParamIndex, RangeIndex)
+				.Size(FVector2D(10.0f, 10.0f))
+			];
+		break;
+
+	case mu::PARAMETER_TYPE::T_INT:
+	{
+		// If we have a list of options, add a combo box too
+		TSharedPtr<SMutableTextSearchBox> ParamComboBox;
+		const int32 ValueCount = MutableParameters->GetIntPossibleValueCount(ParamIndex);
+		if (ValueCount > 0)
+		{
+			const FString ToolTipText = FString("None");
+			TArray<FString> OptionNamesAttribute;
+			const int32 Value = MutableParameters->GetIntValue( ParamIndex , RangeIndex);
+			const int32 ValueIndex = MutableParameters->GetIntValueIndex( ParamIndex, Value );
+			for (int32 i = 0; i < ValueCount; ++i)
+			{
+				const char* ValueText = MutableParameters->GetIntPossibleValueName( ParamIndex, i );
+				OptionNamesAttribute.Add( FString(ANSI_TO_TCHAR(ValueText)) );				
+			}
+
+			ParameterHorizontalBox->AddSlot()
 				.Padding(4.0f)
 				.VAlign(VAlign_Center)
 				[
-					SNew(SGridPanel)
-					.FillColumn(0, 1.0f)
+					SAssignNew(ParamComboBox, SMutableTextSearchBox)
+					.Visibility(this, &SMutableParametersWidget::GetParameterVisibility, ParamIndex)
+					.ToolTipText(FText::FromString(ToolTipText))
+					.PossibleSuggestions(OptionNamesAttribute)
+					.InitialText(FText::FromString(OptionNamesAttribute[ValueIndex]))
+					.MustMatchPossibleSuggestions(TAttribute<bool>(true))
+					.SuggestionListPlacement(EMenuPlacement::MenuPlacement_ComboBox)
+					.OnTextCommitted(this, &SMutableParametersWidget::OnIntParameterTextChanged, ParamIndex,RangeIndex)
+				];
+		}
 
-					+SGridPanel::Slot(0,0)
-					.VAlign(VAlign_Center)
-					[
-						SNew(STextBlock)
-						.Text(FText::FromString(TEXT("Location")))
-						.Font(FAppStyle::GetFontStyle(TEXT("PropertyWindow.NormalFont")))
-					]
+		ParameterHorizontalBox->AddSlot()
+			.Padding(4.0f)
+			.VAlign(VAlign_Center)
+			[
+				SNew(SNumericEntryBox<int32>)
+				.Visibility(this, &SMutableParametersWidget::GetParameterVisibility, ParamIndex)
+				.AllowSpin(true)
+				.MinSliderValue(0)
+				.MaxSliderValue(this, &SMutableParametersWidget::GetIntParameterValueMax, ParamIndex)
+				.Value(this, &SMutableParametersWidget::GetIntParameterValue, ParamIndex, RangeIndex)
+				.OnValueChanged(this, &SMutableParametersWidget::OnIntParameterChanged, ParamIndex, ParamComboBox,RangeIndex)
+			];
+		break;
+	}
 
-					+SGridPanel::Slot(1, 0)
-					[
-						SNew(SNumericVectorInputBox<FVector::FReal>)
-						.X(this, &SMutableParametersWidget::GetProjectorLocation, EAxis::X, ParamIndex)
-						.Y(this, &SMutableParametersWidget::GetProjectorLocation, EAxis::Y, ParamIndex)
-						.Z(this, &SMutableParametersWidget::GetProjectorLocation, EAxis::Z, ParamIndex)
-						.bColorAxisLabels(true)
-						.OnXChanged(this, &SMutableParametersWidget::SetProjectorLocation, ETextCommit::Default, EAxis::X, false, ParamIndex)
-						.OnYChanged(this, &SMutableParametersWidget::SetProjectorLocation, ETextCommit::Default, EAxis::Y, false, ParamIndex)
-						.OnZChanged(this, &SMutableParametersWidget::SetProjectorLocation, ETextCommit::Default, EAxis::Z, false, ParamIndex)
-						.OnXCommitted(this, &SMutableParametersWidget::SetProjectorLocation, EAxis::X, true, ParamIndex)
-						.OnYCommitted(this, &SMutableParametersWidget::SetProjectorLocation, EAxis::Y, true, ParamIndex)
-						.OnZCommitted(this, &SMutableParametersWidget::SetProjectorLocation, EAxis::Z, true, ParamIndex)
-						.AllowSpin(true)
-						.SpinDelta(1)
-					]
-
-					+ SGridPanel::Slot(0, 1)
-					.VAlign(VAlign_Center)
-					[
-						SNew(STextBlock)
-						.Text(FText::FromString(TEXT("Scale")))
-						.Font(FAppStyle::GetFontStyle("PropertyWindow.NormalFont"))
-					]
+	case mu::PARAMETER_TYPE::T_PROJECTOR:
+		ParameterHorizontalBox->AddSlot()
+			.Padding(4.0f)
+			.VAlign(VAlign_Center)
+			[
+				SNew(SGridPanel)
+				.FillColumn(0, 1.0f)
 
-					+ SGridPanel::Slot(1, 1)
-					[
-						SNew(SNumericVectorInputBox<FVector::FReal>)
-						.X(this, &SMutableParametersWidget::GetProjectorScale, EAxis::X, ParamIndex)
-						.Y(this, &SMutableParametersWidget::GetProjectorScale, EAxis::Y, ParamIndex)
-						.Z(this, &SMutableParametersWidget::GetProjectorScale, EAxis::Z, ParamIndex)
-						.bColorAxisLabels(true)
-						.OnXChanged(this, &SMutableParametersWidget::SetProjectorScale, ETextCommit::Default, EAxis::X, false, ParamIndex)
-						.OnYChanged(this, &SMutableParametersWidget::SetProjectorScale, ETextCommit::Default, EAxis::Y, false, ParamIndex)
-						.OnZChanged(this, &SMutableParametersWidget::SetProjectorScale, ETextCommit::Default, EAxis::Z, false, ParamIndex)
-						.OnXCommitted(this, &SMutableParametersWidget::SetProjectorScale, EAxis::X, true, ParamIndex)
-						.OnYCommitted(this, &SMutableParametersWidget::SetProjectorScale, EAxis::Y, true, ParamIndex)
-						.OnZCommitted(this, &SMutableParametersWidget::SetProjectorScale, EAxis::Z, true, ParamIndex)
-						.AllowSpin(true)
-						.SpinDelta(1)
-					]
+				+SGridPanel::Slot(0,0)
+				.VAlign(VAlign_Center)
+				[
+					SNew(STextBlock)
+					.Text(FText::FromString(TEXT("Location")))
+					.Font(FAppStyle::GetFontStyle(TEXT("PropertyWindow.NormalFont")))
+				]
 
-				];
-			break;
+				+SGridPanel::Slot(1, 0)
+				[
+					SNew(SNumericVectorInputBox<FVector::FReal>)
+					.X(this, &SMutableParametersWidget::GetProjectorLocation, EAxis::X, ParamIndex, RangeIndex)
+					.Y(this, &SMutableParametersWidget::GetProjectorLocation, EAxis::Y, ParamIndex, RangeIndex)
+					.Z(this, &SMutableParametersWidget::GetProjectorLocation, EAxis::Z, ParamIndex, RangeIndex)
+					.bColorAxisLabels(true)
+					.OnXChanged(this, &SMutableParametersWidget::SetProjectorLocation, ETextCommit::Default, EAxis::X, false, ParamIndex, RangeIndex)
+					.OnYChanged(this, &SMutableParametersWidget::SetProjectorLocation, ETextCommit::Default, EAxis::Y, false, ParamIndex, RangeIndex)
+					.OnZChanged(this, &SMutableParametersWidget::SetProjectorLocation, ETextCommit::Default, EAxis::Z, false, ParamIndex, RangeIndex)
+					.OnXCommitted(this, &SMutableParametersWidget::SetProjectorLocation, EAxis::X, true, ParamIndex, RangeIndex)
+					.OnYCommitted(this, &SMutableParametersWidget::SetProjectorLocation, EAxis::Y, true, ParamIndex, RangeIndex)
+					.OnZCommitted(this, &SMutableParametersWidget::SetProjectorLocation, EAxis::Z, true, ParamIndex, RangeIndex)
+					.AllowSpin(true)
+					.SpinDelta(1)
+				]
 
-		default:
-			// Unsupported parameter type
-			ParameterBox->AddSlot()
-				.Padding(4.0f)
+				+ SGridPanel::Slot(0, 1)
 				.VAlign(VAlign_Center)
 				[
 					SNew(STextBlock)
-					.Text(FText::FromString(TEXT("Unsupported parameter type.")))
-					.Justification(ETextJustify::Left)
-					.AutoWrapText(true)
-				];
-		}
+					.Text(FText::FromString(TEXT("Scale")))
+					.Font(FAppStyle::GetFontStyle("PropertyWindow.NormalFont"))
+				]
+
+				+ SGridPanel::Slot(1, 1)
+				[
+					SNew(SNumericVectorInputBox<FVector::FReal>)
+					.X(this, &SMutableParametersWidget::GetProjectorScale, EAxis::X, ParamIndex, RangeIndex)
+					.Y(this, &SMutableParametersWidget::GetProjectorScale, EAxis::Y, ParamIndex, RangeIndex)
+					.Z(this, &SMutableParametersWidget::GetProjectorScale, EAxis::Z, ParamIndex, RangeIndex)
+					.bColorAxisLabels(true)
+					.OnXChanged(this, &SMutableParametersWidget::SetProjectorScale, ETextCommit::Default, EAxis::X, false, ParamIndex, RangeIndex)
+					.OnYChanged(this, &SMutableParametersWidget::SetProjectorScale, ETextCommit::Default, EAxis::Y, false, ParamIndex, RangeIndex)
+					.OnZChanged(this, &SMutableParametersWidget::SetProjectorScale, ETextCommit::Default, EAxis::Z, false, ParamIndex, RangeIndex)
+					.OnXCommitted(this, &SMutableParametersWidget::SetProjectorScale, EAxis::X, true, ParamIndex, RangeIndex)
+					.OnYCommitted(this, &SMutableParametersWidget::SetProjectorScale, EAxis::Y, true, ParamIndex, RangeIndex)
+					.OnZCommitted(this, &SMutableParametersWidget::SetProjectorScale, EAxis::Z, true, ParamIndex, RangeIndex)
+					.AllowSpin(true)
+					.SpinDelta(1)
+				]
 
+			];
+		break;
+
+	default:
+		// Unsupported parameter type
+		ParameterHorizontalBox->AddSlot()
+			.Padding(4.0f)
+			.VAlign(VAlign_Center)
+			[
+				SNew(STextBlock)
+				.Text(FText::FromString(TEXT("Unsupported parameter type.")))
+				.Justification(ETextJustify::Left)
+				.AutoWrapText(true)
+			];
 	}
 }
 
@@ -263,7 +410,7 @@ EVisibility SMutableParametersWidget::GetParameterVisibility(int32 ParamIndex) c
 }
 
 
-ECheckBoxState SMutableParametersWidget::GetBoolParameterValue(int32 ParamIndex) const
+ECheckBoxState SMutableParametersWidget::GetBoolParameterValue(int32 ParamIndex,  mu::RangeIndexPtrConst RangeIndex) const
 {
 	if ( !MutableParameters 
 		|| ParamIndex>=MutableParameters->GetCount()
@@ -271,12 +418,12 @@ ECheckBoxState SMutableParametersWidget::GetBoolParameterValue(int32 ParamIndex)
 	{
 		return ECheckBoxState::Undetermined;
 	}	
-
-	return MutableParameters->GetBoolValue(ParamIndex) ? ECheckBoxState::Checked : ECheckBoxState::Unchecked;
+	
+	return MutableParameters->GetBoolValue(ParamIndex,RangeIndex) ? ECheckBoxState::Checked : ECheckBoxState::Unchecked;
 }
 
 
-void SMutableParametersWidget::OnBoolParameterChanged(ECheckBoxState InCheckboxState, int32 ParamIndex)
+void SMutableParametersWidget::OnBoolParameterChanged(ECheckBoxState InCheckboxState, int32 ParamIndex,  mu::RangeIndexPtrConst RangeIndex)
 {
 	if (!MutableParameters
 		|| ParamIndex >= MutableParameters->GetCount()
@@ -285,13 +432,13 @@ void SMutableParametersWidget::OnBoolParameterChanged(ECheckBoxState InCheckboxS
 		return;
 	}
 
-	MutableParameters->SetBoolValue( ParamIndex, InCheckboxState==ECheckBoxState::Checked );
+	MutableParameters->SetBoolValue( ParamIndex, InCheckboxState==ECheckBoxState::Checked, RangeIndex );
 
 	OnParametersValueChanged.ExecuteIfBound(ParamIndex);
 }
 
 
-TOptional<float> SMutableParametersWidget::GetFloatParameterValue(int32 ParamIndex) const
+TOptional<float> SMutableParametersWidget::GetFloatParameterValue(int32 ParamIndex, mu::RangeIndexPtrConst RangeIndex) const
 {
 	if (!MutableParameters
 		|| ParamIndex >= MutableParameters->GetCount()
@@ -300,11 +447,11 @@ TOptional<float> SMutableParametersWidget::GetFloatParameterValue(int32 ParamInd
 		return TOptional<float>();
 	}
 
-	return MutableParameters->GetFloatValue(ParamIndex);
+	return MutableParameters->GetFloatValue(ParamIndex,RangeIndex);
 }
 
 
-void SMutableParametersWidget::OnFloatParameterChanged(float InValue, int32 ParamIndex)
+void SMutableParametersWidget::OnFloatParameterChanged(float InValue, int32 ParamIndex, mu::RangeIndexPtrConst RangeIndex)
 {
 	if (!MutableParameters
 		|| ParamIndex >= MutableParameters->GetCount()
@@ -313,13 +460,22 @@ void SMutableParametersWidget::OnFloatParameterChanged(float InValue, int32 Para
 		return;
 	}
 
-	MutableParameters->SetFloatValue(ParamIndex, InValue);
-
+	MutableParameters->SetFloatValue(ParamIndex, InValue, RangeIndex);
+	
 	OnParametersValueChanged.ExecuteIfBound(ParamIndex);
 }
 
+void SMutableParametersWidget::OnFloatParameterCommitted(float Value, ETextCommit::Type CommitType, int32 ParamIndex)
+{
+	// Think on the possibility of asking for a reconstruction if the new value is a valid integer
+	if (Value >= 1 && Value == (int)Value)
+	{
+		ScheduleUpdateIfRequired(ParamIndex);
+	}
+}
+
 
-TOptional<int32> SMutableParametersWidget::GetIntParameterValue(int32 ParamIndex) const
+TOptional<int32> SMutableParametersWidget::GetIntParameterValue (int32 ParamIndex, mu::RangeIndexPtrConst RangeIndex) const
 {
 	if (!MutableParameters
 		|| ParamIndex >= MutableParameters->GetCount()
@@ -328,7 +484,7 @@ TOptional<int32> SMutableParametersWidget::GetIntParameterValue(int32 ParamIndex
 		return TOptional<int32>();
 	}
 
-	int32 Value = MutableParameters->GetIntValue(ParamIndex);
+	int32 Value = MutableParameters->GetIntValue(ParamIndex, RangeIndex);
 	int32 ValueIndex = MutableParameters->GetIntValueIndex(ParamIndex, Value);
 
 	return ValueIndex;
@@ -355,7 +511,7 @@ TOptional<int32> SMutableParametersWidget::GetIntParameterValueMax(int32 ParamIn
 }
 
 
-void SMutableParametersWidget::OnIntParameterChanged(int32 InValue, int32 ParamIndex, TSharedPtr<SMutableTextSearchBox> Combo )
+void SMutableParametersWidget::OnIntParameterChanged(int32 InValue, int32 ParamIndex, TSharedPtr<SMutableTextSearchBox> Combo, mu::RangeIndexPtrConst RangeIndex )
 {
 	if (!MutableParameters
 		|| ParamIndex >= MutableParameters->GetCount()
@@ -368,8 +524,8 @@ void SMutableParametersWidget::OnIntParameterChanged(int32 InValue, int32 ParamI
 	if (InValue >= 0 && InValue < MutableParameters->GetIntPossibleValueCount(ParamIndex))
 	{
 		int32 RealValue = MutableParameters->GetIntPossibleValue(ParamIndex, InValue);
-		MutableParameters->SetIntValue(ParamIndex, RealValue);
-
+		MutableParameters->SetIntValue(ParamIndex, RealValue,RangeIndex);
+		
 		// Update the text combo if any
 		if (Combo)
 		{
@@ -377,13 +533,15 @@ void SMutableParametersWidget::OnIntParameterChanged(int32 InValue, int32 ParamI
 			Text = ANSI_TO_TCHAR(MutableParameters->GetIntPossibleValueName(ParamIndex, InValue));
 			Combo->SetText(FText::FromString(Text));
 		}
-
+		
 		OnParametersValueChanged.ExecuteIfBound(ParamIndex);
+		
+		ScheduleUpdateIfRequired(ParamIndex);
 	}
 }
 
 
-void SMutableParametersWidget::OnIntParameterTextChanged(TSharedPtr<FString> Selection, ESelectInfo::Type, int32 ParamIndex)
+void SMutableParametersWidget::OnIntParameterTextChanged(TSharedPtr<FString> Selection, ESelectInfo::Type, int32 ParamIndex, mu::RangeIndexPtrConst RangeIndex)
 {
 	if (!Selection
 		|| !MutableParameters
@@ -397,14 +555,16 @@ void SMutableParametersWidget::OnIntParameterTextChanged(TSharedPtr<FString> Sel
 	if (ValueIndex >= 0 && ValueIndex < MutableParameters->GetIntPossibleValueCount(ParamIndex))
 	{
 		int32 RealValue = MutableParameters->GetIntPossibleValue(ParamIndex, ValueIndex);
-		MutableParameters->SetIntValue(ParamIndex, RealValue);
-
+		MutableParameters->SetIntValue(ParamIndex, RealValue,RangeIndex);
+		
 		OnParametersValueChanged.ExecuteIfBound(ParamIndex);
+		
+		ScheduleUpdateIfRequired(ParamIndex);
 	}
 }
 
 
-FLinearColor SMutableParametersWidget::GetColorParameterValue(int32 ParamIndex) const
+FLinearColor SMutableParametersWidget::GetColorParameterValue(int32 ParamIndex, mu::RangeIndexPtrConst RangeIndex) const
 {
 	if (!MutableParameters
 		|| ParamIndex >= MutableParameters->GetCount()
@@ -414,21 +574,21 @@ FLinearColor SMutableParametersWidget::GetColorParameterValue(int32 ParamIndex)
 	}
 
 	FLinearColor Result;
-	MutableParameters->GetColourValue(ParamIndex, &Result.R, &Result.G, &Result.B);
+	MutableParameters->GetColourValue(ParamIndex, &Result.R, &Result.G, &Result.B, RangeIndex);
 	Result.A = 1.0f;
 
 	return Result;
 }
 
 
-FReply SMutableParametersWidget::OnColorBlockMouseButtonDown(const FGeometry& MyGeometry, const FPointerEvent& MouseEvent, int32 ParamIndex)
+FReply SMutableParametersWidget::OnColorBlockMouseButtonDown(const FGeometry& MyGeometry, const FPointerEvent& MouseEvent, int32 ParamIndex, mu::RangeIndexPtrConst RangeIndex)
 {
 	if (MouseEvent.GetEffectingButton() != EKeys::LeftMouseButton)
 	{
 		return FReply::Unhandled();
 	}
 
-	FLinearColor col = GetColorParameterValue(ParamIndex);
+	FLinearColor col = GetColorParameterValue(ParamIndex,RangeIndex);
 
 	TArray<FLinearColor*> LinearColorArray;
 	LinearColorArray.Add(&col);
@@ -438,14 +598,14 @@ FReply SMutableParametersWidget::OnColorBlockMouseButtonDown(const FGeometry& My
 	args.bUseAlpha = false;
 	args.bOnlyRefreshOnMouseUp = false;
 	args.InitialColor = col;
-	args.OnColorCommitted = FOnLinearColorValueChanged::CreateSP(this, &SMutableParametersWidget::OnSetColorFromColorPicker, ParamIndex);
+	args.OnColorCommitted = FOnLinearColorValueChanged::CreateSP(this, &SMutableParametersWidget::OnSetColorFromColorPicker, ParamIndex, RangeIndex);
 	OpenColorPicker(args);
 
 	return FReply::Handled();
 }
 
 
-void SMutableParametersWidget::OnSetColorFromColorPicker(FLinearColor NewColor, int32 ParamIndex)
+void SMutableParametersWidget::OnSetColorFromColorPicker(FLinearColor NewColor, int32 ParamIndex,  mu::RangeIndexPtrConst RangeIndex)
 {
 	if (!MutableParameters
 		|| ParamIndex >= MutableParameters->GetCount()
@@ -453,14 +613,14 @@ void SMutableParametersWidget::OnSetColorFromColorPicker(FLinearColor NewColor,
 	{
 		return;
 	}
-
-	MutableParameters->SetColourValue(ParamIndex, NewColor.R, NewColor.G, NewColor.B);
+	
+	MutableParameters->SetColourValue(ParamIndex, NewColor.R, NewColor.G, NewColor.B, RangeIndex);
 
 	OnParametersValueChanged.ExecuteIfBound(ParamIndex);
 }
 
 
-TOptional<FVector::FReal> SMutableParametersWidget::GetProjectorLocation(EAxis::Type Axis, int32 ParamIndex) const
+TOptional<FVector::FReal> SMutableParametersWidget::GetProjectorLocation(EAxis::Type Axis, int32 ParamIndex, mu::RangeIndexPtrConst RangeIndex) const
 { 
 	if (!MutableParameters
 		|| ParamIndex >= MutableParameters->GetCount()
@@ -468,20 +628,20 @@ TOptional<FVector::FReal> SMutableParametersWidget::GetProjectorLocation(EAxis::
 	{
 		return 0.0f;
 	}
-
+	
 	mu::PROJECTOR_TYPE Type;
 	FVector3f Pos;
 	FVector3f Dir;
 	FVector3f Up;
 	FVector3f Scale;
 	float Angle;
-	MutableParameters->GetProjectorValue(ParamIndex, &Type, &Pos, &Dir, &Up, &Scale, &Angle );
+	MutableParameters->GetProjectorValue(ParamIndex, &Type, &Pos, &Dir, &Up, &Scale, &Angle , RangeIndex);
 
 	return Pos[ Axis-EAxis::X ];
 }
 
 
-void SMutableParametersWidget::SetProjectorLocation(FVector::FReal NewValue, ETextCommit::Type, EAxis::Type Axis, bool bCommitted, int32 ParamIndex)
+void SMutableParametersWidget::SetProjectorLocation(FVector::FReal NewValue, ETextCommit::Type, EAxis::Type Axis, bool bCommitted, int32 ParamIndex, mu::RangeIndexPtrConst RangeIndex)
 {
 	if (!MutableParameters
 		|| ParamIndex >= MutableParameters->GetCount()
@@ -489,24 +649,26 @@ void SMutableParametersWidget::SetProjectorLocation(FVector::FReal NewValue, ETe
 	{
 		return;
 	}
-
+	
+	//mu::RangeIndexPtr UpdatedRange = RangeIndexPayload.GetRangeIndex();
+	
 	mu::PROJECTOR_TYPE Type;
 	FVector3f Pos;
 	FVector3f Dir;
 	FVector3f Up;
 	FVector3f Scale;
 	float Angle;
-	MutableParameters->GetProjectorValue(ParamIndex, &Type, &Pos, &Dir, &Up, &Scale, &Angle);
+	MutableParameters->GetProjectorValue(ParamIndex, &Type, &Pos, &Dir, &Up, &Scale, &Angle, RangeIndex);
 
 	Pos[Axis - EAxis::X] = NewValue;
 
-	MutableParameters->SetProjectorValue(ParamIndex, Pos, Dir, Up, Scale, Angle);
+	MutableParameters->SetProjectorValue(ParamIndex, Pos, Dir, Up, Scale, Angle, RangeIndex);
 
 	OnParametersValueChanged.ExecuteIfBound(ParamIndex);
 }
 
 
-TOptional<FVector::FReal> SMutableParametersWidget::GetProjectorScale(EAxis::Type Axis, int32 ParamIndex) const
+TOptional<FVector::FReal> SMutableParametersWidget::GetProjectorScale(EAxis::Type Axis, int32 ParamIndex, mu::RangeIndexPtrConst RangeIndex) const
 {
 	if (!MutableParameters
 		|| ParamIndex >= MutableParameters->GetCount()
@@ -521,13 +683,13 @@ TOptional<FVector::FReal> SMutableParametersWidget::GetProjectorScale(EAxis::Typ
 	FVector3f Up;
 	FVector3f Scale;
 	float Angle;
-	MutableParameters->GetProjectorValue(ParamIndex, &Type, &Pos, &Dir, &Up, &Scale, &Angle);
+	MutableParameters->GetProjectorValue(ParamIndex, &Type, &Pos, &Dir, &Up, &Scale, &Angle, RangeIndex);
 
 	return Scale[Axis - EAxis::X];
 }
 
 
-void SMutableParametersWidget::SetProjectorScale(FVector::FReal NewValue, ETextCommit::Type, EAxis::Type Axis, bool bCommitted, int32 ParamIndex)
+void SMutableParametersWidget::SetProjectorScale(FVector::FReal NewValue, ETextCommit::Type, EAxis::Type Axis, bool bCommitted, int32 ParamIndex,  mu::RangeIndexPtrConst RangeIndex)
 {
 	if (!MutableParameters
 		|| ParamIndex >= MutableParameters->GetCount()
@@ -536,17 +698,18 @@ void SMutableParametersWidget::SetProjectorScale(FVector::FReal NewValue, ETextC
 		return;
 	}
 
+	
 	mu::PROJECTOR_TYPE Type;
 	FVector3f Pos;
 	FVector3f Dir;
 	FVector3f Up;
 	FVector3f Scale;
 	float Angle;
-	MutableParameters->GetProjectorValue(ParamIndex, &Type, &Pos, &Dir, &Up, &Scale, &Angle);
+	MutableParameters->GetProjectorValue(ParamIndex, &Type, &Pos, &Dir, &Up, &Scale, &Angle, RangeIndex);
 
 	Scale[Axis - EAxis::X] = NewValue;
 
-	MutableParameters->SetProjectorValue(ParamIndex, Pos, Dir, Up, Scale, Angle);
+	MutableParameters->SetProjectorValue(ParamIndex, Pos, Dir, Up, Scale, Angle, RangeIndex);
 
 	OnParametersValueChanged.ExecuteIfBound(ParamIndex);
 }
diff --git a/Engine/Plugins/Experimental/Mutable/Source/CustomizableObjectEditor/Private/MuCOE/SMutableParametersWidget.h b/Engine/Plugins/Experimental/Mutable/Source/CustomizableObjectEditor/Private/MuCOE/SMutableParametersWidget.h
--- a/Engine/Plugins/Experimental/Mutable/Source/CustomizableObjectEditor/Private/MuCOE/SMutableParametersWidget.h
+++ b/Engine/Plugins/Experimental/Mutable/Source/CustomizableObjectEditor/Private/MuCOE/SMutableParametersWidget.h
@@ -2,9 +2,10 @@
 
 #pragma once
 
-#include "MuR/MutableMemory.h"
 #include "MuR/Parameters.h"
 #include "MuR/Ptr.h"
+#include "Templates/SharedPointer.h"
+#include "Widgets/SBoxPanel.h"
 #include "Widgets/SCompoundWidget.h"
 
 enum class ECheckBoxState : uint8;
@@ -64,26 +65,44 @@ private:
 	/** Dynamically filled box with per-parameter widgets */
 	TSharedPtr<SVerticalBox> ParamBox;
 
+	/** Check if the Parameter index provided is from a parameter that controls the amount of values another parameter/s
+	 * should have exposed in the UI. If the parameter does so then the redraw of the parameters widget will be scheduled
+	 * for the next update.
+	 * @param InParameterIndex The index of the parameter we changed it's value and want to check if that should or not
+	 * trigger the Parameter's Widget redraw.
+	 */
+	void ScheduleUpdateIfRequired(const int32& InParameterIndex);
+	
+	/** Parameter slate generation methods. */
+
+	/** Generate a new slate for the provided parameter index and attaches it to the also provided Horizontal box slate.
+	 * @param ParamIndex The index of the mutable parameter to represent using a slate object.
+	 * @param ParameterHorizontalBox The Horizontal box slate object to contain the newly generated slate object representing the parameter.
+	 * @param RangeIndex Range object pointing at the parameters' target dimension and position. Providing a nullptr is expected for parameters with
+	 * a single value (not a multivalue parameter)
+	 * This object can also represent single value parameters but must be provided in either case.
+	 */
+	void GenerateAndAttachParameterSlate(const int32 ParamIndex, TSharedPtr<SHorizontalBox> ParameterHorizontalBox, mu::RangeIndexPtrConst RangeIndex);
+	
 	/** Internal UI callbacks. */
 	EVisibility GetParameterVisibility(int32 ParamIndex) const;
-	ECheckBoxState GetBoolParameterValue(int32 ParamIndex) const;
-	void OnBoolParameterChanged(ECheckBoxState CheckboxState, int32 ParamIndex);
-	TOptional<float> GetFloatParameterValue(int32 ParamIndex) const;
-	void OnFloatParameterChanged(float Value, int32 ParamIndex);
-	FLinearColor GetColorParameterValue(int32 ParamIndex) const;
-	FReply OnColorBlockMouseButtonDown(const FGeometry& MyGeometry, const FPointerEvent& MouseEvent, int32 ParamIndex);
-	void OnSetColorFromColorPicker(FLinearColor NewColor, int32 ParamIndex);
-	TOptional<int32> GetIntParameterValue(int32 ParamIndex) const;
+	ECheckBoxState GetBoolParameterValue(int32 ParamIndex, mu::RangeIndexPtrConst RangeIndex) const;
+	void OnBoolParameterChanged(ECheckBoxState CheckboxState, int32 ParamIndex, mu::RangeIndexPtrConst RangeIndex);
+	TOptional<float> GetFloatParameterValue(int32 ParamIndex, mu::RangeIndexPtrConst RangeIndex) const;
+	void OnFloatParameterChanged(float Value, int32 ParamIndex, mu::RangeIndexPtrConst RangeIndex);
+	void OnFloatParameterCommitted(float Value, ETextCommit::Type CommitType, int32 ParamIndex);
+	FLinearColor GetColorParameterValue(int32 ParamIndex, mu::RangeIndexPtrConst RangeIndex) const;
+	FReply OnColorBlockMouseButtonDown(const FGeometry& MyGeometry, const FPointerEvent& MouseEvent, int32 ParamIndex, mu::RangeIndexPtrConst RangeIndex);
+	void OnSetColorFromColorPicker(FLinearColor NewColor, int32 ParamIndex, mu::RangeIndexPtrConst RangeIndex);
+	TOptional<int32> GetIntParameterValue(int32 ParamIndex,  mu::RangeIndexPtrConst RangeIndex) const;
 	TOptional<int32> GetIntParameterValueMax(int32 ParamIndex) const;
-	void OnIntParameterChanged(int32 Value, int32 ParamIndex, TSharedPtr<SMutableTextSearchBox> );
-	void OnIntParameterTextChanged(TSharedPtr<FString> Selection, ESelectInfo::Type SelectInfo, int32 ParamIndex);
+	void OnIntParameterChanged(int32 Value, int32 ParamIndex, TSharedPtr<SMutableTextSearchBox> Combo,  mu::RangeIndexPtrConst RangeIndex);
+	void OnIntParameterTextChanged(TSharedPtr<FString> Selection, ESelectInfo::Type SelectInfo, int32 ParamIndex, mu::RangeIndexPtrConst RangeIndex);
 
 	/** Projector UI callbacks */
-	TOptional<FVector::FReal> GetProjectorLocation(EAxis::Type Axis, int32 ParamIndex) const;
-	void SetProjectorLocation(FVector::FReal NewValue, ETextCommit::Type, EAxis::Type, bool bCommitted, int32 ParamIndex);
-	TOptional<FVector::FReal> GetProjectorScale(EAxis::Type Axis, int32 ParamIndex) const;
-	void SetProjectorScale(FVector::FReal NewValue, ETextCommit::Type, EAxis::Type, bool bCommitted, int32 ParamIndex);
-
+	TOptional<FVector::FReal> GetProjectorLocation(EAxis::Type Axis, int32 ParamIndex,  mu::RangeIndexPtrConst RangeIndex) const;
+	void SetProjectorLocation(FVector::FReal NewValue, ETextCommit::Type, EAxis::Type, bool bCommitted, int32 ParamIndex,  mu::RangeIndexPtrConst RangeIndex);
+	TOptional<FVector::FReal> GetProjectorScale(EAxis::Type Axis, int32 ParamIndex, mu::RangeIndexPtrConst RangeIndex) const;
+	void SetProjectorScale(FVector::FReal NewValue, ETextCommit::Type, EAxis::Type, bool bCommitted, int32 ParamIndex, mu::RangeIndexPtrConst RangeIndex);
 
 };
-
