[mutable] Wrapped up Mutable Operations tree into an ScrollBox to allow horizontal scroll in high tree depths.
- Also removed unused includes.

diff --git a/Engine/Plugins/Experimental/Mutable/Source/CustomizableObjectEditor/Private/MuCOE/SMutableCodeViewer.cpp b/Engine/Plugins/Experimental/Mutable/Source/CustomizableObjectEditor/Private/MuCOE/SMutableCodeViewer.cpp
--- a/Engine/Plugins/Experimental/Mutable/Source/CustomizableObjectEditor/Private/MuCOE/SMutableCodeViewer.cpp
+++ b/Engine/Plugins/Experimental/Mutable/Source/CustomizableObjectEditor/Private/MuCOE/SMutableCodeViewer.cpp
@@ -4,9 +4,8 @@
 
 #include "DesktopPlatformModule.h"
 #include "EditorDirectories.h"
-#include "Framework/MultiBox/MultiBoxBuilder.h"
-#include "Framework/Views/TableViewMetadata.h"
 #include "IDesktopPlatform.h"
+#include "Framework/MultiBox/MultiBoxBuilder.h"
 #include "Misc/Paths.h"
 #include "MuCO/CustomizableObject.h"
 #include "MuCOE/SMutableBoolViewer.h"
@@ -26,15 +25,14 @@
 #include "MuR/SystemPrivate.h"
 #include "MuT/ErrorLogPrivate.h"
 #include "MuT/Streams.h"
-#include "MuT/Table.h"
 #include "MuT/TypeInfo.h"
+#include "Widgets/SNullWidget.h"
 #include "Widgets/Colors/SColorBlock.h"
 #include "Widgets/Input/SButton.h"
-#include "Widgets/Input/SNumericEntryBox.h"
 #include "Widgets/Input/SComboBox.h"
+#include "Widgets/Input/SNumericEntryBox.h"
 #include "Widgets/Layout/SScrollBox.h"
 #include "Widgets/Views/STreeView.h"
-#include "Widgets/SNullWidget.h"
 
 class FExtender;
 class FReferenceCollector;
@@ -522,6 +520,11 @@ void SMutableCodeViewer::Construct(const FArguments& InArgs, const TSharedPtr<mu
 	ToolbarBuilder.EndSection();
 
 	ToolbarBuilder.AddWidget(SNew(STextBlock).Text(FText::FromString(InArgs._DataTag)));
+
+	TSharedRef<SScrollBar> TreeVertScrollBar =
+		SNew(SScrollBar).
+		Orientation(EOrientation::Orient_Vertical).
+		AlwaysShowScrollbar(false);
 	
 	ChildSlot
 	[
@@ -625,28 +628,50 @@ void SMutableCodeViewer::Construct(const FArguments& InArgs, const TSharedPtr<mu
 					.BorderImage(FAppStyle::GetBrush("ToolPanel.GroupBorder"))
 					.Padding(FMargin(4.0f, 4.0f))
 					[
-						SAssignNew(TreeView, STreeView<TSharedPtr<FMutableCodeTreeElement>>)
-						.TreeItemsSource(&RootNodes)
-						.OnGenerateRow(this, &SMutableCodeViewer::GenerateRowForNodeTree)
-						.OnRowReleased(this, &SMutableCodeViewer::OnRowReleased)
-						.OnGetChildren(this, &SMutableCodeViewer::GetChildrenForInfo)
-						.OnSelectionChanged(this, &SMutableCodeViewer::OnSelectionChanged)
-						.OnSetExpansionRecursive(this, &SMutableCodeViewer::TreeExpandRecursive)
-						.OnContextMenuOpening(this, &SMutableCodeViewer::OnTreeContextMenuOpening)
-						.OnExpansionChanged(this, &SMutableCodeViewer::OnExpansionChanged)
-						.SelectionMode(ESelectionMode::Single)
-						.HeaderRow
-						(
-							SNew(SHeaderRow)
-
-							+ SHeaderRow::Column(MutableCodeTreeViewColumns::OperationsColumnID)
-								.DefaultLabel(LOCTEXT("Operation", "Operation"))
-								.FillWidth(OperationsColumnWidth)
-								
-							+ SHeaderRow::Column(MutableCodeTreeViewColumns::AdditionalDataColumnID)
-								.DefaultLabel(LOCTEXT("OperationFlags", "Flags"))
-								.FillWidth(ExtraDataColumnWidth)
-						)
+						SNew(SHorizontalBox)
+
+						+ SHorizontalBox::Slot()
+						.FillWidth(1.0f)
+						[
+							SNew(SScrollBox)
+							.Orientation(EOrientation::Orient_Horizontal)
+							.ConsumeMouseWheel(EConsumeMouseWheel::Never)
+
+							+ SScrollBox::Slot()
+							.HAlign(HAlign_Fill)
+							[
+								SAssignNew(TreeView, STreeView<TSharedPtr<FMutableCodeTreeElement>>)
+								.TreeItemsSource(&RootNodes)
+								.OnGenerateRow(this, &SMutableCodeViewer::GenerateRowForNodeTree)
+								.OnRowReleased(this, &SMutableCodeViewer::OnRowReleased)
+								.OnGetChildren(this, &SMutableCodeViewer::GetChildrenForInfo)
+								.OnSelectionChanged(this, &SMutableCodeViewer::OnSelectionChanged)
+								.OnSetExpansionRecursive(this, &SMutableCodeViewer::TreeExpandRecursive)
+								.OnContextMenuOpening(this, &SMutableCodeViewer::OnTreeContextMenuOpening)
+								.OnExpansionChanged(this, &SMutableCodeViewer::OnExpansionChanged)
+								.SelectionMode(ESelectionMode::Single)
+								.ExternalScrollbar(TreeVertScrollBar)
+								.HeaderRow
+								(
+									SNew(SHeaderRow)
+									.ResizeMode(ESplitterResizeMode::Fill)
+
+									+ SHeaderRow::Column(MutableCodeTreeViewColumns::OperationsColumnID)
+										.DefaultLabel(LOCTEXT("Operation", "Operation"))
+										.ManualWidth(618.0f)
+				
+									+ SHeaderRow::Column(MutableCodeTreeViewColumns::AdditionalDataColumnID)
+										.DefaultLabel(LOCTEXT("OperationFlags", "Flags"))
+										.FixedWidth(50.0f)
+								)
+							]
+						]
+
+						+ SHorizontalBox::Slot()
+						.AutoWidth()
+						[
+							TreeVertScrollBar
+						]
 					]
 				]
 				
