#rnx

diff --git a/Engine/Plugins/Experimental/Mutable/Source/CustomizableObjectEditor/Private/MuCOE/CustomizableObjectCompiler.cpp b/Engine/Plugins/Experimental/Mutable/Source/CustomizableObjectEditor/Private/MuCOE/CustomizableObjectCompiler.cpp
--- a/Engine/Plugins/Experimental/Mutable/Source/CustomizableObjectEditor/Private/MuCOE/CustomizableObjectCompiler.cpp
+++ b/Engine/Plugins/Experimental/Mutable/Source/CustomizableObjectEditor/Private/MuCOE/CustomizableObjectCompiler.cpp
@@ -215,6 +215,18 @@ void FCustomizableObjectCompiler::Compile(UCustomizableObject& Object, const FCo
 		return;
 	}
 
+	// Now that we know for sure that the CO is locked and there are no pending updates of instances using the CO,
+	// destroy any live update instances, as they become invalid when recompiling the CO
+	for (TObjectIterator<UCustomizableObjectInstance> It; It; ++It)
+	{
+		UCustomizableObjectInstance* Instance = *It;
+		if (Instance &&
+			Instance->GetCustomizableObject() == &Object)
+		{
+			Instance->DestroyLiveUpdateInstance();
+		}
+	}
+
 	CurrentObject = &Object;
 
 	PreloadingReferencerAssets = true;
