[mutable] Made color preview box of MutableColorViewer not display the color as sRGB
- Changed internal slate being used to display the preview box.
- Removed all usages of SMutableColorPreviewBox
- Removed SMutableColorPreviewBox since it is redundant (SColorBox)

diff --git a/Engine/Plugins/Experimental/Mutable/Source/CustomizableObjectEditor/Private/MuCOE/SMutableCodeViewer.cpp b/Engine/Plugins/Experimental/Mutable/Source/CustomizableObjectEditor/Private/MuCOE/SMutableCodeViewer.cpp
--- a/Engine/Plugins/Experimental/Mutable/Source/CustomizableObjectEditor/Private/MuCOE/SMutableCodeViewer.cpp
+++ b/Engine/Plugins/Experimental/Mutable/Source/CustomizableObjectEditor/Private/MuCOE/SMutableCodeViewer.cpp
@@ -6,10 +6,10 @@
 #include "EditorDirectories.h"
 #include "IDesktopPlatform.h"
 #include "Framework/MultiBox/MultiBoxBuilder.h"
+#include "Framework/Views/TableViewMetadata.h"
 #include "Misc/Paths.h"
 #include "MuCO/CustomizableObject.h"
 #include "MuCOE/SMutableBoolViewer.h"
-#include "MuCOE/SMutableColorPreviewBox.h"
 #include "MuCOE/SMutableColorViewer.h"
 #include "MuCOE/SMutableConstantsWidget.h"
 #include "MuCOE/SMutableCurveViewer.h"
@@ -107,6 +107,24 @@ public:
 		}
 	}
 
+	/** Depending on the state of the row returns one color or another to be used by the highlighting system */
+	FLinearColor GetHighlightColor() const
+	{
+		if (bShouldBeHiglighted)
+		{
+			if (RowItem->DuplicatedOf)
+			{
+				return HighlightedDuplicatedBoxColor;
+			}
+			else
+			{
+				return HighlightedUniqueRowBoxColor;
+			}
+		}
+
+		return HighlightBoxDefaultColor;
+	}
+	
 	/** Method intended with the generation of the wanted objects for each column*/
 	virtual TSharedRef<SWidget> GenerateWidgetForColumn(const FName& ColumnName) override
 	{
@@ -243,14 +261,14 @@ public:
 				
 			// First coll showing operation name and type
 			+ SHorizontalBox::Slot()
-			.HAlign(EHorizontalAlignment::HAlign_Left)
+			.HAlign(EHorizontalAlignment::HAlign_Fill)
 			[
 				SNew(SOverlay)
 
 				+ SOverlay::Slot()
 				[
-					SAssignNew(this->HighlightingColorBox, SMutableColorPreviewBox)
-					.BoxColor(HighlightBoxDefaultColor)
+					SAssignNew(this->HighlightingColorBox, SColorBlock)
+					.Color(this, &SMutableCodeTreeRow::GetHighlightColor)
 				]
 
 				+ SOverlay::Slot()
@@ -322,23 +340,16 @@ public:
 		return SNullWidget::NullWidget;
 	}
 	
-	/** Set the row background to one or another highlighting color depending if it is a unique or a duplicated row */
-	void Highlight() const
+	/** Marks the row to be highlighted */
+	void Highlight()
 	{
-		if (RowItem->DuplicatedOf)
-		{
-			HighlightingColorBox->SetColor(HighlightedDuplicatedBoxColor);
-		}
-		else
-		{
-			HighlightingColorBox->SetColor(HighlightedUniqueRowBoxColor);
-		}
+		bShouldBeHiglighted = true;
 	}
 
-	/** Sets the background of the row to the default state */
-	void ResetHighlight() const
+	/** Resets the highlighting status */
+	void ResetHighlight()
 	{
-		HighlightingColorBox->SetColor(HighlightBoxDefaultColor);
+		bShouldBeHiglighted = false;
 	}
 
 	/** Returns a reference to the Element this row is representing */
@@ -360,16 +371,16 @@ private:
 	 */
 	
 	/** Custom Widget used to display a color. Used as the background of the text on the row to serve as highlighting Visual Element*/
-	TSharedPtr<SMutableColorPreviewBox> HighlightingColorBox = nullptr;
+	TSharedPtr<SColorBlock> HighlightingColorBox = nullptr;
 	
 	/** The color used to highlight the row if duplicated from another row */
-	const FSlateColor HighlightedDuplicatedBoxColor = FSlateColor(FLinearColor(1, 1, 1, 0.15));
+	const FLinearColor HighlightedDuplicatedBoxColor = FLinearColor(1, 1, 1, 0.15);
 
 	/** The color used to highlight elements that are originals (not duplicates)  */
-	const FSlateColor HighlightedUniqueRowBoxColor = FSlateColor(FLinearColor(1, 1, 1, 0.28));
+	const FLinearColor HighlightedUniqueRowBoxColor = FLinearColor(1, 1, 1, 0.28);
 
 	/** Default color used when the row is not highlighted */
-	const FSlateColor HighlightBoxDefaultColor = FSlateColor(TransparentColor);
+	const FLinearColor HighlightBoxDefaultColor = TransparentColor;
 
 	/*
 	 * Extra data objects
@@ -392,6 +403,8 @@ private:
 
 	/** Color shown on the extra data column when the resource is found to be State Constant */
 	const FLinearColor StateConstantBoxColor = FLinearColor(1,0,0,0.8);
+
+	bool bShouldBeHiglighted = false;
 };
 
 
@@ -1764,7 +1777,7 @@ void SMutableCodeViewer::HighlightDuplicatesOfEntry(const TSharedPtr<FMutableCod
 		if (TreeItem.Get() != InTargetEntry.Get() && TreeItem->MutableOperation == HighlightedOperation)
 		{
 			TSharedPtr<ITableRow> TableRow = TreeView->WidgetFromItem(TreeItem);
-			const SMutableCodeTreeRow* MutableRow = static_cast<SMutableCodeTreeRow*>(TableRow.Get());
+			SMutableCodeTreeRow* MutableRow = static_cast<SMutableCodeTreeRow*>(TableRow.Get());
 			MutableRow->Highlight();
 		}
 	}
@@ -1783,7 +1796,7 @@ void SMutableCodeViewer::ClearHighlightedItems()
 
 			if (TableRow.IsValid())
 			{
-				const SMutableCodeTreeRow* MutableRow = static_cast<SMutableCodeTreeRow*>(TableRow.Get());
+				SMutableCodeTreeRow* MutableRow = static_cast<SMutableCodeTreeRow*>(TableRow.Get());
 				MutableRow->ResetHighlight();
 			}
 
diff --git a/Engine/Plugins/Experimental/Mutable/Source/CustomizableObjectEditor/Private/MuCOE/SMutableColorPreviewBox.cpp b/Engine/Plugins/Experimental/Mutable/Source/CustomizableObjectEditor/Private/MuCOE/SMutableColorPreviewBox.cpp
deleted file mode 100644
diff --git a/Engine/Plugins/Experimental/Mutable/Source/CustomizableObjectEditor/Private/MuCOE/SMutableColorPreviewBox.h b/Engine/Plugins/Experimental/Mutable/Source/CustomizableObjectEditor/Private/MuCOE/SMutableColorPreviewBox.h
deleted file mode 100644
diff --git a/Engine/Plugins/Experimental/Mutable/Source/CustomizableObjectEditor/Private/MuCOE/SMutableColorViewer.cpp b/Engine/Plugins/Experimental/Mutable/Source/CustomizableObjectEditor/Private/MuCOE/SMutableColorViewer.cpp
--- a/Engine/Plugins/Experimental/Mutable/Source/CustomizableObjectEditor/Private/MuCOE/SMutableColorViewer.cpp
+++ b/Engine/Plugins/Experimental/Mutable/Source/CustomizableObjectEditor/Private/MuCOE/SMutableColorViewer.cpp
@@ -2,8 +2,8 @@
 
 #include "MuCOE/SMutableColorViewer.h"
 
-#include "MuCOE/SMutableColorPreviewBox.h"
 #include "Widgets/SBoxPanel.h"
+#include "Widgets/Colors/SColorBlock.h"
 #include "Widgets/Text/STextBlock.h"
 
 #define LOCTEXT_NAMESPACE "CustomizableObjectEditor"
@@ -18,9 +18,9 @@ void SMutableColorViewer::Construct(const FArguments& InArgs)
 	const FText BlueColorValueTitle = LOCTEXT("BlueColorValueTitle", "Blue : ");
 	const FText AlphaColorValueTitle = LOCTEXT("AlphaColorValueTitle", "Alpha : ");
 
-	const int32 IndentationSpace = 16;	
-	const int32 AfterTitleSpacing = 4;
-	const int32 ColorPreviewHorizontalPadding = 18;
+	constexpr int32 IndentationSpace = 16;
+	constexpr int32 AfterTitleSpacing = 4;
+	constexpr int32 ColorPreviewHorizontalPadding = 18;
 
 	ChildSlot
 	[
@@ -141,8 +141,9 @@ void SMutableColorViewer::Construct(const FArguments& InArgs)
 			.Padding(ColorPreviewHorizontalPadding,0)
 			.MaxWidth(120)
 			[
-				SAssignNew(this->ColorPreview, SMutableColorPreviewBox)
-				// .BoxColor(this, &SMutableColorViewer::GetPreviewColor)
+				SAssignNew(ColorPreview, SColorBlock)
+				.UseSRGB(false)
+				.Color(this, &SMutableColorViewer::GetColor)
 			]
 			// -----------------------------------------------------
 		]
@@ -151,38 +152,43 @@ void SMutableColorViewer::Construct(const FArguments& InArgs)
 
 void SMutableColorViewer::SetColor(float InRed, float InGreen, float InBlue, float InAlpha)
 {
-	this->RedValue = InRed;
-	this->GreenValue = InGreen;
-	this->BlueValue = InBlue;
-	this->AlphaValue = InAlpha;
-
-	this->ColorPreview->SetColor( this->GetPreviewColor() );
+	RedValue = InRed;
+	GreenValue = InGreen;
+	BlueValue = InBlue;
+	AlphaValue = InAlpha;
 }
 
 
-FSlateColor SMutableColorViewer::GetPreviewColor() const
+FLinearColor SMutableColorViewer::GetColor() const
 {
-	return FSlateColor { FLinearColor{this->RedValue, this->GreenValue, this->BlueValue, this->AlphaValue} };
+	FLinearColor BoxColor;
+	
+	BoxColor.R = RedValue;
+	BoxColor.G = GreenValue;
+	BoxColor.B = BlueValue;
+	BoxColor.A = AlphaValue;
+	
+	return BoxColor;
 }
 
 FText SMutableColorViewer::GetRedValue() const
 {
-	return FText::AsNumber(this->RedValue);
+	return FText::AsNumber(RedValue);
 }
 
 FText SMutableColorViewer::GetGreenValue() const
 {
-	return FText::AsNumber(this->GreenValue);
+	return FText::AsNumber(GreenValue);
 }
 
 FText SMutableColorViewer::GetBlueValue() const
 {
-	return FText::AsNumber(this->BlueValue);
+	return FText::AsNumber(BlueValue);
 }
 
 FText SMutableColorViewer::GetAlphaValue() const
 {
-	return FText::AsNumber(this->AlphaValue);
+	return FText::AsNumber(AlphaValue);
 }
 
 #undef LOCTEXT_NAMESPACE
diff --git a/Engine/Plugins/Experimental/Mutable/Source/CustomizableObjectEditor/Private/MuCOE/SMutableColorViewer.h b/Engine/Plugins/Experimental/Mutable/Source/CustomizableObjectEditor/Private/MuCOE/SMutableColorViewer.h
--- a/Engine/Plugins/Experimental/Mutable/Source/CustomizableObjectEditor/Private/MuCOE/SMutableColorViewer.h
+++ b/Engine/Plugins/Experimental/Mutable/Source/CustomizableObjectEditor/Private/MuCOE/SMutableColorViewer.h
@@ -4,7 +4,7 @@
 
 #include "Widgets/SCompoundWidget.h"
 
-class SMutableColorPreviewBox;
+class SColorBlock;
 
 
 class SMutableColorViewer : public SCompoundWidget
@@ -28,17 +28,14 @@ private:
 	float AlphaValue = 1.0f;
 
 	/** Color box widget designed to serve as a preview of the color reported by mutable */
-	TSharedPtr<SMutableColorPreviewBox> ColorPreview;
+	TSharedPtr<SColorBlock> ColorPreview;
 
+	/** Callback method invoked by the ColorPreview slate to get a FLinearColor object to display. */
+	FLinearColor GetColor() const;
+	
 	/** Retrieve the values of each color component for the UI Texts to be updated */
 	FText GetRedValue() const;
 	FText GetGreenValue() const;
 	FText GetBlueValue() const;
 	FText GetAlphaValue() const;
-
-	/*
-	* Get a color object that the SMutableColorPreviewBox is able to display 
-	* @return a new FSlateColor generated from the float values held by this widget
-	*/
-	FSlateColor GetPreviewColor() const;
 };
