[mutable] Added way of getting the default values for the parameters of the model.
- Float, int, Bool, Color and Projector parameter default values are now accessible from C++ and blueprint (from an instance of a CustomizableObject object).
- Texture parameter default value can only be accessed from C++. Uint64 (aliased type) seems not able to be used as part of a UFunction.
- Default values for types not set in EMutableParameterType have been excluded since they are not currently being used by the CO or COI systems.
- Generated utility functions to get the relative projector type from CO->Core and the other way around (Part of CO module).
- Replaced code that was converting the type of the projector manually and replaced it with one of the two functions part of the new MutableProjectorUtils.h header.
- Fixed wrong type being returned by "GetCurrentSupportedVersion" from int to int32.
- Moved implementations out of hosting classes for GetProjectorDefaultValue and for IsParameterMultidimensional into CustomizableObject.

diff --git a/Engine/Plugins/Experimental/Mutable/Source/CustomizableObject/Classes/MuCO/CustomizableObject.h b/Engine/Plugins/Experimental/Mutable/Source/CustomizableObject/Classes/MuCO/CustomizableObject.h
--- a/Engine/Plugins/Experimental/Mutable/Source/CustomizableObject/Classes/MuCO/CustomizableObject.h
+++ b/Engine/Plugins/Experimental/Mutable/Source/CustomizableObject/Classes/MuCO/CustomizableObject.h
@@ -2,10 +2,11 @@
  
 #pragma once
 
+#include "RHIDefinitions.h"
 #include "MuCO/CustomizableObjectClothingTypes.h"
 #include "MuCO/CustomizableObjectIdentifier.h"
+#include "MuCO/CustomizableObjectParameterTypeDefinitions.h"
 #include "MuCO/CustomizableObjectUIData.h"
-#include "RHIDefinitions.h"
 
 class FReply;
 
@@ -1187,6 +1188,89 @@ public:
 	UFUNCTION(BlueprintCallable, Category = CustomizableObject)
 	void UnloadMaskOutCache();
 
+private:
+	
+	/** Returns true or false if the parameter with name can be located and it has the type the caller is looking for. It will also
+	 * check if the model has been set to ensure access to it can take place at the calculated parameter index.
+	 * @param InParameterName The name of the parameter to look for.
+	 * @param InParameterType The type the parameter we are looking for we know has. If the name does not match this type this check will fail and return false.
+	 * @param OutParameterIndex The index of the parameter.
+	 * @return True if the parameter can be accessed for it's default values, false if it can not be accessed.
+	 */
+	bool CanDefaultParameterBeAccessed(const FString& InParameterName,const EMutableParameterType& InParameterType, int32& OutParameterIndex) const;
+
+public:
+	/** Get the default value of a parameter of type Float.
+	 * @param InParameterName The name of the Float parameter to get the default value of.
+	 * @return The default value of the provided parameter name.
+	 */
+	UFUNCTION(BlueprintCallable, Category = CustomizableObject)
+	float GetFloatParameterDefaultValue(UPARAM(DisplayName = "Parameter Name") const FString& InParameterName) const;
+	
+	/** Get the default value of a parameter of type Int. 
+	 * @param InParameterName The name of the Int parameter to get the default value of.
+	 * @return The default value of the provided parameter name.
+	 */
+	UFUNCTION(BlueprintCallable, Category = CustomizableObject)
+	int32 GetIntParameterDefaultValue(UPARAM(DisplayName = "Parameter Name") const FString& InParameterName) const;
+ 
+	/** Get the default value of a parameter of type Bool.
+	 * @param InParameterName The name of the Bool parameter to get the default value of.
+	 * @return The default value of the provided parameter name.
+	 */
+	UFUNCTION(BlueprintCallable, Category = CustomizableObject)
+	bool GetBoolParameterDefaultValue(UPARAM(DisplayName = "Parameter Name") const FString& InParameterName) const;
+
+	/** Get the default value of a parameter of type Color.
+	 * @param InParameterName The name of the Color parameter to get the default value of.
+	 * @return The default value of the provided parameter name.
+	*/
+	UFUNCTION(BlueprintCallable, Category = CustomizableObject)
+	FLinearColor GetColorParameterDefaultValue(UPARAM(DisplayName = "Parameter Name") const FString& InParameterName) const;
+	
+	/** Get the default value of a parameter of type Projector.
+	 * @param InParameterName The name of the Projector parameter to get the default value of.
+	 * @param OutPos The default position of the Projector.
+	 * @param OutDirection The default projection direction of the Projector.
+	 * @param OutUp The default up vector of the Projector.
+	 * @param OutScale The default scale of the Projector.
+	 * @param OutAngle The default angle of the Projector.
+	 * @param OutType The default type of the Projector.
+	 */
+	UFUNCTION(BlueprintCallable, Category = CustomizableObject)
+	void GetProjectorParameterDefaultValue (
+		UPARAM(DisplayName = "Parameter Name") const FString& InParameterName,
+		UPARAM(DisplayName = "Possition") FVector3f& OutPos,
+		UPARAM(DisplayName = "Direction") FVector3f& OutDirection,
+		UPARAM(DisplayName = "Up") FVector3f& OutUp,
+		UPARAM(DisplayName = "Scale") FVector3f& OutScale,
+		UPARAM(DisplayName = "Angle") float& OutAngle,
+		UPARAM(DisplayName = "Type") ECustomizableObjectProjectorType& OutType) const;
+
+	/** Get the default value of a projector with the provided name
+	 * @param InParameterName The name of the parameter to get the default value of.
+	 * @return A data structure containing all the default data for the targeted projector parameter.
+	 */
+	FCustomizableObjectProjector GetProjectorParameterDefaultValue ( const FString& InParameterName) const;
+	
+	/** Get the default value of a parameter of type Texture.
+	 * @param InParameterName The name of the Projector parameter to get the default value of.
+	 * @return An uint64 (EXTERNAL_IMAGE_ID) pointing at the default parameter's texture.
+	 */
+	uint64 GetTextureParameterDefaultValue (const FString& InParameterName) const;
+
+	/** Return true or false depending if the parameter at the index provided is multidimensional or not.
+	 * @param InParamIndex The index of the parameter to check.
+	 * @return True if the parameter is multilayer and false if it is not.
+	 */
+	bool IsParameterMultidimensional(const int32& InParamIndex) const;
+	
+	/** Return true or false depending if the parameter at the index provided is multidimensional or not.
+	 * @param InParameterName The name of the parameter to check.
+	 * @return True if the parameter is multilayer and false if it is not.
+	 */
+	bool IsParameterMultidimensional (const FString& InParameterName) const;
+	
 private: 
 
 	// This is information about the parameters in the model that is generated at model compile time.
@@ -1220,7 +1304,7 @@ public:
 	void UpdateVersionId();
 	FGuid GetVersionId() const { return VersionId; }
 
-	int GetCurrentSupportedVersion() const { return CurrentSupportedVersion; };
+	int32 GetCurrentSupportedVersion() const { return CurrentSupportedVersion; };
 
 	// Compose folder name where the data is stored
 	FString GetCompiledDataFolderPath(bool bIsEditorData) const;
diff --git a/Engine/Plugins/Experimental/Mutable/Source/CustomizableObject/Classes/MuCO/CustomizableObjectInstance.h b/Engine/Plugins/Experimental/Mutable/Source/CustomizableObject/Classes/MuCO/CustomizableObjectInstance.h
--- a/Engine/Plugins/Experimental/Mutable/Source/CustomizableObject/Classes/MuCO/CustomizableObjectInstance.h
+++ b/Engine/Plugins/Experimental/Mutable/Source/CustomizableObject/Classes/MuCO/CustomizableObjectInstance.h
@@ -270,15 +270,13 @@ public:
 	UFUNCTION(BlueprintCallable, Category = CustomizableObjectInstance)
 	bool IsParameterRelevant(const FString& ParamName) const;
 
-	UFUNCTION(BlueprintCallable, Category = CustomizableObjectInstance)
+	UFUNCTION(BlueprintCallable, Category = CustomizableObjectInstance, 
+		meta = (DeprecatedFunction, DeprecationMessage = "Use the method with the same name located at the Customizable Object instead."))
 	bool IsParamMultidimensional(const FString& ParamName) const;
-
+	
 	UFUNCTION(BlueprintCallable, Category = CustomizableObjectInstance)
 	int32 GetProjectorValueRange(const FString& ParamName) const;
 	
-	bool IsParamMultidimensional(int32 ParamIndex) const;
-	FCustomizableObjectProjector GetProjectorDefaultValue(int32 ParamIndex) const;
-
 	bool bShowOnlyRuntimeParameters = true;
 	bool bShowOnlyRelevantParameters = true;
 
diff --git a/Engine/Plugins/Experimental/Mutable/Source/CustomizableObject/CustomizableObject.Build.cs b/Engine/Plugins/Experimental/Mutable/Source/CustomizableObject/CustomizableObject.Build.cs
--- a/Engine/Plugins/Experimental/Mutable/Source/CustomizableObject/CustomizableObject.Build.cs
+++ b/Engine/Plugins/Experimental/Mutable/Source/CustomizableObject/CustomizableObject.Build.cs
@@ -40,7 +40,7 @@ public class CustomizableObject : ModuleRules
             "Engine",
 			"SkeletalMerging",
 			"ClothingSystemRuntimeCommon",
-            "GameplayTags",
+            "GameplayTags"
 		});
 
         PrivateIncludePathModuleNames.AddRange(
diff --git a/Engine/Plugins/Experimental/Mutable/Source/CustomizableObject/Private/MuCO/CustomizableObject.cpp b/Engine/Plugins/Experimental/Mutable/Source/CustomizableObject/Private/MuCO/CustomizableObject.cpp
--- a/Engine/Plugins/Experimental/Mutable/Source/CustomizableObject/Private/MuCO/CustomizableObject.cpp
+++ b/Engine/Plugins/Experimental/Mutable/Source/CustomizableObject/Private/MuCO/CustomizableObject.cpp
@@ -12,13 +12,14 @@
 #include "Input/Reply.h"
 #include "Interfaces/ITargetPlatform.h"
 #include "Materials/MaterialInterface.h"
+#include "Misc/DataValidation.h"
 #include "Misc/PackageName.h"
 #include "Misc/Paths.h"
-#include "Misc/DataValidation.h"
 #include "MuCO/CustomizableObjectInstance.h"
 #include "MuCO/CustomizableObjectPrivate.h"
 #include "MuCO/CustomizableObjectSystem.h"
 #include "MuCO/ICustomizableObjectModule.h"
+#include "MuCO/MutableProjectorTypeUtils.h"
 #include "MuCO/UnrealMutableModelDiskStreamer.h"
 #include "MuR/ModelPrivate.h"
 #include "PhysicsEngine/PhysicsAsset.h"
@@ -1422,6 +1423,187 @@ void UCustomizableObject::UnloadMaskOutCache()
 }
 
 
+bool UCustomizableObject::CanDefaultParameterBeAccessed(const FString& InParameterName,
+	const EMutableParameterType& InParameterType, int32& OutParameterIndex) const
+{
+	OutParameterIndex = FindParameter(InParameterName);
+	if (OutParameterIndex == INDEX_NONE)
+	{
+		UE_LOG(LogMutable,Error,TEXT("The parameter with name [%s] does not exist in the CO with name [%s]."), *InParameterName,*GetName());
+		return false;
+	}
+
+	const TSharedPtr<const mu::Model> MutableModel = PrivateData->GetModel();
+	if (!MutableModel)
+	{
+		UE_LOG(LogMutable,Error,TEXT("Mutable model not set in [%s] CO. Compile this Customizable Object in order to be able to acces the default values of the parameters."),*GetName());
+		return false;
+	}
+	
+	const int32 ParameterCount = MutableModel->GetPrivate()->m_program.m_parameters.Num();
+	check (OutParameterIndex < ParameterCount);		// this failing would mean the data between CO and core is not synced 
+	
+	if ( GetParameterType(OutParameterIndex) != InParameterType)
+	{
+		UE_LOG(LogMutable,Error,TEXT("The parameter with name [%s] is not a [%s] parameter of the [%s] CO."), *InParameterName,*UEnum::GetValueAsString(InParameterType),*GetName());
+		return false;
+	}
+	
+	return true;
+}
+
+float UCustomizableObject::GetFloatParameterDefaultValue(const FString& InParameterName) const
+{
+	float Value = -1.0f;
+	
+	int32 ParameterIndex;
+	if (!CanDefaultParameterBeAccessed(InParameterName,EMutableParameterType::Float,ParameterIndex))
+	{
+		return Value;
+	}
+
+	const TSharedPtr<const mu::Model> MutableModel = PrivateData->GetModel();
+	Value = MutableModel->GetPrivate()->m_program.m_parameters[ParameterIndex].m_defaultValue.m_float;
+	return Value;
+}
+
+
+int32 UCustomizableObject::GetIntParameterDefaultValue(const FString& InParameterName) const
+{
+	int32 Value = -1;
+	
+	int32 ParameterIndex;
+	if (!CanDefaultParameterBeAccessed(InParameterName,EMutableParameterType::Int,ParameterIndex))
+	{
+		return Value;
+	}
+
+	const TSharedPtr<const mu::Model> MutableModel = PrivateData->GetModel();
+	Value = MutableModel->GetPrivate()->m_program.m_parameters[ParameterIndex].m_defaultValue.m_int;
+	return Value;
+}
+
+
+bool UCustomizableObject::GetBoolParameterDefaultValue(const FString& InParameterName) const
+{
+	bool Value = false;
+	
+	int32 ParameterIndex;
+	if (!CanDefaultParameterBeAccessed(InParameterName,EMutableParameterType::Bool,ParameterIndex))
+	{
+		return Value;
+	}
+	
+	const TSharedPtr<const mu::Model> MutableModel = PrivateData->GetModel();
+	Value = MutableModel->GetPrivate()->m_program.m_parameters[ParameterIndex].m_defaultValue.m_bool;
+	return Value;
+}
+
+
+FLinearColor UCustomizableObject::GetColorParameterDefaultValue(const FString& InParameterName) const
+{
+	FLinearColor Value = FLinearColor::Black;
+	
+	int32 ParameterIndex;
+	if (!CanDefaultParameterBeAccessed(InParameterName,EMutableParameterType::Color,ParameterIndex))
+	{
+		return Value;
+	}
+	
+	const mu::Model::Private* ModelPrivate = PrivateData->GetModel()->GetPrivate();
+	Value.R = ModelPrivate->m_program.m_parameters[ParameterIndex].m_defaultValue.m_colour[0];
+	Value.G = ModelPrivate->m_program.m_parameters[ParameterIndex].m_defaultValue.m_colour[1];
+	Value.B = ModelPrivate->m_program.m_parameters[ParameterIndex].m_defaultValue.m_colour[2];
+	Value.A = 1.0f;
+	
+	return Value;
+}
+
+
+void UCustomizableObject::GetProjectorParameterDefaultValue(const FString& InParameterName, FVector3f& OutPos,
+	FVector3f& OutDirection, FVector3f& OutUp, FVector3f& OutScale, float& OutAngle,
+	ECustomizableObjectProjectorType& OutType) const
+{
+	int32 ParameterIndex;
+	if (!CanDefaultParameterBeAccessed(InParameterName,EMutableParameterType::Projector,ParameterIndex))
+	{
+		return;
+	}
+	
+	const TSharedPtr<const mu::Model> MutableModel = PrivateData->GetModel();
+	
+	mu::FProjector MutableProjector = MutableModel->GetPrivate()->m_program.m_parameters[ParameterIndex].m_defaultValue.m_projector;
+	OutType = ProjectorUtils::GetEquivalentProjectorType(MutableProjector.type);
+	OutPos = MutableProjector.position;
+	OutDirection = MutableProjector.direction;
+	OutUp = MutableProjector.up;
+	OutScale = MutableProjector.scale;
+	OutAngle = MutableProjector.projectionAngle;
+}
+
+FCustomizableObjectProjector UCustomizableObject::GetProjectorParameterDefaultValue(const FString& InParameterName) const 
+{
+	FCustomizableObjectProjector Value;
+	
+	int32 ParameterIndex;
+	if (!CanDefaultParameterBeAccessed(InParameterName,EMutableParameterType::Projector,ParameterIndex))
+	{
+		return Value;
+	}
+	
+	const TSharedPtr<const mu::Model> MutableModel = PrivateData->GetModel();
+
+	mu::FProjector MutableProjector = MutableModel->GetPrivate()->m_program.m_parameters[ParameterIndex].m_defaultValue.m_projector;
+	Value.Angle = MutableProjector.projectionAngle;
+	Value.Direction = MutableProjector.direction;
+	Value.Position = MutableProjector.position;
+	Value.Scale = MutableProjector.scale;
+	Value.Up = MutableProjector.up;
+	Value.ProjectionType = ProjectorUtils::GetEquivalentProjectorType(MutableProjector.type);
+
+	return Value;
+}
+
+
+uint64 UCustomizableObject::GetTextureParameterDefaultValue(const FString& InParameterName) const
+{
+	uint64 Value = false;
+	
+	int32 ParameterIndex;
+	if (!CanDefaultParameterBeAccessed(InParameterName,EMutableParameterType::Texture,ParameterIndex))
+	{
+		return Value;
+	}
+	
+	const TSharedPtr<const mu::Model> MutableModel = PrivateData->GetModel();
+	Value = MutableModel->GetPrivate()->m_program.m_parameters[ParameterIndex].m_defaultValue.m_image;
+	return Value;
+}
+
+bool UCustomizableObject::IsParameterMultidimensional(const FString& InParameterName) const
+{
+	const int32 ParameterIndex = FindParameter(InParameterName);
+	if (ParameterIndex == INDEX_NONE)
+	{
+		UE_LOG(LogMutable,Error,TEXT("The parameter with name %s does not exist in the CO with name %s."), *InParameterName,*GetName());
+		checkNoEntry();
+	}
+
+	return IsParameterMultidimensional(ParameterIndex);
+}
+
+bool UCustomizableObject::IsParameterMultidimensional(const int32& InParamIndex) const
+{
+	check(InParamIndex != INDEX_NONE);
+	
+	const mu::ParametersPtr MutableParameters = mu::Model::NewParameters(this->GetPrivate()->GetModel());
+	check(InParamIndex < MutableParameters->GetCount());
+	const mu::RangeIndexPtr RangeIdxPtr = MutableParameters->NewRangeIndex(InParamIndex);
+
+	return RangeIdxPtr.get() != nullptr;
+}
+
+
 void UCustomizableObject::ApplyStateForcedValuesToParameters( int32 State, mu::Parameters* Parameters)
 {
 	FParameterUIData StateMetaData = GetStateUIMetadataFromIndex(State);
diff --git a/Engine/Plugins/Experimental/Mutable/Source/CustomizableObject/Private/MuCO/CustomizableObjectInstance.cpp b/Engine/Plugins/Experimental/Mutable/Source/CustomizableObject/Private/MuCO/CustomizableObjectInstance.cpp
--- a/Engine/Plugins/Experimental/Mutable/Source/CustomizableObject/Private/MuCO/CustomizableObjectInstance.cpp
+++ b/Engine/Plugins/Experimental/Mutable/Source/CustomizableObject/Private/MuCO/CustomizableObjectInstance.cpp
@@ -29,6 +29,7 @@
 
 // Required for engine branch preprocessor defines.
 #include "MuCO/UnrealPortabilityHelpers.h"
+#include "MuR/Model.h"
 #include "MuR/OpImagePixelFormat.h"
 #include "MuR/OpImageSwizzle.h"
 #include "PhysicsEngine/PhysicsConstraintTemplate.h"
@@ -475,22 +476,15 @@ FString UCustomizableObjectInstance::GetDesc()
 	return FString::Printf(TEXT("Instance of [%s]"), *ObjectName);
 }
 
-
-bool UCustomizableObjectInstance::IsParamMultidimensional(const FString& ParamName) const
-{
-	return Descriptor.IsParamMultidimensional(ParamName);
-}
-
-
 int32 UCustomizableObjectInstance::GetProjectorValueRange(const FString& ParamName) const
 {
 	return Descriptor.GetProjectorValueRange(ParamName);
 }
 
 
-bool UCustomizableObjectInstance::IsParamMultidimensional(int32 ParamIndex) const
+bool UCustomizableObjectInstance::IsParamMultidimensional(const FString& ParamName) const
 {
-	return Descriptor.IsParamMultidimensional(ParamIndex);
+	return GetCustomizableObject()->IsParameterMultidimensional(ParamName);
 }
 
 
@@ -519,12 +513,6 @@ void UCustomizableInstancePrivateData::ReleaseMutableTexture(int32 MutableTextur
 }
 
 
-FCustomizableObjectProjector UCustomizableObjectInstance::GetProjectorDefaultValue(const int32 ParamIndex) const
-{
-	return Descriptor.GetProjectorDefaultValue(ParamIndex);
-}
-
-
 void UCustomizableInstancePrivateData::InstanceUpdateFlags(const UCustomizableObjectInstance& Public)
 {
 	const UCustomizableObject* CustomizableObject = Public.GetCustomizableObject();
diff --git a/Engine/Plugins/Experimental/Mutable/Source/CustomizableObject/Private/MuCO/CustomizableObjectInstanceDescriptor.cpp b/Engine/Plugins/Experimental/Mutable/Source/CustomizableObject/Private/MuCO/CustomizableObjectInstanceDescriptor.cpp
--- a/Engine/Plugins/Experimental/Mutable/Source/CustomizableObject/Private/MuCO/CustomizableObjectInstanceDescriptor.cpp
+++ b/Engine/Plugins/Experimental/Mutable/Source/CustomizableObject/Private/MuCO/CustomizableObjectInstanceDescriptor.cpp
@@ -5,6 +5,7 @@
 #include "MuCO/CustomizableObject.h"
 #include "MuCO/CustomizableObjectPrivate.h"
 #include "MuCO/DefaultImageProvider.h"
+#include "MuCO/MutableProjectorTypeUtils.h"
 #include "MuR/Model.h"
 #include "MuR/MutableMemory.h"
 #include "MuR/Parameters.h"
@@ -126,7 +127,7 @@ void FCustomizableObjectInstanceDescriptor::SaveDescriptor(FArchive& Ar)
 					Value = CustomizableObject->FindIntParameterValue(ModelParameterIndex,P.ParameterValueName);
 
 				int32 ParameterIndexInObject = CustomizableObject->FindParameter(IntParameters[IntParameterIndex].ParameterName);
-				bIsParamMultidimensional = IsParamMultidimensional(ParameterIndexInObject);
+				bIsParamMultidimensional = CustomizableObject->IsParameterMultidimensional(ParameterIndexInObject);
 
 				if (bIsParamMultidimensional)
 				{
@@ -314,7 +315,7 @@ void FCustomizableObjectInstanceDescriptor::LoadDescriptor(FArchive& Ar)
 
 			const int32 IntParameterIndex = FindIntParameterNameIndex(Name);
 			const int32 ParameterIndexInObject = CustomizableObject->FindParameter(IntParameters[IntParameterIndex].ParameterName);
-			const bool bIsParamMultidimensional = IsParamMultidimensional(ParameterIndexInObject);
+			const bool bIsParamMultidimensional = CustomizableObject->IsParameterMultidimensional(ParameterIndexInObject);
 
 			if (bUseCompactDescriptor)
 			{
@@ -903,28 +904,15 @@ void FCustomizableObjectInstanceDescriptor::ReloadParameters()
 						&Value.Angle,
 						RangeIndex);
 					
-					switch (Type)
+					Value.ProjectionType = ProjectorUtils::GetEquivalentProjectorType(Type);
+					if (Value.ProjectionType == ECustomizableObjectProjectorType::Cylindrical)
 					{
-					case mu::PROJECTOR_TYPE::PLANAR:
-						Value.ProjectionType = ECustomizableObjectProjectorType::Planar;
-						break;
-
-					case mu::PROJECTOR_TYPE::CYLINDRICAL:
 						// Unapply strange swizzle for scales.
 						// TODO: try to avoid this
-						Value.ProjectionType = ECustomizableObjectProjectorType::Cylindrical;
 						Value.Direction = -Value.Direction;
 						Value.Up = -Value.Up;
 						Value.Scale[2] = -Value.Scale[0];
 						Value.Scale[0] = Value.Scale[1] = Value.Scale[1] * 2.0f;
-						break;
-
-					case mu::PROJECTOR_TYPE::WRAPPING:
-						Value.ProjectionType = ECustomizableObjectProjectorType::Wrapping;
-						break;
-
-					default:
-						check(false); // Not implemented.
 					}
 				}; 
 
@@ -1153,12 +1141,12 @@ const FString& FCustomizableObjectInstanceDescriptor::GetIntParameterSelectedOpt
 	{
 		if (RangeIndex == -1)
 		{
-			check(!IsParamMultidimensional(ParameterIndexInObject)); // This param is multidimensional, it must have a RangeIndex of 0 or more
+			check(!CustomizableObject->IsParameterMultidimensional(ParameterIndexInObject)); // This param is multidimensional, it must have a RangeIndex of 0 or more
 			return IntParameters[IntParamIndex].ParameterValueName;
 		}
 		else
 		{
-			check(IsParamMultidimensional(ParameterIndexInObject)); // This param is not multidimensional, it must have a RangeIndex of -1
+			check(CustomizableObject->IsParameterMultidimensional(ParameterIndexInObject)); // This param is not multidimensional, it must have a RangeIndex of -1
 
 			if (IntParameters[IntParamIndex].ParameterRangeValueNames.IsValidIndex(RangeIndex))
 			{
@@ -1197,12 +1185,12 @@ void FCustomizableObjectInstanceDescriptor::SetIntParameterSelectedOption(const
 			
 			if (RangeIndex == -1)
 			{
-				check(!IsParamMultidimensional(ParameterIndexInObject)); // This param is multidimensional, it must have a RangeIndex of 0 or more
+				check(!CustomizableObject->IsParameterMultidimensional(ParameterIndexInObject)); // This param is multidimensional, it must have a RangeIndex of 0 or more
 				IntParameters[IntParamIndex].ParameterValueName = SelectedOption;
 			}
 			else
 			{
-				check(IsParamMultidimensional(ParameterIndexInObject)); // This param is not multidimensional, it must have a RangeIndex of -1
+				check(CustomizableObject->IsParameterMultidimensional(ParameterIndexInObject)); // This param is not multidimensional, it must have a RangeIndex of -1
 
 				if (!IntParameters[IntParamIndex].ParameterRangeValueNames.IsValidIndex(RangeIndex))
 				{
@@ -1238,12 +1226,12 @@ float FCustomizableObjectInstanceDescriptor::GetFloatParameterSelectedOption(con
 	{
 		if (RangeIndex == -1)
 		{
-			check(!IsParamMultidimensional(ParameterIndexInObject)); // This param is multidimensional, it must have a RangeIndex of 0 or more
+			check(!CustomizableObject->IsParameterMultidimensional(ParameterIndexInObject)); // This param is multidimensional, it must have a RangeIndex of 0 or more
 			return FloatParameters[FloatParamIndex].ParameterValue;
 		}
 		else
 		{
-			check(IsParamMultidimensional(ParameterIndexInObject)); // This param is not multidimensional, it must have a RangeIndex of -1
+			check(CustomizableObject->IsParameterMultidimensional(ParameterIndexInObject)); // This param is not multidimensional, it must have a RangeIndex of -1
 
 			if (FloatParameters[FloatParamIndex].ParameterRangeValues.IsValidIndex(RangeIndex))
 			{
@@ -1268,12 +1256,12 @@ void FCustomizableObjectInstanceDescriptor::SetFloatParameterSelectedOption(cons
 	{
 		if (RangeIndex == -1)
 		{
-			check(!IsParamMultidimensional(ParameterIndexInObject)); // This param is multidimensional, it must have a RangeIndex of 0 or more
+			check(!CustomizableObject->IsParameterMultidimensional(ParameterIndexInObject)); // This param is multidimensional, it must have a RangeIndex of 0 or more
 			FloatParameters[FloatParamIndex].ParameterValue = FloatValue;
 		}
 		else
 		{
-			check(IsParamMultidimensional(ParameterIndexInObject)); // This param is not multidimensional, it must have a RangeIndex of -1
+			check(CustomizableObject->IsParameterMultidimensional(ParameterIndexInObject)); // This param is not multidimensional, it must have a RangeIndex of -1
 
 			if (!FloatParameters[FloatParamIndex].ParameterRangeValues.IsValidIndex(RangeIndex))
 			{
@@ -1301,12 +1289,12 @@ uint64 FCustomizableObjectInstanceDescriptor::GetTextureParameterSelectedOption(
 	{
 		if (RangeIndex == -1)
 		{
-			check(!IsParamMultidimensional(ParameterIndexInObject)); // This param is multidimensional, it must have a RangeIndex of 0 or more
+			check(!CustomizableObject->IsParameterMultidimensional(ParameterIndexInObject)); // This param is multidimensional, it must have a RangeIndex of 0 or more
 			return TextureParameters[TextureParamIndex].ParameterValue;
 		}
 		else
 		{
-			check(IsParamMultidimensional(ParameterIndexInObject)); // This param is not multidimensional, it must have a RangeIndex of -1
+			check(CustomizableObject->IsParameterMultidimensional(ParameterIndexInObject)); // This param is not multidimensional, it must have a RangeIndex of -1
 
 			if (TextureParameters[TextureParamIndex].ParameterRangeValues.IsValidIndex(RangeIndex))
 			{
@@ -1340,12 +1328,12 @@ void FCustomizableObjectInstanceDescriptor::SetTextureParameterSelectedOption(co
 	{
 		if (RangeIndex == -1)
 		{
-			check(!IsParamMultidimensional(ParameterIndexInObject)); // This param is multidimensional, it must have a RangeIndex of 0 or more
+			check(!CustomizableObject->IsParameterMultidimensional(ParameterIndexInObject)); // This param is multidimensional, it must have a RangeIndex of 0 or more
 			TextureParameters[TextureParamIndex].ParameterValue = TextureValue;
 		}
 		else
 		{
-			check(IsParamMultidimensional(ParameterIndexInObject)); // This param is not multidimensional, it must have a RangeIndex of -1
+			check(CustomizableObject->IsParameterMultidimensional(ParameterIndexInObject)); // This param is not multidimensional, it must have a RangeIndex of -1
 
 			if (!TextureParameters[TextureParamIndex].ParameterRangeValues.IsValidIndex(RangeIndex))
 			{
@@ -1463,12 +1451,12 @@ void FCustomizableObjectInstanceDescriptor::SetProjectorValue(const FString& Pro
 
 		if (RangeIndex == -1)
 		{
-			check(!IsParamMultidimensional(ParameterIndexInObject)); // This param is multidimensional, it must have a RangeIndex of 0 or more
+			check(!CustomizableObject->IsParameterMultidimensional(ParameterIndexInObject)); // This param is multidimensional, it must have a RangeIndex of 0 or more
 			ProjectorParameters[ProjectorParamIndex].Value = ProjectorData;
 		}
 		else
 		{
-			check(IsParamMultidimensional(ParameterIndexInObject)); // This param is not multidimensional, it must have a RangeIndex of -1
+			check(CustomizableObject->IsParameterMultidimensional(ParameterIndexInObject)); // This param is not multidimensional, it must have a RangeIndex of -1
 
 			if (!ProjectorParameters[ProjectorParamIndex].RangeValues.IsValidIndex(RangeIndex))
 			{
@@ -1499,12 +1487,12 @@ void FCustomizableObjectInstanceDescriptor::SetProjectorPosition(const FString&
 
 		if (RangeIndex == -1)
 		{
-			check(!IsParamMultidimensional(ParameterIndexInObject)); // This param is multidimensional, it must have a RangeIndex of 0 or more
+			check(!CustomizableObject->IsParameterMultidimensional(ParameterIndexInObject)); // This param is multidimensional, it must have a RangeIndex of 0 or more
 			ProjectorParameters[ProjectorParamIndex].Value = ProjectorData;
 		}
 		else
 		{
-			check(IsParamMultidimensional(ParameterIndexInObject)); // This param is not multidimensional, it must have a RangeIndex of -1
+			check(CustomizableObject->IsParameterMultidimensional(ParameterIndexInObject)); // This param is not multidimensional, it must have a RangeIndex of -1
 
 			if (!ProjectorParameters[ProjectorParamIndex].RangeValues.IsValidIndex(RangeIndex))
 			{
@@ -1537,13 +1525,13 @@ void FCustomizableObjectInstanceDescriptor::GetProjectorValue(const FString& Pro
 
 		if (RangeIndex == -1)
 		{
-			check(!IsParamMultidimensional(ParameterIndexInObject)); // This param is multidimensional, it must have a RangeIndex of 0 or more
+			check(!CustomizableObject->IsParameterMultidimensional(ParameterIndexInObject)); // This param is multidimensional, it must have a RangeIndex of 0 or more
 
 			Projector = &ProjectorParameters[ProjectorParamIndex].Value;
 		}
 		else
 		{
-			check(IsParamMultidimensional(ParameterIndexInObject)); // This param is not multidimensional, it must have a RangeIndex of -1
+			check(CustomizableObject->IsParameterMultidimensional(ParameterIndexInObject)); // This param is not multidimensional, it must have a RangeIndex of -1
 			check(ProjectorParameters[ProjectorParamIndex].RangeValues.IsValidIndex(RangeIndex));
 
 			Projector = &ProjectorParameters[ProjectorParamIndex].RangeValues[RangeIndex];
@@ -1579,13 +1567,13 @@ void FCustomizableObjectInstanceDescriptor::GetProjectorValueF(const FString& Pr
 
 		if (RangeIndex == -1)
 		{
-			check(!IsParamMultidimensional(ParameterIndexInObject)); // This param is multidimensional, it must have a RangeIndex of 0 or more
+			check(!CustomizableObject->IsParameterMultidimensional(ParameterIndexInObject)); // This param is multidimensional, it must have a RangeIndex of 0 or more
 
 			Projector = &ProjectorParameters[ProjectorParamIndex].Value;
 		}
 		else
 		{
-			check(IsParamMultidimensional(ParameterIndexInObject)); // This param is not multidimensional, it must have a RangeIndex of -1
+			check(CustomizableObject->IsParameterMultidimensional(ParameterIndexInObject)); // This param is not multidimensional, it must have a RangeIndex of -1
 			check(ProjectorParameters[ProjectorParamIndex].RangeValues.IsValidIndex(RangeIndex));
 
 			Projector = &ProjectorParameters[ProjectorParamIndex].RangeValues[RangeIndex];
@@ -1851,28 +1839,6 @@ int32 FCustomizableObjectInstanceDescriptor::FindProjectorParameterNameIndex(con
 	return -1;
 }
 
-
-bool FCustomizableObjectInstanceDescriptor::IsParamMultidimensional(const FString& ParamName) const
-{
-	check(CustomizableObject);
-
-	const int32 ParameterIndex = CustomizableObject->FindParameter(ParamName);
-	return IsParamMultidimensional(ParameterIndex);
-}
-
-
-bool FCustomizableObjectInstanceDescriptor::IsParamMultidimensional(const int32 ParamIndex) const
-{
-	check(CustomizableObject);
-
-	const mu::ParametersPtr MutableParameters = mu::Model::NewParameters(CustomizableObject->GetPrivate()->GetModel());
-	check(ParamIndex < MutableParameters->GetCount());
-	const mu::RangeIndexPtr RangeIdxPtr = MutableParameters->NewRangeIndex(ParamIndex);
-
-	return RangeIdxPtr.get() != nullptr;
-}
-
-
 int32 FCustomizableObjectInstanceDescriptor::GetProjectorValueRange(const FString& ParamName) const
 {
 	check(CustomizableObject);
@@ -1925,7 +1891,7 @@ int32 FCustomizableObjectInstanceDescriptor::AddValueToProjectorRange(const FStr
 	if (projectorParameterIndex != -1)
 	{
 		FCustomizableObjectProjectorParameterValue& ProjectorParameter = ProjectorParameters[projectorParameterIndex];
-		const FCustomizableObjectProjector Projector = GetProjectorDefaultValue(CustomizableObject->FindParameter(ParamName));
+		const FCustomizableObjectProjector Projector = GetCustomizableObject()->GetProjectorParameterDefaultValue(ParamName);
 		return ProjectorParameter.RangeValues.Add(Projector);
 	}
 
@@ -2061,53 +2027,6 @@ int32 FCustomizableObjectInstanceDescriptor::RemoveValueFromProjectorRange(const
 }
 
 
-FCustomizableObjectProjector FCustomizableObjectInstanceDescriptor::GetProjectorDefaultValue(int32 const ParamIndex) const
-{
-	check(CustomizableObject);
-
-	const mu::ParametersPtr MutableParameters = mu::Model::NewParameters(CustomizableObject->GetPrivate()->GetModel());
-	check(ParamIndex < MutableParameters->GetCount());
-
-	FCustomizableObjectProjector Projector;
-
-	mu::PROJECTOR_TYPE type;
-	
-	MutableParameters->GetProjectorValue(ParamIndex,
-		&type,
-		&Projector.Position,
-		&Projector.Direction,
-		&Projector.Up,
-		&Projector.Scale,
-		&Projector.Angle,
-		nullptr);
-
-	switch (type)
-	{
-	case mu::PROJECTOR_TYPE::PLANAR:
-		Projector.ProjectionType = ECustomizableObjectProjectorType::Planar;
-		break;
-
-	case mu::PROJECTOR_TYPE::CYLINDRICAL:
-		// Unapply strange swizzle for scales.
-		// TODO: try to avoid this
-		Projector.ProjectionType = ECustomizableObjectProjectorType::Cylindrical;
-		Projector.Direction = -Projector.Direction;
-		Projector.Up = -Projector.Up;
-		Projector.Scale[2] = -Projector.Scale[0];
-		Projector.Scale[0] = Projector.Scale[1] = Projector.Scale[1] * 2.0f;
-		break;
-
-	case mu::PROJECTOR_TYPE::WRAPPING:
-		Projector.ProjectionType = ECustomizableObjectProjectorType::Wrapping;
-		break;
-	default:
-		unimplemented()
-	}
-
-	return Projector;
-}
-
-
 int32 FCustomizableObjectInstanceDescriptor::GetState() const
 {
 	return State;
@@ -2163,7 +2082,7 @@ void FCustomizableObjectInstanceDescriptor::SetRandomValues()
 		const int32 ParameterIndexInCO = CustomizableObject->FindParameter(IntParameters[i].ParameterName);
 
 		// TODO: Randomize multidimensional parameters
-		if (ParameterIndexInCO >= 0 && !IsParamMultidimensional(ParameterIndexInCO))
+		if (ParameterIndexInCO >= 0 && !CustomizableObject->IsParameterMultidimensional(ParameterIndexInCO))
 		{
 			const int32 NumValues = CustomizableObject->GetIntParameterNumOptions(ParameterIndexInCO);
 			if (NumValues > 0)
diff --git a/Engine/Plugins/Experimental/Mutable/Source/CustomizableObject/Private/MuCO/MultilayerProjector.cpp b/Engine/Plugins/Experimental/Mutable/Source/CustomizableObject/Private/MuCO/MultilayerProjector.cpp
--- a/Engine/Plugins/Experimental/Mutable/Source/CustomizableObject/Private/MuCO/MultilayerProjector.cpp
+++ b/Engine/Plugins/Experimental/Mutable/Source/CustomizableObject/Private/MuCO/MultilayerProjector.cpp
@@ -150,7 +150,7 @@ void FMultilayerProjector::CreateLayer(FCustomizableObjectInstanceDescriptor& De
 		const int32 ProjectorParameterIndex = Descriptor.FindProjectorParameterNameIndex(ParamNameString);
 		
 		FCustomizableObjectProjectorParameterValue& ProjectorParameter = Descriptor.GetProjectorParameters()[ProjectorParameterIndex];
-		const FCustomizableObjectProjector Projector = Descriptor.GetProjectorDefaultValue(Object->FindParameter(ParamNameString));
+		const FCustomizableObjectProjector Projector = Descriptor.GetCustomizableObject()->GetProjectorParameterDefaultValue(ParamNameString);
 		ProjectorParameter.RangeValues.Insert(Projector, Index);
 	}
 	
diff --git a/Engine/Plugins/Experimental/Mutable/Source/CustomizableObject/Private/MuCO/MutableProjectorTypeUtils.cpp b/Engine/Plugins/Experimental/Mutable/Source/CustomizableObject/Private/MuCO/MutableProjectorTypeUtils.cpp
new file mode 100644
--- /dev/null
+++ b/Engine/Plugins/Experimental/Mutable/Source/CustomizableObject/Private/MuCO/MutableProjectorTypeUtils.cpp
@@ -0,0 +1,47 @@
+// Copyright Epic Games, Inc. All Rights Reserved.
+
+#include "MuCO/MutableProjectorTypeUtils.h"
+
+#include "MuCO/CustomizableObjectParameterTypeDefinitions.h"
+#include "MuR/Parameters.h"
+
+ECustomizableObjectProjectorType ProjectorUtils::GetEquivalentProjectorType (mu::PROJECTOR_TYPE ProjectorType)
+{
+	// Translate projector type from Mutable Core enum type to CO enum type
+	switch (ProjectorType)
+	{
+	case mu::PROJECTOR_TYPE::PLANAR:
+		return ECustomizableObjectProjectorType::Planar;
+		
+	case mu::PROJECTOR_TYPE::CYLINDRICAL:
+		return ECustomizableObjectProjectorType::Cylindrical;
+		
+	case mu::PROJECTOR_TYPE::WRAPPING:
+		return ECustomizableObjectProjectorType::Wrapping;
+		
+	case mu::PROJECTOR_TYPE::COUNT:
+	default:
+		checkNoEntry();
+		return ECustomizableObjectProjectorType::Planar;
+	}
+}
+
+
+mu::PROJECTOR_TYPE ProjectorUtils::GetEquivalentProjectorType (ECustomizableObjectProjectorType ProjectorType)
+{
+	if (GetEquivalentProjectorType(mu::PROJECTOR_TYPE::PLANAR) == ProjectorType)
+	{
+		return mu::PROJECTOR_TYPE::PLANAR;
+	}
+	if (GetEquivalentProjectorType(mu::PROJECTOR_TYPE::WRAPPING) == ProjectorType)
+	{
+		return mu::PROJECTOR_TYPE::WRAPPING;
+	}
+	if (GetEquivalentProjectorType(mu::PROJECTOR_TYPE::CYLINDRICAL) == ProjectorType)
+	{
+		return mu::PROJECTOR_TYPE::CYLINDRICAL;
+	}
+
+	checkNoEntry();
+	return mu::PROJECTOR_TYPE::COUNT;		// Invalid 
+}
diff --git a/Engine/Plugins/Experimental/Mutable/Source/CustomizableObject/Public/MuCO/CustomizableObjectInstanceDescriptor.h b/Engine/Plugins/Experimental/Mutable/Source/CustomizableObject/Public/MuCO/CustomizableObjectInstanceDescriptor.h
--- a/Engine/Plugins/Experimental/Mutable/Source/CustomizableObject/Public/MuCO/CustomizableObjectInstanceDescriptor.h
+++ b/Engine/Plugins/Experimental/Mutable/Source/CustomizableObject/Public/MuCO/CustomizableObjectInstanceDescriptor.h
@@ -211,12 +211,6 @@ struct CUSTOMIZABLEOBJECT_API FCustomizableObjectInstanceDescriptor
 
 	// Parameter Ranges
 
-	/** Returns true if the parameter is multidimensional (has multiple ranges). */
-	bool IsParamMultidimensional(const FString& ParamName) const;
-
-	/** Returns true if the parameter is multidimensional (has multiple ranges). */
-	bool IsParamMultidimensional(int32 ParamIndex) const;
-
 	/** Gets the range of values of the projector with ParamName, returns -1 if the parameter does not exist. */
 	int32 GetProjectorValueRange(const FString& ParamName) const;
 
@@ -260,10 +254,6 @@ struct CUSTOMIZABLEOBJECT_API FCustomizableObjectInstanceDescriptor
 	/** Remove the RangeIndex element of the projector range of values from the parameter ParamName, returns the index of the last valid projector, -1 if no values left. */
 	int32 RemoveValueFromProjectorRange(const FString& ParamName, int32 RangeIndex);
 
-	// Default values
-
-	FCustomizableObjectProjector GetProjectorDefaultValue(int32 ParamIndex) const;
-	
 	// ------------------------------------------------------------
    	// States
    	// ------------------------------------------------------------
diff --git a/Engine/Plugins/Experimental/Mutable/Source/CustomizableObject/Public/MuCO/MutableProjectorTypeUtils.h b/Engine/Plugins/Experimental/Mutable/Source/CustomizableObject/Public/MuCO/MutableProjectorTypeUtils.h
new file mode 100644
--- /dev/null
+++ b/Engine/Plugins/Experimental/Mutable/Source/CustomizableObject/Public/MuCO/MutableProjectorTypeUtils.h
@@ -0,0 +1,24 @@
+// Copyright Epic Games, Inc. All Rights Reserved.
+ 
+#pragma once
+
+#include "MuCO/CustomizableObjectParameterTypeDefinitions.h"
+#include "MuR/Parameters.h"
+
+namespace ProjectorUtils
+{
+	/** Get with the equivalent type of CO projector provided a mutable internal projector type.
+	 * @param ProjectorType The mutable internal projector type to get the unreal's equivalent projector type.
+	 * @return The equivalent projector type to the one provided to be used with the CO and COI systems.
+	 */
+	CUSTOMIZABLEOBJECT_API
+	ECustomizableObjectProjectorType GetEquivalentProjectorType (mu::PROJECTOR_TYPE ProjectorType);
+
+	/** Get the equivalent mutable projector type to the CO mutable projector type.
+	 * @param ProjectorType The mutable CO projector type to get the mutable's internal equivalent projector type.
+	 * @return the projector type provided but as a mutable core projector type.
+	 */
+	CUSTOMIZABLEOBJECT_API
+	mu::PROJECTOR_TYPE GetEquivalentProjectorType (ECustomizableObjectProjectorType ProjectorType);
+}
+
diff --git a/Engine/Plugins/Experimental/Mutable/Source/CustomizableObjectEditor/Private/MuCOE/GenerateMutableSource/GenerateMutableSourceProjector.cpp b/Engine/Plugins/Experimental/Mutable/Source/CustomizableObjectEditor/Private/MuCOE/GenerateMutableSource/GenerateMutableSourceProjector.cpp
--- a/Engine/Plugins/Experimental/Mutable/Source/CustomizableObjectEditor/Private/MuCOE/GenerateMutableSource/GenerateMutableSourceProjector.cpp
+++ b/Engine/Plugins/Experimental/Mutable/Source/CustomizableObjectEditor/Private/MuCOE/GenerateMutableSource/GenerateMutableSourceProjector.cpp
@@ -2,6 +2,7 @@
 
 #include "MuCOE/GenerateMutableSource/GenerateMutableSourceProjector.h"
 
+#include "MuCO/MutableProjectorTypeUtils.h"
 #include "MuCOE/CustomizableObjectCompiler.h"
 #include "MuCOE/Nodes/CustomizableObjectNodeProjectorConstant.h"
 #include "MuCOE/Nodes/CustomizableObjectNodeProjectorParameter.h"
@@ -32,49 +33,49 @@ mu::NodeProjectorPtr GenerateMutableSourceProjector(const UEdGraphPin* Pin, FMut
 		mu::NodeProjectorConstantPtr ProjectorNode = new mu::NodeProjectorConstant();
 		Result = ProjectorNode;
 
-		switch ((int)TypedNodeConst->Value.ProjectionType)
-		{
-		case 0:
-			ProjectorNode->SetValue(
-				mu::PROJECTOR_TYPE::PLANAR,
-				TypedNodeConst->Value.Position[0], TypedNodeConst->Value.Position[1], TypedNodeConst->Value.Position[2],
-				TypedNodeConst->Value.Direction[0], TypedNodeConst->Value.Direction[1], TypedNodeConst->Value.Direction[2],
-				TypedNodeConst->Value.Up[0], TypedNodeConst->Value.Up[1], TypedNodeConst->Value.Up[2],
-				TypedNodeConst->Value.Scale[0], TypedNodeConst->Value.Scale[1], TypedNodeConst->Value.Scale[2],
-				TypedNodeConst->Value.Angle);
-			break;
-
-		case 1:
-		{
-			// Apply strange swizzle for scales
-			// TODO: try to avoid this
-			float Radius = FMath::Abs(TypedNodeConst->Value.Scale[0] / 2.0f);
-			float Height = TypedNodeConst->Value.Scale[2];
-			ProjectorNode->SetValue(
-				mu::PROJECTOR_TYPE::CYLINDRICAL,
-				TypedNodeConst->Value.Position[0], TypedNodeConst->Value.Position[1], TypedNodeConst->Value.Position[2],
-				-TypedNodeConst->Value.Direction[0], -TypedNodeConst->Value.Direction[1], -TypedNodeConst->Value.Direction[2],
-				-TypedNodeConst->Value.Up[0], -TypedNodeConst->Value.Up[1], -TypedNodeConst->Value.Up[2],
-				-Height, Radius, Radius,
-				TypedNodeConst->Value.Angle);
-			break;
+		const mu::PROJECTOR_TYPE ProjectorType = ProjectorUtils::GetEquivalentProjectorType(TypedNodeConst->Value.ProjectionType);
+		switch (ProjectorType) 
+		{ 
+			case mu::PROJECTOR_TYPE::PLANAR:
+				ProjectorNode->SetValue(
+					ProjectorType,
+					TypedNodeConst->Value.Position[0], TypedNodeConst->Value.Position[1], TypedNodeConst->Value.Position[2],
+					TypedNodeConst->Value.Direction[0], TypedNodeConst->Value.Direction[1], TypedNodeConst->Value.Direction[2],
+					TypedNodeConst->Value.Up[0], TypedNodeConst->Value.Up[1], TypedNodeConst->Value.Up[2],
+					TypedNodeConst->Value.Scale[0], TypedNodeConst->Value.Scale[1], TypedNodeConst->Value.Scale[2],
+					TypedNodeConst->Value.Angle);
+				break;
+			
+			case mu::PROJECTOR_TYPE::CYLINDRICAL:
+				{
+					// Apply strange swizzle for scales
+					// TODO: try to avoid this
+					const float Radius = FMath::Abs(TypedNodeConst->Value.Scale[0] / 2.0f);
+					const float Height = TypedNodeConst->Value.Scale[2];
+					ProjectorNode->SetValue(
+						ProjectorType,
+						TypedNodeConst->Value.Position[0], TypedNodeConst->Value.Position[1], TypedNodeConst->Value.Position[2],
+						-TypedNodeConst->Value.Direction[0], -TypedNodeConst->Value.Direction[1], -TypedNodeConst->Value.Direction[2],
+						-TypedNodeConst->Value.Up[0], -TypedNodeConst->Value.Up[1], -TypedNodeConst->Value.Up[2],
+						-Height, Radius, Radius,
+						TypedNodeConst->Value.Angle);
+				}
+				break;
+			
+			case mu::PROJECTOR_TYPE::WRAPPING:
+				ProjectorNode->SetValue(
+					ProjectorType,
+					TypedNodeConst->Value.Position[0], TypedNodeConst->Value.Position[1], TypedNodeConst->Value.Position[2],
+					TypedNodeConst->Value.Direction[0], TypedNodeConst->Value.Direction[1], TypedNodeConst->Value.Direction[2],
+					TypedNodeConst->Value.Up[0], TypedNodeConst->Value.Up[1], TypedNodeConst->Value.Up[2],
+					TypedNodeConst->Value.Scale[0], TypedNodeConst->Value.Scale[1], TypedNodeConst->Value.Scale[2],
+					TypedNodeConst->Value.Angle);
+				break;
+			
+			case mu::PROJECTOR_TYPE::COUNT: 
+			default:
+				checkNoEntry();
 		}
-
-		case 2:
-			ProjectorNode->SetValue(
-				mu::PROJECTOR_TYPE::WRAPPING,
-				TypedNodeConst->Value.Position[0], TypedNodeConst->Value.Position[1], TypedNodeConst->Value.Position[2],
-				TypedNodeConst->Value.Direction[0], TypedNodeConst->Value.Direction[1], TypedNodeConst->Value.Direction[2],
-				TypedNodeConst->Value.Up[0], TypedNodeConst->Value.Up[1], TypedNodeConst->Value.Up[2],
-				TypedNodeConst->Value.Scale[0], TypedNodeConst->Value.Scale[1], TypedNodeConst->Value.Scale[2],
-				TypedNodeConst->Value.Angle);
-			break;
-
-		default:
-			// Not implemented.
-			check(false);
-		}
-
 	}
 
 	else if (const UCustomizableObjectNodeProjectorParameter* TypedNodeParam = Cast<UCustomizableObjectNodeProjectorParameter>(Node))
diff --git a/Engine/Plugins/Experimental/Mutable/Source/CustomizableObjectEditor/Private/MuCOE/SCustomizableInstanceProperties.cpp b/Engine/Plugins/Experimental/Mutable/Source/CustomizableObjectEditor/Private/MuCOE/SCustomizableInstanceProperties.cpp
--- a/Engine/Plugins/Experimental/Mutable/Source/CustomizableObjectEditor/Private/MuCOE/SCustomizableInstanceProperties.cpp
+++ b/Engine/Plugins/Experimental/Mutable/Source/CustomizableObjectEditor/Private/MuCOE/SCustomizableInstanceProperties.cpp
@@ -506,6 +506,7 @@ void SCustomizableInstanceProperties::FillChildrenMap(int32 ParamIndexInObject)
 void SCustomizableInstanceProperties::AddParameter(int32 ParamIndexInObject)
 {
 	UCustomizableObject* CustomizableObject = CustomInstance->GetCustomizableObject();
+	check(CustomizableObject);
 	
 	FString ParamName = CustomizableObject->GetParameterName(ParamIndexInObject);
 
@@ -517,8 +518,8 @@ void SCustomizableInstanceProperties::AddParameter(int32 ParamIndexInObject)
 	bool bHasCollapsibleChildren = UIData.ParamUIMetadata.ExtraInformation.Find(FString("__HasCollapsibleChildren")) != nullptr;
 
 	bool bHideParam = ParamName.EndsWith(FMultilayerProjector::NUM_LAYERS_PARAMETER_POSTFIX)
-		|| (ParamName.EndsWith(FMultilayerProjector::IMAGE_PARAMETER_POSTFIX) && CustomInstance->IsParamMultidimensional(ParamIndexInObject))
-		|| (ParamName.EndsWith(FMultilayerProjector::OPACITY_PARAMETER_POSTFIX) && CustomInstance->IsParamMultidimensional(ParamIndexInObject))
+		|| (ParamName.EndsWith(FMultilayerProjector::IMAGE_PARAMETER_POSTFIX) && CustomizableObject->IsParameterMultidimensional(ParamIndexInObject))
+		|| (ParamName.EndsWith(FMultilayerProjector::OPACITY_PARAMETER_POSTFIX) && CustomizableObject->IsParameterMultidimensional(ParamIndexInObject))
 		|| (ParamName.EndsWith(FMultilayerProjector::POSE_PARAMETER_POSTFIX));
 	
 	const FString Name = CustomInstance->ParametersSearchFilter.ToString();
@@ -699,7 +700,7 @@ void SCustomizableInstanceProperties::AddParameter(int32 ParamIndexInObject)
 	case EMutableParameterType::Int:
 	{
 		int numValues = CustomizableObject->GetIntParameterNumOptions(ParamIndexInObject);
-		bool bIsParamMultidimensional = CustomInstance->IsParamMultidimensional(ParamIndexInObject);
+		bool bIsParamMultidimensional = CustomInstance->GetCustomizableObject()->IsParameterMultidimensional((ParamIndexInObject));
 
 		if (!bIsParamMultidimensional && numValues)
 		{
@@ -859,7 +860,7 @@ void SCustomizableInstanceProperties::AddParameter(int32 ParamIndexInObject)
 		}
 		else
 		{
-			bool bIsParamMultidimensional = CustomInstance->IsParamMultidimensional(ParamIndexInObject);
+			bool bIsParamMultidimensional = CustomInstance->GetCustomizableObject()->IsParameterMultidimensional(ParamIndexInObject);
 
 			if (!bIsParamMultidimensional)
 			{
@@ -2101,7 +2102,7 @@ FReply SCustomizableInstanceProperties::OnProjectorResetTransform(FString ParamN
 
 	FString ParamName = ParseNameAndIndex(ParamNameWithIndex, RangeIndex);
 
-	const FCustomizableObjectProjector DefaultValue = CustomInstance->GetProjectorDefaultValue(CustomInstance->GetCustomizableObject()->FindParameter(ParamName));
+	const FCustomizableObjectProjector DefaultValue = CustomInstance->GetCustomizableObject()->GetProjectorParameterDefaultValue(ParamName);
 
 	CustomInstance->SetProjectorValue(ParamName,
 		static_cast<FVector>(DefaultValue.Position),
diff --git a/Engine/Plugins/Experimental/Mutable/Source/CustomizableObjectPopulation/Private/MuCOP/CustomizableObjectPopulationGenerator.cpp b/Engine/Plugins/Experimental/Mutable/Source/CustomizableObjectPopulation/Private/MuCOP/CustomizableObjectPopulationGenerator.cpp
--- a/Engine/Plugins/Experimental/Mutable/Source/CustomizableObjectPopulation/Private/MuCOP/CustomizableObjectPopulationGenerator.cpp
+++ b/Engine/Plugins/Experimental/Mutable/Source/CustomizableObjectPopulation/Private/MuCOP/CustomizableObjectPopulationGenerator.cpp
@@ -46,7 +46,7 @@ void FCustomizableObjectPopulationClassGenerator::GenerateParameters(
 					const FString& OptionName = IntParamSampler.GetOptionName(IntParamSampler.Sample(Rand));
 					if (Instance->FindIntParameterNameIndex(ParamName) != INDEX_NONE)
 					{
-						if (!Instance->IsParamMultidimensional(CustomizableObject->FindParameter(ParamName)))
+						if (!CustomizableObject->IsParameterMultidimensional(ParamName))
 						{
 							Instance->SetIntParameterSelectedOption(ParamName, OptionName);
 						}
@@ -63,7 +63,7 @@ void FCustomizableObjectPopulationClassGenerator::GenerateParameters(
 				const float SampledValue = UniformFloatSamplers[SamplerIndex.SamplerIndex].Sample(Rand);
 				if (Instance->FindFloatParameterNameIndex(ParamName) != INDEX_NONE)
 				{
-					if (!Instance->IsParamMultidimensional(CustomizableObject->FindParameter(ParamName)))
+					if (!CustomizableObject->IsParameterMultidimensional(ParamName))
 					{
 						Instance->SetFloatParameterSelectedOption(ParamName, SampledValue);
 					}
