[Mutable] Add bilinear mag filter option for projector and better projection mip selection for the minimization case to ease aliasing artifacts.

#rnx

diff --git a/Engine/Plugins/Experimental/Mutable/Source/CustomizableObject/Classes/MuCO/CustomizableObject.h b/Engine/Plugins/Experimental/Mutable/Source/CustomizableObject/Classes/MuCO/CustomizableObject.h
--- a/Engine/Plugins/Experimental/Mutable/Source/CustomizableObject/Classes/MuCO/CustomizableObject.h
+++ b/Engine/Plugins/Experimental/Mutable/Source/CustomizableObject/Classes/MuCO/CustomizableObject.h
@@ -1287,7 +1287,7 @@ private:
 	// This is a manual version number for the binary blobs in this asset.
 	// Increasing it invalidates all the previously compiled models.
 	// Warning: If while merging code both versions have changed, take the highest+1.
-	static const int32 CurrentSupportedVersion = 370;
+	static const int32 CurrentSupportedVersion = 371;
 
 public:
 
diff --git a/Engine/Plugins/Experimental/Mutable/Source/CustomizableObjectEditor/Classes/MuCOE/Nodes/CustomizableObjectNodeTextureProject.h b/Engine/Plugins/Experimental/Mutable/Source/CustomizableObjectEditor/Classes/MuCOE/Nodes/CustomizableObjectNodeTextureProject.h
--- a/Engine/Plugins/Experimental/Mutable/Source/CustomizableObjectEditor/Classes/MuCOE/Nodes/CustomizableObjectNodeTextureProject.h
+++ b/Engine/Plugins/Experimental/Mutable/Source/CustomizableObjectEditor/Classes/MuCOE/Nodes/CustomizableObjectNodeTextureProject.h
@@ -16,6 +16,19 @@ class UObject;
 struct FEdGraphPinReference;
 struct FPropertyChangedEvent;
 
+UENUM()
+enum class ETextureProjectSamplingMethod
+{
+	Point,
+	BiLinear
+};
+
+UENUM()
+enum class ETextureProjectMinFilterMethod
+{
+	None,
+	TotalAreaHeuristic
+};
 
 UCLASS()
 class CUSTOMIZABLEOBJECTEDITOR_API UCustomizableObjectNodeTextureProject : public UCustomizableObjectNode
@@ -35,6 +48,12 @@ public:
 	UPROPERTY(EditAnywhere, Category = CustomizableObject)
 	bool bEnableAngleFadeOutForAlpha = true;
 
+	UPROPERTY(EditAnywhere, Category = CustomizableObject)
+	ETextureProjectSamplingMethod SamplingMethod = ETextureProjectSamplingMethod::Point;
+
+	UPROPERTY(EditAnywhere, Category = CustomizableObject)
+	ETextureProjectMinFilterMethod MinFilterMethod = ETextureProjectMinFilterMethod::None;
+
 	/** Set the width of the Texture. If the reference texture is not null, it overrides the Size X property */
 	UPROPERTY(EditAnywhere, Category = CustomizableObject)
 	uint32 TextureSizeX = 0;
diff --git a/Engine/Plugins/Experimental/Mutable/Source/CustomizableObjectEditor/Private/MuCOE/GenerateMutableSource/GenerateMutableSourceImage.cpp b/Engine/Plugins/Experimental/Mutable/Source/CustomizableObjectEditor/Private/MuCOE/GenerateMutableSource/GenerateMutableSourceImage.cpp
--- a/Engine/Plugins/Experimental/Mutable/Source/CustomizableObjectEditor/Private/MuCOE/GenerateMutableSource/GenerateMutableSourceImage.cpp
+++ b/Engine/Plugins/Experimental/Mutable/Source/CustomizableObjectEditor/Private/MuCOE/GenerateMutableSource/GenerateMutableSourceImage.cpp
@@ -588,6 +588,31 @@ mu::NodeImagePtr GenerateMutableSourceImage(const UEdGraphPin* Pin, FMutableGrap
 		ImageNode->SetImageSize(TextureSize);
 
 		ImageNode->SetAngleFadeChannels(TypedNodeProject->bEnableAngleFadeOutForRGB, TypedNodeProject->bEnableAngleFadeOutForAlpha);
+		ImageNode->SetSamplingMethod(Invoke(
+			[](ETextureProjectSamplingMethod Method) -> mu::ESamplingMethod 
+			{
+				switch(Method)
+				{
+				case ETextureProjectSamplingMethod::Point:    return mu::ESamplingMethod::Point;
+				case ETextureProjectSamplingMethod::BiLinear: return mu::ESamplingMethod::BiLinear;
+				default: { check(false); return mu::ESamplingMethod::Point; }
+				}
+			}, 
+			TypedNodeProject->SamplingMethod)
+		);
+
+		ImageNode->SetMinFilterMethod(Invoke(
+			[](ETextureProjectMinFilterMethod Method) -> mu::EMinFilterMethod 
+			{
+				switch(Method)
+				{
+				case ETextureProjectMinFilterMethod::None:               return mu::EMinFilterMethod::None;
+				case ETextureProjectMinFilterMethod::TotalAreaHeuristic: return mu::EMinFilterMethod::TotalAreaHeuristic;
+				default: { check(false); return mu::EMinFilterMethod::None; }
+				}
+			}, 
+			TypedNodeProject->MinFilterMethod)
+		);
 
 		if (const UEdGraphPin* ConnectedPin = FollowInputPin(*TypedNodeProject->AngleFadeStartPin()))
 		{
diff --git a/Engine/Plugins/Experimental/Mutable/Source/MutableRuntime/Private/MuR/CodeRunner.cpp b/Engine/Plugins/Experimental/Mutable/Source/MutableRuntime/Private/MuR/CodeRunner.cpp
--- a/Engine/Plugins/Experimental/Mutable/Source/MutableRuntime/Private/MuR/CodeRunner.cpp
+++ b/Engine/Plugins/Experimental/Mutable/Source/MutableRuntime/Private/MuR/CodeRunner.cpp
@@ -3867,29 +3867,99 @@ namespace mu
             switch (item.Stage)
             {
             case 0:
-                    if (!args.image)
-                    {
-                        AddOp( FScheduledOp( item.At, item, 1),
-                            FScheduledOp::FromOpAndOptions( args.mesh, item, 0 ) );
-                    }
-                    else
-                    {
-                        AddOp( FScheduledOp( item.At, item, 1),
-							FScheduledOp::FromOpAndOptions( args.mesh, item, 0),
-							FScheduledOp( args.image, item ),
-							FScheduledOp( args.mask, item ),
-							FScheduledOp::FromOpAndOptions( args.angleFadeProperties, item, 0 ),
-							FScheduledOp::FromOpAndOptions( args.projector, item, 0 ) );
-                    }
+				if (args.image)
+				{
+					AddOp(FScheduledOp(item.At, item, 1),
+						FScheduledOp::FromOpAndOptions(args.mesh, item, 0),
+						FScheduledOp::FromOpAndOptions(args.projector, item, 0));
+				}
+				else
+				{
+					AddOp(FScheduledOp(item.At, item, 1),
+						FScheduledOp::FromOpAndOptions(args.mesh, item, 0));
+				}
                 break;
+			case 1:
+			{
+                MUTABLE_CPUPROFILER_SCOPE(IM_RASTERMESH_1)
 
-            case 1:
+				Ptr<const Mesh> pMesh = GetMemory().GetMesh(FScheduledOp::FromOpAndOptions(args.mesh, item, 0));
+				Ptr<const Projector> pProjector = GetMemory().GetProjector(FScheduledOp::FromOpAndOptions(args.projector, item, 0));
+
+				if (!pMesh)
+				{
+					// Stop execution.
+					check(false);
+					GetMemory().SetImage(item, nullptr);
+					break;
+				}
+
+				// If no image we are generating a flat mesh UV raster. This is the final stage in this case.
+				if (!args.image)
+				{
+					uint16 SizeX = args.sizeX;
+					uint16 SizeY = args.sizeY;
+
+					// Drop mips while possible
+					int32 MipsToDrop = item.ExecutionOptions;
+					while (MipsToDrop && !(SizeX % 2) && !(SizeY % 2))
+					{
+						SizeX = FMath::Max(uint16(1),FMath::DivideAndRoundUp(SizeX, uint16(2)));
+						SizeY = FMath::Max(uint16(1),FMath::DivideAndRoundUp(SizeY, uint16(2)));
+						--MipsToDrop;
+					}
+
+                    // Flat mesh UV raster
+					Ptr<Image> ResultImage = new Image(SizeX, SizeY, 1, EImageFormat::IF_L_UBYTE);
+                    ImageRasterMesh( pMesh.get(), ResultImage.get(), args.blockIndex );
+
+					// Stop execution.
+					GetMemory().SetImage(item, ResultImage);
+					break;
+				}
+
+				const int32 MipsToSkip = item.ExecutionOptions;
+				int32 ProjectionMip = MipsToSkip;
+
+				if (pProjector)
+				{
+					if (static_cast<EMinFilterMethod>(args.MinFilterMethod) == EMinFilterMethod::TotalAreaHeuristic)
+					{
+						FVector2f TargetImageSizeF = FVector2f(
+							FMath::Max(args.sizeX >> MipsToSkip, 1),
+							FMath::Max(args.sizeY >> MipsToSkip, 1));
+						FVector2f SourceImageSizeF = FVector2f(args.SourceSizeX, args.SourceSizeY);
+
+						ProjectionMip = ComputeProjectedFootprintBestMip(pMesh.get(), pProjector->m_value, TargetImageSizeF, SourceImageSizeF);
+					}
+				}
+
+				AddOp(FScheduledOp(item.At, item, 2, ProjectionMip),
+					FScheduledOp::FromOpAndOptions(args.mesh, item, 0),
+					FScheduledOp::FromOpAndOptions(args.projector, item, 0),
+					FScheduledOp::FromOpAndOptions(args.image, item, ProjectionMip),
+					FScheduledOp(args.mask, item),
+					FScheduledOp::FromOpAndOptions(args.angleFadeProperties, item, 0));
+
+				// Set mesh and projector back to the cache to prevent possible code re-execution.
+				GetMemory().SetMesh(FScheduledOp::FromOpAndOptions(args.mesh, item, 0), pMesh);
+				GetMemory().SetProjector(FScheduledOp::FromOpAndOptions(args.projector, item, 0), pProjector);
+
+				break;
+			}
+            case 2:
             {
-                MUTABLE_CPUPROFILER_SCOPE(IM_RASTERMESH_1)
+				MUTABLE_CPUPROFILER_SCOPE(IM_RASTERMESH_2)
 
-                Ptr<const Mesh> pMesh = GetMemory().GetMesh( FScheduledOp::FromOpAndOptions(args.mesh, item, 0) );
+				if (!args.image)
+				{
+					// This case is treated at the previous stage.
+					check(false);
+					GetMemory().SetImage(item, nullptr);
+					break;
+				}
 
-                ImagePtr pNew;
+                Ptr<const Mesh> pMesh = GetMemory().GetMesh(FScheduledOp::FromOpAndOptions(args.mesh, item, 0));
 
 				uint16 SizeX = args.sizeX;
 				uint16 SizeY = args.sizeY;
@@ -3902,110 +3972,112 @@ namespace mu
 					SizeY = FMath::Max(uint16(1),FMath::DivideAndRoundUp(SizeY, uint16(2)));
 					--MipsToDrop;
 				}
-
-                if (!pMesh.get())
-                {
-                    pNew = new Image(SizeX, SizeY, 1, EImageFormat::IF_L_UBYTE);
+	
+				if (!pMesh)
+				{
                     check(false);
-                }
-                else if (args.image)
-                {
-                    // Raster with projection
-                    Ptr<const Image> pSource = GetMemory().GetImage( FCacheAddress(args.image,item) );
+                    Ptr<Image> ResultImage = new Image(SizeX, SizeY, 1, EImageFormat::IF_L_UBYTE);
 
-                    Ptr<const Image> pMask = nullptr;
-                    if ( args.mask )
-                    {
-                        pMask = GetMemory().GetImage( FCacheAddress(args.mask,item) );
+					GetMemory().SetImage(item, ResultImage);
 
-						// TODO: This shouldn't happen, but be defensive.
-						FImageSize ResultSize(SizeX, SizeY);
-						if (pMask && pMask->GetSize()!= ResultSize)
-						{
-							MUTABLE_CPUPROFILER_SCOPE(ImageResize_MaskFixForProjection);
-							pMask = ImageResizeLinear(0, pMask.get(), ResultSize);
-						}
-                    }
+					break;
+				}
 
-                    float fadeStart = 180.0f;
-                    float fadeEnd = 180.0f;
-                    if ( args.angleFadeProperties )
-                    {
-                        FVector4f fadeProperties = GetMemory().GetColour(FScheduledOp::FromOpAndOptions(args.angleFadeProperties, item, 0));
-                        fadeStart = fadeProperties[0];
-                        fadeEnd = fadeProperties[1];
-                    }
-                    fadeStart *= PI / 180.0f;
-                    fadeEnd *= PI / 180.0f;
+				// Raster with projection
+				Ptr<const Image> pSource = GetMemory().GetImage(
+					FCacheAddress(args.image, item.ExecutionIndex, static_cast<uint8>(item.CustomState)));
 
-					EImageFormat format = pSource ? GetUncompressedFormat(pSource->GetFormat()) : EImageFormat::IF_L_UBYTE;
-                    pNew = new Image( SizeX, SizeY, 1, format );
+				Ptr<const Image> pMask = nullptr;
+				if (args.mask)
+				{
+					pMask = GetMemory().GetImage(FCacheAddress(args.mask, item));
 
-                    if (pSource && pSource->GetFormat()!=format)
-                    {
-						MUTABLE_CPUPROFILER_SCOPE(RunCode_RasterMesh_ReformatSource);
-                        pSource = ImagePixelFormat( m_pSettings->GetPrivate()->m_imageCompressionQuality, pSource.get(), format );
-                    }
+					// TODO: This shouldn't happen, but be defensive.
+					FImageSize ResultSize(SizeX, SizeY);
+					if (pMask && pMask->GetSize()!= ResultSize)
+					{
+						MUTABLE_CPUPROFILER_SCOPE(ImageResize_MaskFixForProjection);
+						pMask = ImageResizeLinear(0, pMask.get(), ResultSize);
+					}
+				}
+
+				float fadeStart = 180.0f;
+				float fadeEnd = 180.0f;
+				if ( args.angleFadeProperties )
+				{
+					FVector4f fadeProperties = GetMemory().GetColour(FScheduledOp::FromOpAndOptions(args.angleFadeProperties, item, 0));
+					fadeStart = fadeProperties[0];
+					fadeEnd = fadeProperties[1];
+				}
+				const float FadeStartRad = FMath::DegreesToRadians(fadeStart);
+				const float FadeEndRad = FMath::DegreesToRadians(fadeEnd);
 
-                    // Allocate memory for the temporary buffers
-                    SCRATCH_IMAGE_PROJECT scratch;
-                    scratch.vertices.SetNum( pMesh->GetVertexCount() );
-                    scratch.culledVertex.SetNum( pMesh->GetVertexCount() );
+				EImageFormat Format = pSource ? GetUncompressedFormat(pSource->GetFormat()) : EImageFormat::IF_L_UBYTE;
+                Ptr<Image> pNew = new Image( SizeX, SizeY, 1, Format );
 
-                    // Layout is always 0 because the previous mesh project operations take care of
-                    // moving the right layout channel to the 0.
-                    int layout = 0;
+				if (pSource && pSource->GetFormat()!=Format)
+				{
+					MUTABLE_CPUPROFILER_SCOPE(RunCode_RasterMesh_ReformatSource);
+					pSource = ImagePixelFormat( m_pSettings->GetPrivate()->m_imageCompressionQuality, pSource.get(), Format );
+				}
 
-                    if ( args.projector && pSource && pSource->GetSizeX()>0 && pSource->GetSizeY()>0 )
-                    {
-                        Ptr<const Projector> pProjector = GetMemory().GetProjector(FScheduledOp::FromOpAndOptions(args.projector, item, 0));
-                        if (pProjector)
-                        {
-                            switch (pProjector->m_value.type)
-                            {
-                            case PROJECTOR_TYPE::PLANAR:
-                                ImageRasterProjectedPlanar( pMesh.get(), pNew.get(),
-                                                            pSource.get(), pMask.get(),
-															args.bIsRGBFadingEnabled, args.bIsAlphaFadingEnabled,
-                                                            fadeStart, fadeEnd,
-                                                            layout, args.blockIndex,
-                                                            &scratch );
-                                break;
-
-                            case PROJECTOR_TYPE::WRAPPING:
-                                ImageRasterProjectedWrapping( pMesh.get(), pNew.get(),
-                                                              pSource.get(), pMask.get(),
-															  args.bIsRGBFadingEnabled, args.bIsAlphaFadingEnabled,
-															  fadeStart, fadeEnd,
-                                                              layout, args.blockIndex,
-                                                              &scratch );
-                                break;
-
-                            case PROJECTOR_TYPE::CYLINDRICAL:
-                                ImageRasterProjectedCylindrical( pMesh.get(), pNew.get(),
-                                                                 pSource.get(), pMask.get(),
-																 args.bIsRGBFadingEnabled, args.bIsAlphaFadingEnabled,
-																 fadeStart, fadeEnd,
-                                                                 layout,
-																 pProjector->m_value.projectionAngle,
-                                                                 &scratch );
-                                break;
-
-                            default:
-                                check(false);
-                                break;
-                            }
-                        }
-                    }
-                }
-                else
-                {
-                    // Flat mesh UV raster
-                    pNew = new Image( SizeX, SizeY, 1, EImageFormat::IF_L_UBYTE );
-                    ImageRasterMesh( pMesh.get(), pNew.get(), args.blockIndex );
-                }
+				// Allocate memory for the temporary buffers
+				SCRATCH_IMAGE_PROJECT scratch;
+				scratch.vertices.SetNum( pMesh->GetVertexCount() );
+				scratch.culledVertex.SetNum( pMesh->GetVertexCount() );
+
+				// Layout is always 0 because the previous mesh project operations take care of
+				// moving the right layout channel to the 0.
+				int layout = 0;
+
+				ESamplingMethod SamplingMethod = static_cast<ESamplingMethod>(args.SamplingMethod);
+
+				if (args.projector && pSource && pSource->GetSizeX() > 0 && pSource->GetSizeY() > 0)
+				{
+					Ptr<const Projector> pProjector = GetMemory().GetProjector(FScheduledOp::FromOpAndOptions(args.projector, item, 0));
+					if (pProjector)
+					{
+						switch (pProjector->m_value.type)
+						{
+						case PROJECTOR_TYPE::PLANAR:
+							ImageRasterProjectedPlanar(pMesh.get(), pNew.get(),
+								pSource.get(), pMask.get(),
+								args.bIsRGBFadingEnabled, args.bIsAlphaFadingEnabled,
+								SamplingMethod,
+								FadeStartRad, FadeEndRad,
+								layout, args.blockIndex,
+								&scratch);
+							break;
+
+						case PROJECTOR_TYPE::WRAPPING:
+							ImageRasterProjectedWrapping(pMesh.get(), pNew.get(),
+								pSource.get(), pMask.get(),
+								args.bIsRGBFadingEnabled, args.bIsAlphaFadingEnabled,
+								SamplingMethod,
+								FadeStartRad, FadeEndRad,
+								layout, args.blockIndex,
+								&scratch);
+							break;
+
+						case PROJECTOR_TYPE::CYLINDRICAL:
+							ImageRasterProjectedCylindrical(pMesh.get(), pNew.get(),
+								pSource.get(), pMask.get(),
+								args.bIsRGBFadingEnabled, args.bIsAlphaFadingEnabled,
+								FadeStartRad, FadeEndRad,
+								layout,
+								pProjector->m_value.projectionAngle,
+								&scratch);
+							break;
+
+						default:
+							check(false);
+							break;
+						}
+					}
+				}
+
+				GetMemory().SetImage(item, pNew);
 
-                GetMemory().SetImage( item, pNew );
                 break;
             }
 
@@ -4172,7 +4244,6 @@ namespace mu
         }
     }
 
-
     //---------------------------------------------------------------------------------------------
     Ptr<RangeIndex> CodeRunner::BuildCurrentOpRangeIndex( const FScheduledOp& item, const Parameters* pParams, const Model* pModel, int32 parameterIndex )
     {
diff --git a/Engine/Plugins/Experimental/Mutable/Source/MutableRuntime/Private/MuR/CodeRunnerTasked.cpp b/Engine/Plugins/Experimental/Mutable/Source/MutableRuntime/Private/MuR/CodeRunnerTasked.cpp
--- a/Engine/Plugins/Experimental/Mutable/Source/MutableRuntime/Private/MuR/CodeRunnerTasked.cpp
+++ b/Engine/Plugins/Experimental/Mutable/Source/MutableRuntime/Private/MuR/CodeRunnerTasked.cpp
@@ -102,7 +102,6 @@ namespace mu
 		}
 	}
 
-
 	//---------------------------------------------------------------------------------------------
 	TRACE_DECLARE_INT_COUNTER(MutableRuntime_OpenTask, TEXT("MutableRuntime/OpenTask"));
 	TRACE_DECLARE_INT_COUNTER(MutableRuntime_ClosedTasks, TEXT("MutableRuntime/ClosedTasks"));
@@ -490,7 +489,6 @@ namespace mu
 		}
     }
 
-
 	//---------------------------------------------------------------------------------------------
 	void CodeRunner::GetImageDescResult(FImageDesc& OutDesc)
 	{
@@ -1233,8 +1231,13 @@ namespace mu
 	{
 		MUTABLE_CPUPROFILER_SCOPE(FImageSaturateTask);
 
-		const bool bOptimizeUnchanged = FMath::IsNearlyEqual(Factor, 1.0f);
+		if (!Source)
+		{
+			check(false);
+			return;
+		}
 
+		const bool bOptimizeUnchanged = FMath::IsNearlyEqual(Factor, 1.0f);
 		if (Source->IsUnique())
 		{
 			Ptr<Image> MutableImage = mu::CloneOrTakeOver<>(Source.get());
diff --git a/Engine/Plugins/Experimental/Mutable/Source/MutableRuntime/Private/MuR/Image.cpp b/Engine/Plugins/Experimental/Mutable/Source/MutableRuntime/Private/MuR/Image.cpp
--- a/Engine/Plugins/Experimental/Mutable/Source/MutableRuntime/Private/MuR/Image.cpp
+++ b/Engine/Plugins/Experimental/Mutable/Source/MutableRuntime/Private/MuR/Image.cpp
@@ -4,20 +4,24 @@
 
 #include "HAL/UnrealMemory.h"
 #include "Math/UnrealMathSSE.h"
+#include "Math/NumericLimits.h"
 #include "MuR/ImagePrivate.h"
 #include "MuR/MutableMath.h"
 #include "MuR/MutableTrace.h"
 #include "MuR/OpImagePixelFormat.h"
 #include "MuR/OpImageResize.h"
 
-bool bDisableCompressedImageBlackBlockInit = false;
-
-static FAutoConsoleVariableRef CVarDisableCompressedImageBlackBlockInit(
-	TEXT("mutable.DisableCompressedImageBlackBlockInit"),
-	bDisableCompressedImageBlackBlockInit,
-	TEXT("A value of 1 disables mutable compressed black block initialization"),
-	ECVF_Default);
+#include <initializer_list>
 
+namespace
+{
+	static bool bDisableCompressedImageBlackBlockInit = false;
+	static FAutoConsoleVariableRef CVarDisableCompressedImageBlackBlockInit(
+		TEXT("mutable.DisableCompressedImageBlackBlockInit"),
+		bDisableCompressedImageBlackBlockInit,
+		TEXT("A value of 1 disables mutable compressed black block initialization"),
+		ECVF_Default);
+}
 
 namespace mu
 {
@@ -82,11 +86,15 @@ namespace mu
     Image::Image( uint32_t sizeX, uint32_t sizeY, uint32_t lods, EImageFormat format, EInitializationType Init )
     {
 		MUTABLE_CPUPROFILER_SCOPE(NewImage)
-			LLM_SCOPE_BYNAME(TEXT("MutableRuntime"));
+		LLM_SCOPE_BYNAME(TEXT("MutableRuntime"));
 
         check(format != EImageFormat::IF_NONE);
 		check(format < EImageFormat::IF_COUNT);
 
+		check(sizeX <= TNumericLimits<uint16>::Max())
+		check(sizeY <= TNumericLimits<uint16>::Max())
+		check(lods <= TNumericLimits<uint8>::Max())
+
         m_format = format;
         m_size = FImageSize( (uint16)sizeX, (uint16)sizeY );
         m_lods = (uint8_t)lods;
@@ -120,6 +128,7 @@ namespace mu
 					{
 						m_data.SetNumUninitialized(DataSize);
 
+						check(fdata.m_bytesPerBlock > 0);
 						for (int32 BlockDataOffset = 0; BlockDataOffset < DataSize; BlockDataOffset += fdata.m_bytesPerBlock)
 						{
 							FMemory::Memcpy(m_data.GetData() + BlockDataOffset, fdata.BlackBlock, fdata.m_bytesPerBlock);
diff --git a/Engine/Plugins/Experimental/Mutable/Source/MutableRuntime/Private/MuR/ImagePrivate.h b/Engine/Plugins/Experimental/Mutable/Source/MutableRuntime/Private/MuR/ImagePrivate.h
--- a/Engine/Plugins/Experimental/Mutable/Source/MutableRuntime/Private/MuR/ImagePrivate.h
+++ b/Engine/Plugins/Experimental/Mutable/Source/MutableRuntime/Private/MuR/ImagePrivate.h
@@ -7,15 +7,14 @@
 #include "MuR/SerialisationPrivate.h"
 #include "MuR/MutableMath.h"
 
-#include <initializer_list>
-
 namespace mu
 {
 
 	MUTABLE_DEFINE_ENUM_SERIALISABLE( EBlendType );
 	MUTABLE_DEFINE_ENUM_SERIALISABLE( EMipmapFilterType );
 	MUTABLE_DEFINE_ENUM_SERIALISABLE( ECompositeImageMode );
-
+	MUTABLE_DEFINE_ENUM_SERIALISABLE( ESamplingMethod );
+	MUTABLE_DEFINE_ENUM_SERIALISABLE( EMinFilterMethod );
 
     struct FImageFormatData
 	{
@@ -45,8 +44,10 @@ namespace mu
 			)
 			: FImageFormatData(pixelsPerBlockX, pixelsPerBlockY, bytesPerBlock, channels)
 		{
-			check(BlackBlockInit.size() <= MAX_BYTES_PER_BLOCK);
-			FMemory::Memcpy(BlackBlock, BlackBlockInit.begin(), FMath::Min<SIZE_T>(MAX_BYTES_PER_BLOCK, BlackBlockInit.size()));
+			check(MAX_BYTES_PER_BLOCK >= BlackBlockInit.size());
+
+			const SIZE_T SanitizedBlockSize = FMath::Min<SIZE_T>(MAX_BYTES_PER_BLOCK, BlackBlockInit.size());
+			FMemory::Memcpy(BlackBlock, BlackBlockInit.begin(), SanitizedBlockSize);
 		}
 
 		//! For block based formats, size of the block size. For uncompressed formats it will
diff --git a/Engine/Plugins/Experimental/Mutable/Source/MutableRuntime/Private/MuR/OpImageProject.cpp b/Engine/Plugins/Experimental/Mutable/Source/MutableRuntime/Private/MuR/OpImageProject.cpp
--- a/Engine/Plugins/Experimental/Mutable/Source/MutableRuntime/Private/MuR/OpImageProject.cpp
+++ b/Engine/Plugins/Experimental/Mutable/Source/MutableRuntime/Private/MuR/OpImageProject.cpp
@@ -22,7 +22,7 @@
 #include "Math/IntPoint.h"
 #include "Math/UnrealMathSSE.h"
 #include "Misc/AssertionMacros.h"
-
+#include "Math/UnrealMathUtility.h"
 
 namespace mu
 {
@@ -122,7 +122,6 @@ namespace
                 }
             }
 
-
             if ( factor_8>0 )
             {
                 float u = varying[0];
@@ -196,6 +195,170 @@ namespace
     };
 
 
+
+    //---------------------------------------------------------------------------------------------
+    template<int PIXEL_SIZE, bool bUseSmoothInterpolator = false, bool HAS_MASK = true>
+    class RasterProjectedPixelProcessorLinear_UBYTE
+    {
+    public:
+
+        RasterProjectedPixelProcessorLinear_UBYTE( const Image* Source,
+			const uint8* pTargetData,
+			const uint8* pMaskData,
+			bool bInIsRGBFadingEnabled, bool bInIsAlphaFadingEnabled,
+			float FadeStart, float FadeEnd )
+        {
+            SourceSizeX = static_cast<uint16>(Source->GetSizeX());
+            SourceSizeY = static_cast<uint16>(Source->GetSizeY());
+
+            SourceDataPtr = (SourceSizeX > 0 && SourceSizeY > 0) ? Source->GetData() : nullptr;
+
+			// \TODO: to template parameters?
+			bIsRGBFadingEnabled = bInIsRGBFadingEnabled;
+			bIsAlphaFadingEnabled = bInIsAlphaFadingEnabled;
+
+			const float FadeStartCos = FMath::Cos(FadeStart);
+            FadeEndCos = FMath::Cos(FadeEnd);
+
+			const float FadeCosRangeSafeDiv = 
+				FMath::IsNearlyZero(FadeStartCos - FadeEndCos, UE_KINDA_SMALL_NUMBER) 
+				? UE_KINDA_SMALL_NUMBER 
+				: FadeStartCos - FadeEndCos;
+
+			OneOverFadeRangeTimes255 = 255.0f / FadeCosRangeSafeDiv;
+
+            TargetDataPtr = pTargetData;
+            MaskDataPtr = pMaskData;
+
+
+            check(GetImageFormatData(Source->GetFormat() ).m_bytesPerBlock == PIXEL_SIZE);
+        }
+
+        //-----------------------------------------------------------------------------------------
+        inline void ProcessPixel( uint8* BufferPosPtr, float Varying[4] ) const
+        {
+			if (!SourceDataPtr)
+			{
+				return;
+			}
+
+			const bool bDepthClamp = (Varying[2] < 0.0f) | (Varying[2] > 1.0f);
+            const float AngleCos = Varying[3];
+
+			uint16 Factor = static_cast<uint16>(FMath::Clamp((AngleCos - FadeEndCos) * OneOverFadeRangeTimes255, 0.0f, 255.0f));
+			Factor = Factor < 255 ? Factor : static_cast<uint16>(!bDepthClamp) * 255;
+
+			uint16 MaskFactor = AngleCos > FadeEndCos ? 255 : 0;
+
+			if constexpr (HAS_MASK)
+			{
+				// Only read from memory if needed.
+				if (Factor > 0)
+				{
+					MaskFactor = ((MaskFactor > 0) && MaskDataPtr) 
+							? MaskDataPtr[(BufferPosPtr - TargetDataPtr) / PIXEL_SIZE] 
+							: MaskFactor;
+					Factor = (Factor * MaskFactor) / 255;
+				}
+			}
+
+            if (Factor > 0)
+            {
+				const FVector2f Uv = { Varying[0], Varying[1] };
+				const FVector2f SourceSizeF = FVector2f(SourceSizeX, SourceSizeY);
+
+                if ((Uv.X >= 0.0f) & (Uv.X < 1.0f) & (Uv.Y >= 0.0f) & (Uv.Y < 1.0f))
+                {
+					const float CoordsX = FMath::Clamp(Uv.X * SourceSizeF.X - 0.5f, 0.0f, SourceSizeF.X - 1.0f);
+					const float CoordsY = FMath::Clamp(Uv.Y * SourceSizeF.Y - 0.5f, 0.0f, SourceSizeF.Y - 1.0f);
+
+					auto ComputeInterpolator = [](float T) -> uint16
+					{
+						if constexpr (bUseSmoothInterpolator)
+						{
+							return static_cast<uint16>(255.0f * T * T * (3.0f - 2.0f * T));
+						}
+						else
+						{
+							return static_cast<uint16>(255.0f * T);
+						}
+					};
+
+					const uint16 FracX = ComputeInterpolator(FMath::Frac(CoordsX));
+					const uint16 FracY = ComputeInterpolator(FMath::Frac(CoordsY));
+
+					const uint32 X = static_cast<uint32>(CoordsX);
+					const uint32 Y = static_cast<uint32>(CoordsY);
+
+					const uint32 XPlusOne = FMath::Min<uint32>(SourceSizeX - 1, X + 1);
+					const uint32 YPlusOne = FMath::Min<uint32>(SourceSizeY - 1, Y + 1);
+
+					const uint8* Pixel00Ptr = SourceDataPtr + (Y        * SourceSizeX + X)        * PIXEL_SIZE;
+					const uint8* Pixel10Ptr = SourceDataPtr + (Y        * SourceSizeX + XPlusOne) * PIXEL_SIZE;
+					const uint8* Pixel01Ptr = SourceDataPtr + (YPlusOne * SourceSizeX + X)        * PIXEL_SIZE;
+					const uint8* Pixel11Ptr = SourceDataPtr + (YPlusOne * SourceSizeX + XPlusOne) * PIXEL_SIZE;
+
+					struct FPixelData
+					{
+						alignas(8) uint16 Data[PIXEL_SIZE];
+					};
+
+					auto LoadPixel = [](const uint8* Ptr) -> FPixelData
+					{
+						FPixelData Result;
+						for (int32 C = 0; C < PIXEL_SIZE; ++C)
+						{
+							Result.Data[C] = static_cast<uint16>(Ptr[C]);
+						}
+
+						return Result;
+					};
+
+					FPixelData PixelData00 = LoadPixel(Pixel00Ptr);
+					FPixelData PixelData10 = LoadPixel(Pixel10Ptr);
+					FPixelData PixelData01 = LoadPixel(Pixel01Ptr);
+					FPixelData PixelData11 = LoadPixel(Pixel11Ptr);
+
+					uint16 FilteredPixelData[PIXEL_SIZE];
+					
+					for (int32 C = 0; C < PIXEL_SIZE; ++C)
+					{
+						const uint16 LerpY0 = ((PixelData10.Data[C] * FracX) + PixelData00.Data[C] * (255 - FracX)) / 255;
+						const uint16 LerpY1 = ((PixelData11.Data[C] * FracX) + PixelData01.Data[C] * (255 - FracX)) / 255;
+						FilteredPixelData[C] = ((LerpY1 * FracY) + LerpY0*(255 - FracY)) / 255;
+					} 
+
+					static_assert(PIXEL_SIZE <= 4);
+					const uint16 FadingMask[4] = { 
+						bIsRGBFadingEnabled   ? 1 : 0, 
+						bIsRGBFadingEnabled   ? 1 : 0, 
+						bIsRGBFadingEnabled   ? 1 : 0, 
+						bIsAlphaFadingEnabled ? 1 : 0 };
+
+					for (int32 I = 0; I < PIXEL_SIZE; ++I)
+					{
+						BufferPosPtr[I] = static_cast<uint8>(
+							(FilteredPixelData[I] * (MaskFactor * (1 - FadingMask[I]) + (Factor * FadingMask[I]))) / 255);
+					}				
+                }
+            }
+        }
+
+        const uint8* SourceDataPtr;
+        uint16 SourceSizeX, SourceSizeY;
+
+        //! Cosine of the fading angle range
+		float FadeEndCos;
+		// precomputation of 255 / (FadeStartCos - FadeEndCos) 
+		float OneOverFadeRangeTimes255;
+
+		bool bIsRGBFadingEnabled, bIsAlphaFadingEnabled;
+
+        const uint8* TargetDataPtr;
+        const uint8* MaskDataPtr;
+    };
+
+
     //---------------------------------------------------------------------------------------------
     template<int PIXEL_SIZE,bool HAS_MASK=true>
     class RasterCylindricalProjectedPixelProcessor_UBYTE
@@ -707,6 +870,7 @@ void ImageRasterProjectedPlanar( const Mesh* pMesh,
 	const Image* pSource,
 	const Image* pMask,
 	bool bIsRGBFadingEnabled, bool bIsAlphaFadingEnabled,
+	ESamplingMethod SamplingMethod,
 	float fadeStart, float fadeEnd,
 	int layout, int block,
 	SCRATCH_IMAGE_PROJECT* scratch )
@@ -717,54 +881,114 @@ void ImageRasterProjectedPlanar( const Mesh* pMesh,
 
 	MUTABLE_CPUPROFILER_SCOPE(ImageProject);
 
+	constexpr bool bUseSmoothInterpolator = false;
     if ( pMesh->m_staticFormatFlags & (1<<SMF_PROJECT) )
     {
         // Mesh-optimised version
         if ( pSource->GetFormat() == EImageFormat::IF_RGB_UBYTE )
         {
-            RasterProjectedPixelProcessor_UBYTE<3> pixelProc( 
-				pSource,
-				pImage->GetData(),
-				pMask ? pMask->GetData() : 0,
-				bIsRGBFadingEnabled, bIsAlphaFadingEnabled,
-				fadeStart, fadeEnd );
-
-            ImageRasterProjected_Optimised( pMesh, pImage, pixelProc, fadeEnd, scratch );
+			if (SamplingMethod == ESamplingMethod::Point)
+			{
+				const RasterProjectedPixelProcessor_UBYTE<3> pixelProc(pSource,
+					pImage->GetData(),
+					pMask ? pMask->GetData() : nullptr,
+					bIsRGBFadingEnabled, bIsAlphaFadingEnabled,
+					fadeStart, fadeEnd);
+				ImageRasterProjected_Optimised(pMesh, pImage, pixelProc, fadeEnd, scratch);
+			}
+			else if (SamplingMethod == ESamplingMethod::BiLinear)
+			{
+				const RasterProjectedPixelProcessorLinear_UBYTE<3, bUseSmoothInterpolator> pixelProc(pSource,
+					pImage->GetData(),
+					pMask ? pMask->GetData() : nullptr,
+					bIsRGBFadingEnabled, bIsAlphaFadingEnabled,
+					fadeStart, fadeEnd);
+				ImageRasterProjected_Optimised(pMesh, pImage, pixelProc, fadeEnd, scratch);
+			}
+			else
+			{
+				check(false);
+			}
         }
 
         else if ( pSource->GetFormat() == EImageFormat::IF_RGBA_UBYTE ||
                   pSource->GetFormat() == EImageFormat::IF_BGRA_UBYTE )
         {
-            if (pMask!=nullptr)
+            if (pMask != nullptr)
             {
-                RasterProjectedPixelProcessor_UBYTE<4,true> pixelProc( pSource,
-                    pImage->GetData(),
-                    pMask->GetData(),
-					bIsRGBFadingEnabled, bIsAlphaFadingEnabled,
-					fadeStart, fadeEnd );
-                ImageRasterProjected_Optimised( pMesh, pImage, pixelProc, fadeEnd, scratch );
+				if (SamplingMethod == ESamplingMethod::Point)
+				{
+					const RasterProjectedPixelProcessor_UBYTE<4, true> pixelProc(pSource,
+						pImage->GetData(),
+						pMask->GetData(),
+						bIsRGBFadingEnabled, bIsAlphaFadingEnabled,
+						fadeStart, fadeEnd);
+					ImageRasterProjected_Optimised(pMesh, pImage, pixelProc, fadeEnd, scratch);
+				}
+				else if (SamplingMethod == ESamplingMethod::BiLinear)
+				{
+					const RasterProjectedPixelProcessorLinear_UBYTE<4, bUseSmoothInterpolator, true> pixelProc(pSource,
+						pImage->GetData(),
+						pMask->GetData(),
+						bIsRGBFadingEnabled, bIsAlphaFadingEnabled,
+						fadeStart, fadeEnd);
+					ImageRasterProjected_Optimised(pMesh, pImage, pixelProc, fadeEnd, scratch);
+				}
+				else
+				{
+					check(false);
+				}
             }
             else
             {
-                RasterProjectedPixelProcessor_UBYTE<4,false> pixelProc( pSource,
-                    pImage->GetData(),
-                    nullptr,
-					bIsRGBFadingEnabled, bIsAlphaFadingEnabled,
-					fadeStart, fadeEnd );
-                ImageRasterProjected_Optimised( pMesh, pImage, pixelProc, fadeEnd, scratch );
+				if (SamplingMethod == ESamplingMethod::Point) 
+				{
+					const RasterProjectedPixelProcessor_UBYTE<4, false> pixelProc(pSource,
+						pImage->GetData(),
+						nullptr,
+						bIsRGBFadingEnabled, bIsAlphaFadingEnabled,
+						fadeStart, fadeEnd);
+					ImageRasterProjected_Optimised(pMesh, pImage, pixelProc, fadeEnd, scratch);
+				}
+				else if (SamplingMethod == ESamplingMethod::BiLinear)
+				{
+					const RasterProjectedPixelProcessorLinear_UBYTE<4, bUseSmoothInterpolator, false> pixelProc(pSource,
+						pImage->GetData(),
+						nullptr,
+						bIsRGBFadingEnabled, bIsAlphaFadingEnabled,
+						fadeStart, fadeEnd);
+					ImageRasterProjected_Optimised(pMesh, pImage, pixelProc, fadeEnd, scratch);
+				}
+				else
+				{
+					check(false);
+				}
             }
-
         }
-
         else if ( pSource->GetFormat() == EImageFormat::IF_L_UBYTE )
         {
-            RasterProjectedPixelProcessor_UBYTE<1> pixelProc( pSource,
-				pImage->GetData(),
-				pMask ? pMask->GetData() : 0,
-				bIsRGBFadingEnabled, bIsAlphaFadingEnabled,
-				fadeStart, fadeEnd );
-
-            ImageRasterProjected_Optimised( pMesh, pImage, pixelProc, fadeEnd, scratch );
+			if (SamplingMethod == ESamplingMethod::Point)
+			{
+				const RasterProjectedPixelProcessor_UBYTE<1> pixelProc(pSource,
+					pImage->GetData(),
+					pMask ? pMask->GetData() : nullptr,
+					bIsRGBFadingEnabled, bIsAlphaFadingEnabled,
+					fadeStart, fadeEnd);
+				ImageRasterProjected_Optimised(pMesh, pImage, pixelProc, fadeEnd, scratch);
+			}
+			else if (SamplingMethod == ESamplingMethod::BiLinear)
+			{
+				const RasterProjectedPixelProcessorLinear_UBYTE<1, bUseSmoothInterpolator> pixelProc(pSource,
+					pImage->GetData(),
+					pMask ? pMask->GetData() : nullptr,
+					bIsRGBFadingEnabled, bIsAlphaFadingEnabled,
+					fadeStart, fadeEnd);
+				ImageRasterProjected_Optimised(pMesh, pImage, pixelProc, fadeEnd, scratch);
+			}
+			else
+			{
+				check(false);
+			}
         }
 
         else
@@ -798,6 +1022,7 @@ void ImageRasterProjectedWrapping( const Mesh* pMesh,
                                        const Image* pSource,
                                        const Image* pMask,
 									   bool bIsRGBFadingEnabled, bool bIsAlphaFadingEnabled,
+									   ESamplingMethod SamplingMethod,
                                        float fadeStart, float fadeEnd,
                                        int layout,
                                        int block,
@@ -809,19 +1034,38 @@ void ImageRasterProjectedWrapping( const Mesh* pMesh,
 
 	MUTABLE_CPUPROFILER_SCOPE(ImageProjectWrapping);
 
-    if ( ( pMesh->m_staticFormatFlags & (1<<SMF_PROJECTWRAPPING) )
-         )
+
+	constexpr bool bUseSmoothInterpolator = false;
+
+    if ( ( pMesh->m_staticFormatFlags & (1<<SMF_PROJECTWRAPPING) ) )
     {
         // Mesh-optimised version
         if ( pSource->GetFormat() == EImageFormat::IF_RGB_UBYTE )
         {
-            RasterProjectedPixelProcessor_UBYTE<3> pixelProc(  pSource,
-				pImage->GetData(),
-				pMask ? pMask->GetData() : 0,
-				bIsRGBFadingEnabled, bIsAlphaFadingEnabled,
-				fadeStart, fadeEnd );
+			if (SamplingMethod == ESamplingMethod::Point)
+			{
+				RasterProjectedPixelProcessor_UBYTE<3> pixelProc(pSource,
+					pImage->GetData(),
+					pMask ? pMask->GetData() : 0,
+					bIsRGBFadingEnabled, bIsAlphaFadingEnabled,
+					fadeStart, fadeEnd);
+
+				ImageRasterProjected_OptimisedWrapping(pMesh, pImage, pixelProc, fadeEnd, block, scratch);
+			}
+			else if (SamplingMethod == ESamplingMethod::BiLinear)
+			{
+				RasterProjectedPixelProcessorLinear_UBYTE<3, bUseSmoothInterpolator> pixelProc(pSource,
+					pImage->GetData(),
+					pMask ? pMask->GetData() : 0,
+					bIsRGBFadingEnabled, bIsAlphaFadingEnabled,
+					fadeStart, fadeEnd);
 
-            ImageRasterProjected_OptimisedWrapping( pMesh, pImage, pixelProc, fadeEnd, block, scratch );
+				ImageRasterProjected_OptimisedWrapping(pMesh, pImage, pixelProc, fadeEnd, block, scratch);
+			}
+			else
+			{
+				check(false);
+			}
         }
 
         else if ( pSource->GetFormat() == EImageFormat::IF_RGBA_UBYTE ||
@@ -829,34 +1073,82 @@ void ImageRasterProjectedWrapping( const Mesh* pMesh,
         {
             if (pMask!=nullptr)
             {
-                RasterProjectedPixelProcessor_UBYTE<4,true> pixelProc( pSource,
-					pImage->GetData(),
-					pMask->GetData(),
-					bIsRGBFadingEnabled, bIsAlphaFadingEnabled,
-					fadeStart, fadeEnd );
-                ImageRasterProjected_OptimisedWrapping( pMesh, pImage, pixelProc, fadeEnd, block, scratch );
+				if (SamplingMethod == ESamplingMethod::Point)
+				{
+					RasterProjectedPixelProcessor_UBYTE<4, true> pixelProc(pSource,
+						pImage->GetData(),
+						pMask->GetData(),
+						bIsRGBFadingEnabled, bIsAlphaFadingEnabled,
+						fadeStart, fadeEnd);
+					ImageRasterProjected_OptimisedWrapping(pMesh, pImage, pixelProc, fadeEnd, block, scratch);
+				}
+				else if (SamplingMethod == ESamplingMethod::BiLinear)
+				{
+					RasterProjectedPixelProcessorLinear_UBYTE<4, bUseSmoothInterpolator, true> pixelProc(pSource,
+						pImage->GetData(),
+						pMask->GetData(),
+						bIsRGBFadingEnabled, bIsAlphaFadingEnabled,
+						fadeStart, fadeEnd);
+					ImageRasterProjected_OptimisedWrapping(pMesh, pImage, pixelProc, fadeEnd, block, scratch);
+				}
+				else
+				{
+					check(false);
+				}
             }
             else
             {
-                RasterProjectedPixelProcessor_UBYTE<4,false> pixelProc( pSource,
-					pImage->GetData(),
-					nullptr,
-					bIsRGBFadingEnabled, bIsAlphaFadingEnabled,
-					fadeStart, fadeEnd );
-                ImageRasterProjected_OptimisedWrapping( pMesh, pImage, pixelProc, fadeEnd, block, scratch );
+				if (SamplingMethod == ESamplingMethod::Point)
+				{
+					RasterProjectedPixelProcessor_UBYTE<4, false> pixelProc(pSource,
+						pImage->GetData(),
+						nullptr,
+						bIsRGBFadingEnabled, bIsAlphaFadingEnabled,
+						fadeStart, fadeEnd);
+					ImageRasterProjected_OptimisedWrapping(pMesh, pImage, pixelProc, fadeEnd, block, scratch);
+				}
+				else if (SamplingMethod == ESamplingMethod::BiLinear)
+				{
+					RasterProjectedPixelProcessorLinear_UBYTE<4, bUseSmoothInterpolator, false> pixelProc(pSource,
+						pImage->GetData(),
+						nullptr,
+						bIsRGBFadingEnabled, bIsAlphaFadingEnabled,
+						fadeStart, fadeEnd);
+					ImageRasterProjected_OptimisedWrapping(pMesh, pImage, pixelProc, fadeEnd, block, scratch);
+				}
+				else
+				{
+					check(false);
+				}
             }
-
         }
 
         else if ( pSource->GetFormat() == EImageFormat::IF_L_UBYTE )
         {
-            RasterProjectedPixelProcessor_UBYTE<1> pixelProc( pSource,
-				pImage->GetData(),
-				pMask ? pMask->GetData() : 0,
-				bIsRGBFadingEnabled, bIsAlphaFadingEnabled,
-				fadeStart, fadeEnd );
+			if (SamplingMethod == ESamplingMethod::Point)
+			{
+				RasterProjectedPixelProcessor_UBYTE<1> pixelProc(pSource,
+					pImage->GetData(),
+					pMask ? pMask->GetData() : 0,
+					bIsRGBFadingEnabled, bIsAlphaFadingEnabled,
+					fadeStart, fadeEnd);
 
-            ImageRasterProjected_OptimisedWrapping( pMesh, pImage, pixelProc, fadeEnd, block, scratch );
+				ImageRasterProjected_OptimisedWrapping(pMesh, pImage, pixelProc, fadeEnd, block, scratch);
+			}
+			else if (SamplingMethod == ESamplingMethod::BiLinear)
+			{	
+				RasterProjectedPixelProcessorLinear_UBYTE<1, bUseSmoothInterpolator> pixelProc(pSource,
+					pImage->GetData(),
+					pMask ? pMask->GetData() : 0,
+					bIsRGBFadingEnabled, bIsAlphaFadingEnabled,
+					fadeStart, fadeEnd);
+
+				ImageRasterProjected_OptimisedWrapping(pMesh, pImage, pixelProc, fadeEnd, block, scratch);
+			}
+			else
+			{
+				check(false);
+			}
         }
 
         else
@@ -883,6 +1175,85 @@ void ImageRasterProjectedWrapping( const Mesh* pMesh,
 
 }
 
+int32 ComputeProjectedFootprintBestMip(
+	const Mesh* pMesh, const FProjector& Projector, const FVector2f& TargetSize, const FVector2f& SourceSize)
+{
+	// Compute projected mesh footprint on the traget image and extract a desired size for source image.
+	check(pMesh->GetVertexBuffers().GetElementSize(0) == sizeof(OPTIMISED_VERTEX));
+	const OPTIMISED_VERTEX* VerticesPtr = reinterpret_cast<const OPTIMISED_VERTEX*>(pMesh->GetVertexBuffers().GetBufferData(0));
+
+	check(pMesh->GetIndexBuffers().GetElementSize(0) == sizeof(uint32));
+	const uint32* IndicesPtr = reinterpret_cast<const uint32*>(pMesh->GetIndexBuffers().GetBufferData(0));
+
+	const uint32 NumIndices = pMesh->GetIndexBuffers().GetElementCount();
+
+	// Source area is the fraction of source image that is covered by the projected mesh.
+	float SourceArea = UE_KINDA_SMALL_NUMBER;
+	// Target area is the fraction of the target image that is covered by the projected mesh.
+	float TargetArea = UE_KINDA_SMALL_NUMBER;
+
+	auto ComputeTriangleArea = [](const FVector2f& A, const FVector2f& B, const FVector2f& C) -> float
+	{
+		return (A.X * (B.Y - C.Y) + B.X * (C.Y - A.Y) + C.X * (A.Y - B.Y)) * 0.5f;
+	};
+
+	const bool bIsPlanarProjection = Projector.type == PROJECTOR_TYPE::PLANAR;
+
+	// TODO: evaluate other strategies to get a better mip estimate. Maybe use the median of area change?
+	check(NumIndices % 3 == 0);
+	for (uint32 I = 0; I < NumIndices; I += 3)
+	{
+		const OPTIMISED_VERTEX& A = VerticesPtr[IndicesPtr[I + 0]];
+		const OPTIMISED_VERTEX& B = VerticesPtr[IndicesPtr[I + 1]];
+		const OPTIMISED_VERTEX& C = VerticesPtr[IndicesPtr[I + 2]];
+
+		const float TriangleSourceArea = ComputeTriangleArea(
+			FVector2f(A.pos.X, A.pos.Y), FVector2f(B.pos.X, B.pos.Y), FVector2f(C.pos.X, C.pos.Y));
+
+		const float TriangleTargetArea = ComputeTriangleArea(A.uv, B.uv, C.uv);
+
+		// Set weight to zero if source area is close to zero to prevent posible infinite scaling. 
+		float TriangleWeight = static_cast<float>(!FMath::IsNearlyZero(TriangleSourceArea));
+		if (bIsPlanarProjection)
+		{
+			TriangleWeight *= FMath::Clamp((A.nor.X + B.nor.X + C.nor.X) * (1.0f / 3.0f), 0.0f, 1.0f);
+		}
+
+		SourceArea += TriangleSourceArea * TriangleWeight;
+		TargetArea += TriangleTargetArea * TriangleWeight;
+	}
+
+	const float SourceFootprintAreaOnTarget = TargetArea * (1.0f / SourceArea);
+	const float SourceFootprintAreaOnTargetInPixelsSquared = 
+		SourceFootprintAreaOnTarget * (TargetSize.X * TargetSize.Y);
+
+	// Find the mip that better adapts to the footprint.
+	// Swipe from the smalles to the largest mip finding the best fitting one.
+
+	const FIntVector2 SourceSizeInt = FIntVector2(SourceSize.X, SourceSize.Y);
+
+	int32 Mip = FMath::CeilToInt(FMath::Log2(FMath::Max(SourceSize.GetMax(), 1.0f)));
+	int32 BestMip = Mip;
+	float BestMipAreaInPixelsSquared = 1.0f;
+	for (; Mip > 0; --Mip)
+	{
+		FIntVector2 SourceMipSize =
+			FIntVector2{ FMath::Max(SourceSizeInt.X >> Mip, 1), FMath::Max(SourceSizeInt.Y >> Mip, 1) };
+
+		float AreaOfMipInPixelsSquared = static_cast<float>(SourceMipSize.X * SourceMipSize.Y);
+
+		const float CurrentMipAreaDifference = FMath::Abs(AreaOfMipInPixelsSquared - SourceFootprintAreaOnTargetInPixelsSquared);
+		const float BestMipAreaDifference = FMath::Abs(BestMipAreaInPixelsSquared - SourceFootprintAreaOnTargetInPixelsSquared);
+
+		if (CurrentMipAreaDifference < BestMipAreaDifference)
+		{
+			BestMip = Mip;
+			BestMipAreaInPixelsSquared = AreaOfMipInPixelsSquared;
+		}
+	}
+
+	return BestMip;
+}
 
 //-------------------------------------------------------------------------------------------------
 void ImageRasterProjectedCylindrical( const Mesh* pMesh,
diff --git a/Engine/Plugins/Experimental/Mutable/Source/MutableRuntime/Private/MuR/OpImageProject.h b/Engine/Plugins/Experimental/Mutable/Source/MutableRuntime/Private/MuR/OpImageProject.h
--- a/Engine/Plugins/Experimental/Mutable/Source/MutableRuntime/Private/MuR/OpImageProject.h
+++ b/Engine/Plugins/Experimental/Mutable/Source/MutableRuntime/Private/MuR/OpImageProject.h
@@ -3,7 +3,7 @@
 #pragma once
 
 #include "MuR/Mesh.h"
-
+#include "MuR/Image.h"
 
 namespace mu
 {
@@ -22,6 +22,7 @@ namespace mu
                                             const Image* pSource,
                                             const Image* pMask,
 											bool bIsRGBFadingEnabled, bool bIsAlphaFadingEnabled,
+											ESamplingMethod SamplingMethod,
 											float fadeStart, float fadeEnd,
                                             int layout,
                                             int block,
@@ -42,11 +43,15 @@ namespace mu
                                       const Image* pSource,
                                       const Image* pMask,
 									  bool bIsRGBFadingEnabled, bool bIsAlphaFadingEnabled,
+									  ESamplingMethod SamplingMethod,
 									  float fadeStart, float fadeEnd,
 									  int layout,
                                       int block,
                                       SCRATCH_IMAGE_PROJECT* scratch );
 
+	extern int32 ComputeProjectedFootprintBestMip(
+			const Mesh* pMesh, const FProjector& Projector, const FVector2f& TargetSize, const FVector2f& SourceSize);
+
     extern MeshPtr MeshProject( const Mesh* pMesh,
                                 const FProjector& projector );
 
diff --git a/Engine/Plugins/Experimental/Mutable/Source/MutableRuntime/Private/MuR/Operations.h b/Engine/Plugins/Experimental/Mutable/Source/MutableRuntime/Private/MuR/Operations.h
--- a/Engine/Plugins/Experimental/Mutable/Source/MutableRuntime/Private/MuR/Operations.h
+++ b/Engine/Plugins/Experimental/Mutable/Source/MutableRuntime/Private/MuR/Operations.h
@@ -784,8 +784,14 @@ namespace mu
 			
 			int32 blockIndex;
 			uint16 sizeX, sizeY;
+			uint16 SourceSizeX, SourceSizeY;
 			uint8 bIsRGBFadingEnabled : 1;
 			uint8 bIsAlphaFadingEnabled : 1;
+			
+			// Currently only 2 sampling methods are contemplated, but reserve 3 bits for future uses. 
+			uint8 SamplingMethod : 3;
+			// Currently only 2 min filter methods are contemplated, but reserve 3 bits for future uses. 
+			uint8 MinFilterMethod : 3;
         };
 
         struct ImageMakeGrowMapArgs
diff --git a/Engine/Plugins/Experimental/Mutable/Source/MutableRuntime/Public/MuR/Image.h b/Engine/Plugins/Experimental/Mutable/Source/MutableRuntime/Public/MuR/Image.h
--- a/Engine/Plugins/Experimental/Mutable/Source/MutableRuntime/Public/MuR/Image.h
+++ b/Engine/Plugins/Experimental/Mutable/Source/MutableRuntime/Public/MuR/Image.h
@@ -162,6 +162,23 @@ namespace mu
 		Black
 	};
 
+	enum class ESamplingMethod : uint8
+	{
+		Point = 0,
+		BiLinear,
+		MaxValue
+	};
+	static_assert(uint32(ESamplingMethod::MaxValue) <= (1 << 3), "ESampligMethod enum cannot hold more than 8 values");
+
+	enum class EMinFilterMethod : uint8
+	{
+		None = 0,
+		TotalAreaHeuristic,
+		MaxValue
+	};
+	static_assert(uint32(EMinFilterMethod::MaxValue) <= (1 << 3), "EMinFilterMethod enum cannot hold more than 8 values");
+
+
     //! \brief 2D image resource with mipmaps.
 	//! \ingroup runtime
     class MUTABLERUNTIME_API Image : public RefCounted
diff --git a/Engine/Plugins/Experimental/Mutable/Source/MutableRuntime/Public/MuR/Model.h b/Engine/Plugins/Experimental/Mutable/Source/MutableRuntime/Public/MuR/Model.h
--- a/Engine/Plugins/Experimental/Mutable/Source/MutableRuntime/Public/MuR/Model.h
+++ b/Engine/Plugins/Experimental/Mutable/Source/MutableRuntime/Public/MuR/Model.h
@@ -13,7 +13,7 @@
 //! data structures. Compiled models are not necessarily compatible when the runtime is updated,
 //! so this version number can be used externally to verify this. It is not used internally, and
 //! serializing models from different versions than this runtime will probably result in a crash.
-#define MUTABLE_COMPILED_MODEL_CODE_VERSION		uint32( 57 )
+#define MUTABLE_COMPILED_MODEL_CODE_VERSION		uint32( 58 )
 #define MUTABLE_PARAMETERS_VERSION              uint32( 1 )
 
 
diff --git a/Engine/Plugins/Experimental/Mutable/Source/MutableTools/Private/MuT/ASTOpImageRasterMesh.cpp b/Engine/Plugins/Experimental/Mutable/Source/MutableTools/Private/MuT/ASTOpImageRasterMesh.cpp
--- a/Engine/Plugins/Experimental/Mutable/Source/MutableTools/Private/MuT/ASTOpImageRasterMesh.cpp
+++ b/Engine/Plugins/Experimental/Mutable/Source/MutableTools/Private/MuT/ASTOpImageRasterMesh.cpp
@@ -29,8 +29,12 @@ namespace mu
 		blockIndex = -1;
 		sizeX = 0;
 		sizeY = 0;
+		SourceSizeX = 0;
+		SourceSizeY = 0;
 		bIsRGBFadingEnabled = 1;
 		bIsAlphaFadingEnabled = 1;
+		SamplingMethod = ESamplingMethod::Point;
+		MinFilterMethod = EMinFilterMethod::None;
 	}
 
 
@@ -55,8 +59,12 @@ namespace mu
 				blockIndex == Other->blockIndex &&
 				sizeX == Other->sizeX &&
 				sizeY == Other->sizeY &&
+				SourceSizeX == Other->SourceSizeX &&
+				SourceSizeY == Other->SourceSizeY &&
 				bIsRGBFadingEnabled == Other->bIsRGBFadingEnabled &&
-				bIsAlphaFadingEnabled == Other->bIsAlphaFadingEnabled;
+				bIsAlphaFadingEnabled == Other->bIsAlphaFadingEnabled &&
+				SamplingMethod == Other->SamplingMethod &&
+				MinFilterMethod == Other->MinFilterMethod;
 		}
 		return false;
 	}
@@ -88,8 +96,12 @@ namespace mu
 		n->blockIndex = blockIndex;
 		n->sizeX = sizeX;
 		n->sizeY = sizeY;
+		n->SourceSizeX = SourceSizeX;
+		n->SourceSizeY = SourceSizeY;
 		n->bIsRGBFadingEnabled = bIsRGBFadingEnabled;
 		n->bIsAlphaFadingEnabled = bIsAlphaFadingEnabled;
+		n->SamplingMethod = SamplingMethod;
+		n->MinFilterMethod = MinFilterMethod;
 		return n;
 	}
 
@@ -117,8 +129,12 @@ namespace mu
 			args.blockIndex = blockIndex;
 			args.sizeX = sizeX;
 			args.sizeY = sizeY;
+			args.SourceSizeX = SourceSizeX;
+			args.SourceSizeY = SourceSizeY;
 			args.bIsRGBFadingEnabled = bIsRGBFadingEnabled;
 			args.bIsAlphaFadingEnabled = bIsAlphaFadingEnabled;
+			args.SamplingMethod = static_cast<uint8>(SamplingMethod);
+			args.MinFilterMethod = static_cast<uint8>(MinFilterMethod);
 
 			if (mesh) args.mesh = mesh->linkedAddress;
 			if (image) args.image = image->linkedAddress;
diff --git a/Engine/Plugins/Experimental/Mutable/Source/MutableTools/Private/MuT/ASTOpImageRasterMesh.h b/Engine/Plugins/Experimental/Mutable/Source/MutableTools/Private/MuT/ASTOpImageRasterMesh.h
--- a/Engine/Plugins/Experimental/Mutable/Source/MutableTools/Private/MuT/ASTOpImageRasterMesh.h
+++ b/Engine/Plugins/Experimental/Mutable/Source/MutableTools/Private/MuT/ASTOpImageRasterMesh.h
@@ -25,8 +25,11 @@ namespace mu
 
 		int32 blockIndex;
 		uint16 sizeX, sizeY;
+		uint16 SourceSizeX, SourceSizeY;
 		uint8 bIsRGBFadingEnabled : 1;
 		uint8 bIsAlphaFadingEnabled : 1;
+		ESamplingMethod SamplingMethod;
+		EMinFilterMethod MinFilterMethod;
 
 	public:
 
diff --git a/Engine/Plugins/Experimental/Mutable/Source/MutableTools/Private/MuT/CodeGenerator_Image.cpp b/Engine/Plugins/Experimental/Mutable/Source/MutableTools/Private/MuT/CodeGenerator_Image.cpp
--- a/Engine/Plugins/Experimental/Mutable/Source/MutableTools/Private/MuT/CodeGenerator_Image.cpp
+++ b/Engine/Plugins/Experimental/Mutable/Source/MutableTools/Private/MuT/CodeGenerator_Image.cpp
@@ -1556,6 +1556,8 @@ namespace mu
 
             // Generate
             op->image = Generate( node.m_pImage.get() );
+			op->SourceSizeX = desc.m_size.x();
+			op->SourceSizeY = desc.m_size.y();
 
             // Restore rect
             m_imageState.Pop();
@@ -1575,6 +1577,8 @@ namespace mu
 
 		op->bIsRGBFadingEnabled = node.bIsRGBFadingEnabled;
 		op->bIsAlphaFadingEnabled = node.bIsAlphaFadingEnabled;
+		op->SamplingMethod = node.SamplingMethod;
+		op->MinFilterMethod = node.MinFilterMethod;
 
 		// Fading angles are optional, and stored in a colour. If one exists, we generate both.
 		if (node.m_pAngleFadeStart || node.m_pAngleFadeEnd)
@@ -1612,6 +1616,10 @@ namespace mu
         rasterop->blockIndex = op->blockIndex;
         rasterop->sizeX = op->sizeX;
         rasterop->sizeY = op->sizeY;
+		rasterop->SourceSizeX = op->sizeX;
+		rasterop->SourceSizeY = op->sizeX;
+		rasterop->SamplingMethod = ESamplingMethod::Point;
+		rasterop->MinFilterMethod = EMinFilterMethod::None;
 
         Ptr<ASTOpImageMakeGrowMap> MakeGrowMapOp = new ASTOpImageMakeGrowMap();
 		MakeGrowMapOp->Mask = rasterop;
@@ -1628,8 +1636,8 @@ namespace mu
 
         Ptr<ASTOpFixed> disop = new ASTOpFixed();
         disop->op.type = OP_TYPE::IM_DISPLACE;
-        disop->SetChild( disop->op.args.ImageDisplace.displacementMap, MakeGrowMapOp);
-        disop->SetChild( disop->op.args.ImageDisplace.source, op );
+        disop->SetChild(disop->op.args.ImageDisplace.displacementMap, MakeGrowMapOp);
+        disop->SetChild(disop->op.args.ImageDisplace.source, op );
 
         result.op = disop;
     }
diff --git a/Engine/Plugins/Experimental/Mutable/Source/MutableTools/Private/MuT/NodeImageProject.cpp b/Engine/Plugins/Experimental/Mutable/Source/MutableTools/Private/MuT/NodeImageProject.cpp
--- a/Engine/Plugins/Experimental/Mutable/Source/MutableTools/Private/MuT/NodeImageProject.cpp
+++ b/Engine/Plugins/Experimental/Mutable/Source/MutableTools/Private/MuT/NodeImageProject.cpp
@@ -162,6 +162,18 @@ namespace mu
 		m_pD->m_pAngleFadeEnd = pNode;
 	}
 
+	//---------------------------------------------------------------------------------------------
+	void NodeImageProject::SetSamplingMethod( ESamplingMethod SamplingMethod )
+	{
+		m_pD->SamplingMethod = SamplingMethod;
+	}
+
+	//---------------------------------------------------------------------------------------------
+	void NodeImageProject::SetMinFilterMethod( EMinFilterMethod MinFilterMethod )
+	{
+		m_pD->MinFilterMethod = MinFilterMethod;
+	}
+
 
 	//---------------------------------------------------------------------------------------------
 	NodeImagePtr NodeImageProject::GetTargetMask() const
diff --git a/Engine/Plugins/Experimental/Mutable/Source/MutableTools/Private/MuT/NodeImageProjectPrivate.h b/Engine/Plugins/Experimental/Mutable/Source/MutableTools/Private/MuT/NodeImageProjectPrivate.h
--- a/Engine/Plugins/Experimental/Mutable/Source/MutableTools/Private/MuT/NodeImageProjectPrivate.h
+++ b/Engine/Plugins/Experimental/Mutable/Source/MutableTools/Private/MuT/NodeImageProjectPrivate.h
@@ -34,11 +34,13 @@ namespace mu
 		uint8 m_layout = 0;
 		bool bIsRGBFadingEnabled = true;
 		bool bIsAlphaFadingEnabled = true;
+		ESamplingMethod SamplingMethod = ESamplingMethod::Point;
+		EMinFilterMethod MinFilterMethod = EMinFilterMethod::None;
 
 		//!
 		void Serialise( OutputArchive& arch ) const
 		{
-            uint32 ver = 3;
+            uint32 ver = 4;
 			arch << ver;
 
 			arch << m_pProjector;
@@ -51,6 +53,8 @@ namespace mu
 			arch << m_imageSize;
 			arch << bIsRGBFadingEnabled;
 			arch << bIsAlphaFadingEnabled;
+			arch << SamplingMethod;
+			arch << MinFilterMethod;
 		}
 
 		//!
@@ -58,7 +62,7 @@ namespace mu
 		{
             uint32 ver;
 			arch >> ver;
-            check(ver>=2 && ver<=3);
+            check(ver>=2 && ver<=4);
 
 			arch >> m_pProjector;
 			arch >> m_pMesh;
@@ -73,6 +77,12 @@ namespace mu
 				arch >> bIsRGBFadingEnabled;
 				arch >> bIsAlphaFadingEnabled;
 			}
+
+			if (ver >= 4)
+			{
+				arch >> SamplingMethod;
+				arch >> MinFilterMethod;
+			}
 		}
 	};
 
diff --git a/Engine/Plugins/Experimental/Mutable/Source/MutableTools/Public/MuT/NodeImageProject.h b/Engine/Plugins/Experimental/Mutable/Source/MutableTools/Public/MuT/NodeImageProject.h
--- a/Engine/Plugins/Experimental/Mutable/Source/MutableTools/Public/MuT/NodeImageProject.h
+++ b/Engine/Plugins/Experimental/Mutable/Source/MutableTools/Public/MuT/NodeImageProject.h
@@ -8,6 +8,7 @@
 #include "MuR/RefCounted.h"
 #include "MuT/Node.h"
 #include "MuT/NodeImage.h"
+#include "MuR/Image.h"
 
 namespace mu
 {
@@ -36,7 +37,6 @@ namespace mu
 	class InputArchive;
 	class OutputArchive;
 
-
 	//!
 	//! \ingroup model
 	class MUTABLETOOLS_API NodeImageProject : public NodeImage
@@ -93,6 +93,12 @@ namespace mu
 		//! Set the node generating the fading end angle
 		void SetAngleFadeEnd( NodeScalarPtr );
 
+		//! Set sampling method.
+		void SetSamplingMethod(ESamplingMethod SamplingMethod);
+		
+		//! Set min filter method
+		void SetMinFilterMethod(EMinFilterMethod MinFilterMethod);
+
         //! Get the node generating the image to project.
         NodeImagePtr GetImage() const;
         void SetImage( NodeImagePtr );
