[Mutable] Fix android compilation errors.

#rnx

diff --git a/Engine/Plugins/Experimental/Mutable/Source/MutableRuntime/Private/MuR/OpImageProject.cpp b/Engine/Plugins/Experimental/Mutable/Source/MutableRuntime/Private/MuR/OpImageProject.cpp
--- a/Engine/Plugins/Experimental/Mutable/Source/MutableRuntime/Private/MuR/OpImageProject.cpp
+++ b/Engine/Plugins/Experimental/Mutable/Source/MutableRuntime/Private/MuR/OpImageProject.cpp
@@ -330,10 +330,10 @@ namespace
 
 					static_assert(PIXEL_SIZE <= 4);
 					const uint16 FadingMask[4] = { 
-						bIsRGBFadingEnabled   ? 1 : 0, 
-						bIsRGBFadingEnabled   ? 1 : 0, 
-						bIsRGBFadingEnabled   ? 1 : 0, 
-						bIsAlphaFadingEnabled ? 1 : 0 };
+						static_cast<uint16>(bIsRGBFadingEnabled), 
+						static_cast<uint16>(bIsRGBFadingEnabled), 
+						static_cast<uint16>(bIsRGBFadingEnabled), 
+						static_cast<uint16>(bIsAlphaFadingEnabled)};
 
 					for (int32 I = 0; I < PIXEL_SIZE; ++I)
 					{
