#rnx

diff --git a/Engine/Plugins/Experimental/Mutable/Source/CustomizableObject/Private/MuCO/CustomizableObjectInstance.cpp b/Engine/Plugins/Experimental/Mutable/Source/CustomizableObject/Private/MuCO/CustomizableObjectInstance.cpp
--- a/Engine/Plugins/Experimental/Mutable/Source/CustomizableObject/Private/MuCO/CustomizableObjectInstance.cpp
+++ b/Engine/Plugins/Experimental/Mutable/Source/CustomizableObject/Private/MuCO/CustomizableObjectInstance.cpp
@@ -2487,12 +2487,6 @@ void CopyTextureProperties(UTexture2D* Texture, const UTexture2D* SourceTexture)
 	Texture->AddressY = SourceTexture->AddressY;
 }
 
-static int32 MaxTextureSizeToGenerate = 0;
-FAutoConsoleVariableRef CVarMaxTextureSizeToGenerate(
-	TEXT("Mutable.MaxTextureSizeToGenerate"),
-	MaxTextureSizeToGenerate,
-	TEXT("Max texture size on non streamed textures. Mip 0 will be the first mip with max size equal or less than MaxTextureSizeToGenerate."
-		"If a texture doesn't have small enough mips, mip 0 will be the last mip available."));
 
 // The memory allocated in the function and pointed by the returned pointer is owned by the caller and must be freed. 
 // If assigned to a UTexture2D, it will be freed by that UTexture2D
@@ -2531,6 +2525,11 @@ FTexturePlatformData* UCustomizableInstancePrivateData::MutableCreateImagePlatfo
 	}
 
 	// Reduce final texture size if we surpass the max size we can generate.
+	UCustomizableObjectSystem* System = UCustomizableObjectSystem::GetInstance();
+	FCustomizableObjectSystemPrivate* SystemPrivate = System ? System->GetPrivate() : nullptr;
+
+	int32 MaxTextureSizeToGenerate = SystemPrivate ? SystemPrivate->MaxTextureSizeToGenerate : 0;
+
 	if (MaxTextureSizeToGenerate > 0)
 	{
 		// Skip mips only if texture streaming is disabled 
diff --git a/Engine/Plugins/Experimental/Mutable/Source/CustomizableObject/Private/MuCO/CustomizableObjectMipDataProvider.cpp b/Engine/Plugins/Experimental/Mutable/Source/CustomizableObject/Private/MuCO/CustomizableObjectMipDataProvider.cpp
--- a/Engine/Plugins/Experimental/Mutable/Source/CustomizableObject/Private/MuCO/CustomizableObjectMipDataProvider.cpp
+++ b/Engine/Plugins/Experimental/Mutable/Source/CustomizableObject/Private/MuCO/CustomizableObjectMipDataProvider.cpp
@@ -97,13 +97,31 @@ namespace impl
 				// runtime cannot remember all the resources it has built, and only remembers a fixed amount.
 				mu::RESOURCE_ID MipImageID = Instance->GetImageId(ImageRef.LOD, ImageRef.Component, SurfaceIndex, ImageRef.Image);
 
-				// TODO: Why do we need to do this again? The full size should be stored in the initial image creation.
-				mu::FImageDesc ImageDesc;
+				UCustomizableObjectSystem* COSystem = UCustomizableObjectSystem::GetInstance();
+				FCustomizableObjectSystemPrivate* SystemPrivate = COSystem ? COSystem->GetPrivate() : nullptr;
+				int32 MaxTextureSizeToGenerate = SystemPrivate ? SystemPrivate->MaxTextureSizeToGenerate : 0;
+
+				int32 ExtraMipsToSkip = 0;
+				
+				if (MaxTextureSizeToGenerate > 0)
+				{
+					// TODO: Why do we need to do this again? The full size should be stored in the initial image creation.
+					mu::FImageDesc ImageDesc;
+					System->GetImageDesc(InstanceID, MipImageID, ImageDesc);
+
+					uint16 MaxSize = FMath::Max(ImageDesc.m_size[0], ImageDesc.m_size[1]);
+
+					if (MaxSize > MaxTextureSizeToGenerate)
+					{
+						ExtraMipsToSkip = FMath::CeilLogTwo(MaxSize / MaxTextureSizeToGenerate);
+					}
+				}
+
 				mu::ImagePtrConst Image;
 				{
 					MUTABLE_CPUPROFILER_SCOPE(GetImage);
 
-					Image = System->GetImage(InstanceID, MipImageID, OperationData->MipsToSkip, ImageRef.LOD);
+					Image = System->GetImage(InstanceID, MipImageID, OperationData->MipsToSkip + ExtraMipsToSkip, ImageRef.LOD);
 				}
 
 				check(Image);
diff --git a/Engine/Plugins/Experimental/Mutable/Source/CustomizableObject/Private/MuCO/CustomizableObjectSystem.cpp b/Engine/Plugins/Experimental/Mutable/Source/CustomizableObject/Private/MuCO/CustomizableObjectSystem.cpp
--- a/Engine/Plugins/Experimental/Mutable/Source/CustomizableObject/Private/MuCO/CustomizableObjectSystem.cpp
+++ b/Engine/Plugins/Experimental/Mutable/Source/CustomizableObject/Private/MuCO/CustomizableObjectSystem.cpp
@@ -707,6 +707,14 @@ static FAutoConsoleVariableRef CVarSkipGenerateResidentMips(
 	TEXT("If 1 or greater, resident mip generation will be optional. If 0, resident mips will be always generated"),
 	ECVF_Default);
 
+int32 FCustomizableObjectSystemPrivate::MaxTextureSizeToGenerate = 0;
+
+FAutoConsoleVariableRef CVarMaxTextureSizeToGenerate(
+	TEXT("Mutable.MaxTextureSizeToGenerate"),
+	FCustomizableObjectSystemPrivate::MaxTextureSizeToGenerate,
+	TEXT("Max texture size on Mutable textures. Mip 0 will be the first mip with max size equal or less than MaxTextureSizeToGenerate."
+		"If a texture doesn't have small enough mips, mip 0 will be the last mip available."));
+
 
 /** Update the given Instance Skeletal Meshes and call its callbacks. */
 void UpdateSkeletalMesh(UCustomizableObjectInstance* CustomizableObjectInstance, const FDescriptorRuntimeHash& UpdatedDescriptorRuntimeHash)
@@ -1425,8 +1433,18 @@ namespace impl
 			// This should only be done when using progressive images, since GetImageDesc does some actual processing.
 			{
 				System->GetImageDesc(OperationData->InstanceID, Image.ImageID, ImageDesc);
-				Image.FullImageSizeX = ImageDesc.m_size[0];
-				Image.FullImageSizeY = ImageDesc.m_size[1];
+
+				uint16 MaxTextureSizeToGenerate = uint16(CustomizableObjectSystemPrivateData->MaxTextureSizeToGenerate);
+				uint16 MaxSize = FMath::Max(ImageDesc.m_size[0], ImageDesc.m_size[1]);
+				uint16 Reduction = 1;
+
+				if (MaxTextureSizeToGenerate > 0 && MaxSize > MaxTextureSizeToGenerate)
+				{
+					Reduction = MaxSize / MaxTextureSizeToGenerate;
+				}
+
+				Image.FullImageSizeX = ImageDesc.m_size[0] / Reduction;
+				Image.FullImageSizeY = ImageDesc.m_size[1] / Reduction;
 			}
 
 			bool bCached = false;
@@ -1515,7 +1533,7 @@ namespace impl
 			}
 		}
 	}
-
+	
 
 	void Subtask_Mutable_PrepareSkeletonData(const TSharedPtr<FMutableOperationData>& OperationData)
 	{
diff --git a/Engine/Plugins/Experimental/Mutable/Source/CustomizableObject/Private/MuCO/CustomizableObjectSystemPrivate.h b/Engine/Plugins/Experimental/Mutable/Source/CustomizableObject/Private/MuCO/CustomizableObjectSystemPrivate.h
--- a/Engine/Plugins/Experimental/Mutable/Source/CustomizableObject/Private/MuCO/CustomizableObjectSystemPrivate.h
+++ b/Engine/Plugins/Experimental/Mutable/Source/CustomizableObject/Private/MuCO/CustomizableObjectSystemPrivate.h
@@ -567,6 +567,7 @@ public:
 	static int32 EnableMutableAnimInfoDebugging;
 	static int32 EnableSkipGenerateResidentMips;
 	static int32 EnableOnlyGenerateRequestedLODs;
+	static int32 MaxTextureSizeToGenerate;
 	static bool bEnableMutableReusePreviousUpdateData;
 
 	/** */
