[mutable] Added new methods to reduce the layout blocks in a fixed layout.
- Added method to reduce the blocks by units instead of by halves.
- Added method to reduce both sides of a block at the same time.
- Improved order of reduction when using reduction priorities (ordered by priority AND size).

#rnx

diff --git a/Engine/Plugins/Experimental/Mutable/Source/CustomizableObjectEditor/Classes/MuCOE/CustomizableObjectLayout.h b/Engine/Plugins/Experimental/Mutable/Source/CustomizableObjectEditor/Classes/MuCOE/CustomizableObjectLayout.h
--- a/Engine/Plugins/Experimental/Mutable/Source/CustomizableObjectEditor/Classes/MuCOE/CustomizableObjectLayout.h
+++ b/Engine/Plugins/Experimental/Mutable/Source/CustomizableObjectEditor/Classes/MuCOE/CustomizableObjectLayout.h
@@ -14,12 +14,22 @@ enum class ECustomizableObjectTextureLayoutPackingStrategy : uint8
 	Fixed = 1 UMETA(DisplayName = "Fixed Layout")
 };
 
+// Fixed Layout reduction methods
+UENUM()
+enum class ECustomizableObjectLayoutBlockReductionMethod : uint8
+{
+	// Layout blocks will be reduced by halves
+	Halve = 0 UMETA(DisplayName = "Reduce by Half"),
+	// LAyout blocks will be reduced by a grid unit
+	Unitary = 1 UMETA(DisplayName = "Reduce by Unit")
+};
+
 USTRUCT()
 struct CUSTOMIZABLEOBJECTEDITOR_API FCustomizableObjectLayoutBlock
 {
 	GENERATED_USTRUCT_BODY()
 
-		FCustomizableObjectLayoutBlock()
+	FCustomizableObjectLayoutBlock()
 	{
 		Min = FIntPoint(0, 0);
 		Max = FIntPoint(1, 1);
@@ -27,17 +37,20 @@ struct CUSTOMIZABLEOBJECTEDITOR_API FCustomizableObjectLayoutBlock
 	}
 
 	UPROPERTY(EditAnywhere, Category = CustomizableObject)
-		FIntPoint Min;
+	FIntPoint Min;
 
 	UPROPERTY(EditAnywhere, Category = CustomizableObject)
-		FIntPoint Max;
+	FIntPoint Max;
 
 	UPROPERTY(EditAnywhere, Category = CustomizableObject)
-		uint32 Priority;
+	uint32 Priority;
 
 	//! Unique unchangeable id used to reference this block from other nodes.
 	UPROPERTY()
-		FGuid Id;
+	FGuid Id;
+
+	UPROPERTY(EditAnywhere, Category = CustomizableObject)
+	bool bUseSymmetry = false;
 };
 
 UCLASS()
@@ -57,6 +70,7 @@ public:
 	void SetLayoutName(FString Name);
 	void SetIgnoreVertexLayoutWarnings(bool bValue);
 	void SetIgnoreWarningsLOD(int32 LODValue);
+	void SetBlockReductionMethod(ECustomizableObjectLayoutBlockReductionMethod Method);
 
 
 	int32 GetLOD() const { return LOD; }
@@ -69,6 +83,7 @@ public:
 	ECustomizableObjectTextureLayoutPackingStrategy GetPackingStrategy() const { return PackingStrategy; }
 	bool GetIgnoreVertexLayoutWarnings() const { return bIgnoreUnassignedVertexWarning; };
 	int32 GetFirstLODToIgnoreWarnings() { return FirstLODToIgnore; };
+	ECustomizableObjectLayoutBlockReductionMethod GetBlockReductionMethod()const { return BlockReductionMethod; }
 
 	void GetUVChannel(TArray<FVector2f>& UVs, int32 UVChannelIndex = 0) const;
 
@@ -117,4 +132,7 @@ private:
 	UPROPERTY()
 	int32 FirstLODToIgnore = 0;
 
+	UPROPERTY()
+	ECustomizableObjectLayoutBlockReductionMethod BlockReductionMethod = ECustomizableObjectLayoutBlockReductionMethod::Halve;
+
 };
diff --git a/Engine/Plugins/Experimental/Mutable/Source/CustomizableObjectEditor/Private/MuCOE/CustomizableObjectLayout.cpp b/Engine/Plugins/Experimental/Mutable/Source/CustomizableObjectEditor/Private/MuCOE/CustomizableObjectLayout.cpp
--- a/Engine/Plugins/Experimental/Mutable/Source/CustomizableObjectEditor/Private/MuCOE/CustomizableObjectLayout.cpp
+++ b/Engine/Plugins/Experimental/Mutable/Source/CustomizableObjectEditor/Private/MuCOE/CustomizableObjectLayout.cpp
@@ -24,9 +24,11 @@ UCustomizableObjectLayout::UCustomizableObjectLayout()
 	Block.Max = FIntPoint(4, 4);
 	Block.Id = FGuid::NewGuid();
 	Block.Priority = 0;
+	Block.bUseSymmetry = false;
 	Blocks.Add(Block);
 
 	PackingStrategy = ECustomizableObjectTextureLayoutPackingStrategy::Resizable;
+	BlockReductionMethod = ECustomizableObjectLayoutBlockReductionMethod::Halve;
 }
 
 
@@ -110,7 +112,7 @@ void UCustomizableObjectLayout::GenerateBlocksFromUVs()
 		{
 			Blocks.Empty();
 		
-			//Generating the layout blocks with the mutable layout		
+			//Generating the layout blocks with the mutable layout
 			for (int i = 0; i < Layout->GetBlockCount(); ++i)
 			{
 				int minX, minY, sizeX, sizeY;
@@ -122,6 +124,7 @@ void UCustomizableObjectLayout::GenerateBlocksFromUVs()
 				block.Max = FIntPoint(minX + sizeX, minY + sizeY);
 				block.Id = FGuid::NewGuid();
 				block.Priority = 0;
+				block.bUseSymmetry = false;
 				Blocks.Add(block);
 			}
 		
@@ -179,4 +182,9 @@ void UCustomizableObjectLayout::SetIgnoreWarningsLOD(int32 LODValue)
 	FirstLODToIgnore = LODValue;
 }
 
+
+void UCustomizableObjectLayout::SetBlockReductionMethod(ECustomizableObjectLayoutBlockReductionMethod Method)
+{
+	BlockReductionMethod = Method;
+}
 #undef LOCTEXT_NAMESPACE
diff --git a/Engine/Plugins/Experimental/Mutable/Source/CustomizableObjectEditor/Private/MuCOE/GenerateMutableSource/GenerateMutableSourceLayout.cpp b/Engine/Plugins/Experimental/Mutable/Source/CustomizableObjectEditor/Private/MuCOE/GenerateMutableSource/GenerateMutableSourceLayout.cpp
--- a/Engine/Plugins/Experimental/Mutable/Source/CustomizableObjectEditor/Private/MuCOE/GenerateMutableSource/GenerateMutableSourceLayout.cpp
+++ b/Engine/Plugins/Experimental/Mutable/Source/CustomizableObjectEditor/Private/MuCOE/GenerateMutableSource/GenerateMutableSourceLayout.cpp
@@ -59,22 +59,8 @@ mu::NodeLayoutPtr GenerateMutableSourceLayout(const UEdGraphPin * Pin, FMutableG
 		LayoutNode->SetMaxGridSize(TypedNodeBlocks->Layout->GetMaxGridSize().X, TypedNodeBlocks->Layout->GetMaxGridSize().Y);
 		LayoutNode->SetBlockCount(TypedNodeBlocks->Layout->Blocks.Num() ? TypedNodeBlocks->Layout->Blocks.Num() : 1);
 		LayoutNode->SetIgnoreWarningsLOD(TypedNodeBlocks->Layout->GetIgnoreVertexLayoutWarnings() ? TypedNodeBlocks->Layout->GetFirstLODToIgnoreWarnings() : -1);
-
-		mu::EPackStrategy strategy = mu::EPackStrategy::RESIZABLE_LAYOUT;
-
-		switch (TypedNodeBlocks->Layout->GetPackingStrategy())
-		{
-		case ECustomizableObjectTextureLayoutPackingStrategy::Resizable:
-			strategy = mu::EPackStrategy::RESIZABLE_LAYOUT;
-			break;
-		case ECustomizableObjectTextureLayoutPackingStrategy::Fixed:
-			strategy = mu::EPackStrategy::FIXED_LAYOUT;
-			break;
-		default:
-			break;
-		}
-
-		LayoutNode->SetLayoutPackingStrategy(strategy);
+		LayoutNode->SetLayoutPackingStrategy(TypedNodeBlocks->Layout->GetPackingStrategy() == ECustomizableObjectTextureLayoutPackingStrategy::Fixed ? mu::EPackStrategy::FIXED_LAYOUT : mu::EPackStrategy::RESIZABLE_LAYOUT);
+		LayoutNode->SetBlockReductionMethod(TypedNodeBlocks->Layout->GetBlockReductionMethod() == ECustomizableObjectLayoutBlockReductionMethod::Halve ? mu::EReductionMethod::HALVE_REDUCTION : mu::EReductionMethod::UNITARY_REDUCTION);
 
 		if (TypedNodeBlocks->Layout->Blocks.Num())
 		{
@@ -86,7 +72,7 @@ mu::NodeLayoutPtr GenerateMutableSourceLayout(const UEdGraphPin * Pin, FMutableG
 					TypedNodeBlocks->Layout->Blocks[BlockIndex].Max.X - TypedNodeBlocks->Layout->Blocks[BlockIndex].Min.X,
 					TypedNodeBlocks->Layout->Blocks[BlockIndex].Max.Y - TypedNodeBlocks->Layout->Blocks[BlockIndex].Min.Y);
 
-				LayoutNode->SetBlockPriority(BlockIndex, TypedNodeBlocks->Layout->Blocks[BlockIndex].Priority);
+				LayoutNode->SetBlockOptions(BlockIndex, TypedNodeBlocks->Layout->Blocks[BlockIndex].Priority, TypedNodeBlocks->Layout->Blocks[BlockIndex].bUseSymmetry);
 			}
 		}
 		else
@@ -95,7 +81,7 @@ mu::NodeLayoutPtr GenerateMutableSourceLayout(const UEdGraphPin * Pin, FMutableG
 			GenerationContext.Compiler->CompilerLog(FText::FromString(msg), Node, EMessageSeverity::Warning);
 
 			LayoutNode->SetBlock(0, 0, 0, TypedNodeBlocks->Layout->GetGridSize().X, TypedNodeBlocks->Layout->GetGridSize().Y);
-			LayoutNode->SetBlockPriority(0, 0);
+			LayoutNode->SetBlockOptions(0, 0, false);
 		}
 	}
 	
diff --git a/Engine/Plugins/Experimental/Mutable/Source/CustomizableObjectEditor/Private/MuCOE/GenerateMutableSource/GenerateMutableSourceMesh.cpp b/Engine/Plugins/Experimental/Mutable/Source/CustomizableObjectEditor/Private/MuCOE/GenerateMutableSource/GenerateMutableSourceMesh.cpp
--- a/Engine/Plugins/Experimental/Mutable/Source/CustomizableObjectEditor/Private/MuCOE/GenerateMutableSource/GenerateMutableSourceMesh.cpp
+++ b/Engine/Plugins/Experimental/Mutable/Source/CustomizableObjectEditor/Private/MuCOE/GenerateMutableSource/GenerateMutableSourceMesh.cpp
@@ -2763,9 +2763,10 @@ mu::NodeMeshPtr GenerateMutableSourceMesh(const UEdGraphPin * Pin,
 				LayoutNode->SetGridSize(GridSize, GridSize);
 				LayoutNode->SetMaxGridSize(GridSize, GridSize);
 				LayoutNode->SetLayoutPackingStrategy(mu::EPackStrategy::RESIZABLE_LAYOUT);
+				LayoutNode->SetBlockReductionMethod(mu::EReductionMethod::HALVE_REDUCTION);
 				LayoutNode->SetBlockCount(1);
 				LayoutNode->SetBlock(0, 0, 0, GridSize, GridSize);
-				LayoutNode->SetBlockPriority(0, 0);
+				LayoutNode->SetBlockOptions(0, 0, false);
 				MeshNode->SetLayout(0, LayoutNode);
 
 				// We need it here because we create multiple nodes.
@@ -2839,9 +2840,10 @@ mu::NodeMeshPtr GenerateMutableSourceMesh(const UEdGraphPin * Pin,
 					LayoutNode->SetGridSize(GridSize, GridSize);
 					LayoutNode->SetMaxGridSize(GridSize, GridSize);
 					LayoutNode->SetLayoutPackingStrategy(mu::EPackStrategy::RESIZABLE_LAYOUT);
+					LayoutNode->SetBlockReductionMethod(mu::EReductionMethod::HALVE_REDUCTION);
 					LayoutNode->SetBlockCount(1);
 					LayoutNode->SetBlock(0, 0, 0, GridSize, GridSize);
-					LayoutNode->SetBlockPriority(0, 0);
+					LayoutNode->SetBlockOptions(0, 0, false);
 					MeshNode->SetLayout(0, LayoutNode);
 
 					// We need it here because we create multiple nodes.
@@ -3374,22 +3376,8 @@ mu::NodeMeshPtr GenerateMutableSourceMesh(const UEdGraphPin * Pin,
 					LayoutNode->SetGridSize(Layouts[i]->GetGridSize().X, Layouts[i]->GetGridSize().Y);
 					LayoutNode->SetMaxGridSize(Layouts[i]->GetMaxGridSize().X, Layouts[i]->GetMaxGridSize().Y);
 					LayoutNode->SetBlockCount(Layouts[i]->Blocks.Num() ? Layouts[i]->Blocks.Num() : 1);
-
-					mu::EPackStrategy strategy = mu::EPackStrategy::RESIZABLE_LAYOUT;
-
-					switch (Layouts[i]->GetPackingStrategy())
-					{
-					case ECustomizableObjectTextureLayoutPackingStrategy::Resizable:
-						strategy = mu::EPackStrategy::RESIZABLE_LAYOUT;
-						break;
-					case ECustomizableObjectTextureLayoutPackingStrategy::Fixed:
-						strategy = mu::EPackStrategy::FIXED_LAYOUT;
-						break;
-					default:
-						break;
-					}
-
-					LayoutNode->SetLayoutPackingStrategy(strategy);
+					LayoutNode->SetLayoutPackingStrategy(Layouts[i]->GetPackingStrategy() == ECustomizableObjectTextureLayoutPackingStrategy::Fixed ? mu::EPackStrategy::FIXED_LAYOUT : mu::EPackStrategy::RESIZABLE_LAYOUT);
+					LayoutNode->SetBlockReductionMethod(Layouts[i]->GetBlockReductionMethod() == ECustomizableObjectLayoutBlockReductionMethod::Halve ? mu::EReductionMethod::HALVE_REDUCTION : mu::EReductionMethod::UNITARY_REDUCTION);
 
 					if (Layouts[i]->Blocks.Num())
 					{
@@ -3401,7 +3389,7 @@ mu::NodeMeshPtr GenerateMutableSourceMesh(const UEdGraphPin * Pin,
 								Layouts[i]->Blocks[BlockIndex].Max.X - Layouts[i]->Blocks[BlockIndex].Min.X,
 								Layouts[i]->Blocks[BlockIndex].Max.Y - Layouts[i]->Blocks[BlockIndex].Min.Y);
 
-							LayoutNode->SetBlockPriority(BlockIndex, Layouts[i]->Blocks[BlockIndex].Priority);
+							LayoutNode->SetBlockOptions(BlockIndex, Layouts[i]->Blocks[BlockIndex].Priority, Layouts[i]->Blocks[BlockIndex].bUseSymmetry);
 						}
 					}
 					else
@@ -3410,7 +3398,7 @@ mu::NodeMeshPtr GenerateMutableSourceMesh(const UEdGraphPin * Pin,
 						GenerationContext.Compiler->CompilerLog(FText::FromString(msg), Node, EMessageSeverity::Warning);
 
 						LayoutNode->SetBlock(0, 0, 0, Layouts[i]->GetGridSize().X, Layouts[i]->GetGridSize().Y);
-						LayoutNode->SetBlockPriority(0, 0);
+						LayoutNode->SetBlockOptions(0, 0, false);
 					}
 
 					MeshTableNode->SetLayout(i, LayoutNode);
diff --git a/Engine/Plugins/Experimental/Mutable/Source/CustomizableObjectEditor/Private/MuCOE/Nodes/CustomizableObjectNodeTable.cpp b/Engine/Plugins/Experimental/Mutable/Source/CustomizableObjectEditor/Private/MuCOE/Nodes/CustomizableObjectNodeTable.cpp
--- a/Engine/Plugins/Experimental/Mutable/Source/CustomizableObjectEditor/Private/MuCOE/Nodes/CustomizableObjectNodeTable.cpp
+++ b/Engine/Plugins/Experimental/Mutable/Source/CustomizableObjectEditor/Private/MuCOE/Nodes/CustomizableObjectNodeTable.cpp
@@ -556,6 +556,7 @@ void UCustomizableObjectNodeTable::RemapPinsData(const TMap<UEdGraphPin*, UEdGra
 							NewLayout->SetGridSize(OldLayout->GetGridSize());
 							NewLayout->SetMaxGridSize(OldLayout->GetMaxGridSize());
 							NewLayout->SetPackingStrategy(OldLayout->GetPackingStrategy());
+							NewLayout->SetBlockReductionMethod(OldLayout->GetBlockReductionMethod());
 
 							break;
 						}
diff --git a/Engine/Plugins/Experimental/Mutable/Source/CustomizableObjectEditor/Private/MuCOE/SCustomizableObjectLayoutGrid.cpp b/Engine/Plugins/Experimental/Mutable/Source/CustomizableObjectEditor/Private/MuCOE/SCustomizableObjectLayoutGrid.cpp
--- a/Engine/Plugins/Experimental/Mutable/Source/CustomizableObjectEditor/Private/MuCOE/SCustomizableObjectLayoutGrid.cpp
+++ b/Engine/Plugins/Experimental/Mutable/Source/CustomizableObjectEditor/Private/MuCOE/SCustomizableObjectLayoutGrid.cpp
@@ -10,6 +10,7 @@
 #include "MuCO/UnrealPortabilityHelpers.h"
 #include "Rendering/SlateRenderer.h"
 #include "RenderingThread.h"
+#include "Styling/SlateTypes.h"
 #include "UnrealClient.h"
 #include "Widgets/Input/SNumericEntryBox.h"
 
@@ -141,6 +142,7 @@ void SCustomizableObjectLayoutGrid::Construct( const FArguments& InArgs )
 	DeleteBlocksDelegate = InArgs._OnDeleteBlocks;
 	AddBlockAtDelegate = InArgs._OnAddBlockAt;
 	OnSetBlockPriority = InArgs._OnSetBlockPriority;
+	OnSetReduceBlockSymmetrically = InArgs._OnSetReduceBlockSymmetrically;
 
 	HasDragged = false;
 	Dragging = false;
@@ -425,11 +427,12 @@ FReply SCustomizableObjectLayoutGrid::OnMouseButtonDown( const FGeometry& MyGeom
 
 			MenuBuilder.BeginSection("Fixed Layout Strategy", LOCTEXT("BlockActionsTitle", "Fixed Layout Actions"));
 			{
-				if (SelectedBlocks.Num() && LayoutStrategy == ECustomizableObjectTextureLayoutPackingStrategy::Fixed)
+				if (SelectedBlocks.Num())
 				{
 					MenuBuilder.AddWidget(
 						SNew(SBox)
 						.WidthOverride(125.0f)
+						.ToolTipText(LOCTEXT("SetBlockPriority_Tooltip", "Sets the block priority for a Fixed Layout Strategy."))
 						[
 							SNew(SNumericEntryBox<int32>)
 							.MinValue(0)
@@ -439,10 +442,20 @@ FReply SCustomizableObjectLayoutGrid::OnMouseButtonDown( const FGeometry& MyGeom
 							.Value(this, &SCustomizableObjectLayoutGrid::GetBlockPriortyValue)
 							.UndeterminedString(LOCTEXT("MultipleValues", "Multiples Values"))
 							.OnValueChanged(this, &SCustomizableObjectLayoutGrid::OnBlockPriorityChanged)
-							.ToolTipText(LOCTEXT("SetBlockPriorityTooltip", "Sets the block priority for a Fixed Layout Strategy"))
 							.EditableTextBoxStyle(&FAppStyle::GetWidgetStyle<FEditableTextBoxStyle>("NormalEditableTextBox"))
 						]
 					, FText::FromString("Block Priority"), true);
+
+					MenuBuilder.AddWidget(
+						SNew(SBox)
+						.WidthOverride(125.0f)
+						.ToolTipText(LOCTEXT("SetBlockSymmetry_Tooltip", "if true, this block will be reduced in both axes at the same time in a Fixed Layout Strategy."))
+						[
+							SNew(SCheckBox)
+							.IsChecked(this, &SCustomizableObjectLayoutGrid::GetReductionMethodValue)
+							.OnCheckStateChanged(this, &SCustomizableObjectLayoutGrid::OnReduceBlockSymmetricallyChanged)
+						]
+					, FText::FromString("Reduce Symmetrically"), true);
 				}
 			}
 			MenuBuilder.EndSection();
@@ -1066,6 +1079,37 @@ TOptional<int32> SCustomizableObjectLayoutGrid::GetBlockPriortyValue() const
 }
 
 
+ECheckBoxState SCustomizableObjectLayoutGrid::GetReductionMethodValue() const
+{
+	if (SelectedBlocks.Num())
+	{
+		TArray<FCustomizableObjectLayoutBlock> CurrentSelectedBlocks;
+
+		for (const FCustomizableObjectLayoutBlock& Block : Blocks.Get())
+		{
+			if (SelectedBlocks.Contains(Block.Id))
+			{
+				CurrentSelectedBlocks.Add(Block);
+			}
+		}
+
+		bool bUsesSymmetry = CurrentSelectedBlocks[0].bUseSymmetry;
+
+		for (const FCustomizableObjectLayoutBlock& Block : CurrentSelectedBlocks)
+		{
+			if (Block.bUseSymmetry != bUsesSymmetry)
+			{
+				return ECheckBoxState::Undetermined;
+			}
+		}
+
+		return bUsesSymmetry ? ECheckBoxState::Checked : ECheckBoxState::Unchecked;
+	}
+
+	return ECheckBoxState::Undetermined;
+}
+
+
 void SCustomizableObjectLayoutGrid::OnBlockPriorityChanged(int32 InValue)
 {
 	if (SelectedBlocks.Num())
@@ -1075,9 +1119,12 @@ void SCustomizableObjectLayoutGrid::OnBlockPriorityChanged(int32 InValue)
 }
 
 
-void SCustomizableObjectLayoutGrid::SetLayoutStrategy(ECustomizableObjectTextureLayoutPackingStrategy Strategy)
+void SCustomizableObjectLayoutGrid::OnReduceBlockSymmetricallyChanged(ECheckBoxState InCheckboxState)
 {
-	LayoutStrategy = Strategy;
+	if (SelectedBlocks.Num())
+	{
+		OnSetReduceBlockSymmetrically.ExecuteIfBound(InCheckboxState == ECheckBoxState::Checked);
+	}
 }
 
 
diff --git a/Engine/Plugins/Experimental/Mutable/Source/CustomizableObjectEditor/Private/MuCOE/SCustomizableObjectLayoutGrid.h b/Engine/Plugins/Experimental/Mutable/Source/CustomizableObjectEditor/Private/MuCOE/SCustomizableObjectLayoutGrid.h
--- a/Engine/Plugins/Experimental/Mutable/Source/CustomizableObjectEditor/Private/MuCOE/SCustomizableObjectLayoutGrid.h
+++ b/Engine/Plugins/Experimental/Mutable/Source/CustomizableObjectEditor/Private/MuCOE/SCustomizableObjectLayoutGrid.h
@@ -15,6 +15,7 @@ struct FKeyEvent;
 struct FPointerEvent;
 //#include "PreviewScene.h"
 
+enum class ECheckBoxState : uint8;
 
 struct FRect2D
 {
@@ -45,6 +46,7 @@ public:
 	DECLARE_DELEGATE(FDeleteBlockDelegate);
 	DECLARE_DELEGATE_TwoParams(FAddBlockAtDelegate, FIntPoint, FIntPoint);
 	DECLARE_DELEGATE_OneParam(FSetBlockPriority, int32);
+	DECLARE_DELEGATE_OneParam(FSetReduceBlockSymmetrically, bool);
 
 	SLATE_BEGIN_ARGS( SCustomizableObjectLayoutGrid ){}
 
@@ -59,6 +61,7 @@ public:
 		SLATE_EVENT(FDeleteBlockDelegate, OnDeleteBlocks)
 		SLATE_EVENT(FAddBlockAtDelegate, OnAddBlockAt)
 		SLATE_EVENT(FSetBlockPriority, OnSetBlockPriority)
+		SLATE_EVENT(FSetReduceBlockSymmetrically, OnSetReduceBlockSymmetrically)
 
 	SLATE_END_ARGS()
 
@@ -114,8 +117,11 @@ public:
 	/** Callback when the priority of a block changes */
 	void OnBlockPriorityChanged(int32 InValue);
 
-	/** Sets the packing strategy of the current layout */
-	void SetLayoutStrategy(ECustomizableObjectTextureLayoutPackingStrategy Strategy);
+	/** Callback when a block reduction method changes */
+	void OnReduceBlockSymmetricallyChanged(ECheckBoxState InCheckboxState);
+
+	/** Gets block reduction value of the selected blocks */
+	ECheckBoxState GetReductionMethodValue() const;
 
 private:
 
@@ -129,6 +135,7 @@ private:
 	FDeleteBlockDelegate DeleteBlocksDelegate;
 	FAddBlockAtDelegate AddBlockAtDelegate;
 	FSetBlockPriority OnSetBlockPriority;
+	FSetReduceBlockSymmetrically OnSetReduceBlockSymmetrically;
 
 	/** Size of the grid in blocks */
 	TAttribute<FIntPoint> GridSize;
@@ -203,6 +210,4 @@ private:
 	/** Custom Slate drawing element. Used to improve the UVs drawing performance. */
 	TSharedPtr<class FUVCanvasDrawer, ESPMode::ThreadSafe> UVCanvasDrawer;
 
-	/** Current packing strategy of the represented layout */
-	ECustomizableObjectTextureLayoutPackingStrategy LayoutStrategy;
 };
diff --git a/Engine/Plugins/Experimental/Mutable/Source/CustomizableObjectEditor/Private/MuCOE/SCustomizableObjectNodeLayoutBlocksEditor.cpp b/Engine/Plugins/Experimental/Mutable/Source/CustomizableObjectEditor/Private/MuCOE/SCustomizableObjectNodeLayoutBlocksEditor.cpp
--- a/Engine/Plugins/Experimental/Mutable/Source/CustomizableObjectEditor/Private/MuCOE/SCustomizableObjectNodeLayoutBlocksEditor.cpp
+++ b/Engine/Plugins/Experimental/Mutable/Source/CustomizableObjectEditor/Private/MuCOE/SCustomizableObjectNodeLayoutBlocksEditor.cpp
@@ -119,13 +119,9 @@ void SCustomizableObjectNodeLayoutBlocksEditor::SetCurrentLayout(UCustomizableOb
 			.OnDeleteBlocks(this, &SCustomizableObjectNodeLayoutBlocksEditor::OnRemoveBlock)
 			.OnAddBlockAt(this, &SCustomizableObjectNodeLayoutBlocksEditor::OnAddBlockAt)
 			.OnSetBlockPriority(this, &SCustomizableObjectNodeLayoutBlocksEditor::OnSetBlockPriority)
+			.OnSetReduceBlockSymmetrically(this, &SCustomizableObjectNodeLayoutBlocksEditor::OnSetBlockReductionSymmetry)
 		]
-	];
-
-	if (CurrentLayout && LayoutGridWidget.IsValid())
-	{
-		LayoutGridWidget->SetLayoutStrategy(CurrentLayout->GetPackingStrategy());
-	}	
+	];	
 }
 
 
@@ -228,7 +224,7 @@ TSharedRef<SWidget> SCustomizableObjectNodeLayoutBlocksEditor::BuildLayoutToolBa
 		.Padding(4,0)
 		[
 			SNew(SBorder)
-			.Padding(0)
+			.Padding(2.0f)
 			.BorderImage(FAppStyle::GetBrush("NoBorder"))
 			.IsEnabled( FSlateApplication::Get().GetNormalExecutionAttribute() )
 			[
@@ -248,33 +244,45 @@ TSharedRef<SWidget> SCustomizableObjectNodeLayoutBlocksEditor::BuildLayoutStrate
 		MaxLayoutGridSizes.Add(MakeShareable(new FString(FString::Printf(TEXT("%d x %d"), Size, Size))));
 	}
 
-	MaxLayoutGridSizeCombo = SNew(STextComboBox)
-		.OptionsSource(&MaxLayoutGridSizes)
-		.OnSelectionChanged(this,&SCustomizableObjectNodeLayoutBlocksEditor::OnMaxGridSizeChanged)
-		.IsEnabled(FSlateApplication::Get().GetNormalExecutionAttribute());
-
 	LayoutPackingStrategies.Empty();
 	LayoutPackingStrategies.Add(MakeShareable(new FString("Resizable")));
 	LayoutPackingStrategies.Add(MakeShareable(new FString("Fixed")));
 
+	BlockReductionMethods.Empty();
+	BlockReductionMethods.Add(MakeShareable(new FString("Halve")));
+	BlockReductionMethods.Add(MakeShareable(new FString("Unitary")));
+	
+	MaxLayoutGridSizeCombo = SNew(STextComboBox)
+		.OptionsSource(&MaxLayoutGridSizes)
+		.OnSelectionChanged(this, &SCustomizableObjectNodeLayoutBlocksEditor::OnMaxGridSizeChanged)
+		.IsEnabled(FSlateApplication::Get().GetNormalExecutionAttribute());
+
 	LayoutPackingStrategyCombo = SNew(STextComboBox)
 		.OptionsSource(&LayoutPackingStrategies)
 		.OnSelectionChanged(this, &SCustomizableObjectNodeLayoutBlocksEditor::OnLayoutPackingStrategyChanged)
 		.IsEnabled(FSlateApplication::Get().GetNormalExecutionAttribute());
 
+	BlockReductionMethodsCombo = SNew(STextComboBox)
+		.OptionsSource(&BlockReductionMethods)
+		.OnSelectionChanged(this, &SCustomizableObjectNodeLayoutBlocksEditor::OnReductionMethodChanged)
+		.IsEnabled(FSlateApplication::Get().GetNormalExecutionAttribute());
+
 
 	LayoutStrategyWidget = SNew(SHorizontalBox)
-		+ SHorizontalBox::Slot().Padding(2.0f).AutoWidth()
+		+ SHorizontalBox::Slot()
+		.AutoWidth()
+		.Padding(0.0f,2.0f,2.0f,0.0f)
 		[
 			SNew(STextBlock)
-			.Text(LOCTEXT("Layout Strategy:", "Layout Strategy:"))
+			.Text(LOCTEXT("LayoutStrategy_Text", "Layout Strategy:"))
 			.ToolTipText(LOCTEXT("LayoutStrategyTooltup","Selects the packing strategy: Resizable Layout or Fixed Layout"))
-			.Margin(FMargin(0, 4, 0, 0))
 			.ShadowOffset(FVector2D::UnitVector)
 			.ColorAndOpacity(FLinearColor::Gray)
 		]
 
-		+ SHorizontalBox::Slot().Padding(2.0f).AutoWidth()
+		+ SHorizontalBox::Slot()
+		.AutoWidth()
+		.HAlign(EHorizontalAlignment::HAlign_Right)
 		[
 			SNew(SVerticalBox)
 			+ SVerticalBox::Slot()
@@ -284,20 +292,54 @@ TSharedRef<SWidget> SCustomizableObjectNodeLayoutBlocksEditor::BuildLayoutStrate
 			]
 		];
 
-	FixedLayoutWidget = SNew(SHorizontalBox)
-		+ SHorizontalBox::Slot().Padding(2.0f).AutoWidth()
+	FixedLayoutWidget = SNew(SVerticalBox)
+		+SVerticalBox::Slot()
+		.AutoHeight()
 		[
-			SNew(STextBlock)
-			.Text(LOCTEXT("Max Layout Size:", "Max Layout Size:"))
-			.Margin(FMargin(0, 4, 0, 0))
-			.ShadowOffset(FVector2D::UnitVector)
-			.ColorAndOpacity(FLinearColor::Gray)
+			SNew(SHorizontalBox)
+			+SHorizontalBox::Slot()
+			.AutoWidth()
+			.Padding(0.0f, 2.0f, 2.0f, 0.0f)
+			[
+				SNew(STextBlock)
+				.Text(LOCTEXT("MaxLayoutSize_Text", "Max Layout Size:"))
+				.ShadowOffset(FVector2D::UnitVector)
+				.ColorAndOpacity(FLinearColor::Gray)
+			]
+			+ SHorizontalBox::Slot()
+			.AutoWidth()
+			.HAlign(EHorizontalAlignment::HAlign_Right)
+			[
+				MaxLayoutGridSizeCombo.ToSharedRef()
+			]
 		]
-		+ SHorizontalBox::Slot().Padding(2.0f).AutoWidth()
+		+ SVerticalBox::Slot()
+		.AutoHeight()
 		[
-			MaxLayoutGridSizeCombo.ToSharedRef()
+			SNew(SHorizontalBox)
+			+ SHorizontalBox::Slot()
+			.AutoWidth()
+			.Padding(0.0f, 2.0f, 2.0f, 0.0f)
+			[
+				SNew(STextBlock)
+				.Text(LOCTEXT("ReductionMethod_Text", "Reduction Method:"))
+				.ToolTipText(LOCTEXT("Reduction_Method_Tooltip", "Select how blocks will be reduced in case that they do not fit in the layout:"
+				"\n Halve: blocks will be reduced by half each time."
+				"\n Unit: blocks will be reduced by one unit each time."))
+				.ShadowOffset(FVector2D::UnitVector)
+				.ColorAndOpacity(FLinearColor::Gray)
+			]
+			
+			+ SHorizontalBox::Slot()
+			.AutoWidth()
+			.Padding(0.0f, 0.0f, 0.0f, 8.0f)
+			.HAlign(EHorizontalAlignment::HAlign_Right)
+			[
+				BlockReductionMethodsCombo.ToSharedRef()
+			]
 		];
 
+
 	//Modify after the creation of all the widgets
 	if (CurrentLayout)
 	{
@@ -318,6 +360,11 @@ TSharedRef<SWidget> SCustomizableObjectNodeLayoutBlocksEditor::BuildLayoutStrate
 		{
 			LayoutPackingStrategyCombo->SetSelectedItem(LayoutPackingStrategies[(uint32)CurrentLayout->GetPackingStrategy()]);
 		}
+
+		if (BlockReductionMethods.IsValidIndex((uint32)CurrentLayout->GetBlockReductionMethod()))
+		{
+			BlockReductionMethodsCombo->SetSelectedItem(BlockReductionMethods[(uint32)CurrentLayout->GetBlockReductionMethod()]);
+		}
 	}
 
 	return SNew(SVerticalBox)
@@ -368,6 +415,7 @@ void SCustomizableObjectNodeLayoutBlocksEditor::OnAddBlock()
 		block.Max = FIntPoint(1, 1);
 		block.Id = FGuid::NewGuid();
 		block.Priority = 0;
+		block.bUseSymmetry = false;
 		CurrentLayout->Blocks.Add( block );
 		CurrentLayout->MarkPackageDirty();
 
@@ -388,6 +436,7 @@ void SCustomizableObjectNodeLayoutBlocksEditor::OnAddBlockAt(const FIntPoint Min
 		block.Max = Max;
 		block.Id = FGuid::NewGuid();
 		block.Priority = 0;
+		block.bUseSymmetry = false;
 
 		CurrentLayout->Blocks.Add(block);
 		CurrentLayout->MarkPackageDirty();
@@ -395,7 +444,6 @@ void SCustomizableObjectNodeLayoutBlocksEditor::OnAddBlockAt(const FIntPoint Min
 }
 
 
-
 void SCustomizableObjectNodeLayoutBlocksEditor::OnRemoveBlock()
 {
 	if (CurrentLayout)
@@ -479,7 +527,28 @@ void SCustomizableObjectNodeLayoutBlocksEditor::OnSetBlockPriority(int32 InValue
 				}
 			}
 		}
-	}	
+	}
+}
+
+
+void SCustomizableObjectNodeLayoutBlocksEditor::OnSetBlockReductionSymmetry(bool bInValue)
+{
+	if (CurrentLayout)
+	{
+		if (LayoutGridWidget.IsValid())
+		{
+			TArray<FGuid> SelectedBlocks = LayoutGridWidget->GetSelectedBlocks();
+
+			for (int i = 0; i < CurrentLayout->Blocks.Num(); ++i)
+			{
+				if (SelectedBlocks.Contains(CurrentLayout->Blocks[i].Id))
+				{
+					CurrentLayout->Blocks[i].bUseSymmetry = bInValue;
+					CurrentLayout->MarkPackageDirty();
+				}
+			}
+		}
+	}
 }
 
 
@@ -501,12 +570,6 @@ void SCustomizableObjectNodeLayoutBlocksEditor::OnLayoutPackingStrategyChanged(T
 		if (CurrentLayout->GetPackingStrategy() != (ECustomizableObjectTextureLayoutPackingStrategy)selection)
 		{
 			CurrentLayout->SetPackingStrategy((ECustomizableObjectTextureLayoutPackingStrategy)selection);
-
-			if (LayoutGridWidget.IsValid())
-			{
-				LayoutGridWidget->SetLayoutStrategy(CurrentLayout->GetPackingStrategy());
-			}
-
 			CurrentLayout->MarkPackageDirty();
 		}
 	}
@@ -528,6 +591,21 @@ void SCustomizableObjectNodeLayoutBlocksEditor::OnMaxGridSizeChanged(TSharedPtr<
 }
 
 
+void SCustomizableObjectNodeLayoutBlocksEditor::OnReductionMethodChanged(TSharedPtr<FString> NewSelection, ESelectInfo::Type SelectInfo)
+{
+	if (CurrentLayout)
+	{
+		uint32 selection = BlockReductionMethods.IndexOfByKey(NewSelection);
+
+		if (CurrentLayout->GetBlockReductionMethod() != (ECustomizableObjectLayoutBlockReductionMethod)selection)
+		{
+			CurrentLayout->SetBlockReductionMethod((ECustomizableObjectLayoutBlockReductionMethod)selection);
+			CurrentLayout->MarkPackageDirty();
+		}
+	}
+}
+
+
 TArray<FCustomizableObjectLayoutBlock> SCustomizableObjectNodeLayoutBlocksEditor::GetBlocks() const
 {
 	TArray<FCustomizableObjectLayoutBlock> Blocks;
diff --git a/Engine/Plugins/Experimental/Mutable/Source/CustomizableObjectEditor/Private/MuCOE/SCustomizableObjectNodeLayoutBlocksEditor.h b/Engine/Plugins/Experimental/Mutable/Source/CustomizableObjectEditor/Private/MuCOE/SCustomizableObjectNodeLayoutBlocksEditor.h
--- a/Engine/Plugins/Experimental/Mutable/Source/CustomizableObjectEditor/Private/MuCOE/SCustomizableObjectNodeLayoutBlocksEditor.h
+++ b/Engine/Plugins/Experimental/Mutable/Source/CustomizableObjectEditor/Private/MuCOE/SCustomizableObjectNodeLayoutBlocksEditor.h
@@ -73,14 +73,20 @@ private:
 	/** List of available layout packing strategies. */
 	TArray< TSharedPtr< FString > > LayoutPackingStrategies;
 
-	/** Widget for displaying the available layout packing strategies. */
+	/** Widget to select the available layout packing strategies. */
 	TSharedPtr< STextComboBox > LayoutPackingStrategyCombo;
 
+	/** List of available block reduction methods. */
+	TArray< TSharedPtr< FString > > BlockReductionMethods;
+
+	/** Widget to select the available layout block reduction methods. */
+	TSharedPtr< STextComboBox > BlockReductionMethodsCombo;
+
 	/** Widget to select the layout packing strategy */
 	TSharedPtr< SHorizontalBox> LayoutStrategyWidget;
 
 	/** Widget to select the fixed layout properties */
-	TSharedPtr< SHorizontalBox> FixedLayoutWidget;
+	TSharedPtr< SVerticalBox> FixedLayoutWidget;
 
 	TSharedPtr<SWidget> StrategyWidget;
 
@@ -97,6 +103,7 @@ private:
 	/** Callbacks from the layout block editor. */
 	TSharedRef<SWidget> BuildLayoutToolBar();
 	TSharedRef<SWidget> BuildLayoutStrategyWidgets(const ISlateStyle* Style, const FName& StyleName);
+
 	void OnAddBlock();
 	void OnAddBlockAt(const FIntPoint Min, const FIntPoint Max);
 	void OnRemoveBlock();
@@ -107,10 +114,14 @@ private:
 	/** . */
 	void OnGridSizeChanged( TSharedPtr<FString> NewSelection, ESelectInfo::Type SelectInfo );
 	void OnMaxGridSizeChanged(TSharedPtr<FString> NewSelection, ESelectInfo::Type SelectInfo);
+	void OnReductionMethodChanged(TSharedPtr<FString> NewSelection, ESelectInfo::Type SelectInfo);
 
 	/** Sets the block priority from the input text. */
 	void OnSetBlockPriority(int32 InValue);
 
+	/** Sets the block reduction symmetry method. */
+	void OnSetBlockReductionSymmetry(bool bInValue);
+
 	/** Called when the packing strategy has changed. */
 	void OnLayoutPackingStrategyChanged(TSharedPtr<FString> NewSelection, ESelectInfo::Type SelectInfo);
 
diff --git a/Engine/Plugins/Experimental/Mutable/Source/MutableRuntime/Private/MuR/CodeRunner.cpp b/Engine/Plugins/Experimental/Mutable/Source/MutableRuntime/Private/MuR/CodeRunner.cpp
--- a/Engine/Plugins/Experimental/Mutable/Source/MutableRuntime/Private/MuR/CodeRunner.cpp
+++ b/Engine/Plugins/Experimental/Mutable/Source/MutableRuntime/Private/MuR/CodeRunner.cpp
@@ -5165,6 +5165,7 @@ namespace mu
 					scratch.positions.SetNum(BlockCount);
 					scratch.priorities.SetNum(BlockCount);
 					scratch.reductions.SetNum(BlockCount);
+					scratch.useSymmetry.SetNum(BlockCount);
 
 					LayoutPack3(pResult.get(), pSource.get(), &scratch);
 				}
diff --git a/Engine/Plugins/Experimental/Mutable/Source/MutableRuntime/Private/MuR/Layout.cpp b/Engine/Plugins/Experimental/Mutable/Source/MutableRuntime/Private/MuR/Layout.cpp
--- a/Engine/Plugins/Experimental/Mutable/Source/MutableRuntime/Private/MuR/Layout.cpp
+++ b/Engine/Plugins/Experimental/Mutable/Source/MutableRuntime/Private/MuR/Layout.cpp
@@ -43,6 +43,7 @@ namespace mu {
 		pResult->m_blocks = m_blocks;
 		pResult->m_strategy = m_strategy;
 		pResult->FirstLODToIgnoreWarnings = FirstLODToIgnoreWarnings;
+		pResult->ReductionMethod = ReductionMethod;
 		return pResult;
 	}
 
@@ -55,7 +56,8 @@ namespace mu {
 			(m_blocks == o.m_blocks) &&
 			(m_strategy == o.m_strategy) &&
 			// maybe this is not needed
-			(FirstLODToIgnoreWarnings == o.FirstLODToIgnoreWarnings);;
+			(FirstLODToIgnoreWarnings == o.FirstLODToIgnoreWarnings) &&
+			(ReductionMethod==o.ReductionMethod);
 	}
 
 
@@ -130,15 +132,21 @@ namespace mu {
 
 
 	//---------------------------------------------------------------------------------------------
-	void Layout::GetBlockPriority(int index, int* pPriority) const
+	void Layout::GetBlockOptions(int index, int* pPriority, bool* bUseSymmetry) const
 	{
 		check(index >= 0 && index < m_blocks.Num());
 		check(pPriority);
+		check(bUseSymmetry);
 
 		if (pPriority)
 		{
 			*pPriority = m_blocks[index].m_priority;
 		}
+
+		if (bUseSymmetry)
+		{
+			*bUseSymmetry = m_blocks[index].bUseSymmetry;
+		}
 	}
 
 
@@ -154,11 +162,12 @@ namespace mu {
 
 
 	//---------------------------------------------------------------------------------------------
-	void Layout::SetBlockPriority(int index, int priority)
+	void Layout::SetBlockOptions(int index, int priority, bool bUseSymmetry)
 	{
 		check(index >= 0 && index < m_blocks.Num());
 
 		m_blocks[index].m_priority = priority;
+		m_blocks[index].bUseSymmetry = bUseSymmetry;
 	}
 
 
@@ -235,5 +244,19 @@ namespace mu {
 	{
 		return FirstLODToIgnoreWarnings;
 	}
+
+
+	//---------------------------------------------------------------------------------------------
+	void Layout::SetBlockReductionMethod(EReductionMethod Method)
+	{
+		ReductionMethod = Method;
+	}
+
+
+	//---------------------------------------------------------------------------------------------
+	EReductionMethod Layout::GetBlockReductionMethod() const
+	{
+		return ReductionMethod;
+	}
 }
 
diff --git a/Engine/Plugins/Experimental/Mutable/Source/MutableRuntime/Private/MuR/OpLayoutPack.h b/Engine/Plugins/Experimental/Mutable/Source/MutableRuntime/Private/MuR/OpLayoutPack.h
--- a/Engine/Plugins/Experimental/Mutable/Source/MutableRuntime/Private/MuR/OpLayoutPack.h
+++ b/Engine/Plugins/Experimental/Mutable/Source/MutableRuntime/Private/MuR/OpLayoutPack.h
@@ -143,16 +143,18 @@ namespace mu
 			index = -1;
 		}
 
-		LAY_BLOCK( int i, vec2<int> s, int p=0 )
+		LAY_BLOCK( int i, vec2<int> s, int p=0, bool sym = false )
 		{
 			index = i;
 			size = s;
 			priority = p;
+			useSymmetry = sym;
 		}
 
 		int index;
 		vec2<int> size;
 		int priority;
+		bool useSymmetry;
 	};
 
     inline bool CompareBlocks( const LAY_BLOCK& a, const LAY_BLOCK& b )
@@ -187,6 +189,11 @@ namespace mu
 
 	inline bool CompareBlocksPriority(const LAY_BLOCK& a, const LAY_BLOCK& b)
 	{
+		if (a.priority == b.priority)
+		{
+			return CompareBlocks(a, b);
+		}
+
 		return a.priority > b.priority;
 	}
 
@@ -201,6 +208,7 @@ namespace mu
 		TArray< vec2<int> > positions;
 		TArray< int > priorities;
 		TArray< vec2<int> > reductions;
+		TArray< int > useSymmetry;
     };
 
 
@@ -356,11 +364,21 @@ namespace mu
     // }
 
 
+	inline int32 ReductionOperation(int& BlockSize, EReductionMethod ReductionMethod)
+	{
+		if (ReductionMethod == EReductionMethod::UNITARY_REDUCTION)
+		{
+			return BlockSize -= 1;
+		}
+		
+		return BlockSize /= 2;
+	}
+
 
     //---------------------------------------------------------------------------------------------
     //! Even more expensive but precise version
     //---------------------------------------------------------------------------------------------
-	inline void ReduceBlock(int blockCount, int* area, int* reverse_it, SCRATCH_LAYOUT_PACK* scratch)
+	inline void ReduceBlock(int blockCount, int* area, int* reverse_it, SCRATCH_LAYOUT_PACK* scratch, EReductionMethod ReductionMethod)
 	{
 		int r_it = (*reverse_it);
 		bool pass = false;
@@ -369,36 +387,39 @@ namespace mu
 
 		if (scratch->sorted[r_it].size[0] != 1 || scratch->sorted[r_it].size[1] != 1)
 		{
-			if (scratch->reductions[scratch->sorted[r_it].index][0] > scratch->reductions[scratch->sorted[r_it].index][1])
+			if (scratch->sorted[r_it].useSymmetry)
 			{
-				if (scratch->sorted[r_it].size[1] > 1)
+				// We reduce both sides of the block at the same time
+				for (int32 Index = 0; Index <= 1; ++Index)
 				{
-					scratch->sorted[r_it].size[1] /= 2;
-					scratch->blocks[scratch->sorted[r_it].index][1] /= 2;
-					pass = true;
-				}
+					if (scratch->sorted[r_it].size[Index] > 1)
+					{
+						ReductionOperation(scratch->sorted[r_it].size[Index], ReductionMethod);
+						ReductionOperation(scratch->blocks[scratch->sorted[r_it].index][Index], ReductionMethod);
 
-				scratch->reductions[scratch->sorted[r_it].index][1] += 1;
-			}
-			else if (scratch->reductions[scratch->sorted[r_it].index][0] < scratch->reductions[scratch->sorted[r_it].index][1])
-			{
-				if (scratch->sorted[r_it].size[0] > 1)
-				{
-					scratch->sorted[r_it].size[0] /= 2;
-					scratch->blocks[scratch->sorted[r_it].index][0] /= 2;
-					pass = true;
+						pass = true;
+					}
 				}
-
-				scratch->reductions[scratch->sorted[r_it].index][0] += 1;
 			}
 			else
 			{
-				if (r_it % 2 == 0)
+				if (scratch->reductions[scratch->sorted[r_it].index][0] > scratch->reductions[scratch->sorted[r_it].index][1])
+				{
+					if (scratch->sorted[r_it].size[1] > 1)
+					{
+						ReductionOperation(scratch->sorted[r_it].size[1], ReductionMethod);
+						ReductionOperation(scratch->blocks[scratch->sorted[r_it].index][1], ReductionMethod);
+						pass = true;
+					}
+
+					scratch->reductions[scratch->sorted[r_it].index][1] += 1;
+				}
+				else if (scratch->reductions[scratch->sorted[r_it].index][0] < scratch->reductions[scratch->sorted[r_it].index][1])
 				{
 					if (scratch->sorted[r_it].size[0] > 1)
 					{
-						scratch->sorted[r_it].size[0] /= 2;
-						scratch->blocks[scratch->sorted[r_it].index][0] /= 2;
+						ReductionOperation(scratch->sorted[r_it].size[0], ReductionMethod);
+						ReductionOperation(scratch->blocks[scratch->sorted[r_it].index][0], ReductionMethod);
 						pass = true;
 					}
 
@@ -406,14 +427,23 @@ namespace mu
 				}
 				else
 				{
-					if (scratch->sorted[r_it].size[1] > 1)
+					// we select the first side to reduce "randomly"
+					int32 Index = r_it % 2;
+
+					// if we can't reduce a dimension then we try to reduce the other one
+					if (scratch->sorted[r_it].size[Index] <= 1)
+					{
+						Index = r_it == 0 ? 1 : 0;
+					}
+
+					if (scratch->sorted[r_it].size[Index] > 1)
 					{
-						scratch->sorted[r_it].size[1] /= 2;
-						scratch->blocks[scratch->sorted[r_it].index][1] /= 2;
+						ReductionOperation(scratch->sorted[r_it].size[Index], ReductionMethod);
+						ReductionOperation(scratch->blocks[scratch->sorted[r_it].index][Index], ReductionMethod);
 						pass = true;
 					}
 
-					scratch->reductions[scratch->sorted[r_it].index][1] += 1;
+					scratch->reductions[scratch->sorted[r_it].index][Index] += 1;
 				}
 			}
 		}
@@ -437,6 +467,7 @@ namespace mu
 		}
 	}
 
+
 	inline bool SetPositions(int bestY,int layoutSizeY, int* maxX, int* maxY, SCRATCH_LAYOUT_PACK* scratch, EPackStrategy packStrategy)
 	{
 		bool fits = true;
@@ -635,7 +666,8 @@ namespace mu
             pSourceLayout->GetBlock( index, &b.min[0], &b.min[1], &b.size[0], &b.size[1] );
 
 			int p;
-			pSourceLayout->GetBlockPriority(index, &p);
+			bool Symmetry;
+			pSourceLayout->GetBlockOptions(index, &p, &Symmetry);
 
 			vec2<int> reductions;
 			reductions[0] = 0;
@@ -654,6 +686,7 @@ namespace mu
             scratch->blocks[index] = b.size;
 			scratch->priorities[index] = p;
 			scratch->reductions[index] = reductions;
+			scratch->useSymmetry[index] = (int)Symmetry;
         }
 
 
@@ -770,7 +803,7 @@ namespace mu
         check( (int)scratch->sorted.Num()==blockCount );
         for ( int index=0; index<blockCount; ++index )
         {
-            scratch->sorted[index] = LAY_BLOCK( index, scratch->blocks[index], scratch->priorities[index] );
+            scratch->sorted[index] = LAY_BLOCK( index, scratch->blocks[index], scratch->priorities[index], (bool)scratch->useSymmetry[index]);
         }
 
 		scratch->sorted.Sort( CompareBlocks);
@@ -785,7 +818,7 @@ namespace mu
 			//Shrink blocks in case we do not have enough space to pack everything
 			while (maxX*maxY < area)
 			{
-				ReduceBlock(blockCount, &area, &r_it, scratch);
+				ReduceBlock(blockCount, &area, &r_it, scratch, pSourceLayout->GetBlockReductionMethod());
 			}
 			
 		}
@@ -811,7 +844,7 @@ namespace mu
 					scratch->sorted.Sort(CompareBlocksPriority);
 				}
 
-				ReduceBlock(blockCount, &area, &r_it, scratch);
+				ReduceBlock(blockCount, &area, &r_it, scratch, pSourceLayout->GetBlockReductionMethod());
 			}
 		}
 
@@ -819,6 +852,7 @@ namespace mu
         pResult->SetGridSize( maxX, maxY );
         pResult->SetMaxGridSize(layoutSizeX, layoutSizeY);
 		pResult->SetLayoutPackingStrategy(pSourceLayout->GetLayoutPackingStrategy());
+		pResult->SetBlockReductionMethod(pSourceLayout->GetBlockReductionMethod());
 
         for ( int index=0; index<blockCount; ++index )
         {
@@ -829,8 +863,7 @@ namespace mu
                     scratch->blocks[index][0], scratch->blocks[index][1]
                 );
 
-			pResult->SetBlockPriority(index, scratch->priorities[index]);
+			pResult->SetBlockOptions(index, scratch->priorities[index], (bool)scratch->useSymmetry[index]);
         }
     }
-
 }
diff --git a/Engine/Plugins/Experimental/Mutable/Source/MutableRuntime/Public/MuR/Layout.h b/Engine/Plugins/Experimental/Mutable/Source/MutableRuntime/Public/MuR/Layout.h
--- a/Engine/Plugins/Experimental/Mutable/Source/MutableRuntime/Public/MuR/Layout.h
+++ b/Engine/Plugins/Experimental/Mutable/Source/MutableRuntime/Public/MuR/Layout.h
@@ -27,6 +27,13 @@ namespace mu
 		FIXED_LAYOUT
 	};
 
+	//! Types of layout reduction methods 
+	enum class EReductionMethod : uint32
+	{
+		HALVE_REDUCTION,	// Divide axis by 2
+		UNITARY_REDUCTION	// Reduces 1 block the axis 
+	};
+
 
     //! \brief Image block layout class.
     //!
@@ -98,8 +105,9 @@ namespace mu
 
 		//! Returns the reduction priority of a block.
 		//! \param index Block to get priority.
-		//! \param[out] priority will be set to the reduction priority of the block
-		void GetBlockPriority(int index, int* priority) const;
+		//! \param[out] priority reduction priority of the block
+		//! \param[out] bUSeSymmetry reduction method of the block
+		void GetBlockOptions(int index, int* priority, bool* bUseSymmetry) const;
 
 		//! Set a block of the layout.
 		//! "Position" here means the lower-left corner of the block.
@@ -112,11 +120,11 @@ namespace mu
 		//! \param sizeY will be set to the y size of the block
         void SetBlock( int index, int minX, int minY, int sizeX, int sizeY );
 
-		//! Set the reduction priority of a block
-		//! The blocks with the highest values will be the last to be reduced
-		//! \param index Block to set priority
-		//! \param priority will be set to the reduction priority of the block
-		void SetBlockPriority(int index, int priority);
+		//! Set the reduction options of a block
+		//! \param index Block to set the options
+		//! \param priority will be set to the reduction priority of the block. The blocks with the highest values will be the last to be reduced
+		//! \param bUseSymmetry will be set to reduce the block in both axis at the same time.
+		void SetBlockOptions(int index, int priority, bool bUseSymmetry);
 
 		//! Set the texture layout packing strategy
 		//! By default the texture layout packing strategy is set to resizable layout
@@ -131,6 +139,12 @@ namespace mu
 		//! Get the LOD where the unassigned vertices warnings starts to be ignored
 		int32 GetIgnoreLODWarnings();
 
+		//! Set the block reduction method a the Fixed_Layout strategy
+		void SetBlockReductionMethod(EReductionMethod Method);
+
+		//! Returns the block reduction method
+		EReductionMethod GetBlockReductionMethod() const;
+
 	protected:
 
 		//! Forbidden. Manage with the Ptr<> template.
@@ -148,6 +162,7 @@ namespace mu
 				m_size = size;
 				m_id = -1;
 				m_priority = 0;
+				bUseSymmetry = false;
 			}
 
 			vec2<uint16> m_min;
@@ -159,6 +174,9 @@ namespace mu
 			//! Priority value to control the shrink texture layout strategy
 			int32 m_priority;
 
+			//! Value to control the method to reduce the block
+			bool bUseSymmetry;
+
 
 			//!
 			void Serialise(OutputArchive& arch) const
@@ -167,6 +185,7 @@ namespace mu
 				arch << m_size;
 				arch << m_id;
 				arch << m_priority;
+				arch << bUseSymmetry;
 			}
 
 			//!
@@ -176,6 +195,7 @@ namespace mu
 				arch >> m_size;
 				arch >> m_id;
 				arch >> m_priority;
+				arch >> bUseSymmetry;
 			}
 
 			//!
@@ -184,7 +204,8 @@ namespace mu
 				return (m_min == o.m_min) &&
 					(m_size == o.m_size) &&
 					(m_id == o.m_id) &&
-					(m_priority == o.m_priority);
+					(m_priority == o.m_priority) &&
+					(bUseSymmetry == o.bUseSymmetry);
 			}
 
 			inline bool IsSimilar(const FBlock& o) const
@@ -192,7 +213,18 @@ namespace mu
 				// All but ids
 				return (m_min == o.m_min) &&
 					(m_size == o.m_size) &&
-					(m_priority == o.m_priority);
+					(m_priority == o.m_priority) &&
+					(bUseSymmetry == o.bUseSymmetry);
+			}
+
+			//!
+			void UnserialiseOldVersion(InputArchive& Archive, const int32 Version)
+			{
+				// Use the version if the Layout version changes to 6
+				Archive >> m_min;
+				Archive >> m_size;
+				Archive >> m_id;
+				Archive >> m_priority;
 			}
 		};
 
@@ -210,11 +242,13 @@ namespace mu
 
 		int32 FirstLODToIgnoreWarnings;
 
+		EReductionMethod ReductionMethod = EReductionMethod::HALVE_REDUCTION;
+
 
 		//!
 		void Serialise(OutputArchive& arch) const
 		{
-			uint32 ver = 4;
+			uint32 ver = 5;
 			arch << ver;
 
 			arch << m_size;
@@ -223,6 +257,7 @@ namespace mu
 			arch << m_maxsize;
 			arch << uint32(m_strategy);
 			arch << FirstLODToIgnoreWarnings;
+			arch << uint32(ReductionMethod);
 		}
 
 		//!
@@ -230,11 +265,26 @@ namespace mu
 		{
 			uint32 ver;
 			arch >> ver;
-			check(ver <= 4);
+			check(ver <= 5);
 
 			arch >> m_size;
 
-			arch >> m_blocks;
+			if (ver < 5)
+			{
+				uint32_t Size = 0;
+				arch >> Size;
+				m_blocks.SetNum(Size);
+
+				for (uint32_t BlockIndex = 0; BlockIndex < Size; ++BlockIndex)
+				{
+					m_blocks[BlockIndex].UnserialiseOldVersion(arch, ver);
+				}
+			}
+			else
+			{
+				arch >> m_blocks;
+			}
+
 			arch >> m_maxsize;
 
 			uint32 temp;
@@ -245,6 +295,12 @@ namespace mu
 			{
 				arch >> FirstLODToIgnoreWarnings;
 			}
+
+			if (ver >= 5)
+			{
+				arch >> temp;
+				ReductionMethod = EReductionMethod(temp);
+			}
 		}
 
 		//!
diff --git a/Engine/Plugins/Experimental/Mutable/Source/MutableTools/Private/MuT/NodeLayout.cpp b/Engine/Plugins/Experimental/Mutable/Source/MutableTools/Private/MuT/NodeLayout.cpp
--- a/Engine/Plugins/Experimental/Mutable/Source/MutableTools/Private/MuT/NodeLayout.cpp
+++ b/Engine/Plugins/Experimental/Mutable/Source/MutableTools/Private/MuT/NodeLayout.cpp
@@ -186,9 +186,9 @@ namespace mu
 
 
 	//---------------------------------------------------------------------------------------------
-	void NodeLayoutBlocks::SetBlockPriority(int index, int priority)
+	void NodeLayoutBlocks::SetBlockOptions(int index, int priority, bool bUseSymmetry)
 	{
-		m_pD->m_pLayout->SetBlockPriority(index, priority);
+		m_pD->m_pLayout->SetBlockOptions(index, priority, bUseSymmetry);
 	}
 
 
@@ -198,6 +198,7 @@ namespace mu
 		m_pD->m_pLayout->SetLayoutPackingStrategy(strategy);
 	}
 
+
 	//---------------------------------------------------------------------------------------------
 	NodeLayoutBlocksPtr NodeLayoutBlocks::GenerateLayoutBlocks(const MeshPtr pMesh, int layoutIndex, int gridSizeX, int gridSizeY)
 	{
@@ -415,6 +416,20 @@ namespace mu
 	{
 		return m_pD->m_pLayout->GetIgnoreLODWarnings();
 	}
+
+
+	//---------------------------------------------------------------------------------------------
+	void NodeLayoutBlocks::SetBlockReductionMethod(EReductionMethod Method)
+	{
+		m_pD->m_pLayout->SetBlockReductionMethod(Method);
+	}
+
+
+	//---------------------------------------------------------------------------------------------
+	EReductionMethod NodeLayoutBlocks::GetBlockReductionMethod()
+	{
+		return m_pD->m_pLayout->GetBlockReductionMethod();
+	}
 }
 
 
diff --git a/Engine/Plugins/Experimental/Mutable/Source/MutableTools/Public/MuT/NodeLayout.h b/Engine/Plugins/Experimental/Mutable/Source/MutableTools/Public/MuT/NodeLayout.h
--- a/Engine/Plugins/Experimental/Mutable/Source/MutableTools/Public/MuT/NodeLayout.h
+++ b/Engine/Plugins/Experimental/Mutable/Source/MutableTools/Public/MuT/NodeLayout.h
@@ -140,8 +140,8 @@ namespace mu
 		//! minx and miny refer to the lowest left corner of the block.
 		void GetBlock(int index, int* minx, int* miny, int* sizex, int* sizey);
 
-		//! Set priority of a block.
-		void SetBlockPriority(int index, int priority);
+		//! Set reduction block options like priority or if the block has to be reduced symmetrically.
+		void SetBlockOptions(int index, int priority, bool bUseSymmetry);
 
 		//! Set the texture layout packing strategy 
 		void SetLayoutPackingStrategy(EPackStrategy strategy);
@@ -155,6 +155,12 @@ namespace mu
 		//! Get the LOD where the unassigned vertices warnings starts to be ignored
 		int32 GetIgnoreWarningsLOD();
 
+		//! Set the block reduction method a the Fixed_Layout strategy
+		void SetBlockReductionMethod(EReductionMethod strategy);
+
+		//! Returns the block reduction method
+		EReductionMethod GetBlockReductionMethod();
+
 
 		//-----------------------------------------------------------------------------------------
 		// Interface pattern
