[mutable] Retargeted IsParameterMultidimensional calls to point at method in CO
- Made call to IsParameterMultidimensional(ParameterName) part of the CO blueprint callable.

diff --git a/Engine/Plugins/Experimental/Mutable/Source/CustomizableObject/Classes/MuCO/CustomizableObject.h b/Engine/Plugins/Experimental/Mutable/Source/CustomizableObject/Classes/MuCO/CustomizableObject.h
--- a/Engine/Plugins/Experimental/Mutable/Source/CustomizableObject/Classes/MuCO/CustomizableObject.h
+++ b/Engine/Plugins/Experimental/Mutable/Source/CustomizableObject/Classes/MuCO/CustomizableObject.h
@@ -1261,15 +1261,16 @@ public:
 
 	/** Return true or false depending if the parameter at the index provided is multidimensional or not.
 	 * @param InParamIndex The index of the parameter to check.
-	 * @return True if the parameter is multilayer and false if it is not.
+	 * @return True if the parameter is multidimensional and false if it is not.
 	 */
 	bool IsParameterMultidimensional(const int32& InParamIndex) const;
 	
 	/** Return true or false depending if the parameter at the index provided is multidimensional or not.
 	 * @param InParameterName The name of the parameter to check.
-	 * @return True if the parameter is multilayer and false if it is not.
+	 * @return True if the parameter is multidimensional and false if it is not.
 	 */
-	bool IsParameterMultidimensional (const FString& InParameterName) const;
+	UFUNCTION(BlueprintCallable, Category = CustomizableObject)
+	bool IsParameterMultidimensional ( UPARAM(DisplayName = "Parameter Name") const FString& InParameterName) const;
 	
 private: 
 
