[mutable] Fixed crash in table nodes when there was an empty texture cell in the data table.

#rnx

diff --git a/Engine/Plugins/Experimental/Mutable/Source/MutableTools/Private/MuT/CodeGenerator_Image.cpp b/Engine/Plugins/Experimental/Mutable/Source/MutableTools/Private/MuT/CodeGenerator_Image.cpp
--- a/Engine/Plugins/Experimental/Mutable/Source/MutableTools/Private/MuT/CodeGenerator_Image.cpp
+++ b/Engine/Plugins/Experimental/Mutable/Source/MutableTools/Private/MuT/CodeGenerator_Image.cpp
@@ -1841,8 +1841,13 @@ namespace mu
 			[this](const NodeImageTable::Private& node, int colIndex, int row, ErrorLog* pErrorLog)
 			{
 				NodeImageConstantPtr pCell = new NodeImageConstant();
+				ImagePtrConst pImage = nullptr;
+				
+				if (Ptr<ResourceProxy<Image>> pProxyImage = node.m_pTable->GetPrivate()->m_rows[row].m_values[colIndex].m_pProxyImage)
+				{
+					pImage = pProxyImage->Get();
+				}
 
-				ImagePtrConst pImage = node.m_pTable->GetPrivate()->m_rows[row].m_values[colIndex].m_pProxyImage->Get();
 				if (!pImage)
 				{
 					char temp[256];
