[Mutable] Fix inconsistent COI state when skipping updates, which resulted in meshes remaining in TPose.

#rnx

diff --git a/Engine/Plugins/Experimental/Mutable/Source/CustomizableObject/Private/MuCO/CustomizableObjectSystem.cpp b/Engine/Plugins/Experimental/Mutable/Source/CustomizableObject/Private/MuCO/CustomizableObjectSystem.cpp
--- a/Engine/Plugins/Experimental/Mutable/Source/CustomizableObject/Private/MuCO/CustomizableObjectSystem.cpp
+++ b/Engine/Plugins/Experimental/Mutable/Source/CustomizableObject/Private/MuCO/CustomizableObjectSystem.cpp
@@ -2172,6 +2172,7 @@ namespace impl
 		// Skip update, the requested update is equal to the running update.
 		if (Operation->InstanceDescriptorRuntimeHash.IsSubset(CandidateInstance->GetDescriptorRuntimeHash()))
 		{
+			CandidateInstance->SkeletalMeshStatus = ESkeletalMeshState::Correct;
 			UpdateSkeletalMesh(CandidateInstance, CandidateInstance->GetDescriptorRuntimeHash());
 			bCancel = true;
 		}
