[Mutable] Fix recurrent warning due to a null world.

#rnx

diff --git a/Engine/Plugins/Experimental/Mutable/Source/CustomizableObject/Private/MuCO/CustomizableObjectSystem.cpp b/Engine/Plugins/Experimental/Mutable/Source/CustomizableObject/Private/MuCO/CustomizableObjectSystem.cpp
--- a/Engine/Plugins/Experimental/Mutable/Source/CustomizableObject/Private/MuCO/CustomizableObjectSystem.cpp
+++ b/Engine/Plugins/Experimental/Mutable/Source/CustomizableObject/Private/MuCO/CustomizableObjectSystem.cpp
@@ -586,7 +586,12 @@ void FCustomizableObjectSystemPrivate::UpdateStats()
 				{
 					AActor* ParentActor = CustomizableSkeletalComponent->GetAttachmentRootActor();
 					UCustomizableObjectInstance* Instance = CustomizableSkeletalComponent->CustomizableObjectInstance;
-					APawn* PlayerPawn = UGameplayStatics::GetPlayerPawn(CustomizableSkeletalComponent->GetWorld(), 0);
+
+					APawn* PlayerPawn = nullptr;
+					if (UWorld* World = CustomizableSkeletalComponent->GetWorld())
+					{
+						PlayerPawn = UGameplayStatics::GetPlayerPawn(World, 0);						
+					}
 					
 					if (ParentActor && (ParentActor == PlayerPawn) && Instance)
 					{
