#rnx

diff --git a/Engine/Plugins/Experimental/Mutable/Source/CustomizableObject/Private/MuCO/CustomizableObjectSystem.cpp b/Engine/Plugins/Experimental/Mutable/Source/CustomizableObject/Private/MuCO/CustomizableObjectSystem.cpp
--- a/Engine/Plugins/Experimental/Mutable/Source/CustomizableObject/Private/MuCO/CustomizableObjectSystem.cpp
+++ b/Engine/Plugins/Experimental/Mutable/Source/CustomizableObject/Private/MuCO/CustomizableObjectSystem.cpp
@@ -2139,6 +2139,12 @@ namespace impl
 		UCustomizableInstancePrivateData* CandidateInstancePrivateData = CandidateInstance->GetPrivate();
 		check(CandidateInstancePrivateData != nullptr);
 
+		if (!CandidateInstancePrivateData)
+		{
+			System->ClearCurrentMutableOperation();
+			return;
+		}
+
 		if (CandidateInstancePrivateData && CandidateInstancePrivateData->HasCOInstanceFlags(PendingLODsUpdate))
 		{
 			CandidateInstancePrivateData->ClearCOInstanceFlags(PendingLODsUpdate);
@@ -2190,10 +2196,7 @@ namespace impl
 
 		if (bCancel)
 		{
-			if (CandidateInstancePrivateData) 
-			{
-				CandidateInstancePrivateData->ClearCOInstanceFlags(Updating);
-			}
+			CandidateInstancePrivateData->ClearCOInstanceFlags(Updating);
 
 			System->ClearCurrentMutableOperation();
 			return;
