[Mutable] Copy AssetUserData from the reference skeletal mesh to the generated meshes.

#rnx

diff --git a/Engine/Plugins/Experimental/Mutable/Source/CustomizableObject/Classes/MuCO/CustomizableObject.h b/Engine/Plugins/Experimental/Mutable/Source/CustomizableObject/Classes/MuCO/CustomizableObject.h
--- a/Engine/Plugins/Experimental/Mutable/Source/CustomizableObject/Classes/MuCO/CustomizableObject.h
+++ b/Engine/Plugins/Experimental/Mutable/Source/CustomizableObject/Classes/MuCO/CustomizableObject.h
@@ -24,6 +24,7 @@ class FText;
 class IAsyncReadFileHandle;
 class ITargetPlatform;
 class UAnimInstance;
+class UAssetUserData;
 class UCustomizableObject;
 class UEdGraph;
 class UMaterialInterface;
@@ -667,13 +668,24 @@ struct CUSTOMIZABLEOBJECT_API FMutableCachedPlatformData
 };
 
 
+USTRUCT()
 struct FMutableRefLODInfo
 {
+	GENERATED_BODY()
+
+	UPROPERTY()
 	float ScreenSize = 0.f;
+
+	UPROPERTY()
 	float LODHysteresis = 0.f;
+
+	UPROPERTY()
 	bool bSupportUniformlyDistributedSampling = false;
+
+	UPROPERTY()
 	bool bAllowCPUAccess = false;
 
+#if WITH_EDITORONLY_DATA
 	friend FArchive& operator<<(FArchive& Ar, FMutableRefLODInfo& Data)
 	{
 		Ar << Data.ScreenSize;
@@ -683,14 +695,22 @@ struct FMutableRefLODInfo
 
 		return Ar;
 	}
+#endif
 };
 
 
+USTRUCT()
 struct FMutableRefLODRenderData
 {
+	GENERATED_BODY()
+
+	UPROPERTY()
 	bool bIsLODOptional = false;
+
+	UPROPERTY()
 	bool bStreamedDataInlined = false;
 	
+#if WITH_EDITORONLY_DATA
 	friend FArchive& operator<<(FArchive& Ar, FMutableRefLODRenderData& Data)
 	{
 		Ar << Data.bIsLODOptional;
@@ -698,15 +718,22 @@ struct FMutableRefLODRenderData
 
 		return Ar;
 	}
+#endif
 };
 
 
+USTRUCT()
 struct FMutableRefLODData
 {
+	GENERATED_BODY()
+
+	UPROPERTY()
 	FMutableRefLODInfo LODInfo;
 
+	UPROPERTY()
 	FMutableRefLODRenderData RenderData;
 	
+#if WITH_EDITORONLY_DATA
 	friend FArchive& operator<<(FArchive& Ar, FMutableRefLODData& Data)
 	{
 		Ar << Data.LODInfo;
@@ -714,6 +741,7 @@ struct FMutableRefLODData
 		
 		return Ar;
 	}
+#endif
 };
 
 
@@ -758,6 +786,7 @@ struct FMutableRefSocket
 		return false;
 	}
 	
+#if WITH_EDITORONLY_DATA
 	friend FArchive& operator<<(FArchive& Ar, FMutableRefSocket& Data)
 	{
 		Ar << Data.SocketName;
@@ -770,15 +799,22 @@ struct FMutableRefSocket
 
 		return Ar;
 	}
+#endif
 };
 
 
+USTRUCT()
 struct FMutableRefSkeletalMeshSettings
 {
+	GENERATED_BODY()
+
+	UPROPERTY()
 	bool bEnablePerPolyCollision = false;
 
+	UPROPERTY()
 	float DefaultUVChannelDensity = 0.f;
 
+#if WITH_EDITORONLY_DATA
 	friend FArchive& operator<<(FArchive& Ar, FMutableRefSkeletalMeshSettings& Data)
 	{
 		Ar << Data.bEnablePerPolyCollision;
@@ -786,42 +822,72 @@ struct FMutableRefSkeletalMeshSettings
 
 		return Ar;
 	}
+#endif
 };
 
+
+USTRUCT()
+struct FMutableRefAssetUserData
+{
+	GENERATED_BODY()
+
+	UPROPERTY()
+	TObjectPtr<UAssetUserData> AssetUserData;
+
+#if WITH_EDITORONLY_DATA
+	FString ClassPath;
+	TArray<uint8> Bytes;
+
+	friend FArchive& operator<<(FArchive& Ar, FMutableRefAssetUserData& Data);
+
+	void InitResources(UObject* InOuter);
+#endif
+
+};
+
+
 USTRUCT()
 struct FMutableRefSkeletalMeshData
 {
 	GENERATED_BODY()
 	
 	// LOD info
+	UPROPERTY()
 	TArray<FMutableRefLODData> LODData;
 	
 	// Sockets
+	UPROPERTY()
 	TArray<FMutableRefSocket> Sockets;
 
 	// Bounding Box
+	UPROPERTY()
 	FBoxSphereBounds Bounds;
 
 	// Settings
+	UPROPERTY()
 	FMutableRefSkeletalMeshSettings Settings;
 
 	// Skeleton, must be stored in the ReferencedSkeletons too
-	UPROPERTY(Transient)
+	UPROPERTY()
 	TSoftObjectPtr<USkeleton> Skeleton;
 	
 	// PhysicsAsset, must be stored in the PhysicsAssetMap too
-	UPROPERTY(Transient)
+	UPROPERTY()
 	TSoftObjectPtr<UPhysicsAsset> PhysicsAsset;
 	
 	// Post Processing AnimBP, must be stored in the AnimBPAssetsMap too
-	UPROPERTY(Transient) 
+	UPROPERTY() 
 	TSoftClassPtr<UAnimInstance> PostProcessAnimInst;
 	
 	// Shadow PhysicsAsset, must be stored in the PhysicsAssetMap too
-	UPROPERTY(Transient)
+	UPROPERTY()
 	TSoftObjectPtr<UPhysicsAsset> ShadowPhysicsAsset;
 
+	// Asset user data
+	UPROPERTY()
+	TArray<FMutableRefAssetUserData> AssetUserData;
 
+#if WITH_EDITORONLY_DATA
 	friend FArchive& operator<<(FArchive& Ar, FMutableRefSkeletalMeshData& Data)
 	{
 		Ar << Data.LODData;
@@ -829,8 +895,7 @@ struct FMutableRefSkeletalMeshData
 		Ar << Data.Bounds;
 		Ar << Data.Settings;
 
-
-		if(Ar.IsSaving())
+		if (Ar.IsSaving())
 		{
 			FString AssetPath = Data.Skeleton.ToSoftObjectPath().ToString();
 			Ar << AssetPath;
@@ -864,9 +929,21 @@ struct FMutableRefSkeletalMeshData
 			Data.ShadowPhysicsAsset = TSoftObjectPtr<UPhysicsAsset>(FSoftObjectPath(ShadowPhysicsAssetPath));
 		}
 
+		Ar << Data.AssetUserData;
+
 		return Ar;
 	}
 
+	void InitResources(UObject* InOuter)
+	{
+		check(InOuter);
+		for (FMutableRefAssetUserData& Data : AssetUserData)
+		{
+			Data.InitResources(InOuter);
+		}
+	}
+#endif
+
 };
 
 USTRUCT()
@@ -1004,7 +1081,7 @@ public:
 	 * the Reference Skeletal Meshes to avoid having them loaded at all times. This includes data like LOD distances,
 	 * LOD render data settings, Mesh sockets, Bounding volumes, etc.
 	 */
-	UPROPERTY(Transient)
+	UPROPERTY()
 	TArray<FMutableRefSkeletalMeshData> ReferenceSkeletalMeshesData;
 
 	// Hide for now, since it is not supported yet
@@ -1288,7 +1365,7 @@ private:
 	// This is a manual version number for the binary blobs in this asset.
 	// Increasing it invalidates all the previously compiled models.
 	// Warning: If while merging code both versions have changed, take the highest+1.
-	static const int32 CurrentSupportedVersion = 371;
+	static const int32 CurrentSupportedVersion = 372;
 
 public:
 
diff --git a/Engine/Plugins/Experimental/Mutable/Source/CustomizableObject/Private/MuCO/CustomizableObject.cpp b/Engine/Plugins/Experimental/Mutable/Source/CustomizableObject/Private/MuCO/CustomizableObject.cpp
--- a/Engine/Plugins/Experimental/Mutable/Source/CustomizableObject/Private/MuCO/CustomizableObject.cpp
+++ b/Engine/Plugins/Experimental/Mutable/Source/CustomizableObject/Private/MuCO/CustomizableObject.cpp
@@ -7,6 +7,7 @@
 #include "Async/AsyncFileHandle.h"
 #include "EdGraph/EdGraph.h"
 #include "Engine/SkeletalMesh.h"
+#include "Engine/AssetUserData.h"
 #include "HAL/FileManager.h"
 #include "HAL/PlatformFileManager.h"
 #include "Input/Reply.h"
@@ -457,7 +458,13 @@ void UCustomizableObject::LoadCompiledData(FArchive& MemoryReader, bool bSkipEdi
 		UCustomizableObjectSystem::GetInstance();
 
 		MemoryReader << ReferenceSkeletalMeshesData;
-		
+
+		// Some resources require an outer to initialize themselves. 
+		for(FMutableRefSkeletalMeshData& ReferenceSkeletalMeshData : ReferenceSkeletalMeshesData)
+		{
+			ReferenceSkeletalMeshData.InitResources(this);
+		}
+
 		// We can load
 		int32 NumReferencedMaterials = 0;
 		MemoryReader << NumReferencedMaterials;
@@ -831,11 +838,6 @@ void UCustomizableObject::SaveEmbeddedData(FArchive& Ar)
 
 	if (PrivateData->GetModel())
 	{
-		// Serialize RefSkeletalMeshesData
-		{
-			Ar << ReferenceSkeletalMeshesData;
-		}
-		
 		// Serialize morph data
 		{
 			Ar << ContributingMorphTargetsInfo;
@@ -876,11 +878,6 @@ void UCustomizableObject::LoadEmbeddedData(FArchive& Ar)
 
 	if(CurrentSupportedVersion == InternalVersion)
 	{
-		// Load RefSkeletalMeshesData
-		{
-			Ar << ReferenceSkeletalMeshesData;
-		}
-	
 		// Load morph data
 		{
 			Ar << ContributingMorphTargetsInfo;
@@ -2002,6 +1999,58 @@ void UCustomizableObjectBulk::PrepareBulkData(UCustomizableObject* InOuter, cons
 	}
 }
 
+#if WITH_EDITORONLY_DATA
+
+//-------------------------------------------------------------------------------------------------
+//-------------------------------------------------------------------------------------------------
+//-------------------------------------------------------------------------------------------------
+void FMutableRefAssetUserData::InitResources(UObject* InOuter)
+{
+	UClass* AssetUserDataClass = FindObject<UClass>(nullptr, *ClassPath);
+	if (AssetUserDataClass)
+	{
+		AssetUserData = NewObject<UAssetUserData>(InOuter, AssetUserDataClass);
+		if (AssetUserData)
+		{
+			FMemoryReaderView MemoryReader(Bytes);
+			AssetUserData->Serialize(MemoryReader);
+
+			ClassPath.Empty();
+			Bytes.Empty();
+		}
+	}
+}
+
+
+FArchive& operator<<(FArchive& Ar, FMutableRefAssetUserData& Data)
+{
+	if (Ar.IsSaving())
+	{
+		if (Data.AssetUserData)
+		{
+			Data.ClassPath = Data.AssetUserData->GetClass()->GetPathName();
+
+			FMemoryWriter MemoryWriter(Data.Bytes);
+			Data.AssetUserData->Serialize(MemoryWriter);
+		}
+
+		Ar << Data.ClassPath;
+		Ar << Data.Bytes;
+
+		Data.ClassPath.Empty();
+		Data.Bytes.Empty();
+	}
+	else
+	{
+		Ar << Data.ClassPath;
+		Ar << Data.Bytes;
+	}
+
+	return Ar;
+}
+
+#endif
+
 #endif
 
 #undef LOCTEXT_NAMESPACE
diff --git a/Engine/Plugins/Experimental/Mutable/Source/CustomizableObject/Private/MuCO/CustomizableObjectInstance.cpp b/Engine/Plugins/Experimental/Mutable/Source/CustomizableObject/Private/MuCO/CustomizableObjectInstance.cpp
--- a/Engine/Plugins/Experimental/Mutable/Source/CustomizableObject/Private/MuCO/CustomizableObjectInstance.cpp
+++ b/Engine/Plugins/Experimental/Mutable/Source/CustomizableObject/Private/MuCO/CustomizableObjectInstance.cpp
@@ -2868,7 +2868,14 @@ void UCustomizableInstancePrivateData::InitSkeletalMeshData(const TSharedPtr<FMu
 	SkeletalMesh->SetPhysicsAsset(RefSkeletalMeshData->PhysicsAsset.Get());
 	SkeletalMesh->SetEnablePerPolyCollision(RefSkeletalMeshData->Settings.bEnablePerPolyCollision);
 
-
+	// Asset User Data
+	for (const FMutableRefAssetUserData& MutAssetUserData : RefSkeletalMeshData->AssetUserData)
+	{
+		if (MutAssetUserData.AssetUserData)
+		{
+			SkeletalMesh->AddAssetUserData(MutAssetUserData.AssetUserData);
+		}
+	}
 
 	// Allocate resources for rendering and add LOD Info
 	{
diff --git a/Engine/Plugins/Experimental/Mutable/Source/CustomizableObjectEditor/Private/MuCOE/GenerateMutableSource/GenerateMutableSource.cpp b/Engine/Plugins/Experimental/Mutable/Source/CustomizableObjectEditor/Private/MuCOE/GenerateMutableSource/GenerateMutableSource.cpp
--- a/Engine/Plugins/Experimental/Mutable/Source/CustomizableObjectEditor/Private/MuCOE/GenerateMutableSource/GenerateMutableSource.cpp
+++ b/Engine/Plugins/Experimental/Mutable/Source/CustomizableObjectEditor/Private/MuCOE/GenerateMutableSource/GenerateMutableSource.cpp
@@ -1339,6 +1339,21 @@ void PopulateReferenceSkeletalMeshesData(FMutableGraphGenerationContext& Generat
 			Data.PhysicsAsset = PhysicsAsset;
 			GenerationContext.PhysicsAssetMap.FindOrAdd(Data.PhysicsAsset.ToString(), Data.PhysicsAsset);
 		}
+
+		// Asset User Data
+		if (const TArray<UAssetUserData*>* AssetUserDataArray = RefSkeletalMesh->GetAssetUserDataArray())
+		{
+			for (UAssetUserData* AssetUserData : *AssetUserDataArray)
+			{
+				if (AssetUserData)
+				{
+					FMutableRefAssetUserData& MutAssetUserData = Data.AssetUserData.AddDefaulted_GetRef();
+
+					// Duplicate the AssetUserData using the CO as an outer, so we can save it.
+					MutAssetUserData.AssetUserData = Cast<UAssetUserData>(StaticDuplicateObject(AssetUserData, GenerationContext.Object));
+				}
+			}
+		}
 	}
 }
 
