[Mutable] Ensure skeletons to be merged are compatible

#rnx

diff --git a/Engine/Plugins/Experimental/Mutable/Source/CustomizableObject/Private/MuCO/CustomizableObjectInstance.cpp b/Engine/Plugins/Experimental/Mutable/Source/CustomizableObject/Private/MuCO/CustomizableObjectInstance.cpp
--- a/Engine/Plugins/Experimental/Mutable/Source/CustomizableObject/Private/MuCO/CustomizableObjectInstance.cpp
+++ b/Engine/Plugins/Experimental/Mutable/Source/CustomizableObject/Private/MuCO/CustomizableObjectInstance.cpp
@@ -894,6 +894,107 @@ void UCustomizableInstancePrivateData::UpdateInstanceIfNotGenerated(UCustomizabl
 }
 
 
+#if !UE_BUILD_SHIPPING
+bool AreSkeletonsCompatible(const TArray<TObjectPtr<USkeleton>>& InSkeletons)
+{
+	MUTABLE_CPUPROFILER_SCOPE(AreSkeletonsCompatible);
+	bool bCompatible = true;
+
+	struct FBoneToMergeInfo
+	{
+		FBoneToMergeInfo(const uint32 InBonePathHash, const uint32 InSkeletonIndex, const uint32 InParentBoneSkeletonIndex) :
+		BonePathHash(InBonePathHash), SkeletonIndex(InSkeletonIndex), ParentBoneSkeletonIndex(InParentBoneSkeletonIndex)
+		{}
+
+		uint32 BonePathHash = 0;
+		uint32 SkeletonIndex = 0;
+		uint32 ParentBoneSkeletonIndex = 0;
+	};
+
+	// Accumulated hierarchy hash from parent-bone to root bone
+	TMap<FName, FBoneToMergeInfo> BoneNamesToBoneInfo;
+	BoneNamesToBoneInfo.Reserve(InSkeletons[0] ? InSkeletons[0]->GetReferenceSkeleton().GetNum() : 0);
+	
+	for (int32 SkeletonIndex = 0; SkeletonIndex < InSkeletons.Num(); ++SkeletonIndex)
+	{
+		const TObjectPtr<USkeleton> Skeleton = InSkeletons[SkeletonIndex];
+		check(Skeleton);
+
+		const FReferenceSkeleton& ReferenceSkeleton = Skeleton->GetReferenceSkeleton();
+		const TArray<FMeshBoneInfo>& Bones = ReferenceSkeleton.GetRawRefBoneInfo();
+		const TArray<FTransform>& BonePoses = ReferenceSkeleton.GetRawRefBonePose();
+
+		const int32 NumBones = Bones.Num();
+		for (int32 BoneIndex = 0; BoneIndex < NumBones; ++BoneIndex)
+		{
+			const FMeshBoneInfo& Bone = Bones[BoneIndex];
+
+			// Retrieve parent bone name and respective hash, root-bone is assumed to have a parent hash of 0
+			const FName ParentName = Bone.ParentIndex != INDEX_NONE ? Bones[Bone.ParentIndex].Name : NAME_None;
+			const uint32 ParentHash = Bone.ParentIndex != INDEX_NONE ? GetTypeHash(ParentName) : 0;
+
+			// Look-up the path-hash from root to the parent bone
+			const FBoneToMergeInfo* ParentBoneInfo = BoneNamesToBoneInfo.Find(ParentName);
+			const uint32 ParentBonePathHash = ParentBoneInfo ? ParentBoneInfo->BonePathHash : 0;
+			const uint32 ParentBoneSkeletonIndex = ParentBoneInfo ? ParentBoneInfo->SkeletonIndex : 0;
+
+			// Append parent hash to path to give full path hash to current bone
+			const uint32 BonePathHash = HashCombine(ParentBonePathHash, ParentHash);
+
+			// Check if the bone exists in the hierarchy 
+			const FBoneToMergeInfo* ExistingBoneInfo = BoneNamesToBoneInfo.Find(Bone.Name);
+
+			// If the hash differs from the existing one it means skeletons are incompatible
+			if (!ExistingBoneInfo)
+			{
+				// Add path hash to current bone
+				BoneNamesToBoneInfo.Add(Bone.Name, FBoneToMergeInfo(BonePathHash, SkeletonIndex, ParentBoneSkeletonIndex));
+			}
+			else if (ExistingBoneInfo->BonePathHash != BonePathHash)
+			{
+				if (bCompatible)
+				{
+					// Print the skeletons to merge
+					FString Msg = TEXT("Failed to merge skeletons. Skeletons to merge: ");
+					for (int32 AuxSkeletonIndex = 0; AuxSkeletonIndex < InSkeletons.Num(); ++AuxSkeletonIndex)
+					{
+						if (InSkeletons[AuxSkeletonIndex] != nullptr)
+						{
+							Msg += FString::Printf(TEXT("\n\t- %s"), *InSkeletons[AuxSkeletonIndex].GetName());
+						}
+					}
+
+					UE_LOG(LogMutable, Error, TEXT("%s"), *Msg);
+
+#if WITH_EDITOR
+					FNotificationInfo Info(FText::FromString(TEXT("Mutable: Failed to merge skeletons. Invalid parent chain detected. Please check the output log for more information.")));
+					Info.bFireAndForget = true;
+					Info.FadeOutDuration = 1.0f;
+					Info.ExpireDuration = 10.0f;
+					FSlateNotificationManager::Get().AddNotification(Info);
+#endif
+
+					bCompatible = false;
+				}
+				
+				// Print the first non compatible bone in the bone chain, since all child bones will be incompatible too.
+				if (ExistingBoneInfo->ParentBoneSkeletonIndex != SkeletonIndex)
+				{
+					// Different skeletons can't be used if they are incompatible with the reference skeleton.
+					UE_LOG(LogMutable, Error, TEXT("[%s] parent bone is different in skeletons [%s] and [%s]."),
+						*Bone.Name.ToString(),
+						*InSkeletons[SkeletonIndex]->GetName(),
+						*InSkeletons[ExistingBoneInfo->ParentBoneSkeletonIndex]->GetName());
+				}
+			}
+		}
+	}
+
+	return bCompatible;
+}
+#endif
+
+
 USkeleton* UCustomizableInstancePrivateData::MergeSkeletons(UCustomizableObjectInstance* Public, const FMutableRefSkeletalMeshData* RefSkeletalMeshData, int32 ComponentIndex)
 {
 	MUTABLE_CPUPROFILER_SCOPE(BuildSkeletonData_MergeSkeletons);
@@ -917,6 +1018,14 @@ USkeleton* UCustomizableInstancePrivateData::MergeSkeletons(UCustomizableObjectI
 		return RefSkeleton;
 	}
 
+#if !UE_BUILD_SHIPPING
+	// Test Skeleton compatibility before attempting the merge to avoid a crash.
+	if (!AreSkeletonsCompatible(ReferencedSkeletons.SkeletonsToMerge))
+	{
+		return nullptr;
+	}
+#endif
+
 	FSkeletonMergeParams Params;
 	Exchange(Params.SkeletonsToMerge, ReferencedSkeletons.SkeletonsToMerge);
 
