[Mtuable] Replace deprecated ANSI_TO_TCHAR macros.

#rnx

diff --git a/Engine/Plugins/Experimental/Mutable/Source/CustomizableObject/Private/MuCO/CustomizableObject.cpp b/Engine/Plugins/Experimental/Mutable/Source/CustomizableObject/Private/MuCO/CustomizableObject.cpp
--- a/Engine/Plugins/Experimental/Mutable/Source/CustomizableObject/Private/MuCO/CustomizableObject.cpp
+++ b/Engine/Plugins/Experimental/Mutable/Source/CustomizableObject/Private/MuCO/CustomizableObject.cpp
@@ -968,7 +968,7 @@ int32 UCustomizableObject::FindState( const FString& Name ) const
 	int32 Result = -1;
 	if (PrivateData->GetModel())
 	{
-		Result = PrivateData->GetModel()->FindState( TCHAR_TO_ANSI(*Name) );
+		Result = PrivateData->GetModel()->FindState( StringCast<ANSICHAR>(*Name).Get() );
 	}
 
 	return Result;
@@ -994,7 +994,7 @@ FString UCustomizableObject::GetStateName(int32 StateIndex) const
 
 	if (PrivateData->GetModel())
 	{
-		Result = ANSI_TO_TCHAR( PrivateData->GetModel()->GetStateName(StateIndex) );
+		Result = StringCast<TCHAR>( PrivateData->GetModel()->GetStateName(StateIndex) ).Get();
 	}
 
 	return Result;
diff --git a/Engine/Plugins/Experimental/Mutable/Source/CustomizableObject/Private/MuCO/CustomizableObjectInstanceDescriptor.cpp b/Engine/Plugins/Experimental/Mutable/Source/CustomizableObject/Private/MuCO/CustomizableObjectInstanceDescriptor.cpp
--- a/Engine/Plugins/Experimental/Mutable/Source/CustomizableObject/Private/MuCO/CustomizableObjectInstanceDescriptor.cpp
+++ b/Engine/Plugins/Experimental/Mutable/Source/CustomizableObject/Private/MuCO/CustomizableObjectInstanceDescriptor.cpp
@@ -2058,7 +2058,7 @@ void FCustomizableObjectInstanceDescriptor::SetCurrentState(const FString& State
 	}
 	else
 	{
-		UE_LOG(LogMutable, Error, TEXT("%s: Unknown %s state."), ANSI_TO_TCHAR(__FUNCTION__), *StateName);
+		UE_LOG(LogMutable, Error, TEXT("%s: Unknown %s state."), StringCast<TCHAR>(__FUNCTION__).Get(), *StateName);
 	}
 }
 
diff --git a/Engine/Plugins/Experimental/Mutable/Source/CustomizableObjectEditor/Private/MuCOE/GenerateMutableSource/GenerateMutableSource.cpp b/Engine/Plugins/Experimental/Mutable/Source/CustomizableObjectEditor/Private/MuCOE/GenerateMutableSource/GenerateMutableSource.cpp
--- a/Engine/Plugins/Experimental/Mutable/Source/CustomizableObjectEditor/Private/MuCOE/GenerateMutableSource/GenerateMutableSource.cpp
+++ b/Engine/Plugins/Experimental/Mutable/Source/CustomizableObjectEditor/Private/MuCOE/GenerateMutableSource/GenerateMutableSource.cpp
@@ -186,7 +186,7 @@ void FMutableGraphGenerationContext::GenerateClippingCOInternalTags()
 
 		for (j = 0; j < It->Value.Num(); ++j)
 		{
-			It->Value[j]->AddTag(TCHAR_TO_ANSI(*TagName));
+			It->Value[j]->AddTag(StringCast<ANSICHAR>(*TagName).Get());
 		}
 
 		UCustomizableObjectNodeMeshClipWithMesh* CustomizableObjectNodeMeshClipWithMesh = It->Key;
@@ -213,7 +213,7 @@ void FMutableGraphGenerationContext::GenerateClippingCOInternalTags()
 
 					for (int32 k = 0; k < It2->Value.Num(); ++k)
 					{
-						It2->Value[k]->AddTag(TCHAR_TO_ANSI(*TagName));
+						It2->Value[k]->AddTag(StringCast<ANSICHAR>(*TagName).Get());
 					}
 				}
 			}
@@ -564,8 +564,8 @@ mu::NodeMeshApplyPosePtr CreateNodeMeshApplyPose(mu::NodeMeshPtr InputMeshNode,
 
 	for (int32 i = 0; i < ArrayBoneName.Num(); ++i)
 	{
-		MutableSkeleton->SetBoneName(i, TCHAR_TO_ANSI(*ArrayBoneName[i]));
-		MutableMesh->SetBonePose(i, TCHAR_TO_ANSI(*ArrayBoneName[i]), (FTransform3f)ArrayTransform[i], mu::EBoneUsageFlags::Skinning);
+		MutableSkeleton->SetBoneName(i, StringCast<ANSICHAR>(*ArrayBoneName[i]).Get());
+		MutableMesh->SetBonePose(i, StringCast<ANSICHAR>(*ArrayBoneName[i]).Get(), (FTransform3f)ArrayTransform[i], mu::EBoneUsageFlags::Skinning);
 	}
 
 	mu::NodeMeshApplyPosePtr NodeMeshApplyPose = new mu::NodeMeshApplyPose;
@@ -607,8 +607,8 @@ mu::NodeObjectPtr GenerateMutableSource(const UEdGraphPin * Pin, FMutableGraphGe
 		mu::NodeObjectNewPtr ObjectNode = new mu::NodeObjectNew();
 		Result = ObjectNode;
 
-		ObjectNode->SetName(TCHAR_TO_ANSI(*TypedNodeObj->ObjectName));
-		ObjectNode->SetUid(TCHAR_TO_ANSI(*GenerationContext.GetNodeIdUnique(TypedNodeObj).ToString()));
+		ObjectNode->SetName(StringCast<ANSICHAR>(*TypedNodeObj->ObjectName).Get());
+		ObjectNode->SetUid(StringCast<ANSICHAR>(*GenerationContext.GetNodeIdUnique(TypedNodeObj).ToString()).Get());
 
 		// LOD
 		const int32 NumLODs = TypedNodeObj->GetNumLODPins();
@@ -699,10 +699,10 @@ mu::NodeObjectPtr GenerateMutableSource(const UEdGraphPin * Pin, FMutableGraphGe
 		for (int StateIndex = 0; StateIndex < NumStates && bFilterStates; ++StateIndex)
 		{
 			const FCustomizableObjectState& State = TypedNodeObj->States[StateIndex];
-			ObjectNode->SetStateName(StateIndex, TCHAR_TO_ANSI(*State.Name));
+			ObjectNode->SetStateName(StateIndex, StringCast<ANSICHAR>(*State.Name).Get());
 			for (int ParamIndex = 0; ParamIndex < State.RuntimeParameters.Num(); ++ParamIndex)
 			{
-				ObjectNode->AddStateParam(StateIndex, TCHAR_TO_ANSI(*State.RuntimeParameters[ParamIndex]));
+				ObjectNode->AddStateParam(StateIndex, StringCast<ANSICHAR>(*State.RuntimeParameters[ParamIndex]).Get());
 			}
 
 			const ITargetPlatform* TargetPlatform = GenerationContext.Options.TargetPlatform;
@@ -894,8 +894,8 @@ mu::NodeObjectPtr GenerateMutableSource(const UEdGraphPin * Pin, FMutableGraphGe
 		GenerationContext.SocketPriorityStack.Push(TypedNodeGroup->SocketPriority);
 
 		GenerationContext.AddParameterNameUnique(TypedNodeGroup, TypedNodeGroup->GroupName);
-		GroupNode->SetName(TCHAR_TO_ANSI(*TypedNodeGroup->GroupName));
-		GroupNode->SetUid(TCHAR_TO_ANSI(*TypedNodeGroup->NodeGuid.ToString()));
+		GroupNode->SetName(StringCast<ANSICHAR>(*TypedNodeGroup->GroupName).Get());
+		GroupNode->SetUid(StringCast<ANSICHAR>(*TypedNodeGroup->NodeGuid.ToString()).Get());
 
 		// Get all group projectors and put them in the generation context so that they are available to the child material nodes of this group node
 		uint32 NumProjectorCountBeforeNode = GenerationContext.ProjectorGroupMap.Num();
@@ -981,7 +981,7 @@ mu::NodeObjectPtr GenerateMutableSource(const UEdGraphPin * Pin, FMutableGraphGe
 			else
 			{
 				mu::NodeObjectPtr ChildNode = new mu::NodeObjectNew;
-				ChildNode->SetName(TCHAR_TO_ANSI(*CustomizableObjectNodeObject->ObjectName));
+				ChildNode->SetName(StringCast<ANSICHAR>(*CustomizableObjectNodeObject->ObjectName).Get());
 				GroupNode->SetChild(ChildIndex, ChildNode.get());
 			}
 		}
@@ -1069,7 +1069,7 @@ mu::NodeObjectPtr GenerateMutableSource(const UEdGraphPin * Pin, FMutableGraphGe
 			else
 			{
 				mu::NodeObjectPtr ChildNode = new mu::NodeObjectNew;
-				ChildNode->SetName(TCHAR_TO_ANSI(*CustomizableObjectNodeObject->ObjectName));
+				ChildNode->SetName(StringCast<ANSICHAR>(*CustomizableObjectNodeObject->ObjectName).Get());
 				GroupNode->SetChild(ChildIndex, ChildNode.get());
 			}
 
diff --git a/Engine/Plugins/Experimental/Mutable/Source/CustomizableObjectEditor/Private/MuCOE/GenerateMutableSource/GenerateMutableSourceColor.cpp b/Engine/Plugins/Experimental/Mutable/Source/CustomizableObjectEditor/Private/MuCOE/GenerateMutableSource/GenerateMutableSourceColor.cpp
--- a/Engine/Plugins/Experimental/Mutable/Source/CustomizableObjectEditor/Private/MuCOE/GenerateMutableSource/GenerateMutableSourceColor.cpp
+++ b/Engine/Plugins/Experimental/Mutable/Source/CustomizableObjectEditor/Private/MuCOE/GenerateMutableSource/GenerateMutableSourceColor.cpp
@@ -64,8 +64,8 @@ mu::NodeColourPtr GenerateMutableSourceColor(const UEdGraphPin* Pin, FMutableGra
 
 		GenerationContext.AddParameterNameUnique(Node, TypedNodeColorParam->ParameterName);
 
-		ColorNode->SetName(TCHAR_TO_ANSI(*TypedNodeColorParam->ParameterName));
-		ColorNode->SetUid(TCHAR_TO_ANSI(*GenerationContext.GetNodeIdUnique(Node).ToString()));
+		ColorNode->SetName(StringCast<ANSICHAR>(*TypedNodeColorParam->ParameterName).Get());
+		ColorNode->SetUid(StringCast<ANSICHAR>(*GenerationContext.GetNodeIdUnique(Node).ToString()).Get());
 		ColorNode->SetDefaultValue(TypedNodeColorParam->DefaultValue.R, TypedNodeColorParam->DefaultValue.G, TypedNodeColorParam->DefaultValue.B);
 
 		GenerationContext.ParameterUIDataMap.Add(TypedNodeColorParam->ParameterName, FParameterUIData(
@@ -258,7 +258,7 @@ mu::NodeColourPtr GenerateMutableSourceColor(const UEdGraphPin* Pin, FMutableGra
 			const UEdGraphPin* VariationPin = TypedNodeColorVar->VariationPin(VariationIndex);
 			if (!VariationPin) continue;
 
-			ColorNode->SetVariationTag(VariationIndex, TCHAR_TO_ANSI(*TypedNodeColorVar->Variations[VariationIndex].Tag));
+			ColorNode->SetVariationTag(VariationIndex, StringCast<ANSICHAR>(*TypedNodeColorVar->Variations[VariationIndex].Tag).Get());
 			if (const UEdGraphPin* ConnectedPin = FollowInputPin(*VariationPin))
 			{
 				mu::NodeColourPtr ChildNode = GenerateMutableSourceColor(ConnectedPin, GenerationContext);
@@ -286,18 +286,18 @@ mu::NodeColourPtr GenerateMutableSourceColor(const UEdGraphPin* Pin, FMutableGra
 			Table = GenerateMutableSourceTable(TypedNodeTable->Table->GetName(), Pin, GenerationContext);
 
 			// Generating a new Color column if not exists
-			if (Table && Table->FindColumn(TCHAR_TO_ANSI(*ColumnName)) == INDEX_NONE)
+			if (Table && Table->FindColumn(StringCast<ANSICHAR>(*ColumnName).Get()) == INDEX_NONE)
 			{
 				GenerateTableColumn(TypedNodeTable, Pin, Table, ColumnName, GenerationContext.CurrentLOD, GenerationContext);
 			}
 
 			ColorTableNode->SetTable(Table);
-			ColorTableNode->SetColumn(TCHAR_TO_ANSI(*ColumnName));
-			ColorTableNode->SetParameterName(TCHAR_TO_ANSI(*TypedNodeTable->ParameterName));
+			ColorTableNode->SetColumn(StringCast<ANSICHAR>(*ColumnName).Get());
+			ColorTableNode->SetParameterName(StringCast<ANSICHAR>(*TypedNodeTable->ParameterName).Get());
 
 			GenerationContext.AddParameterNameUnique(Node, TypedNodeTable->ParameterName);
 
-			if (Table->FindColumn(TCHAR_TO_ANSI(*ColumnName)) == INDEX_NONE)
+			if (Table->FindColumn(StringCast<ANSICHAR>(*ColumnName).Get()) == INDEX_NONE)
 			{
 				FString Msg = FString::Printf(TEXT("Couldn't find pin column with name %s"), *ColumnName);
 				GenerationContext.Compiler->CompilerLog(FText::FromString(Msg), Node);
diff --git a/Engine/Plugins/Experimental/Mutable/Source/CustomizableObjectEditor/Private/MuCOE/GenerateMutableSource/GenerateMutableSourceFloat.cpp b/Engine/Plugins/Experimental/Mutable/Source/CustomizableObjectEditor/Private/MuCOE/GenerateMutableSource/GenerateMutableSourceFloat.cpp
--- a/Engine/Plugins/Experimental/Mutable/Source/CustomizableObjectEditor/Private/MuCOE/GenerateMutableSource/GenerateMutableSourceFloat.cpp
+++ b/Engine/Plugins/Experimental/Mutable/Source/CustomizableObjectEditor/Private/MuCOE/GenerateMutableSource/GenerateMutableSourceFloat.cpp
@@ -91,8 +91,8 @@ mu::NodeScalarPtr GenerateMutableSourceFloat(const UEdGraphPin* Pin, FMutableGra
 
 		GenerationContext.AddParameterNameUnique(Node, FloatParameterNode->ParameterName);
 
-		ScalarNode->SetName(TCHAR_TO_ANSI(*FloatParameterNode->ParameterName));
-		ScalarNode->SetUid(TCHAR_TO_ANSI(*GenerationContext.GetNodeIdUnique(Node).ToString()));
+		ScalarNode->SetName(StringCast<ANSICHAR>(*FloatParameterNode->ParameterName).Get());
+		ScalarNode->SetUid(StringCast<ANSICHAR>(*GenerationContext.GetNodeIdUnique(Node).ToString()).Get());
 		ScalarNode->SetDefaultValue(FloatParameterNode->DefaultValue);
 
 		const int32 NumDescriptionImage = FloatParameterNode->GetNumDescriptionImage();
@@ -128,8 +128,8 @@ mu::NodeScalarPtr GenerateMutableSourceFloat(const UEdGraphPin* Pin, FMutableGra
 
 		GenerationContext.AddParameterNameUnique(Node, EnumParamNode->ParameterName);
 
-		EnumParameterNode->SetName(TCHAR_TO_ANSI(*EnumParamNode->ParameterName));
-		EnumParameterNode->SetUid(TCHAR_TO_ANSI(*GenerationContext.GetNodeIdUnique(Node).ToString()));
+		EnumParameterNode->SetName(StringCast<ANSICHAR>(*EnumParamNode->ParameterName).Get());
+		EnumParameterNode->SetUid(StringCast<ANSICHAR>(*GenerationContext.GetNodeIdUnique(Node).ToString()).Get());
 		EnumParameterNode->SetValueCount(NumSelectors);
 		EnumParameterNode->SetDefaultValueIndex(DefaultValue);
 
@@ -138,7 +138,7 @@ mu::NodeScalarPtr GenerateMutableSourceFloat(const UEdGraphPin* Pin, FMutableGra
 
 		for (int SelectorIndex = 0; SelectorIndex < NumSelectors; ++SelectorIndex)
 		{
-			EnumParameterNode->SetValue(SelectorIndex, (float)SelectorIndex, TCHAR_TO_ANSI(*EnumParamNode->Values[SelectorIndex].Name));
+			EnumParameterNode->SetValue(SelectorIndex, (float)SelectorIndex, StringCast<ANSICHAR>(*EnumParamNode->Values[SelectorIndex].Name).Get());
 
 			ParameterUIData.ArrayIntegerParameterOption.Add(FIntegerParameterUIData(
 				EnumParamNode->Values[SelectorIndex].Name,
@@ -301,7 +301,7 @@ mu::NodeScalarPtr GenerateMutableSourceFloat(const UEdGraphPin* Pin, FMutableGra
 			UEdGraphPin* VariationPin = TypedNodeFloatVar->VariationPin(VariationIndex);
 			if (!VariationPin) continue;
 
-			FloatNode->SetVariationTag(VariationIndex, TCHAR_TO_ANSI(*TypedNodeFloatVar->Variations[VariationIndex].Tag));
+			FloatNode->SetVariationTag(VariationIndex, StringCast<ANSICHAR>(*TypedNodeFloatVar->Variations[VariationIndex].Tag).Get());
 			if (const UEdGraphPin* ConnectedPin = FollowInputPin(*VariationPin))
 			{
 				mu::NodeScalarPtr ChildNode = GenerateMutableSourceFloat(ConnectedPin, GenerationContext);
@@ -329,18 +329,18 @@ mu::NodeScalarPtr GenerateMutableSourceFloat(const UEdGraphPin* Pin, FMutableGra
 			Table = GenerateMutableSourceTable(TypedNodeTable->Table->GetName(), Pin, GenerationContext);
 
 			// Generating a new Float column if not exists
-			if (Table && Table->FindColumn(TCHAR_TO_ANSI(*ColumnName)) == INDEX_NONE)
+			if (Table && Table->FindColumn(StringCast<ANSICHAR>(*ColumnName).Get()) == INDEX_NONE)
 			{
 				GenerateTableColumn(TypedNodeTable, Pin, Table, ColumnName, GenerationContext.CurrentLOD, GenerationContext);
 			}
 			
 			ScalarTableNode->SetTable(Table);
-			ScalarTableNode->SetColumn(TCHAR_TO_ANSI(*ColumnName));
-			ScalarTableNode->SetParameterName(TCHAR_TO_ANSI(*TypedNodeTable->ParameterName));
+			ScalarTableNode->SetColumn(StringCast<ANSICHAR>(*ColumnName).Get());
+			ScalarTableNode->SetParameterName(StringCast<ANSICHAR>(*TypedNodeTable->ParameterName).Get());
 
 			GenerationContext.AddParameterNameUnique(Node, TypedNodeTable->ParameterName);
 			
-			if (Table->FindColumn(TCHAR_TO_ANSI(*ColumnName)) == INDEX_NONE)
+			if (Table->FindColumn(StringCast<ANSICHAR>(*ColumnName).Get()) == INDEX_NONE)
 			{
 				Table = new mu::Table();
 
diff --git a/Engine/Plugins/Experimental/Mutable/Source/CustomizableObjectEditor/Private/MuCOE/GenerateMutableSource/GenerateMutableSourceGroupProjector.cpp b/Engine/Plugins/Experimental/Mutable/Source/CustomizableObjectEditor/Private/MuCOE/GenerateMutableSource/GenerateMutableSourceGroupProjector.cpp
--- a/Engine/Plugins/Experimental/Mutable/Source/CustomizableObjectEditor/Private/MuCOE/GenerateMutableSource/GenerateMutableSourceGroupProjector.cpp
+++ b/Engine/Plugins/Experimental/Mutable/Source/CustomizableObjectEditor/Private/MuCOE/GenerateMutableSource/GenerateMutableSourceGroupProjector.cpp
@@ -432,8 +432,8 @@ bool GenerateMutableSourceGroupProjector(const UEdGraphPin* Pin, FMutableGraphGe
 
 			mu::NodeScalarParameterPtr NodeScalarParam = new mu::NodeScalarParameter;
 			FString NodeScalarParamName = ProjParamNode->ParameterName + FMultilayerProjector::NUM_LAYERS_PARAMETER_POSTFIX;
-			NodeScalarParam->SetName(TCHAR_TO_ANSI(*(NodeScalarParamName)));
-			NodeScalarParam->SetUid(TCHAR_TO_ANSI(*(ProjectorParamUid + FString("_NL"))));
+			NodeScalarParam->SetName(StringCast<ANSICHAR>(*(NodeScalarParamName)).Get());
+			NodeScalarParam->SetUid(StringCast<ANSICHAR>(*(ProjectorParamUid + FString("_NL"))).Get());
 			GenerationContext.AddParameterNameUnique(originalGroup, NodeScalarParamName);
 			NodeScalarParam->SetDefaultValue(0.f);
 
@@ -451,8 +451,8 @@ bool GenerateMutableSourceGroupProjector(const UEdGraphPin* Pin, FMutableGraphGe
 			
 			mu::NodeScalarEnumParameterPtr PoseEnumParameterNode = new mu::NodeScalarEnumParameter;
 			FString PoseNodeEnumParamName = ProjParamNode->ParameterName + FMultilayerProjector::POSE_PARAMETER_POSTFIX;
-			PoseEnumParameterNode->SetName(TCHAR_TO_ANSI(*(PoseNodeEnumParamName)));
-			PoseEnumParameterNode->SetUid(TCHAR_TO_ANSI(*(ProjectorParamUid + FString("_SP"))));
+			PoseEnumParameterNode->SetName(StringCast<ANSICHAR>(*(PoseNodeEnumParamName)).Get());
+			PoseEnumParameterNode->SetUid(StringCast<ANSICHAR>(*(ProjectorParamUid + FString("_SP"))).Get());
 			GenerationContext.AddParameterNameUnique(originalGroup, PoseNodeEnumParamName);
 			PoseEnumParameterNode->SetValueCount(ProjParamNode->OptionPoses.Num() + 1);
 			PoseEnumParameterNode->SetDefaultValueIndex(0);
@@ -465,8 +465,8 @@ bool GenerateMutableSourceGroupProjector(const UEdGraphPin* Pin, FMutableGraphGe
 
 			mu::NodeScalarParameterPtr OpacityParameterNode = new mu::NodeScalarParameter;
 			FString OpacityParameterNodeName = ProjParamNode->ParameterName + FMultilayerProjector::OPACITY_PARAMETER_POSTFIX;
-			OpacityParameterNode->SetName(TCHAR_TO_ANSI(*(OpacityParameterNodeName)));
-			OpacityParameterNode->SetUid(TCHAR_TO_ANSI(*(ProjectorParamUid + FString("_O"))));
+			OpacityParameterNode->SetName(StringCast<ANSICHAR>(*(OpacityParameterNodeName)).Get());
+			OpacityParameterNode->SetUid(StringCast<ANSICHAR>(*(ProjectorParamUid + FString("_O"))).Get());
 			GenerationContext.AddParameterNameUnique(originalGroup, OpacityParameterNodeName);
 			OpacityParameterNode->SetDefaultValue(0.75f);
 			OpacityParameterNode->SetRangeCount(1);
@@ -491,7 +491,7 @@ bool GenerateMutableSourceGroupProjector(const UEdGraphPin* Pin, FMutableGraphGe
 
 			for (int PoseIndex = 0; PoseIndex < ProjParamNode->OptionPoses.Num(); ++PoseIndex)
 			{
-				PoseEnumParameterNode->SetValue(PoseIndex + 1, (float)PoseIndex + 1.f, TCHAR_TO_ANSI(*ProjParamNode->OptionPoses[PoseIndex].PoseName));
+				PoseEnumParameterNode->SetValue(PoseIndex + 1, (float)PoseIndex + 1.f, StringCast<ANSICHAR>(*ProjParamNode->OptionPoses[PoseIndex].PoseName).Get());
 
 				TArray<FString> ArrayBoneName;
 				TArray<FTransform> ArrayTransform;
@@ -512,8 +512,8 @@ bool GenerateMutableSourceGroupProjector(const UEdGraphPin* Pin, FMutableGraphGe
 		
 			mu::NodeScalarEnumParameterPtr EnumParameterNode = new mu::NodeScalarEnumParameter;
 			FString NodeEnumParamName = ProjParamNode->ParameterName + FMultilayerProjector::IMAGE_PARAMETER_POSTFIX;
-			EnumParameterNode->SetName(TCHAR_TO_ANSI(*(NodeEnumParamName)));
-			EnumParameterNode->SetUid(TCHAR_TO_ANSI(*(ProjectorParamUid + FString("_SI"))));
+			EnumParameterNode->SetName(StringCast<ANSICHAR>(*NodeEnumParamName).Get());
+			EnumParameterNode->SetUid(StringCast<ANSICHAR>(*(ProjectorParamUid + FString("_SI"))).Get());
 			GenerationContext.AddParameterNameUnique(originalGroup, NodeEnumParamName);
 			EnumParameterNode->SetValueCount(ArrayOptionImage.Num());
 			EnumParameterNode->SetDefaultValueIndex(0);
@@ -526,7 +526,7 @@ bool GenerateMutableSourceGroupProjector(const UEdGraphPin* Pin, FMutableGraphGe
 
 			for (int ImageIndex = 0; ImageIndex < ArrayOptionImage.Num(); ++ImageIndex)
 			{
-				EnumParameterNode->SetValue(ImageIndex, (float)ImageIndex, TCHAR_TO_ANSI(*ArrayOptionImage[ImageIndex].OptionName));
+				EnumParameterNode->SetValue(ImageIndex, (float)ImageIndex, StringCast<ANSICHAR>(*ArrayOptionImage[ImageIndex].OptionName).Get());
 
 				FMutableParamUIMetadata optionMetadata = ParameterUIData.ParamUIMetadata;
 				optionMetadata.UIThumbnail = ArrayOptionImage[ImageIndex].OptionImage;
diff --git a/Engine/Plugins/Experimental/Mutable/Source/CustomizableObjectEditor/Private/MuCOE/GenerateMutableSource/GenerateMutableSourceImage.cpp b/Engine/Plugins/Experimental/Mutable/Source/CustomizableObjectEditor/Private/MuCOE/GenerateMutableSource/GenerateMutableSourceImage.cpp
--- a/Engine/Plugins/Experimental/Mutable/Source/CustomizableObjectEditor/Private/MuCOE/GenerateMutableSource/GenerateMutableSourceImage.cpp
+++ b/Engine/Plugins/Experimental/Mutable/Source/CustomizableObjectEditor/Private/MuCOE/GenerateMutableSource/GenerateMutableSourceImage.cpp
@@ -175,8 +175,8 @@ mu::NodeImagePtr GenerateMutableSourceImage(const UEdGraphPin* Pin, FMutableGrap
 
 		GenerationContext.AddParameterNameUnique(Node, TypedNodeParam->ParameterName);
 
-		TextureNode->SetName(TCHAR_TO_ANSI(*TypedNodeParam->ParameterName));
-		TextureNode->SetUid(TCHAR_TO_ANSI(*GenerationContext.GetNodeIdUnique(Node).ToString()));
+		TextureNode->SetName(StringCast<ANSICHAR>(*TypedNodeParam->ParameterName).Get());
+		TextureNode->SetUid(StringCast<ANSICHAR>(*GenerationContext.GetNodeIdUnique(Node).ToString()).Get());
 
 		// TODO: Set a default value for the texture parameter?
 		//ColorNode->SetDefaultValue(TypedNodeParam->DefaultValue.R, TypedNodeParam->DefaultValue.G, TypedNodeParam->DefaultValue.B);
@@ -428,7 +428,7 @@ mu::NodeImagePtr GenerateMutableSourceImage(const UEdGraphPin* Pin, FMutableGrap
 			const UEdGraphPin* VariationPin = TypedNodeImageVar->VariationPin(VariationIndex);
 			if (!VariationPin) continue;
 
-			TextureNode->SetVariationTag(VariationIndex, TCHAR_TO_ANSI(*TypedNodeImageVar->Variations[VariationIndex].Tag));
+			TextureNode->SetVariationTag(VariationIndex, StringCast<ANSICHAR>(*TypedNodeImageVar->Variations[VariationIndex].Tag).Get());
 			if (const UEdGraphPin* ConnectedPin = FollowInputPin(*VariationPin))
 			{
 				mu::NodeImagePtr ChildNode = GenerateMutableSourceImage(ConnectedPin, GenerationContext, MaxTextureSize);
@@ -851,18 +851,18 @@ mu::NodeImagePtr GenerateMutableSourceImage(const UEdGraphPin* Pin, FMutableGrap
 				Table = GenerateMutableSourceTable(TypedNodeTable->Table->GetName(), Pin, GenerationContext);
 
 				// Generating a new Texture column if not exists
-				if (Table && Table->FindColumn(TCHAR_TO_ANSI(*ColumnName)) == INDEX_NONE)
+				if (Table && Table->FindColumn(StringCast<ANSICHAR>(*ColumnName).Get()) == INDEX_NONE)
 				{
 					GenerateTableColumn(TypedNodeTable, Pin, Table, ColumnName, GenerationContext.CurrentLOD, GenerationContext);
 				}
 				
 				ImageTableNode->SetTable(Table);
-				ImageTableNode->SetColumn(TCHAR_TO_ANSI(*ColumnName));
-				ImageTableNode->SetParameterName(TCHAR_TO_ANSI(*TypedNodeTable->ParameterName));
+				ImageTableNode->SetColumn(StringCast<ANSICHAR>(*ColumnName).Get());
+				ImageTableNode->SetParameterName(StringCast<ANSICHAR>(*TypedNodeTable->ParameterName).Get());
 
 				GenerationContext.AddParameterNameUnique(Node, TypedNodeTable->ParameterName);
 				
-				if (Table->FindColumn(TCHAR_TO_ANSI(*ColumnName)) == INDEX_NONE)
+				if (Table->FindColumn(StringCast<ANSICHAR>(*ColumnName).Get()) == INDEX_NONE)
 				{
 					FString Msg = FString::Printf(TEXT("Couldn't find pin column with name %s"), *ColumnName);
 					GenerationContext.Compiler->CompilerLog(FText::FromString(Msg), Node);
diff --git a/Engine/Plugins/Experimental/Mutable/Source/CustomizableObjectEditor/Private/MuCOE/GenerateMutableSource/GenerateMutableSourceMesh.cpp b/Engine/Plugins/Experimental/Mutable/Source/CustomizableObjectEditor/Private/MuCOE/GenerateMutableSource/GenerateMutableSourceMesh.cpp
--- a/Engine/Plugins/Experimental/Mutable/Source/CustomizableObjectEditor/Private/MuCOE/GenerateMutableSource/GenerateMutableSourceMesh.cpp
+++ b/Engine/Plugins/Experimental/Mutable/Source/CustomizableObjectEditor/Private/MuCOE/GenerateMutableSource/GenerateMutableSourceMesh.cpp
@@ -1317,7 +1317,7 @@ mu::MeshPtr ConvertSkeletalMeshToMutable(USkeletalMesh* InSkeletalMesh, int LOD,
 				MutableSkeleton->SetBoneCount(BoneMapIndex + 1);
 			}
 
-			MutableSkeleton->SetBoneName(BoneMapIndex, TCHAR_TO_ANSI(*BoneName));
+			MutableSkeleton->SetBoneName(BoneMapIndex, StringCast<ANSICHAR>(*BoneName).Get());
 			int32 ParentRefSkeletonIndex = SkeletalMesh->GetRefSkeleton().GetParentIndex(RefSkelIndex);
 			int ParentBoneMapIndex = ParentRefSkeletonIndex >= 0 ? InverseBoneMap[ParentRefSkeletonIndex] : -1;
 			MutableSkeleton->SetBoneParent(BoneMapIndex, ParentBoneMapIndex);
@@ -1333,7 +1333,7 @@ mu::MeshPtr ConvertSkeletalMeshToMutable(USkeletalMesh* InSkeletalMesh, int LOD,
 				EnumAddFlags(BoneUsageFlags, mu::EBoneUsageFlags::Root);
 			}
 
-			MutableMesh->SetBonePose(BoneMapIndex, TCHAR_TO_ANSI(*BoneName), BaseInvTransform.Inverse(), BoneUsageFlags);
+			MutableMesh->SetBonePose(BoneMapIndex, StringCast<ANSICHAR>(*BoneName).Get(), BaseInvTransform.Inverse(), BoneUsageFlags);
 		}
 	}
 
@@ -1367,7 +1367,7 @@ mu::MeshPtr ConvertSkeletalMeshToMutable(USkeletalMesh* InSkeletalMesh, int LOD,
 				{
 					RelevantBodySetups.Add(BodySetup);
 					DiscardedBodySetups[BodySetupIndex] = 0;
-					int32 BonePoseIndex = MutableMesh->FindBonePose(TCHAR_TO_ANSI(*BodyBoneName));
+					int32 BonePoseIndex = MutableMesh->FindBonePose(StringCast<ANSICHAR>(*BodyBoneName).Get());
 					
 					EnumAddFlags(MutableMesh->BonePoses[BonePoseIndex].BoneUsageFlags, mu::EBoneUsageFlags::Physics);
 				}
@@ -1426,7 +1426,7 @@ mu::MeshPtr ConvertSkeletalMeshToMutable(USkeletalMesh* InSkeletalMesh, int LOD,
 			TObjectPtr<USkeletalBodySetup>& BodySetup = RelevantBodySetups[B];
 			
 			FString BodyBoneName = BodySetup->BoneName.ToString();
-			PhysicsBody->SetBodyBoneName( B, TCHAR_TO_ANSI(*BodyBoneName) );
+			PhysicsBody->SetBodyBoneName( B, StringCast<ANSICHAR>(*BodyBoneName).Get() );
 			
 			const int32 NumSpheres = BodySetup->AggGeom.SphereElems.Num();
 			PhysicsBody->SetSphereCount( B, NumSpheres );
@@ -1437,7 +1437,7 @@ mu::MeshPtr ConvertSkeletalMeshToMutable(USkeletalMesh* InSkeletalMesh, int LOD,
 				PhysicsBody->SetSphere( B, I, FVector3f(SphereElem.Center), SphereElem.Radius );
 
 				const FString ElemName = SphereElem.GetName().ToString();
-				PhysicsBody->SetSphereName(B, I, TCHAR_TO_ANSI(*ElemName));	
+				PhysicsBody->SetSphereName(B, I, StringCast<ANSICHAR>(*ElemName).Get());	
 				PhysicsBody->SetSphereFlags(B, I, GetKBodyElemFlags(SphereElem));
 			}
 
@@ -1453,7 +1453,7 @@ mu::MeshPtr ConvertSkeletalMeshToMutable(USkeletalMesh* InSkeletalMesh, int LOD,
 						FVector3f(BoxElem.X, BoxElem.Y, BoxElem.Z));
 
 				const FString KElemName = BoxElem.GetName().ToString();
-				PhysicsBody->SetBoxName(B, I, TCHAR_TO_ANSI(*KElemName));
+				PhysicsBody->SetBoxName(B, I, StringCast<ANSICHAR>(*KElemName).Get());
 				PhysicsBody->SetBoxFlags(B, I, GetKBodyElemFlags(BoxElem));
 			}
 
@@ -1478,7 +1478,7 @@ mu::MeshPtr ConvertSkeletalMeshToMutable(USkeletalMesh* InSkeletalMesh, int LOD,
 				PhysicsBody->SetConvexTransform(B, I, FTransform3f(ConvexElem.GetTransform()));
 				
 				const FString KElemName = ConvexElem.GetName().ToString();
-				PhysicsBody->SetConvexName(B, I, TCHAR_TO_ANSI(*KElemName));
+				PhysicsBody->SetConvexName(B, I, StringCast<ANSICHAR>(*KElemName).Get());
 				PhysicsBody->SetConvexFlags(B, I, GetKBodyElemFlags(ConvexElem));
 			}
 
@@ -1494,7 +1494,7 @@ mu::MeshPtr ConvertSkeletalMeshToMutable(USkeletalMesh* InSkeletalMesh, int LOD,
 						SphylElem.Radius, SphylElem.Length );
 
 				const FString KElemName = SphylElem.GetName().ToString();
-				PhysicsBody->SetSphylName(B, I, TCHAR_TO_ANSI(*KElemName));
+				PhysicsBody->SetSphylName(B, I, StringCast<ANSICHAR>(*KElemName).Get());
 				PhysicsBody->SetSphylFlags(B, I, GetKBodyElemFlags(SphylElem));
 			}
 
@@ -1510,7 +1510,7 @@ mu::MeshPtr ConvertSkeletalMeshToMutable(USkeletalMesh* InSkeletalMesh, int LOD,
 						TaperedCapsuleElem.Radius0, TaperedCapsuleElem.Radius1, TaperedCapsuleElem.Length );
 				
 				const FString KElemName = TaperedCapsuleElem.GetName().ToString();
-				PhysicsBody->SetTaperedCapsuleName(B, I, TCHAR_TO_ANSI(*KElemName));
+				PhysicsBody->SetTaperedCapsuleName(B, I, StringCast<ANSICHAR>(*KElemName).Get());
 				PhysicsBody->SetTaperedCapsuleFlags( B, I, GetKBodyElemFlags(TaperedCapsuleElem));
 			}
 		}
@@ -2381,11 +2381,11 @@ mu::NodeMeshPtr GenerateMorphMesh(const UEdGraphPin* Pin,
 		{
 			const FName RowName = TypedNodeTable->GetRowNames()[RowIndex];
 
-			ColumnIndex = Table->FindColumn(TCHAR_TO_ANSI(*ColumnName));
+			ColumnIndex = Table->FindColumn(StringCast<ANSICHAR>(*ColumnName).Get());
 
 			if (ColumnIndex == INDEX_NONE)
 			{
-				ColumnIndex = Table->AddColumn(TCHAR_TO_ANSI(*ColumnName), mu::TABLE_COLUMN_TYPE::TCT_MESH);
+				ColumnIndex = Table->AddColumn(StringCast<ANSICHAR>(*ColumnName).Get(), mu::TABLE_COLUMN_TYPE::TCT_MESH);
 			}
 
 			mu::MeshPtr MorphedSourceTableMesh = BuildMorphedMutableMesh(Pin, TypedNodeMorphs[MorphIndex].MorphTargetName, GenerationContext, RowName);
@@ -2398,8 +2398,8 @@ mu::NodeMeshPtr GenerateMorphMesh(const UEdGraphPin* Pin,
 
 			mu::NodeMeshTablePtr MorphedSourceMeshNodeTable = new mu::NodeMeshTable;
 			MorphedSourceMeshNodeTable->SetTable(Table);
-			MorphedSourceMeshNodeTable->SetColumn(TCHAR_TO_ANSI(*ColumnName));
-			MorphedSourceMeshNodeTable->SetParameterName(TCHAR_TO_ANSI(*TypedNodeTable->ParameterName));
+			MorphedSourceMeshNodeTable->SetColumn(StringCast<ANSICHAR>(*ColumnName).Get());
+			MorphedSourceMeshNodeTable->SetParameterName(StringCast<ANSICHAR>(*TypedNodeTable->ParameterName).Get());
 			MorphedSourceMeshNodeTable->SetMessageContext(MorphNode);
 
 			// A null target will leave the base unchanged
@@ -2460,7 +2460,7 @@ mu::NodeMeshPtr GenerateMorphMesh(const UEdGraphPin* Pin,
 
 						for (const FString& BoneName : BonesToDeform)
 						{
-							Result->AddBoneToDeform(TCHAR_TO_ANSI(*BoneName));
+							Result->AddBoneToDeform(StringCast<ANSICHAR>(*BoneName).Get());
 						}
 					}
 
@@ -2476,7 +2476,7 @@ mu::NodeMeshPtr GenerateMorphMesh(const UEdGraphPin* Pin,
 	
 						for (const FString& PhysicsBoneName : PhysicsToDeform)
 						{
-							Result->AddPhysicsBodyToDeform(TCHAR_TO_ANSI(*PhysicsBoneName));
+							Result->AddPhysicsBodyToDeform(StringCast<ANSICHAR>(*PhysicsBoneName).Get());
 						}
 					}
 						
@@ -3045,7 +3045,7 @@ mu::NodeMeshPtr GenerateMutableSourceMesh(const UEdGraphPin * Pin,
 			const UEdGraphPin* VariationPin = TypedNodeMeshVar->VariationPin(VariationIndex);
 			if (!VariationPin) continue;
 
-			MeshNode->SetVariationTag(VariationIndex, TCHAR_TO_ANSI(*TypedNodeMeshVar->Variations[VariationIndex].Tag));
+			MeshNode->SetVariationTag(VariationIndex, StringCast<ANSICHAR>(*TypedNodeMeshVar->Variations[VariationIndex].Tag).Get());
 			if (const UEdGraphPin* ConnectedPin = FollowInputPin(*VariationPin))
 			{
 				FMutableGraphMeshGenerationData VariationMeshData;
@@ -3166,7 +3166,7 @@ mu::NodeMeshPtr GenerateMutableSourceMesh(const UEdGraphPin * Pin,
 				
 				for (const FString& BoneName : BonesToDeform)
 				{
-					MeshNode->AddBoneToDeform(TCHAR_TO_ANSI(*BoneName));
+					MeshNode->AddBoneToDeform(StringCast<ANSICHAR>(*BoneName).Get());
 				}
 			}
 
@@ -3181,7 +3181,7 @@ mu::NodeMeshPtr GenerateMutableSourceMesh(const UEdGraphPin * Pin,
 
 				for (const FString& PhysicsBoneName : PhysicsToDeform)
 				{
-					MeshNode->AddPhysicsBodyToDeform(TCHAR_TO_ANSI(*PhysicsBoneName));
+					MeshNode->AddPhysicsBodyToDeform(StringCast<ANSICHAR>(*PhysicsBoneName).Get());
 				}	
 			}
 			
@@ -3337,14 +3337,14 @@ mu::NodeMeshPtr GenerateMutableSourceMesh(const UEdGraphPin * Pin,
 			FString DataTableColumnName = TypedNodeTable->GetColumnNameByPin(Pin);
 
 			// Generating a new Mesh column if not exists
-			if (Table && Table->FindColumn(TCHAR_TO_ANSI(*MutableColumnName)) == INDEX_NONE)
+			if (Table && Table->FindColumn(StringCast<ANSICHAR>(*MutableColumnName).Get()) == INDEX_NONE)
 			{
 				GenerateTableColumn(TypedNodeTable, Pin, Table, DataTableColumnName, CurrentLOD, GenerationContext);
 			}
 
 			MeshTableNode->SetTable(Table);
-			MeshTableNode->SetColumn(TCHAR_TO_ANSI(*MutableColumnName));
-			MeshTableNode->SetParameterName(TCHAR_TO_ANSI(*TypedNodeTable->ParameterName));
+			MeshTableNode->SetColumn(StringCast<ANSICHAR>(*MutableColumnName).Get());
+			MeshTableNode->SetParameterName(StringCast<ANSICHAR>(*TypedNodeTable->ParameterName).Get());
 
 			GenerationContext.AddParameterNameUnique(Node, TypedNodeTable->ParameterName);
 
@@ -3411,7 +3411,7 @@ mu::NodeMeshPtr GenerateMutableSourceMesh(const UEdGraphPin * Pin,
 				MorphResult = GenerateMorphMesh(Pin, GenerationContext.MeshMorphStack, 0, Result, GenerationContext, MeshData, MutableColumnName);
 			}
 
-			if (Table->FindColumn(TCHAR_TO_ANSI(*MutableColumnName)) == INDEX_NONE)
+			if (Table->FindColumn(StringCast<ANSICHAR>(*MutableColumnName).Get()) == INDEX_NONE)
 			{
 				FString Msg = FString::Printf(TEXT("Couldn't find pin column with name %s"), *MutableColumnName);
 				GenerationContext.Compiler->CompilerLog(FText::FromString(Msg), Node);
diff --git a/Engine/Plugins/Experimental/Mutable/Source/CustomizableObjectEditor/Private/MuCOE/GenerateMutableSource/GenerateMutableSourceModifier.cpp b/Engine/Plugins/Experimental/Mutable/Source/CustomizableObjectEditor/Private/MuCOE/GenerateMutableSource/GenerateMutableSourceModifier.cpp
--- a/Engine/Plugins/Experimental/Mutable/Source/CustomizableObjectEditor/Private/MuCOE/GenerateMutableSource/GenerateMutableSourceModifier.cpp
+++ b/Engine/Plugins/Experimental/Mutable/Source/CustomizableObjectEditor/Private/MuCOE/GenerateMutableSource/GenerateMutableSourceModifier.cpp
@@ -58,11 +58,11 @@ mu::NodeModifierPtr GenerateMutableSourceModifier(const UEdGraphPin * Pin, FMuta
 		ClipNode->SetParams(TypedNodeClip->B, TypedNodeClip->Exponent);
 		ClipNode->SetMorphEllipse(TypedNodeClip->Radius, TypedNodeClip->Radius2, TypedNodeClip->RotationAngle);
 
-		ClipNode->SetVertexSelectionBone(TCHAR_TO_ANSI(*TypedNodeClip->BoneName.ToString()), TypedNodeClip->MaxEffectRadius);
+		ClipNode->SetVertexSelectionBone(StringCast<ANSICHAR>(*TypedNodeClip->BoneName.ToString()).Get(), TypedNodeClip->MaxEffectRadius);
 
 		for (const FString& Tag : TypedNodeClip->Tags)
 		{
-			 ClipNode->AddTag(TCHAR_TO_ANSI(*Tag));
+			 ClipNode->AddTag(StringCast<ANSICHAR>(*Tag).Get());
 		}
 	}
 
@@ -100,7 +100,7 @@ mu::NodeModifierPtr GenerateMutableSourceModifier(const UEdGraphPin * Pin, FMuta
 	
 		for (const FString& Tag : TypedNodeClipDeform->Tags)
 		{
-			ClipNode->AddTag(TCHAR_TO_ANSI(*Tag));
+			ClipNode->AddTag(StringCast<ANSICHAR>(*Tag).Get());
 		}		
 	}
 
@@ -165,7 +165,7 @@ mu::NodeModifierPtr GenerateMutableSourceModifier(const UEdGraphPin * Pin, FMuta
 
 		for (const FString& Tag : TypedNodeClipMesh->Tags)
 		{
-			 ClipNode->AddTag(TCHAR_TO_ANSI(*Tag));
+			 ClipNode->AddTag(StringCast<ANSICHAR>(*Tag).Get());
 		}
 
 		if (TypedNodeClipMesh->CustomizableObjectToClipWith != nullptr)
diff --git a/Engine/Plugins/Experimental/Mutable/Source/CustomizableObjectEditor/Private/MuCOE/GenerateMutableSource/GenerateMutableSourceProjector.cpp b/Engine/Plugins/Experimental/Mutable/Source/CustomizableObjectEditor/Private/MuCOE/GenerateMutableSource/GenerateMutableSourceProjector.cpp
--- a/Engine/Plugins/Experimental/Mutable/Source/CustomizableObjectEditor/Private/MuCOE/GenerateMutableSource/GenerateMutableSourceProjector.cpp
+++ b/Engine/Plugins/Experimental/Mutable/Source/CustomizableObjectEditor/Private/MuCOE/GenerateMutableSource/GenerateMutableSourceProjector.cpp
@@ -85,8 +85,8 @@ mu::NodeProjectorPtr GenerateMutableSourceProjector(const UEdGraphPin* Pin, FMut
 
 		GenerationContext.AddParameterNameUnique(Node, TypedNodeParam->ParameterName);
 
-		ProjectorNode->SetName(TCHAR_TO_ANSI(*TypedNodeParam->ParameterName));
-		ProjectorNode->SetUid(TCHAR_TO_ANSI(*GenerationContext.GetNodeIdUnique(Node).ToString()));
+		ProjectorNode->SetName(StringCast<ANSICHAR>(*TypedNodeParam->ParameterName).Get());
+		ProjectorNode->SetUid(StringCast<ANSICHAR>(*GenerationContext.GetNodeIdUnique(Node).ToString()).Get());
 		switch ((int)TypedNodeParam->DefaultValue.ProjectionType)
 		{
 		case 0:
diff --git a/Engine/Plugins/Experimental/Mutable/Source/CustomizableObjectEditor/Private/MuCOE/GenerateMutableSource/GenerateMutableSourceSurface.cpp b/Engine/Plugins/Experimental/Mutable/Source/CustomizableObjectEditor/Private/MuCOE/GenerateMutableSource/GenerateMutableSourceSurface.cpp
--- a/Engine/Plugins/Experimental/Mutable/Source/CustomizableObjectEditor/Private/MuCOE/GenerateMutableSource/GenerateMutableSourceSurface.cpp
+++ b/Engine/Plugins/Experimental/Mutable/Source/CustomizableObjectEditor/Private/MuCOE/GenerateMutableSource/GenerateMutableSourceSurface.cpp
@@ -387,7 +387,7 @@ mu::NodeSurfacePtr GenerateMutableSourceSurface(const UEdGraphPin * Pin, FMutabl
 				mu::TablePtr Table = GenerateMutableSourceTable(TypedNodeTable->Table->GetName(), ConnectedPin, GenerationContext);
 				
 				// Generating a column for each modified parameter(texture, color & float) if not exists
-				if (Table && Table->FindColumn(TCHAR_TO_ANSI(*ColumnName)) == INDEX_NONE)
+				if (Table && Table->FindColumn(StringCast<ANSICHAR>(*ColumnName).Get()) == INDEX_NONE)
 				{
 					GenerateTableColumn(TypedNodeTable, ConnectedPin, Table, ColumnName, GenerationContext.CurrentLOD, GenerationContext);
 				}
@@ -719,10 +719,10 @@ mu::NodeSurfacePtr GenerateMutableSourceSurface(const UEdGraphPin * Pin, FMutabl
 						SurfNodeImageName += "-MutableLayerParam:" + FString::FromInt(LayerIndex);
 					}
 
-					SurfNode->SetImageName(ImageIndex, TCHAR_TO_ANSI(*SurfNodeImageName));
+					SurfNode->SetImageName(ImageIndex, StringCast<ANSICHAR>(*SurfNodeImageName).Get());
 					int32 UVLayout = TypedNodeMat->GetImageUVLayout(ImageIndex);
 					SurfNode->SetImageLayoutIndex(ImageIndex, UVLayout);
-					SurfNode->SetImageAdditionalNames(ImageIndex, TCHAR_TO_ANSI(*TypedNodeMat->Material->GetName()), TCHAR_TO_ANSI(*ImageName));
+					SurfNode->SetImageAdditionalNames(ImageIndex, StringCast<ANSICHAR>(*TypedNodeMat->Material->GetName()).Get(), StringCast<ANSICHAR>(*ImageName).Get());
 
 					if (bShareProjectionTexturesBetweenLODs && bIsGroupProjectorImage)
 					{
@@ -740,7 +740,7 @@ mu::NodeSurfacePtr GenerateMutableSourceSurface(const UEdGraphPin * Pin, FMutabl
 				ensure(LOD > GenerationContext.FirstLODAvailable);
 				check(ProjectorInfo->SurfNode->GetImage(ImageIndex) == ProjectorInfo->ImageNode);
 				SurfNode->SetImage(ImageIndex, ProjectorInfo->ImageNode);
-				SurfNode->SetImageName(ImageIndex, TCHAR_TO_ANSI(*ProjectorInfo->TextureName));
+				SurfNode->SetImageName(ImageIndex, StringCast<ANSICHAR>(*ProjectorInfo->TextureName).Get());
 				SurfNode->SetImageLayoutIndex(ImageIndex, ProjectorInfo->UVLayout);
 
 				TextureNameToProjectionResFactor.Add(ProjectorInfo->RealTextureName, ProjectorInfo->AlternateProjectionResolutionFactor);
@@ -785,7 +785,7 @@ mu::NodeSurfacePtr GenerateMutableSourceSurface(const UEdGraphPin * Pin, FMutabl
 					}
 
 					SurfNode->SetVector(VectorIndex, ColorNode);
-					SurfNode->SetVectorName(VectorIndex, TCHAR_TO_ANSI(*VectorName));
+					SurfNode->SetVectorName(VectorIndex, StringCast<ANSICHAR>(*VectorName).Get());
 				}
 			}
 		}
@@ -827,14 +827,14 @@ mu::NodeSurfacePtr GenerateMutableSourceSurface(const UEdGraphPin * Pin, FMutabl
 					}
 
 					SurfNode->SetScalar(ScalarIndex, ScalarNode);
-					SurfNode->SetScalarName(ScalarIndex, TCHAR_TO_ANSI(*ScalarName));
+					SurfNode->SetScalarName(ScalarIndex, StringCast<ANSICHAR>(*ScalarName).Get());
 				}
 			}
 		}
 
 		for (const FString& Tag : TypedNodeMat->Tags)
 		{
-			SurfNode->AddTag(TCHAR_TO_ANSI(*Tag));
+			SurfNode->AddTag(StringCast<ANSICHAR>(*Tag).Get());
 		}
 
 		TArray<mu::NodeSurfaceNewPtr>* ArraySurfaceNodePtr = GenerationContext.MapMaterialNodeToMutableSurfaceNodeArray.Find(TypedNodeMat);
@@ -877,7 +877,7 @@ mu::NodeSurfacePtr GenerateMutableSourceSurface(const UEdGraphPin * Pin, FMutabl
 
 			for (const FString& Tag : TypedNodeMat->Tags)
 			{
-				SurfNode2->AddTag(TCHAR_TO_ANSI(*Tag));
+				SurfNode2->AddTag(StringCast<ANSICHAR>(*Tag).Get());
 			}
 
 			SurfNode2->SetImageCount(SurfNode->GetImageCount());
@@ -927,7 +927,7 @@ mu::NodeSurfacePtr GenerateMutableSourceSurface(const UEdGraphPin * Pin, FMutabl
 			mu::NodeSurfaceVariationPtr SurfaceVariation = new mu::NodeSurfaceVariation;
 			SurfaceVariation->SetVariationType(mu::NodeSurfaceVariation::VariationType::State);
 			SurfaceVariation->SetVariationCount(1);
-			SurfaceVariation->SetVariationTag(0, TCHAR_TO_ANSI(*AlternateResStateName));
+			SurfaceVariation->SetVariationTag(0, StringCast<ANSICHAR>(*AlternateResStateName).Get());
 
 			SurfaceVariation->AddDefaultSurface(&*SurfNode);
 			SurfaceVariation->AddVariationSurface(0, &*SurfNode2);
@@ -1086,7 +1086,7 @@ mu::NodeSurfacePtr GenerateMutableSourceSurface(const UEdGraphPin * Pin, FMutabl
 		
 			for (const FString& Tag : TypedNodeExt->Tags)
 			{
-				SurfNode->AddTag(TCHAR_TO_ANSI(*Tag));
+				SurfNode->AddTag(StringCast<ANSICHAR>(*Tag).Get());
 			}
 		}();
 	}
@@ -1437,7 +1437,7 @@ mu::NodeSurfacePtr GenerateMutableSourceSurface(const UEdGraphPin * Pin, FMutabl
 
 			if (UEdGraphPin* VariationPin = TypedNodeVar->VariationPin(VariationIndex))
 			{
-				SurfNode->SetVariationTag(VariationIndex, TCHAR_TO_ANSI(*TypedNodeVar->Variations[VariationIndex].Tag));
+				SurfNode->SetVariationTag(VariationIndex, StringCast<ANSICHAR>(*TypedNodeVar->Variations[VariationIndex].Tag).Get());
 				for (const UEdGraphPin* ConnectedPin : FollowInputPinArray(*VariationPin))
 				{
 					// Is it a modifier?
diff --git a/Engine/Plugins/Experimental/Mutable/Source/CustomizableObjectEditor/Private/MuCOE/GenerateMutableSource/GenerateMutableSourceTable.cpp b/Engine/Plugins/Experimental/Mutable/Source/CustomizableObjectEditor/Private/MuCOE/GenerateMutableSource/GenerateMutableSourceTable.cpp
--- a/Engine/Plugins/Experimental/Mutable/Source/CustomizableObjectEditor/Private/MuCOE/GenerateMutableSource/GenerateMutableSourceTable.cpp
+++ b/Engine/Plugins/Experimental/Mutable/Source/CustomizableObjectEditor/Private/MuCOE/GenerateMutableSource/GenerateMutableSourceTable.cpp
@@ -179,11 +179,11 @@ void FillTableColumn(const UCustomizableObjectNodeTable* TableNode,	mu::TablePtr
 			{
 				FString MutableColumnName = ColumnName + FString::Printf(TEXT(" LOD_%d "), LOD) + FString::Printf(TEXT("Mat_%d"), MatIndex);
 
-				CurrentColumn = MutableTable.get()->FindColumn(TCHAR_TO_ANSI(*MutableColumnName));
+				CurrentColumn = MutableTable.get()->FindColumn(StringCast<ANSICHAR>(*MutableColumnName).Get());
 
 				if (CurrentColumn == -1)
 				{
-					CurrentColumn = MutableTable->AddColumn(TCHAR_TO_ANSI(*MutableColumnName), mu::TABLE_COLUMN_TYPE::TCT_MESH);
+					CurrentColumn = MutableTable->AddColumn(StringCast<ANSICHAR>(*MutableColumnName).Get(), mu::TABLE_COLUMN_TYPE::TCT_MESH);
 				}
 
 				// First process the mesh tags that are going to make the mesh unique and affect whether it's repeated in 
@@ -291,11 +291,11 @@ void FillTableColumn(const UCustomizableObjectNodeTable* TableNode,	mu::TablePtr
 			{
 				FString MutableColumnName = ColumnName + FString::Printf(TEXT(" LOD_%d "), LOD) + FString::Printf(TEXT("Mat_%d"), MatIndex);
 
-				CurrentColumn = MutableTable.get()->FindColumn(TCHAR_TO_ANSI(*MutableColumnName));
+				CurrentColumn = MutableTable.get()->FindColumn(StringCast<ANSICHAR>(*MutableColumnName).Get());
 
 				if (CurrentColumn == -1)
 				{
-					CurrentColumn = MutableTable->AddColumn(TCHAR_TO_ANSI(*MutableColumnName), mu::TABLE_COLUMN_TYPE::TCT_MESH);
+					CurrentColumn = MutableTable->AddColumn(StringCast<ANSICHAR>(*MutableColumnName).Get(), mu::TABLE_COLUMN_TYPE::TCT_MESH);
 				}
 
 				mu::MeshPtr MutableMesh = GenerateMutableMesh(StaticMesh, CurrentLOD, MatIndex, FString(), GenerationContext, TableNode);
@@ -317,11 +317,11 @@ void FillTableColumn(const UCustomizableObjectNodeTable* TableNode,	mu::TablePtr
 		else if (UTexture2D* Texture = Cast<UTexture2D>(Object))
 		{
 			// Getting column index from column name
-			CurrentColumn = MutableTable->FindColumn(TCHAR_TO_ANSI(*ColumnName));
+			CurrentColumn = MutableTable->FindColumn(StringCast<ANSICHAR>(*ColumnName).Get());
 
 			if (CurrentColumn == INDEX_NONE)
 			{
-				CurrentColumn = MutableTable->AddColumn(TCHAR_TO_ANSI(*ColumnName), mu::TABLE_COLUMN_TYPE::TCT_IMAGE);
+				CurrentColumn = MutableTable->AddColumn(StringCast<ANSICHAR>(*ColumnName).Get(), mu::TABLE_COLUMN_TYPE::TCT_IMAGE);
 			}
 
 			GenerationContext.ArrayTextureUnrealToMutableTask.Add(FTextureUnrealToMutableTask(MutableTable, Texture, TableNode, CurrentColumn, RowIdx));
@@ -330,9 +330,9 @@ void FillTableColumn(const UCustomizableObjectNodeTable* TableNode,	mu::TablePtr
 		else if (UMaterialInstance* Material = Cast<UMaterialInstance>(Object))
 		{
 			//Adding an empty column for searching purposes
-			if (MutableTable.get()->FindColumn(TCHAR_TO_ANSI(*ColumnName)) == -1)
+			if (MutableTable.get()->FindColumn(StringCast<ANSICHAR>(*ColumnName).Get()) == -1)
 			{
-				CurrentColumn = MutableTable->AddColumn(TCHAR_TO_ANSI(*ColumnName), mu::TABLE_COLUMN_TYPE::TCT_NONE);
+				CurrentColumn = MutableTable->AddColumn(StringCast<ANSICHAR>(*ColumnName).Get(), mu::TABLE_COLUMN_TYPE::TCT_NONE);
 			}
 
 			UMaterialInstance* ReferenceMaterial = TableNode->GetColumnDefaultAssetByType<UMaterialInstance>(ColumnName);
@@ -384,12 +384,12 @@ void FillTableColumn(const UCustomizableObjectNodeTable* TableNode,	mu::TablePtr
 						}
 
 						// Getting column index from parameter name
-						ColumnIndex = MutableTable->FindColumn(TCHAR_TO_ANSI(*ParameterGuid));
+						ColumnIndex = MutableTable->FindColumn(StringCast<ANSICHAR>(*ParameterGuid).Get());
 
 						if (ColumnIndex == INDEX_NONE)
 						{
 							// If there is no column with the parameters name, we generate a new one
-							ColumnIndex = MutableTable->AddColumn(TCHAR_TO_ANSI(*ParameterGuid), mu::TABLE_COLUMN_TYPE::TCT_IMAGE);
+							ColumnIndex = MutableTable->AddColumn(StringCast<ANSICHAR>(*ParameterGuid).Get(), mu::TABLE_COLUMN_TYPE::TCT_IMAGE);
 						}
 
 						UTexture* OutTexture = nullptr;
@@ -461,12 +461,12 @@ void FillTableColumn(const UCustomizableObjectNodeTable* TableNode,	mu::TablePtr
 					}
 					
 					// Getting column index from parameter name
-					ColumnIndex = MutableTable->FindColumn(TCHAR_TO_ANSI(*ParameterGuid));
+					ColumnIndex = MutableTable->FindColumn(StringCast<ANSICHAR>(*ParameterGuid).Get());
 
 					if (ColumnIndex == INDEX_NONE)
 					{
 						// If there is no column with the parameters name, we generate a new one
-						ColumnIndex = MutableTable->AddColumn(TCHAR_TO_ANSI(*ParameterGuid), mu::TABLE_COLUMN_TYPE::TCT_COLOUR);
+						ColumnIndex = MutableTable->AddColumn(StringCast<ANSICHAR>(*ParameterGuid).Get(), mu::TABLE_COLUMN_TYPE::TCT_COLOUR);
 					}
 
 					// Getting the parameter value from the parent material
@@ -527,12 +527,12 @@ void FillTableColumn(const UCustomizableObjectNodeTable* TableNode,	mu::TablePtr
 					}
 
 					// Getting column index from parameter name
-					ColumnIndex = MutableTable->FindColumn(TCHAR_TO_ANSI(*ParameterGuid));
+					ColumnIndex = MutableTable->FindColumn(StringCast<ANSICHAR>(*ParameterGuid).Get());
 
 					if (ColumnIndex == INDEX_NONE)
 					{
 						// If there is no column with the parameters name, we generate a new one
-						ColumnIndex = MutableTable->AddColumn(TCHAR_TO_ANSI(*ParameterGuid), mu::TABLE_COLUMN_TYPE::TCT_SCALAR);
+						ColumnIndex = MutableTable->AddColumn(StringCast<ANSICHAR>(*ParameterGuid).Get(), mu::TABLE_COLUMN_TYPE::TCT_SCALAR);
 					}
 
 					// Getting the parameter value from the parent material
@@ -578,11 +578,11 @@ void FillTableColumn(const UCustomizableObjectNodeTable* TableNode,	mu::TablePtr
 	{
 		if (StructProperty->Struct == TBaseStructure<FLinearColor>::Get())
 		{
-			CurrentColumn = MutableTable->FindColumn(TCHAR_TO_ANSI(*ColumnName));
+			CurrentColumn = MutableTable->FindColumn(StringCast<ANSICHAR>(*ColumnName).Get());
 
 			if (CurrentColumn == INDEX_NONE)
 			{
-				CurrentColumn = MutableTable->AddColumn(TCHAR_TO_ANSI(*ColumnName), mu::TABLE_COLUMN_TYPE::TCT_COLOUR);
+				CurrentColumn = MutableTable->AddColumn(StringCast<ANSICHAR>(*ColumnName).Get(), mu::TABLE_COLUMN_TYPE::TCT_COLOUR);
 			}
 
 			// Setting cell value
@@ -593,11 +593,11 @@ void FillTableColumn(const UCustomizableObjectNodeTable* TableNode,	mu::TablePtr
 
 	else if (const FNumericProperty* FloatNumProperty = CastField<FFloatProperty>(Property))
 	{
-		CurrentColumn = MutableTable->FindColumn(TCHAR_TO_ANSI(*ColumnName));
+		CurrentColumn = MutableTable->FindColumn(StringCast<ANSICHAR>(*ColumnName).Get());
 
 		if (CurrentColumn == INDEX_NONE)
 		{
-			CurrentColumn = MutableTable->AddColumn(TCHAR_TO_ANSI(*ColumnName), mu::TABLE_COLUMN_TYPE::TCT_SCALAR);
+			CurrentColumn = MutableTable->AddColumn(StringCast<ANSICHAR>(*ColumnName).Get(), mu::TABLE_COLUMN_TYPE::TCT_SCALAR);
 		}
 
 		// Setting cell value
@@ -607,11 +607,11 @@ void FillTableColumn(const UCustomizableObjectNodeTable* TableNode,	mu::TablePtr
 
 	else if (const FNumericProperty* DoubleNumProperty = CastField<FDoubleProperty>(Property))
 	{
-		CurrentColumn = MutableTable->FindColumn(TCHAR_TO_ANSI(*ColumnName));
+		CurrentColumn = MutableTable->FindColumn(StringCast<ANSICHAR>(*ColumnName).Get());
 	
 		if (CurrentColumn == INDEX_NONE)
 		{
-			CurrentColumn = MutableTable->AddColumn(TCHAR_TO_ANSI(*ColumnName), mu::TABLE_COLUMN_TYPE::TCT_SCALAR);
+			CurrentColumn = MutableTable->AddColumn(StringCast<ANSICHAR>(*ColumnName).Get(), mu::TABLE_COLUMN_TYPE::TCT_SCALAR);
 		}
 	
 		// Setting cell value
@@ -726,7 +726,7 @@ mu::TablePtr GenerateMutableSourceTable(const FString& TableName, const UEdGraph
 			{
 				MutableTable->AddRow(i);
 				FString RowName= RowNames[i].ToString();
-				MutableTable->SetCell(0, i, TCHAR_TO_ANSI(*RowName));
+				MutableTable->SetCell(0, i, StringCast<ANSICHAR>(*RowName).Get());
 				ParameterUIData.ArrayIntegerParameterOption.Add(FIntegerParameterUIData(
 					RowName,
 					FMutableParamUIMetadata()));
diff --git a/Engine/Plugins/Experimental/Mutable/Source/CustomizableObjectEditor/Private/MuCOE/Nodes/CustomizableObjectNodeLayoutBlocksDetails.cpp b/Engine/Plugins/Experimental/Mutable/Source/CustomizableObjectEditor/Private/MuCOE/Nodes/CustomizableObjectNodeLayoutBlocksDetails.cpp
--- a/Engine/Plugins/Experimental/Mutable/Source/CustomizableObjectEditor/Private/MuCOE/Nodes/CustomizableObjectNodeLayoutBlocksDetails.cpp
+++ b/Engine/Plugins/Experimental/Mutable/Source/CustomizableObjectEditor/Private/MuCOE/Nodes/CustomizableObjectNodeLayoutBlocksDetails.cpp
@@ -193,7 +193,7 @@ void FCustomizableObjectNodeLayoutBlocksDetails::OnGridComboBoxSelectionChanged(
 		mu::Model* Model = CustomInstance->CustomizableObject->GetModel();
 
 		CustomInstance->PreEditChange(NULL);
-		CustomInstance->State = Model->FindState( TCHAR_TO_ANSI(**Selection) );
+		CustomInstance->State = Model->FindState( StringCast<ANSICHAR>(**Selection).Get() );
 		CustomInstance->PostEditChange();
 
 		ResetParamBox();
diff --git a/Engine/Plugins/Experimental/Mutable/Source/CustomizableObjectEditor/Private/MuCOE/SMutableCodeViewer.cpp b/Engine/Plugins/Experimental/Mutable/Source/CustomizableObjectEditor/Private/MuCOE/SMutableCodeViewer.cpp
--- a/Engine/Plugins/Experimental/Mutable/Source/CustomizableObjectEditor/Private/MuCOE/SMutableCodeViewer.cpp
+++ b/Engine/Plugins/Experimental/Mutable/Source/CustomizableObjectEditor/Private/MuCOE/SMutableCodeViewer.cpp
@@ -149,7 +149,7 @@ public:
 			{
 				mu::OP::ParameterArgs Args = Program.GetOpArgs<mu::OP::ParameterArgs>(RowItem->MutableOperation);
 				OpName += TEXT(" ");
-				OpName += ANSI_TO_TCHAR(Program.m_parameters[int32(Args.variable)].m_name.c_str());
+				OpName += StringCast<TCHAR>(Program.m_parameters[int32(Args.variable)].m_name.c_str()).Get();
 				break;
 			}
 
@@ -157,7 +157,7 @@ public:
 			{
 				mu::OP::ImageSwizzleArgs Args = Program.GetOpArgs<mu::OP::ImageSwizzleArgs>(RowItem->MutableOperation);
 				OpName += TEXT(" ");
-				OpName += ANSI_TO_TCHAR(mu::TypeInfo::s_imageFormatName[int32(Args.format)]);
+				OpName += StringCast<TCHAR>(mu::TypeInfo::s_imageFormatName[int32(Args.format)]).Get();
 				break;
 			}
 
@@ -165,9 +165,9 @@ public:
 			{
 				mu::OP::ImagePixelFormatArgs Args = Program.GetOpArgs<mu::OP::ImagePixelFormatArgs>(RowItem->MutableOperation);
 				OpName += TEXT(" ");
-				OpName += ANSI_TO_TCHAR(mu::TypeInfo::s_imageFormatName[int32(Args.format)]);
+				OpName += StringCast<TCHAR>(mu::TypeInfo::s_imageFormatName[int32(Args.format)]).Get();
 				OpName += TEXT(" or ");
-				OpName += ANSI_TO_TCHAR(mu::TypeInfo::s_imageFormatName[int32(Args.formatIfAlpha)]);
+				OpName += StringCast<TCHAR>(mu::TypeInfo::s_imageFormatName[int32(Args.formatIfAlpha)]).Get();
 				break;
 			}
 
@@ -511,7 +511,7 @@ void SMutableCodeViewer::Construct(const FArguments& InArgs, const TSharedPtr<mu
 					{
 						const FString SaveFileName = FString(SaveFilenames[0]);
 
-						mu::OutputFileStream Stream(TCHAR_TO_ANSI(*SaveFileName));
+						mu::OutputFileStream Stream(StringCast<ANSICHAR>(*SaveFileName).Get());
 						Stream.Write(MUTABLE_COMPILED_MODEL_FILETAG, 4);
 						mu::OutputArchive Archive(&Stream);
 						mu::Model::Serialise(InMutableModel.Get(), Archive);
@@ -1218,7 +1218,7 @@ void SMutableCodeViewer::GenerateAllTreeElements()
 	for ( uint32 StateIndex=0; StateIndex<StateCount; ++StateIndex )
 	{
 		const mu::FProgram::FState& State = ModelPrivate->m_program.m_states[StateIndex];
-		FString Caption = FString::Printf( TEXT("state [%s]"), ANSI_TO_TCHAR(State.m_name.c_str()) );
+		FString Caption = FString::Printf( TEXT("state [%s]"), StringCast<TCHAR>(State.m_name.c_str()).Get() );
 
 		const FSlateColor LabelColor = ColorPerComputationalCost[StaticCast<uint8>(GetOperationTypeComputationalCost(
 			ModelPrivate->m_program.GetOpType(State.m_root)))];
@@ -2652,7 +2652,7 @@ FReply SMutableCodeViewer::OnDragOver(const FGeometry& MyGeometry, const FDragDr
 				if (DraggedFileExtension == TEXT(".mutable_compiled"))
 				{
 					// Dump source model to a file.
-					mu::InputFileStream stream(TCHAR_TO_ANSI(*Files[0]));
+					mu::InputFileStream stream(StringCast<ANSICHAR>(*Files[0]).Get());
 
 					char MutableSourceTag[4] = {};
 					stream.Read(MutableSourceTag, 4);
@@ -2688,7 +2688,7 @@ FReply SMutableCodeViewer::OnDrop(const FGeometry& MyGeometry, const FDragDropEv
 				if (DraggedFileExtension == TEXT(".mutable_compiled"))
 				{
 					// Read a mutable compiled model file.
-					mu::InputFileStream stream(TCHAR_TO_ANSI(*Files[0]));
+					mu::InputFileStream stream(StringCast<ANSICHAR>(*Files[0]).Get());
 
 					char MutableSourceTag[4] = {};
 					stream.Read(MutableSourceTag, 4);
diff --git a/Engine/Plugins/Experimental/Mutable/Source/CustomizableObjectEditor/Private/MuCOE/SMutableGraphViewer.cpp b/Engine/Plugins/Experimental/Mutable/Source/CustomizableObjectEditor/Private/MuCOE/SMutableGraphViewer.cpp
--- a/Engine/Plugins/Experimental/Mutable/Source/CustomizableObjectEditor/Private/MuCOE/SMutableGraphViewer.cpp
+++ b/Engine/Plugins/Experimental/Mutable/Source/CustomizableObjectEditor/Private/MuCOE/SMutableGraphViewer.cpp
@@ -44,10 +44,10 @@ public:
 
 		const char* TypeName = RowItem->MutableNode->GetType()->m_strName;
 
-		FText MainLabel = FText::FromString(ANSI_TO_TCHAR(TypeName));
+		FText MainLabel = FText::FromString(StringCast<TCHAR>(TypeName).Get());
 		if (RowItem->DuplicatedOf)
 		{
-			MainLabel = FText::FromString( FString::Printf(TEXT("%s (Duplicated)"), ANSI_TO_TCHAR(TypeName)));
+			MainLabel = FText::FromString( FString::Printf(TEXT("%s (Duplicated)"), StringCast<TCHAR>(TypeName).Get()));
 		}
 
 		// TODO
@@ -153,7 +153,7 @@ void SMutableGraphViewer::Construct(const FArguments& InArgs, const mu::NodePtr&
 
 					// Dump source model to a file.
 					FString SaveFileName = FString(SaveFilenames[0]);
-					mu::OutputFileStream stream(TCHAR_TO_ANSI(*SaveFileName));
+					mu::OutputFileStream stream(StringCast<ANSICHAR>(*SaveFileName).Get());
 					stream.Write(MUTABLE_SOURCE_MODEL_FILETAG, 4);
 					mu::OutputArchive arch(&stream);
 					mu::Node::Serialise(InRootNode.get(), arch);
@@ -460,7 +460,7 @@ FReply SMutableGraphViewer::OnDragOver(const FGeometry& MyGeometry, const FDragD
 				if (DraggedFileExtension == TEXT(".mutable_source"))
 				{
 					// Dump source model to a file.
-					mu::InputFileStream stream(TCHAR_TO_ANSI(*Files[0]));
+					mu::InputFileStream stream(StringCast<ANSICHAR>(*Files[0]).Get());
 
 					char MutableSourceTag[4] = {};
 					stream.Read(MutableSourceTag, 4);
@@ -496,7 +496,7 @@ FReply SMutableGraphViewer::OnDrop(const FGeometry& MyGeometry, const FDragDropE
 				if (DraggedFileExtension == TEXT(".mutable_source"))
 				{
 					// Dump source model to a file.
-					mu::InputFileStream stream(TCHAR_TO_ANSI(*Files[0]));
+					mu::InputFileStream stream(StringCast<ANSICHAR>(*Files[0]).Get());
 
 					char MutableSourceTag[4] = {};
 					stream.Read(MutableSourceTag, 4);
diff --git a/Engine/Plugins/Experimental/Mutable/Source/CustomizableObjectEditor/Private/MuCOE/SMutableImageViewer.cpp b/Engine/Plugins/Experimental/Mutable/Source/CustomizableObjectEditor/Private/MuCOE/SMutableImageViewer.cpp
--- a/Engine/Plugins/Experimental/Mutable/Source/CustomizableObjectEditor/Private/MuCOE/SMutableImageViewer.cpp
+++ b/Engine/Plugins/Experimental/Mutable/Source/CustomizableObjectEditor/Private/MuCOE/SMutableImageViewer.cpp
@@ -285,7 +285,7 @@ FText SMutableImageViewer::GetImageDescriptionLabel() const
 
 	mu::EImageFormat Format = MutableImage->GetFormat();
 	FString Label = FString::Printf(TEXT("%s - %d x %d - %d LODs"), 
-		ANSI_TO_TCHAR(mu::TypeInfo::s_imageFormatName[(size_t)Format]),
+		StringCast<TCHAR>(mu::TypeInfo::s_imageFormatName[(size_t)Format]).Get(),
 		MutableImage->GetSizeX(), 
 		MutableImage->GetSizeY(),
 		MutableImage->GetLODCount() );
diff --git a/Engine/Plugins/Experimental/Mutable/Source/CustomizableObjectEditor/Private/MuCOE/SMutableParametersWidget.cpp b/Engine/Plugins/Experimental/Mutable/Source/CustomizableObjectEditor/Private/MuCOE/SMutableParametersWidget.cpp
--- a/Engine/Plugins/Experimental/Mutable/Source/CustomizableObjectEditor/Private/MuCOE/SMutableParametersWidget.cpp
+++ b/Engine/Plugins/Experimental/Mutable/Source/CustomizableObjectEditor/Private/MuCOE/SMutableParametersWidget.cpp
@@ -59,7 +59,7 @@ void SMutableParametersWidget::Tick(const FGeometry& AllottedGeometry, const dou
 
 	for ( int32 ParamIndex=0; ParamIndex< MutableParameters->GetCount(); ++ParamIndex )
 	{
-		FString ParamName = ANSI_TO_TCHAR(MutableParameters->GetName(ParamIndex));
+		FString ParamName = StringCast<TCHAR>(MutableParameters->GetName(ParamIndex)).Get();
 		TSharedPtr<SHorizontalBox> ParameterBox;
 
 		ParamBox->AddSlot()
@@ -139,7 +139,7 @@ void SMutableParametersWidget::Tick(const FGeometry& AllottedGeometry, const dou
 				for (int i = 0; i < ValueCount; ++i)
 				{
 					const char* ValueText = MutableParameters->GetIntPossibleValueName( ParamIndex, i );
-					OptionNamesAttribute.Add( FString(ANSI_TO_TCHAR(ValueText)) );				
+					OptionNamesAttribute.Add( FString(StringCast<TCHAR>(ValueText).Get()) );
 				}
 
 				ParameterBox->AddSlot()
@@ -373,8 +373,7 @@ void SMutableParametersWidget::OnIntParameterChanged(int32 InValue, int32 ParamI
 		// Update the text combo if any
 		if (Combo)
 		{
-			FString Text = TEXT("Invalid");
-			Text = ANSI_TO_TCHAR(MutableParameters->GetIntPossibleValueName(ParamIndex, InValue));
+			const FString Text = StringCast<TCHAR>(MutableParameters->GetIntPossibleValueName(ParamIndex, InValue)).Get();
 			Combo->SetText(FText::FromString(Text));
 		}
 
@@ -393,7 +392,7 @@ void SMutableParametersWidget::OnIntParameterTextChanged(TSharedPtr<FString> Sel
 		return;
 	}
 
-	int32 ValueIndex = MutableParameters->GetIntValueIndex(ParamIndex, TCHAR_TO_ANSI(**Selection));
+	int32 ValueIndex = MutableParameters->GetIntValueIndex(ParamIndex, StringCast<ANSICHAR>(**Selection).Get());
 	if (ValueIndex >= 0 && ValueIndex < MutableParameters->GetIntPossibleValueCount(ParamIndex))
 	{
 		int32 RealValue = MutableParameters->GetIntPossibleValue(ParamIndex, ValueIndex);
diff --git a/Engine/Plugins/Experimental/Mutable/Source/CustomizableObjectPopulationEditor/Private/MuCOPE/CustomizableObjectPopulationEditorModule.cpp b/Engine/Plugins/Experimental/Mutable/Source/CustomizableObjectPopulationEditor/Private/MuCOPE/CustomizableObjectPopulationEditorModule.cpp
--- a/Engine/Plugins/Experimental/Mutable/Source/CustomizableObjectPopulationEditor/Private/MuCOPE/CustomizableObjectPopulationEditorModule.cpp
+++ b/Engine/Plugins/Experimental/Mutable/Source/CustomizableObjectPopulationEditor/Private/MuCOPE/CustomizableObjectPopulationEditorModule.cpp
@@ -49,14 +49,14 @@ IMPLEMENT_MODULE( FCustomizableObjectPopulationEditorModule, CustomizableObjectP
 static void LogWarning(const char* msg)
 {
 	// Can be called from any thread.
-	UE_LOG(LogMutable, Warning, TEXT("%s"), ANSI_TO_TCHAR(msg) );
+	UE_LOG(LogMutable, Warning, TEXT("%s"), StringCast<TCHAR>(msg).Get() );
 }
 
 
 static void LogError(const char* msg)
 {
 	// Can be called from any thread.
-	UE_LOG(LogMutable, Error, TEXT("%s"), ANSI_TO_TCHAR(msg));
+	UE_LOG(LogMutable, Error, TEXT("%s"), StringCast<TCHAR>(msg).Get());
 }
 
 
diff --git a/Engine/Plugins/Experimental/Mutable/Source/MutableTools/Private/MuT/CodeGenerator.cpp b/Engine/Plugins/Experimental/Mutable/Source/MutableTools/Private/MuT/CodeGenerator.cpp
--- a/Engine/Plugins/Experimental/Mutable/Source/MutableTools/Private/MuT/CodeGenerator.cpp
+++ b/Engine/Plugins/Experimental/Mutable/Source/MutableTools/Private/MuT/CodeGenerator.cpp
@@ -1305,7 +1305,7 @@ namespace mu
 
 							// It is already visible in the error log
 							//UE_LOG( LogMutableCore, Warning, TEXT("A texture [%s] has been resized from [%d x %d] to [%d x %d] because it didn't fit the layout [%d x %d]. "),
-							//	ANSI_TO_TCHAR(node.m_images[t].m_name.c_str()), oldSize[0], oldSize[1], desc.m_size[0], desc.m_size[1], grid[0], grid[1] );
+							//	StringCast<TCHAR>(node.m_images[t].m_name.c_str()).Get(), oldSize[0], oldSize[1], desc.m_size[0], desc.m_size[1], grid[0], grid[1] );
 
 							int currentLOD = m_currentParents.Last().m_lod;
 							char buf[256];
diff --git a/Engine/Plugins/Experimental/Mutable/Source/MutableTools/Private/MuT/CodeGenerator_Mesh.cpp b/Engine/Plugins/Experimental/Mutable/Source/MutableTools/Private/MuT/CodeGenerator_Mesh.cpp
--- a/Engine/Plugins/Experimental/Mutable/Source/MutableTools/Private/MuT/CodeGenerator_Mesh.cpp
+++ b/Engine/Plugins/Experimental/Mutable/Source/MutableTools/Private/MuT/CodeGenerator_Mesh.cpp
@@ -807,7 +807,7 @@ class Node;
 				const char* Aux = nullptr;
 				Aux = tag.c_str();
                 m_pErrorLog->GetPrivate()->Add( 
-					FString::Printf(TEXT("Unknown tag found in mesh variation [%s]."), ANSI_TO_TCHAR(Aux)),
+					FString::Printf(TEXT("Unknown tag found in mesh variation [%s]."), StringCast<TCHAR>(Aux).Get()),
 					ELMT_WARNING, 
 					node.m_errorContext
 				);
diff --git a/Engine/Plugins/Experimental/Mutable/Source/MutableTools/Private/MuT/Compiler.cpp b/Engine/Plugins/Experimental/Mutable/Source/MutableTools/Private/MuT/Compiler.cpp
--- a/Engine/Plugins/Experimental/Mutable/Source/MutableTools/Private/MuT/Compiler.cpp
+++ b/Engine/Plugins/Experimental/Mutable/Source/MutableTools/Private/MuT/Compiler.cpp
@@ -265,8 +265,8 @@ namespace mu
 						"The state [%s] refers to a parameter [%s] "
 						"that has not been found in the model. This error can be "
 						"safely dismissed in case of partial compilation."), 
-						ANSI_TO_TCHAR(s.nodeState.m_name.c_str()),
-						ANSI_TO_TCHAR(s.nodeState.m_runtimeParams[p].c_str()));
+						StringCast<TCHAR>(s.nodeState.m_name.c_str()).Get(),
+						StringCast<TCHAR>(s.nodeState.m_runtimeParams[p].c_str()).Get());
                     m_pD->m_pErrorLog->GetPrivate()->Add(Temp, ELMT_WARNING, pNode->GetBasePrivate()->m_errorContext );
                 }
             }
