#rnx

diff --git a/Engine/Plugins/Experimental/Mutable/Source/CustomizableObjectEditor/Private/MuCOE/Nodes/CustomizableObjectNodeSkeletalMesh.cpp b/Engine/Plugins/Experimental/Mutable/Source/CustomizableObjectEditor/Private/MuCOE/Nodes/CustomizableObjectNodeSkeletalMesh.cpp
--- a/Engine/Plugins/Experimental/Mutable/Source/CustomizableObjectEditor/Private/MuCOE/Nodes/CustomizableObjectNodeSkeletalMesh.cpp
+++ b/Engine/Plugins/Experimental/Mutable/Source/CustomizableObjectEditor/Private/MuCOE/Nodes/CustomizableObjectNodeSkeletalMesh.cpp
@@ -667,6 +667,7 @@ UMaterialInterface* UCustomizableObjectNodeSkeletalMesh::GetMaterialInterfaceFor
 	return nullptr;
 }
 
+
 FSkeletalMaterial* UCustomizableObjectNodeSkeletalMesh::GetSkeletalMaterialFor(const int LODIndex, const int MaterialIndex, const FSkeletalMeshModel* ImportedModel) const
 {
 	if (!SkeletalMesh)
@@ -675,14 +676,23 @@ FSkeletalMaterial* UCustomizableObjectNodeSkeletalMesh::GetSkeletalMaterialFor(c
 	}
 
 	// We assume that LODIndex and MaterialIndex are valid for the imported model
-	int SkeletalMeshMaterialIndex;
-	if (ImportedModel)
+	int32 SkeletalMeshMaterialIndex = -1;
+
+	if (ImportedModel && ImportedModel->LODModels.IsValidIndex(LODIndex) && ImportedModel->LODModels[LODIndex].Sections.IsValidIndex(MaterialIndex))
 	{
 		SkeletalMeshMaterialIndex = ImportedModel->LODModels[LODIndex].Sections[MaterialIndex].MaterialIndex;
 	}
 	else
 	{
-		SkeletalMeshMaterialIndex = Helper_GetImportedModel(SkeletalMesh)->LODModels[LODIndex].Sections[MaterialIndex].MaterialIndex;
+		FSkeletalMeshModel* AuxImportedModel = Helper_GetImportedModel(SkeletalMesh);
+
+		if (AuxImportedModel)
+		{
+			if (AuxImportedModel->LODModels.IsValidIndex(LODIndex) && AuxImportedModel->LODModels[LODIndex].Sections.IsValidIndex(MaterialIndex))
+			{
+				SkeletalMeshMaterialIndex = AuxImportedModel->LODModels[LODIndex].Sections[MaterialIndex].MaterialIndex;
+			}
+		}
 	}
 	
 	// Check if we have lod info map to get the correct material index
@@ -702,6 +712,7 @@ FSkeletalMaterial* UCustomizableObjectNodeSkeletalMesh::GetSkeletalMaterialFor(c
 	return nullptr;
 }
 
+
 //SGraphNode --------------------------------------------
 
 void SGraphNodeSkeletalMesh::Construct(const FArguments& InArgs, UEdGraphNode* InGraphNode)
