[Mutable] Empty the Streaming CacheCore once the generation has completed.

- Renamed b.MutableStreamingMemory CVAR.

#rnx

diff --git a/Engine/Plugins/Experimental/Mutable/Source/CustomizableObject/Classes/MuCO/CustomizableObjectSystem.h b/Engine/Plugins/Experimental/Mutable/Source/CustomizableObject/Classes/MuCO/CustomizableObjectSystem.h
--- a/Engine/Plugins/Experimental/Mutable/Source/CustomizableObject/Classes/MuCO/CustomizableObjectSystem.h
+++ b/Engine/Plugins/Experimental/Mutable/Source/CustomizableObject/Classes/MuCO/CustomizableObjectSystem.h
@@ -66,6 +66,9 @@ DECLARE_DWORD_ACCUMULATOR_STAT_EXTERN(TEXT("Streaming Ops"), STAT_MutableStreami
 DECLARE_DWORD_ACCUMULATOR_STAT_EXTERN(TEXT("Streaming Cache"), STAT_MutableStreamingCache, STATGROUP_Mutable, );
 
 
+extern TAutoConsoleVariable<bool> CVarClearStreamingCacheOnUpdateEnd;
+
+
 #if WITH_EDITOR
 
 // Struct used to keep a copy of the EditorSettings needed to compile Customizable Objects.
diff --git a/Engine/Plugins/Experimental/Mutable/Source/CustomizableObject/Private/MuCO/CustomizableObjectMipDataProvider.cpp b/Engine/Plugins/Experimental/Mutable/Source/CustomizableObject/Private/MuCO/CustomizableObjectMipDataProvider.cpp
--- a/Engine/Plugins/Experimental/Mutable/Source/CustomizableObject/Private/MuCO/CustomizableObjectMipDataProvider.cpp
+++ b/Engine/Plugins/Experimental/Mutable/Source/CustomizableObject/Private/MuCO/CustomizableObjectMipDataProvider.cpp
@@ -134,6 +134,11 @@ namespace impl
 				MUTABLE_CPUPROFILER_SCOPE(EndUpdate);
 				System->EndUpdate(InstanceID);
 				System->ReleaseInstance(InstanceID);
+
+				if (CVarClearStreamingCacheOnUpdateEnd.GetValueOnAnyThread())
+				{
+					System->ClearStreamingCache();
+				}
 			}
 
 			{
diff --git a/Engine/Plugins/Experimental/Mutable/Source/CustomizableObject/Private/MuCO/CustomizableObjectSystem.cpp b/Engine/Plugins/Experimental/Mutable/Source/CustomizableObject/Private/MuCO/CustomizableObjectSystem.cpp
--- a/Engine/Plugins/Experimental/Mutable/Source/CustomizableObject/Private/MuCO/CustomizableObjectSystem.cpp
+++ b/Engine/Plugins/Experimental/Mutable/Source/CustomizableObject/Private/MuCO/CustomizableObjectSystem.cpp
@@ -78,12 +78,19 @@ UCustomizableObjectSystem* FCustomizableObjectSystemPrivate::SSystem = nullptr;
 
 
 static TAutoConsoleVariable<int32> CVarStreamingMemory(
-	TEXT("b.MutableStreamingMemory"),
+	TEXT("mutable.StreamingCacheMemory"),
 	-1,
 	TEXT("If different than 0, limit the amount of memory (in KB) to use to cache streaming data when building characters. 0 means no limit, -1 means use default (40Mb in PC and 20Mb in consoles)."),
 	ECVF_Scalability);
 
 
+TAutoConsoleVariable<bool> CVarClearStreamingCacheOnUpdateEnd(
+	TEXT("mutable.ClearStreamingCahceOnUpdateEnd"),
+	false,
+	TEXT("Clear the Streaming Cache when releasing an Instance."),
+	ECVF_Scalability);
+
+
 bool FMutableQueue::IsEmpty() const
 {
 	return Array.Num() == 0;
@@ -777,8 +784,6 @@ void UpdateSkeletalMesh(UCustomizableObjectInstance* CustomizableObjectInstance,
 
 
 #if WITH_EDITOR
-			CustomizableObjectInstance->InstanceUpdated = true;
-
 			// \TODO: Review if we can avoid all this editor code here.
 			for (int32 MeshIndex = 0; MeshIndex < CustomizableObjectInstance->SkeletalMeshes.Num(); ++MeshIndex)
 			{
@@ -1709,6 +1714,11 @@ namespace impl
 				OperationData->InstanceID = 0;
 			}
 		}
+
+		if (CVarClearStreamingCacheOnUpdateEnd.GetValueOnAnyThread())
+		{
+			MutableSystem->ClearStreamingCache();
+		}
 	}
 
 
@@ -1724,6 +1734,11 @@ namespace impl
 			MutableSystem->ReleaseInstance(OperationData->InstanceID);
 			OperationData->InstanceID = 0;
 		}
+
+		if (CVarClearStreamingCacheOnUpdateEnd.GetValueOnAnyThread())
+		{
+			MutableSystem->ClearStreamingCache();			
+		}
 	}
 
 
@@ -2704,10 +2719,6 @@ void FMutableOperation::MutableIsDisabledCase()
 
 	// We must invalidate the current DescriptorRuntimeHash, since we're discarding everything.
 	CustomizableObjectInstance->Updated(EUpdateResult::Success, FDescriptorRuntimeHash());
-
-#if WITH_EDITOR
-	CustomizableObjectInstance->InstanceUpdated = true;
-#endif
 }
 
 
diff --git a/Engine/Plugins/Experimental/Mutable/Source/MutableRuntime/Private/MuR/System.cpp b/Engine/Plugins/Experimental/Mutable/Source/MutableRuntime/Private/MuR/System.cpp
--- a/Engine/Plugins/Experimental/Mutable/Source/MutableRuntime/Private/MuR/System.cpp
+++ b/Engine/Plugins/Experimental/Mutable/Source/MutableRuntime/Private/MuR/System.cpp
@@ -147,7 +147,16 @@ namespace mu
         m_pD->SetStreamingCache( bytes );
     }
 
+	
+	//---------------------------------------------------------------------------------------------
+	void System::ClearStreamingCache()
+    {
+		LLM_SCOPE_BYNAME(TEXT("MutableRuntime"));
 
+        m_pD->ClearStreamingCache();
+   }
+
+	
     //---------------------------------------------------------------------------------------------
     void System::SetImageParameterGenerator( ImageParameterGenerator* pInterface )
     {
@@ -617,6 +626,24 @@ namespace mu
 	}
 
 
+	//---------------------------------------------------------------------------------------------
+	void System::Private::ClearStreamingCache()
+	{
+		for (FModelCache::FModelCacheEntry& ModelCache : m_modelCache.m_cachePerModel)
+	    {
+		    if (const TSharedPtr<const Model> CacheModel = ModelCache.m_pModel.Pin())
+			{
+				FProgram& Program = CacheModel->GetPrivate()->m_program;
+
+				for (int32 RomIndex=0; RomIndex < Program.m_roms.Num(); ++RomIndex)
+				{
+					Program.UnloadRom(RomIndex);		
+    			}
+			}
+ 		}
+	}
+	
+
 	//---------------------------------------------------------------------------------------------
 	void System::Private::BeginBuild(const TSharedPtr<const Model>& pModel)
 	{
diff --git a/Engine/Plugins/Experimental/Mutable/Source/MutableRuntime/Private/MuR/SystemPrivate.h b/Engine/Plugins/Experimental/Mutable/Source/MutableRuntime/Private/MuR/SystemPrivate.h
--- a/Engine/Plugins/Experimental/Mutable/Source/MutableRuntime/Private/MuR/SystemPrivate.h
+++ b/Engine/Plugins/Experimental/Mutable/Source/MutableRuntime/Private/MuR/SystemPrivate.h
@@ -974,8 +974,8 @@ namespace mu
 		MUTABLERUNTIME_API Ptr<const Mesh> BuildMesh(const TSharedPtr<const Model>&, const Parameters* pParams, OP::ADDRESS at) ;
 		MUTABLERUNTIME_API Ptr<const Layout> BuildLayout(const TSharedPtr<const Model>&, const Parameters* pParams, OP::ADDRESS at) ;
     	MUTABLERUNTIME_API Ptr<const Projector> BuildProjector(const TSharedPtr<const Model>&, const Parameters* pParams, OP::ADDRESS at) ;
-		void ClearCache() ;
 		void SetStreamingCache(uint64 bytes) ;
+		void ClearStreamingCache() ;
 
         //!
         SettingsPtrConst m_pSettings;
diff --git a/Engine/Plugins/Experimental/Mutable/Source/MutableRuntime/Public/MuR/System.h b/Engine/Plugins/Experimental/Mutable/Source/MutableRuntime/Public/MuR/System.h
--- a/Engine/Plugins/Experimental/Mutable/Source/MutableRuntime/Public/MuR/System.h
+++ b/Engine/Plugins/Experimental/Mutable/Source/MutableRuntime/Public/MuR/System.h
@@ -116,6 +116,9 @@ namespace mu
         //! Refer to Settings::SetStreamingCache for more information.
         void SetStreamingCache( uint64 bytes );
 
+        //! Clears the streaming cache content. It does not modify the memory limit.
+		void ClearStreamingCache();
+		
         //! Set a new provider for external image data. This is only necessary if image parameters
         //! are used in the models.
 		//! The provided pointer becomes owned by the system.
