[mutable] Initial implementation of code execution heuristic to minimize memory usage.
Extended option that controls per-state texture compression, to support uncompressed textures when no runtime parameter is used.
Removed deprecated operation "MeshSubtract".
Removed unused NodeObjectState.
#rnx

diff --git a/Engine/Plugins/Experimental/Mutable/Source/CustomizableObject/Classes/MuCO/CustomizableObject.h b/Engine/Plugins/Experimental/Mutable/Source/CustomizableObject/Classes/MuCO/CustomizableObject.h
--- a/Engine/Plugins/Experimental/Mutable/Source/CustomizableObject/Classes/MuCO/CustomizableObject.h
+++ b/Engine/Plugins/Experimental/Mutable/Source/CustomizableObject/Classes/MuCO/CustomizableObject.h
@@ -1365,7 +1365,7 @@ private:
 	// This is a manual version number for the binary blobs in this asset.
 	// Increasing it invalidates all the previously compiled models.
 	// Warning: If while merging code both versions have changed, take the highest+1.
-	static const int32 CurrentSupportedVersion = 372;
+	static const int32 CurrentSupportedVersion = 374;
 
 public:
 
diff --git a/Engine/Plugins/Experimental/Mutable/Source/CustomizableObject/CustomizableObject.Build.cs b/Engine/Plugins/Experimental/Mutable/Source/CustomizableObject/CustomizableObject.Build.cs
--- a/Engine/Plugins/Experimental/Mutable/Source/CustomizableObject/CustomizableObject.Build.cs
+++ b/Engine/Plugins/Experimental/Mutable/Source/CustomizableObject/CustomizableObject.Build.cs
@@ -30,7 +30,6 @@ public class CustomizableObject : ModuleRules
 			"ClothingSystemRuntimeInterface",
 			//"ClothingSystemEditor",
 			"UMG",
-
 			"MutableRuntime",
 		});
 
@@ -40,7 +39,8 @@ public class CustomizableObject : ModuleRules
             "Engine",
 			"SkeletalMerging",
 			"ClothingSystemRuntimeCommon",
-            "GameplayTags"
+            "GameplayTags",
+			"MutableRuntime",
 		});
 
         PrivateIncludePathModuleNames.AddRange(
diff --git a/Engine/Plugins/Experimental/Mutable/Source/CustomizableObject/Private/MuCO/CustomizableObjectInstance.cpp b/Engine/Plugins/Experimental/Mutable/Source/CustomizableObject/Private/MuCO/CustomizableObjectInstance.cpp
--- a/Engine/Plugins/Experimental/Mutable/Source/CustomizableObject/Private/MuCO/CustomizableObjectInstance.cpp
+++ b/Engine/Plugins/Experimental/Mutable/Source/CustomizableObject/Private/MuCO/CustomizableObjectInstance.cpp
@@ -2601,8 +2601,8 @@ void CopyTextureProperties(UTexture2D* Texture, const UTexture2D* SourceTexture)
 // If assigned to a UTexture2D, it will be freed by that UTexture2D
 FTexturePlatformData* UCustomizableInstancePrivateData::MutableCreateImagePlatformData(mu::Ptr<const mu::Image> MutableImage, int32 OnlyLOD, uint16 FullSizeX, uint16 FullSizeY)
 {
-	int32 SizeX = FMath::Max(MutableImage->GetSize().x(), FullSizeX);
-	int32 SizeY = FMath::Max(MutableImage->GetSize().y(), FullSizeY);
+	int32 SizeX = FMath::Max(MutableImage->GetSize()[0], FullSizeX);
+	int32 SizeY = FMath::Max(MutableImage->GetSize()[1], FullSizeY);
 
 	if (SizeX <= 0 || SizeY <= 0)
 	{
diff --git a/Engine/Plugins/Experimental/Mutable/Source/CustomizableObject/Private/MuCO/CustomizableObjectSystem.cpp b/Engine/Plugins/Experimental/Mutable/Source/CustomizableObject/Private/MuCO/CustomizableObjectSystem.cpp
--- a/Engine/Plugins/Experimental/Mutable/Source/CustomizableObject/Private/MuCO/CustomizableObjectSystem.cpp
+++ b/Engine/Plugins/Experimental/Mutable/Source/CustomizableObject/Private/MuCO/CustomizableObjectSystem.cpp
@@ -821,7 +821,9 @@ void FCustomizableObjectSystemPrivate::InitUpdateSkeletalMesh(UCustomizableObjec
 	
 	const FString CurrentState = Instance.GetCurrentState();
 	const FParameterUIData* State = Instance.GetCustomizableObject()->StateUIDataMap.Find(CurrentState);
-	const bool bNeverStream = State ? State->bDontCompressRuntimeTextures : false;
+
+	// \TODO: This should be controllable independently
+	const bool bNeverStream = State ? State->TextureCompressionStrategy!=ETextureCompressionStrategy::None : false;
 	const bool bUseMipmapStreaming = !bNeverStream;
 	int32 MipsToSkip = 0; // 0 means all mips
 
@@ -1235,18 +1237,13 @@ namespace impl
 		// Main instance generation step
 		{
 			// LOD mask, set to all ones to build  all LODs
-			uint32 LODMask = 0xFFFFFFFF;
+			uint32 LODMask = mu::System::AllLODs;
 
 			Instance = System->BeginUpdate(OperationData->InstanceID, MutableParameters, State, LODMask);
 
 			if (!Instance)
 			{
-				const mu::Error MutableError = mu::GetError();
-				if (MutableError == mu::Error::Unsupported)
-				{
-					UE_LOG(LogMutable, Log, TEXT("The necessary functionality is not supported in this version of Mutable"));
-				}
-
+				UE_LOG(LogMutable, Log, TEXT("An Instace update has failed."));
 				return;
 			}
 		}
diff --git a/Engine/Plugins/Experimental/Mutable/Source/CustomizableObject/Public/MuCO/CustomizableObjectUIData.h b/Engine/Plugins/Experimental/Mutable/Source/CustomizableObject/Public/MuCO/CustomizableObjectUIData.h
--- a/Engine/Plugins/Experimental/Mutable/Source/CustomizableObject/Public/MuCO/CustomizableObjectUIData.h
+++ b/Engine/Plugins/Experimental/Mutable/Source/CustomizableObject/Public/MuCO/CustomizableObjectUIData.h
@@ -5,6 +5,8 @@
 #include "Engine/Texture2D.h"
 #include "MuCO/CustomizableObjectParameterTypeDefinitions.h"
 #include "MuCO/CustomizableObjectCustomVersion.h"
+#include "MuR/System.h"
+
 #include "CustomizableObjectUIData.generated.h"
 
 class UTexture2D;
@@ -256,8 +258,14 @@ struct FParameterUIData
 	UPROPERTY(BlueprintReadWrite, Category = UI)
 	ECustomizableObjectGroupType IntegerParameterGroupType;
 
-	UPROPERTY(BlueprintReadWrite, Category = CustomizableObject)
-	bool bDontCompressRuntimeTextures = false; // Only useful for state metadata
+	/** Replaced by TextureCompressionStrategy */
+	UPROPERTY()
+	bool bDontCompressRuntimeTextures_DEPRECATED = false; // Only useful for state metadata
+
+	/** Not really relevant for UI, but apparently bDontCompressRuntimeTextures_DEPRECATED was used to decided some texture properties at runtime.
+	* TODO: Try to handle the NeverStream without going through this.
+	*/
+	ETextureCompressionStrategy TextureCompressionStrategy = ETextureCompressionStrategy::None;
 
 	/** In this mode instances and their temp data will be reused between updates. It will be much faster but spend as much as ten times the memory.
 	    Useful for customization lockers with few characters that are going to have their parameters changed many times, not for in-game */
@@ -291,7 +299,7 @@ struct FParameterUIData
 		Ar << UIData.Type;
 		Ar << UIData.ArrayIntegerParameterOption;
 		Ar << UIData.IntegerParameterGroupType;
-		Ar << UIData.bDontCompressRuntimeTextures;
+		Ar << UIData.TextureCompressionStrategy;
 		Ar << UIData.bLiveUpdateMode;
 		Ar << UIData.bReuseInstanceTextures;
 		Ar << UIData.ForcedParameterValues;
diff --git a/Engine/Plugins/Experimental/Mutable/Source/CustomizableObjectEditor/Classes/MuCOE/Nodes/CustomizableObjectNodeObject.h b/Engine/Plugins/Experimental/Mutable/Source/CustomizableObjectEditor/Classes/MuCOE/Nodes/CustomizableObjectNodeObject.h
--- a/Engine/Plugins/Experimental/Mutable/Source/CustomizableObjectEditor/Classes/MuCOE/Nodes/CustomizableObjectNodeObject.h
+++ b/Engine/Plugins/Experimental/Mutable/Source/CustomizableObjectEditor/Classes/MuCOE/Nodes/CustomizableObjectNodeObject.h
@@ -3,6 +3,7 @@
 #pragma once
 
 #include "MuCOE/Nodes/CustomizableObjectNode.h"
+#include "MuR/System.h"
 
 #include "CustomizableObjectNodeObject.generated.h"
 
@@ -27,8 +28,13 @@ struct CUSTOMIZABLEOBJECTEDITOR_API FCustomizableObjectState
 	UPROPERTY(EditAnywhere, Category=CustomizableObject)
 	TArray<FString> RuntimeParameters;
 
+	/** This is now TextureCompressionStrategy.  */
+	UPROPERTY()
+	bool bDontCompressRuntimeTextures_DEPRECATED = false;
+		
+	/** Special treatment of texture compression for this state. */
 	UPROPERTY(EditAnywhere, Category = CustomizableObject)
-	bool bDontCompressRuntimeTextures = false;
+	ETextureCompressionStrategy TextureCompressionStrategy = ETextureCompressionStrategy::None;
 
 	/** LiveUpdateMode will reuse instance temp. data between updates and speed up update times, but spend much more memory. Good for customization screens, not for actual gameplay modes. */
 	UPROPERTY(EditAnywhere, Category = CustomizableObject)
@@ -153,6 +159,7 @@ public:
 
 	// UObject interface.
 	void PostEditChangeProperty(FPropertyChangedEvent& PropertyChangedEvent) override;
+	void PostLoad() override;
 
 	// EdGraphNode interface
 	FText GetNodeTitle(ENodeTitleType::Type TitleType) const override;
diff --git a/Engine/Plugins/Experimental/Mutable/Source/CustomizableObjectEditor/Private/MuCOE/GenerateMutableSource/GenerateMutableSource.cpp b/Engine/Plugins/Experimental/Mutable/Source/CustomizableObjectEditor/Private/MuCOE/GenerateMutableSource/GenerateMutableSource.cpp
--- a/Engine/Plugins/Experimental/Mutable/Source/CustomizableObjectEditor/Private/MuCOE/GenerateMutableSource/GenerateMutableSource.cpp
+++ b/Engine/Plugins/Experimental/Mutable/Source/CustomizableObjectEditor/Private/MuCOE/GenerateMutableSource/GenerateMutableSource.cpp
@@ -604,7 +604,7 @@ mu::NodeObjectPtr GenerateMutableSource(const UEdGraphPin * Pin, FMutableGraphGe
 			bStartedComponent = true;
 		}
 
-		mu::NodeObjectNewPtr ObjectNode = new mu::NodeObjectNew();
+		mu::Ptr<mu::NodeObjectNew> ObjectNode = new mu::NodeObjectNew();
 		Result = ObjectNode;
 
 		ObjectNode->SetName(StringCast<ANSICHAR>(*TypedNodeObj->ObjectName).Get());
@@ -716,11 +716,11 @@ mu::NodeObjectPtr GenerateMutableSource(const UEdGraphPin * Pin, FMutableGraphGe
 				NumExtraLODsToBuildAfterFirstLOD = *AuxNumExtraLODs;
 			}
 
-			ObjectNode->SetStateProperties(StateIndex, State.bDontCompressRuntimeTextures, State.bBuildOnlyFirstLOD, GenerationContext.FirstLODAvailable, NumExtraLODsToBuildAfterFirstLOD);
+			ObjectNode->SetStateProperties(StateIndex, State.TextureCompressionStrategy, State.bBuildOnlyFirstLOD, GenerationContext.FirstLODAvailable, NumExtraLODsToBuildAfterFirstLOD);
 
 			// UI Data
 			FParameterUIData ParameterUIData(State.Name, State.StateUIMetadata, EMutableParameterType::None);
-			ParameterUIData.bDontCompressRuntimeTextures = State.bDontCompressRuntimeTextures;
+			ParameterUIData.TextureCompressionStrategy = State.TextureCompressionStrategy;
 			ParameterUIData.bLiveUpdateMode = State.bLiveUpdateMode;
 			ParameterUIData.bReuseInstanceTextures = State.bReuseInstanceTextures;
 			ParameterUIData.ForcedParameterValues = State.ForcedParameterValues;
diff --git a/Engine/Plugins/Experimental/Mutable/Source/CustomizableObjectEditor/Private/MuCOE/Nodes/CustomizableObjectNodeObject.cpp b/Engine/Plugins/Experimental/Mutable/Source/CustomizableObjectEditor/Private/MuCOE/Nodes/CustomizableObjectNodeObject.cpp
--- a/Engine/Plugins/Experimental/Mutable/Source/CustomizableObjectEditor/Private/MuCOE/Nodes/CustomizableObjectNodeObject.cpp
+++ b/Engine/Plugins/Experimental/Mutable/Source/CustomizableObjectEditor/Private/MuCOE/Nodes/CustomizableObjectNodeObject.cpp
@@ -34,6 +34,23 @@ UCustomizableObjectNodeObject::UCustomizableObjectNodeObject()
 }
 
 
+void UCustomizableObjectNodeObject::PostLoad()
+{
+	Super::PostLoad();
+
+	for (FCustomizableObjectState& State : States)
+	{
+		if (State.TextureCompressionStrategy== ETextureCompressionStrategy::None
+			&&
+			State.bDontCompressRuntimeTextures_DEPRECATED)
+		{
+			State.bDontCompressRuntimeTextures_DEPRECATED = false;
+			State.TextureCompressionStrategy = ETextureCompressionStrategy::DontCompressRuntime;
+		}
+	}
+}
+
+
 void UCustomizableObjectNodeObject::PostEditChangeProperty(FPropertyChangedEvent& PropertyChangedEvent)
 {
 	if (!Identifier.IsValid())
diff --git a/Engine/Plugins/Experimental/Mutable/Source/CustomizableObjectEditor/Private/MuCOE/SMutableConstantsWidget.cpp b/Engine/Plugins/Experimental/Mutable/Source/CustomizableObjectEditor/Private/MuCOE/SMutableConstantsWidget.cpp
--- a/Engine/Plugins/Experimental/Mutable/Source/CustomizableObjectEditor/Private/MuCOE/SMutableConstantsWidget.cpp
+++ b/Engine/Plugins/Experimental/Mutable/Source/CustomizableObjectEditor/Private/MuCOE/SMutableConstantsWidget.cpp
@@ -1288,7 +1288,7 @@ void SMutableConstantsWidget::LoadConstantMatrices()
 		ConstantMatrixElement->Matrix = &(MutableProgramPtr->m_constantMatrices[MatrixIndex]);
 		ConstantMatrixElement->IndexOnSourceVector = MatrixIndex;
 		
-		ConstantMatrixElement->Matrix->Serialise(Archive);
+		Archive << *ConstantMatrixElement->Matrix;
 		
 		ConstantMatrixElements.Add(ConstantMatrixElement);
 	}
diff --git a/Engine/Plugins/Experimental/Mutable/Source/MutableRuntime/MutableRuntime.Build.cs b/Engine/Plugins/Experimental/Mutable/Source/MutableRuntime/MutableRuntime.Build.cs
--- a/Engine/Plugins/Experimental/Mutable/Source/MutableRuntime/MutableRuntime.Build.cs
+++ b/Engine/Plugins/Experimental/Mutable/Source/MutableRuntime/MutableRuntime.Build.cs
@@ -16,7 +16,8 @@ namespace UnrealBuildTool.Rules
 
 			PublicDependencyModuleNames.AddRange(
 				new string[] {
-				"Core"
+				"Core",
+				"CoreUObject",
 				}
 			);
 
diff --git a/Engine/Plugins/Experimental/Mutable/Source/MutableRuntime/Private/MuR/CodeRunner.cpp b/Engine/Plugins/Experimental/Mutable/Source/MutableRuntime/Private/MuR/CodeRunner.cpp
--- a/Engine/Plugins/Experimental/Mutable/Source/MutableRuntime/Private/MuR/CodeRunner.cpp
+++ b/Engine/Plugins/Experimental/Mutable/Source/MutableRuntime/Private/MuR/CodeRunner.cpp
@@ -59,7 +59,6 @@
 #include "MuR/OpMeshRemapIndices.h"
 #include "MuR/OpMeshRemoveChart.h"
 #include "MuR/OpMeshReshape.h"
-#include "MuR/OpMeshSubtract.h"
 #include "MuR/OpMeshTransform.h"
 #include "MuR/OpMeshOptimizeSkinning.h"
 #include "MuR/Operations.h"
@@ -78,13 +77,19 @@ namespace mu
 {
 
     //---------------------------------------------------------------------------------------------
-    CodeRunner::CodeRunner(const SettingsPtrConst& pSettings,
-		class System::Private* s,
-		const TSharedPtr<const Model>& pModel,
-		const Parameters* pParams,
+    CodeRunner::CodeRunner(const SettingsPtrConst& InSettings,
+		class System::Private* InSystem,
+		EExecutionStrategy InExecutionStrategy,
+		const TSharedPtr<const Model>& InModel,
+		const Parameters* InParams,
 		OP::ADDRESS at,
-		uint32 lodMask, uint8 executionOptions, int32 InImageLOD, FScheduledOp::EType Type )
-		: m_pSettings(pSettings), m_pSystem(s), m_pModel(pModel), m_pParams(pParams), m_lodMask(lodMask)
+		uint32 InLodMask, uint8 executionOptions, int32 InImageLOD, FScheduledOp::EType Type )
+		: m_pSettings(InSettings)
+		, ExecutionStrategy(InExecutionStrategy)
+		, m_pSystem(InSystem)
+		, m_pModel(InModel)
+		, m_pParams(InParams)
+		, m_lodMask(InLodMask)
 	{
 		ScheduledStagePerOp.resize(m_pModel->GetPrivate()->m_program.m_opAddress.Num());
 
@@ -1514,37 +1519,6 @@ namespace mu
             break;
         }
 
-        case OP_TYPE::ME_SUBTRACT:
-        {
-			OP::MeshSubtractArgs args = pModel->GetPrivate()->m_program.GetOpArgs<OP::MeshSubtractArgs>(item.At);
-            switch (item.Stage)
-            {
-            case 0:
-                    AddOp( FScheduledOp( item.At, item, 1),
-                           FScheduledOp( args.a, item),
-                           FScheduledOp( args.b, item) );
-                break;
-
-            case 1:
-            {
-           		MUTABLE_CPUPROFILER_SCOPE(ME_SUBTRACT_1)
-            	
-                Ptr<const Mesh> pA = GetMemory().GetMesh( FCacheAddress(args.a,item) );
-                Ptr<const Mesh> pB = GetMemory().GetMesh( FCacheAddress(args.b,item) );
-
-                MeshPtr pResult = MeshSubtract( pA.get(), pB.get() );
-
-                GetMemory().SetMesh( item, pResult );
-                break;
-            }
-
-            default:
-                check(false);
-            }
-
-            break;
-        }
-
         case OP_TYPE::ME_FORMAT:
         {
 			OP::MeshFormatArgs args = pModel->GetPrivate()->m_program.GetOpArgs<OP::MeshFormatArgs>(item.At);
@@ -2528,28 +2502,52 @@ namespace mu
         case OP_TYPE::IM_LAYER:
         {
 			OP::ImageLayerArgs args = pModel->GetPrivate()->m_program.GetOpArgs<OP::ImageLayerArgs>(item.At);
-            switch (item.Stage)
-            {
-			case 0:
-				AddOp(FScheduledOp(item.At, item, 1),
-					FScheduledOp(args.base, item));
-				break;
 
-			case 1:
-				// Request the rest of the data.
-				AddOp(FScheduledOp(item.At, item, 2),
-					FScheduledOp(args.blended, item),
-					FScheduledOp(args.mask, item));
-				break;
+			if (ExecutionStrategy == EExecutionStrategy::MinimizeMemory)
+			{
+				switch (item.Stage)
+				{
+				case 0:
+					AddOp(FScheduledOp(item.At, item, 1),
+						FScheduledOp(args.base, item));
+					break;
 
-			case 2:
-				// This has been moved to a task. It should have been intercepted in IssueOp.
-				check(false);
-				break;
-			
-            default:
-                check(false);
-            }
+				case 1:
+					// Request the rest of the data.
+					AddOp(FScheduledOp(item.At, item, 2),
+						FScheduledOp(args.blended, item),
+						FScheduledOp(args.mask, item));
+					break;
+
+				case 2:
+					// This has been moved to a task. It should have been intercepted in IssueOp.
+					check(false);
+					break;
+
+				default:
+					check(false);
+				}
+			}
+			else
+			{
+				switch (item.Stage)
+				{
+				case 0:
+					AddOp(FScheduledOp(item.At, item, 1),
+						FScheduledOp(args.base, item),
+						FScheduledOp(args.blended, item),
+						FScheduledOp(args.mask, item));
+					break;
+
+				case 1:
+					// This has been moved to a task. It should have been intercepted in IssueOp.
+					check(false);
+					break;
+
+				default:
+					check(false);
+				}
+			}
 
             break;
         }
@@ -4142,7 +4140,7 @@ namespace mu
 				// is resolution sensitive (pixel offsets). If the size doesn't match, scale the source, apply 
 				// displacement and then unscale it.
 				FImageSize OriginalSourceScale = pSource->GetSize();
-				if (OriginalSourceScale.x()>0 && OriginalSourceScale.y()>0 && OriginalSourceScale != pMap->GetSize())
+				if (OriginalSourceScale[0]>0 && OriginalSourceScale[1]>0 && OriginalSourceScale != pMap->GetSize())
 				{
 					MUTABLE_CPUPROFILER_SCOPE(ImageResize_EmergencyHackForDisplacementStep1);
 					pSource = ImageResizeLinear(0, pSource.get(), pMap->GetSize());
@@ -4153,7 +4151,7 @@ namespace mu
 				//ImagePtr pNew = new Image(pSource->GetSizeX(), pSource->GetSizeY(), 1, pSource->GetFormat());
 				Ptr<Image> pNew = mu::CloneOrTakeOver(pSource.get());
 
-				if (OriginalSourceScale.x() > 0 && OriginalSourceScale.y() > 0)
+				if (OriginalSourceScale[0] > 0 && OriginalSourceScale[1] > 0)
 				{
 					ImageDisplace(pNew.get(), pSource.get(), pMap.get());
 
diff --git a/Engine/Plugins/Experimental/Mutable/Source/MutableRuntime/Private/MuR/CodeRunner.h b/Engine/Plugins/Experimental/Mutable/Source/MutableRuntime/Private/MuR/CodeRunner.h
--- a/Engine/Plugins/Experimental/Mutable/Source/MutableRuntime/Private/MuR/CodeRunner.h
+++ b/Engine/Plugins/Experimental/Mutable/Source/MutableRuntime/Private/MuR/CodeRunner.h
@@ -43,6 +43,7 @@ namespace  mu
     {
     public:
 		CodeRunner(const SettingsPtrConst&, class System::Private*, 
+			EExecutionStrategy,
 			const TSharedPtr<const Model>&, const Parameters* pParams,
 			OP::ADDRESS at, uint32 lodMask, uint8 executionOptions, int32 InImageLOD, FScheduledOp::EType );
 
@@ -214,6 +215,7 @@ namespace  mu
 
 		void AddOp(const FScheduledOp& op)
 		{
+			// It has no dependencies, so add it directly to the open tasks list.
 			OpenTasks.Add(op);
 			ScheduledStagePerOp[op] = op.Stage + 1;
 		}
@@ -301,10 +303,19 @@ namespace  mu
 
 	protected:
 
-		//! Stack of pending operations, and the execution stage they are in.
+		/** Strategy to choose the order of execution of operations. */
+		EExecutionStrategy ExecutionStrategy = EExecutionStrategy::None; 
+
+		/** List of pending operations that we don't know if they cannot be run yet because of dependencies. */
 		TArray< FTask > ClosedTasks;
+
+		/** List of tasks that can be run because they don't have any unmet dependency. */
 		TArray< FScheduledOp > OpenTasks;
+
+		/** For every op, up to what stage it has been scheduled to run. */
 		CodeContainer<uint8> ScheduledStagePerOp;
+
+		/** List of tasks that have been set to run concurrently and their completion is unknown. */
 		TArray< TSharedPtr<FIssuedTask> > IssuedTasks;
 
 	// TODO: protect this.
@@ -356,6 +367,9 @@ namespace  mu
 		/** Try to create a concurrent task for the given op. Return null if not possible. */
 		TSharedPtr<FIssuedTask> IssueOp(FScheduledOp item);
 
+		/** Calculate an heuristic to select op execution based on memory usage. */
+		int32 GetOpEstimatedMemoryDelta(const FScheduledOp& Candidate, const FProgram& Program);
+
 		/** */
 		void CompleteRomLoadOp(FRomLoadOp& o);
 		
diff --git a/Engine/Plugins/Experimental/Mutable/Source/MutableRuntime/Private/MuR/CodeRunnerTasked.cpp b/Engine/Plugins/Experimental/Mutable/Source/MutableRuntime/Private/MuR/CodeRunnerTasked.cpp
--- a/Engine/Plugins/Experimental/Mutable/Source/MutableRuntime/Private/MuR/CodeRunnerTasked.cpp
+++ b/Engine/Plugins/Experimental/Mutable/Source/MutableRuntime/Private/MuR/CodeRunnerTasked.cpp
@@ -185,6 +185,175 @@ namespace mu
 	}
 
 
+	//---------------------------------------------------------------------------------------------
+	/** This table encodes an heuristic for every execution stage of every operation  that tries to 
+	* guess the memory delta of that stage: 
+	* - if negative: after the stage the working memory may have reduced
+	* - if positive: after the stage the working memory may have increased
+	* The actual number is just a unit-less weight that represents "amount of additional memory that will be used".
+	* Some operations don't have a constant memory delta and are dealt with specially, but still need 
+	* to be in this table.
+	*/
+	static constexpr int32 sMemoryWeightsStageCount = 4;
+	static const int8 sMemoryWeights[int32(OP_TYPE::COUNT)][sMemoryWeightsStageCount] =
+	{ 
+		{   0,   0,   0,   0 },	// NONE
+		    
+		{   0,   0,   0,   0 },	// BO_CONSTANT
+		{   0,   0,   0,   0 },	// NU_CONSTANT
+		{   0,   0,   0,   0 },	// SC_CONSTANT
+		{   0,   0,   0,   0 },	// CO_CONSTANT
+		{  20,   0,   0,   0 },	// IM_CONSTANT
+		{  20,   0,   0,   0 },	// ME_CONSTANT
+		{   0,   0,   0,   0 },	// LA_CONSTANT
+		{   0,   0,   0,   0 },	// PR_CONSTANT
+		{   0,   0,   0,   0 },	// ST_CONSTANT
+		    
+		{   0,   0,   0,   0 },	// BO_PARAMETER
+		{   0,   0,   0,   0 },	// NU_PARAMETER
+		{   0,   0,   0,   0 },	// SC_PARAMETER
+		{   0,   0,   0,   0 },	// CO_PARAMETER
+		{   0,   0,   0,   0 },	// PR_PARAMETER
+		{  20,   0,   0,   0 },	// IM_PARAMETER
+		{   0,   0,   0,   0 },	// ST_PARAMETER
+		    
+		{   0,   0,   0,   0 },	// NU_CONDITIONAL
+		{   0,   0,   0,   0 },	// SC_CONDITIONAL
+		{   0,   0,   0,   0 },	// CO_CONDITIONAL
+		{   0,   0,   0,   0 },	// IM_CONDITIONAL
+		{   0,   0,   0,   0 },	// ME_CONDITIONAL
+		{   0,   0,   0,   0 },	// LA_CONDITIONAL
+		{   0,   0,   0,   0 },	// IN_CONDITIONAL
+		    
+		{   0,   0,   0,   0 },	// NU_SWITCH
+		{   0,   0,   0,   0 },	// SC_SWITCH
+		{   0,   0,   0,   0 },	// CO_SWITCH
+		{   0,   0,   0,   0 },	// IM_SWITCH
+		{   0,   0,   0,   0 },	// ME_SWITCH
+		{   0,   0,   0,   0 },	// LA_SWITCH
+		{   0,   0,   0,   0 },	// IN_SWITCH
+		    
+		{   0,   0,   0,   0 },	// BO_LESS
+		{   0,   0,   0,   0 },	// BO_EQUAL_SC_CONST
+		{   0,   0,   0,   0 },	// BO_AND
+		{   0,   0,   0,   0 },	// BO_OR
+		{   0,   0,   0,   0 },	// BO_NOT
+		    
+		{   0,   0,   0,   0 },	// SC_MULTIPLYADD
+		{   0,   0,   0,   0 },	// SC_ARITHMETIC
+		{   0,   0,   0,   0 },	// SC_CURVE
+		    
+		{   0,   0,   0,   0 },	// CO_SAMPLEIMAGE
+		{   0,   0,   0,   0 },	// CO_SWIZZLE
+		{   0,   0,   0,   0 },	// CO_IMAGESIZE
+		{   0,   0,   0,   0 },	// CO_LAYOUTBLOCKTRANSFORM
+		{   0,   0,   0,   0 },	// CO_FROMSCALARS
+		{   0,   0,   0,   0 },	// CO_ARITHMETIC
+		    
+		{   0, -20, -20,   0 },	// IM_LAYER
+		{   0,   0,   0,   0 },	// IM_LAYERCOLOUR
+		{   0,   0,   0,   0 },	// IM_PIXELFORMAT	(special case)
+		{   0,  10,   0,   0 },	// IM_MIPMAP
+		{   0,   0,   0,   0 },	// IM_RESIZE		(special case)
+		{   0,   0,   0,   0 },	// IM_RESIZELIKE	(to be deprecated?)
+		{   0,   0,   0,   0 },	// IM_RESIZEREL		(special case)
+		{   0,  20,   0,   0 },	// IM_BLANKLAYOUT
+		{   0,   0, -20,   0 },	// IM_COMPOSE
+		{   0,   0,   0,   0 },	// IM_DIFFERENCE
+		{   0,   0, -20,   0 },	// IM_INTERPOLATE
+		{   0,   0, -30,   0 },	// IM_INTERPOLATE3
+		{   0,   0,   0,   0 },	// IM_SATURATE
+		{   0,   0,   0,   0 },	// IM_LUMINANCE
+		{   0,   0,   0,   0 },	// IM_SWIZZLE
+		{   0,   0,   0,   0 },	// IM_SELECTCOLOUR
+		{   0,   0,   0,   0 },	// IM_COLOURMAP
+		{   0,   5,   0,   0 },	// IM_GRADIENT
+		{   0,   0,   0,   0 },	// IM_BINARISE
+		{   0,  20,   0,   0 },	// IM_PLAINCOLOUR
+		{   0,   0,   0,   0 },	// IM_GPU
+		{   0, -10,   0,   0 },	// IM_CROP
+		{   0, -10,   0,   0 },	// IM_PATCH
+		{   0,  10,  10,   0 },	// IM_RASTERMESH
+		{   0,   0,   0,   0 },	// IM_MAKEGROWMAP
+		{   0, -10,   0,   0 },	// IM_DISPLACE
+		{   0,   0, -20, -20 },	// IM_MULTILAYER
+		{   0,   0,   0,   0 },	// IM_INVERT
+		{   0,   0,   0,   0 },	// IM_NORMALCOMPOSITE
+		{   0,   0,   0,   0 },	// IM_TRANSFORM
+		    
+		{   0,   0,   0,   0 },	// ME_APPLYLAYOUT
+		{   0, -20,   0,   0 },	// ME_DIFFERENCE
+		{   0,   0, -10,   0 },	// ME_MORPH2
+		{   0,   0,   0,   0 },	// ME_MERGE
+		{   0,   0, -20,   0 },	// ME_INTERPOLATE
+		{   0, -10,   0,   0 },	// ME_MASKCLIPMESH
+		{   0, -10,   0,   0 },	// ME_MASKDIFF
+		{   0,   0, -10,   0 },	// ME_REMOVEMASK
+		{   0,   0,   0,   0 },	// ME_FORMAT
+		{   0, -10,   0,   0 },	// ME_EXTRACTLAYOUTBLOCK
+		{   0, -10,   0,   0 },	// ME_EXTRACTFACEGROUP
+		{   0,   0,   0,   0 },	// ME_TRANSFORM
+		{   0, -10,   0,   0 },	// ME_CLIPMORPHPLANE
+		{   0, -10,   0,   0 },	// ME_CLIPWITHMESH
+		{   0,   0,   0,   0 },	// ME_SETSKELETON
+		{   0,   0,   0,   0 },	// ME_PROJECT
+		{   0,   0,   0,   0 },	// ME_APPLYPOSE
+		{   0,   0,   0,   0 },	// ME_REMAPINDICES
+		{   0,   0,   0,   0 },	// ME_GEOMETRYOPERATION
+		{   0, -20,   0,   0 },	// ME_BINDSHAPE
+		{   0, -10,   0,   0 },	// ME_APPLYSHAPE
+		{   0, -10,   0,   0 },	// ME_CLIPDEFORM
+		{   0, -10,   0,   0 },	// ME_MORPHRESHAPE
+		{   0,   0,   0,   0 },	// ME_OPTIMIZESKINNING
+		    
+		{   0,   0,   0,   0 },	// IN_ADDMESH
+		{   0,   0,   0,   0 },	// IN_ADDIMAGE
+		{   0,   0,   0,   0 },	// IN_ADDVECTOR
+		{   0,   0,   0,   0 },	// IN_ADDSCALAR
+		{   0,   0,   0,   0 },	// IN_ADDSTRING
+		{   0,   0,   0,   0 },	// IN_ADDSURFACE
+		{   0,   0,   0,   0 },	// IN_ADDCOMPONENT
+		{   0,   0,   0,   0 },	// IN_ADDLOD
+		    
+		{   0,   0,   0,   0 },	// LA_PACK
+		{   0,   0,   0,   0 },	// LA_MERGE
+		{   0,  -1,   0,   0 },	// LA_REMOVEBLOCKS
+		{   0,   0,   0,   0 },	// LA_FROMMESH
+	};
+
+	static_assert(sizeof(sMemoryWeights)/sMemoryWeightsStageCount == int32(OP_TYPE::COUNT));
+
+
+	int32 CodeRunner::GetOpEstimatedMemoryDelta( const FScheduledOp& Candidate, const FProgram& Program )
+	{
+		int32 OpDelta = 0;
+		int32 Stage = FMath::Min(int32(Candidate.Stage), sMemoryWeightsStageCount - 1);
+		OP_TYPE OpType = Program.GetOpType(Candidate.At);
+
+		if (OpType == OP_TYPE::IM_PIXELFORMAT && Stage == 1)
+		{
+			// TODO: We should get the actual format from the arguments for a more precise calculation.
+			OP::ImagePixelFormatArgs args = Program.GetOpArgs<OP::ImagePixelFormatArgs>(Candidate.At);
+			bool bIsCompressed = GetUncompressedFormat(args.format) != args.format;
+			bool bIsSmallish = GetImageFormatData(args.format).m_bytesPerBlock < 2;
+			if (bIsCompressed || bIsSmallish)
+			{
+				OpDelta = -10;
+			}
+			else
+			{
+				OpDelta = 10;
+			}
+		}
+		else
+		{
+			OpDelta = sMemoryWeights[int32(OpType)][Stage];
+		}
+
+		return OpDelta;
+	}
+
+
     //---------------------------------------------------------------------------------------------
     void CodeRunner::Run()
     {
@@ -249,7 +418,47 @@ namespace mu
 
 			while (!OpenTasks.IsEmpty())
 			{
-				FScheduledOp item = OpenTasks.Pop();
+				// Get a new task to run
+				FScheduledOp item;
+				switch (ExecutionStrategy)
+				{
+				case EExecutionStrategy::MinimizeMemory:
+				{
+					// TODO: This should be done when an operation is added to the OpenTasks array instead of every time.
+					int32 BestOp = 0;
+					int8 BestDelta = TNumericLimits<int8>::Max();
+					int32 OpenOpCount = OpenTasks.Num();
+					const FProgram& Program = m_pModel->GetPrivate()->m_program;
+					for (int32 OpIndex = 0; OpIndex < OpenOpCount; ++OpIndex)
+					{
+						const FScheduledOp& Candidate = OpenTasks[OpIndex];
+						int32 OpDelta = GetOpEstimatedMemoryDelta(Candidate, Program);
+
+						if (OpDelta < BestDelta)
+						{
+							BestOp = OpIndex;
+							BestDelta = OpDelta;
+
+							// Shortcut: If we are freeing memory, we don't care how much: it is already the best op
+							if (BestDelta < 0)
+							{
+								break;
+							}
+						}
+					}
+
+					item = OpenTasks[BestOp];
+					OpenTasks.RemoveAtSwap(BestOp);
+					break;
+				}
+
+				case EExecutionStrategy::None:
+				default:
+					// Just get one.
+					item = OpenTasks.Pop();
+					break;
+
+				}
 
 				// Special processing in case it is an ImageDesc operation
 				if ( item.Type==FScheduledOp::EType::ImageDesc )
@@ -362,12 +571,11 @@ namespace mu
 				if (o.m_romIndex>=0 && m_pSystem->m_pStreamInterface->IsReadCompleted(o.m_streamID))
 				{
 					CompleteRomLoadOp(o);
+					UpdateTraces();
 				}
-
-				UpdateTraces();
 			}
 
-			// Look for a closed task with dependencies satisfied and move them to open.
+			// Look for a closed task with dependencies satisfied and move them to the open task list.
 			bool bSomeWasReady = false;
 			for (int Index = 0; Index<ClosedTasks.Num(); )
 			{
@@ -418,6 +626,9 @@ namespace mu
 					UE_LOG(LogMutableCore, Log, TEXT("%s"), *TaskDesc);
 				}
 				check(false);
+
+				// This should never happen but if it does, abort the code execution.
+				return;
 			}
 
 			// If at this point there is no open op and we haven't finished, we need to wait for an issued op to complete.
@@ -2077,7 +2288,10 @@ namespace mu
 
 		case OP_TYPE::IM_LAYER:
 		{
-			if (item.Stage == 2)
+			if ((ExecutionStrategy == EExecutionStrategy::MinimizeMemory && item.Stage == 2)
+				||
+				(ExecutionStrategy != EExecutionStrategy::MinimizeMemory && item.Stage == 1)
+				)
 			{
 				OP::ImageLayerArgs args = program.GetOpArgs<OP::ImageLayerArgs>(item.At);
 				Issued = MakeShared<FImageLayerTask>(item, GetMemory(), args, m_pSettings->GetPrivate()->m_imageCompressionQuality);
diff --git a/Engine/Plugins/Experimental/Mutable/Source/MutableRuntime/Private/MuR/Image.cpp b/Engine/Plugins/Experimental/Mutable/Source/MutableRuntime/Private/MuR/Image.cpp
--- a/Engine/Plugins/Experimental/Mutable/Source/MutableRuntime/Private/MuR/Image.cpp
+++ b/Engine/Plugins/Experimental/Mutable/Source/MutableRuntime/Private/MuR/Image.cpp
@@ -150,6 +150,38 @@ namespace mu
         arch << *p;
     }
 
+	void Image::Serialise(OutputArchive& arch) const
+	{
+		uint32 ver = 3;
+		arch << ver;
+
+		arch << m_size;
+		arch << m_lods;
+		arch << (uint8)m_format;
+		arch << m_data;
+
+		// Remove non-persistent flags.
+		uint8 flags = m_flags & ~IF_HAS_RELEVANCY_MAP;
+		arch << flags;
+	}
+
+	void Image::Unserialise(InputArchive& arch)
+	{
+		uint32 ver;
+		arch >> ver;
+		check(ver == 3);
+
+		arch >> m_size;
+		arch >> m_lods;
+
+		uint8 format;
+		arch >> format;
+		m_format = (EImageFormat)format;
+		arch >> m_data;
+
+		arch >> m_flags;
+	}
+
 
     //---------------------------------------------------------------------------------------------
     Ptr<Image> Image::StaticUnserialise( InputArchive& arch )
@@ -406,7 +438,7 @@ namespace mu
     {
 		FIntVector2 res(0,0);
 
-		FIntVector2 s = FIntVector2(m_size.x(), m_size.y());
+		FIntVector2 s = FIntVector2(m_size[0], m_size[1]);
 
         for ( int l=0; l<mip+1; ++l )
         {
@@ -781,7 +813,8 @@ namespace mu
      void Image::GetNonBlackRect_Reference(FImageRect& rect) const
      {            
          rect.min[0] = rect.min[1] = 0;
-         rect.size = m_size;
+		 rect.size[0] = m_size[0];
+		 rect.size[1] = m_size[1];
 
          if ( !rect.size[0] || !rect.size[1] )
              return;
diff --git a/Engine/Plugins/Experimental/Mutable/Source/MutableRuntime/Private/MuR/ImagePrivate.h b/Engine/Plugins/Experimental/Mutable/Source/MutableRuntime/Private/MuR/ImagePrivate.h
--- a/Engine/Plugins/Experimental/Mutable/Source/MutableRuntime/Private/MuR/ImagePrivate.h
+++ b/Engine/Plugins/Experimental/Mutable/Source/MutableRuntime/Private/MuR/ImagePrivate.h
@@ -15,6 +15,7 @@ namespace mu
 	MUTABLE_DEFINE_ENUM_SERIALISABLE( ECompositeImageMode );
 	MUTABLE_DEFINE_ENUM_SERIALISABLE( ESamplingMethod );
 	MUTABLE_DEFINE_ENUM_SERIALISABLE( EMinFilterMethod );
+	MUTABLE_DEFINE_ENUM_SERIALISABLE( EImageFormat );
 
     struct FImageFormatData
 	{
diff --git a/Engine/Plugins/Experimental/Mutable/Source/MutableRuntime/Private/MuR/Layout.cpp b/Engine/Plugins/Experimental/Mutable/Source/MutableRuntime/Private/MuR/Layout.cpp
--- a/Engine/Plugins/Experimental/Mutable/Source/MutableRuntime/Private/MuR/Layout.cpp
+++ b/Engine/Plugins/Experimental/Mutable/Source/MutableRuntime/Private/MuR/Layout.cpp
@@ -156,8 +156,8 @@ namespace mu {
 		check( index >=0 && index < m_blocks.Num() );
 
 		// Keeps the id
-		m_blocks[index].m_min = vec2<uint16>((uint16)minx, (uint16)miny);
-		m_blocks[index].m_size = vec2<uint16>((uint16)sizex, (uint16)sizey);
+		m_blocks[index].m_min = UE::Math::TIntVector2<uint16>((uint16)minx, (uint16)miny);
+		m_blocks[index].m_size = UE::Math::TIntVector2<uint16>((uint16)sizex, (uint16)sizey);
 	}
 
 
@@ -223,7 +223,7 @@ namespace mu {
 	bool Layout::IsSingleBlockAndFull() const
 	{
 		if (m_blocks.Num() == 1
-			&& m_blocks[0].m_min == vec2<uint16>(0, 0)
+			&& m_blocks[0].m_min == UE::Math::TIntVector2<uint16>(0, 0)
 			&& m_blocks[0].m_size == m_size)
 		{
 			return true;
diff --git a/Engine/Plugins/Experimental/Mutable/Source/MutableRuntime/Private/MuR/Mesh.cpp b/Engine/Plugins/Experimental/Mutable/Source/MutableRuntime/Private/MuR/Mesh.cpp
--- a/Engine/Plugins/Experimental/Mutable/Source/MutableRuntime/Private/MuR/Mesh.cpp
+++ b/Engine/Plugins/Experimental/Mutable/Source/MutableRuntime/Private/MuR/Mesh.cpp
@@ -610,9 +610,9 @@ bool Mesh::HasCompatibleFormat( const Mesh* pOther ) const
 
 
 //---------------------------------------------------------------------------------------------
-vec3<uint32_t>  Mesh::GetFaceVertexIndices( int f ) const
+UE::Math::TIntVector3<uint32_t>  Mesh::GetFaceVertexIndices( int f ) const
 {
-    vec3<uint32_t> res;
+	UE::Math::TIntVector3<uint32> res;
 
     MeshBufferIteratorConst<MBF_UINT32,uint32_t,1> it( m_IndexBuffers, MBS_VERTEXINDEX );
     it += f*3;
@@ -640,13 +640,13 @@ bool Mesh::HasFace
 {
     bool found = false;
 
-    vec3<uint32_t> ov = other.GetFaceVertexIndices( otherFaceIndex );
+	UE::Math::TIntVector3<uint32> ov = other.GetFaceVertexIndices( otherFaceIndex );
 
     UntypedMeshBufferIteratorConst it( m_IndexBuffers, MBS_VERTEXINDEX );
 
     for ( int f=0; !found && f<m_FaceBuffers.GetElementCount(); f++ )
     {
-        vec3<uint32_t> v;
+		UE::Math::TIntVector3<uint32> v;
         v[0] = it.GetAsUINT32(); ++it;
         v[1] = it.GetAsUINT32(); ++it;
         v[2] = it.GetAsUINT32(); ++it;
diff --git a/Engine/Plugins/Experimental/Mutable/Source/MutableRuntime/Private/MuR/MutableMath.h b/Engine/Plugins/Experimental/Mutable/Source/MutableRuntime/Private/MuR/MutableMath.h
deleted file mode 100644
diff --git a/Engine/Plugins/Experimental/Mutable/Source/MutableRuntime/Private/MuR/OpImageResize.h b/Engine/Plugins/Experimental/Mutable/Source/MutableRuntime/Private/MuR/OpImageResize.h
--- a/Engine/Plugins/Experimental/Mutable/Source/MutableRuntime/Private/MuR/OpImageResize.h
+++ b/Engine/Plugins/Experimental/Mutable/Source/MutableRuntime/Private/MuR/OpImageResize.h
@@ -207,7 +207,6 @@ namespace mu
         default:
             // Case not implemented
             check( false );
-            mu::Halt();
         }
     }
 
@@ -490,7 +489,6 @@ namespace mu
         default:
             // Case not implemented
             check( false );
-            mu::Halt();
         }
 
     }
@@ -619,7 +617,6 @@ namespace mu
         default:
             // Case not implemented
             check( false );
-            mu::Halt();
         }
     }
 
diff --git a/Engine/Plugins/Experimental/Mutable/Source/MutableRuntime/Private/MuR/OpLayoutPack.h b/Engine/Plugins/Experimental/Mutable/Source/MutableRuntime/Private/MuR/OpLayoutPack.h
--- a/Engine/Plugins/Experimental/Mutable/Source/MutableRuntime/Private/MuR/OpLayoutPack.h
+++ b/Engine/Plugins/Experimental/Mutable/Source/MutableRuntime/Private/MuR/OpLayoutPack.h
@@ -5,8 +5,7 @@
 #include "MuR/Layout.h"
 #include "MuR/MutableTrace.h"
 #include "MuR/Platform.h"
-
-#include <limits>
+#include "GenericPlatform/GenericPlatformMath.h"
 
 
 namespace mu
@@ -506,9 +505,9 @@ namespace mu
 				auto candidateSizeY = scratch->sorted[candidate].size[1];
 
 				// Seek for the lowest span where the block fits
-				int currentLevel = std::numeric_limits<int>::max();
+				int currentLevel = TNumericLimits<int>::Max();
 				int currentX = 0;
-				int currentLevelWithoutHole = std::numeric_limits<int>::max();
+				int currentLevelWithoutHole = TNumericLimits<int>::Max();
 				int currentXWithoutHole = 0;
 				for (int x = 0; x <= *maxX - candidateSizeX; ++x)
 				{
@@ -525,8 +524,8 @@ namespace mu
 					}
 
 					// Does it make an unfillable hole with the top or side?
-					int minX = std::numeric_limits<int>::max();
-					int minY = std::numeric_limits<int>::max();
+					int minX = TNumericLimits<int>::Max();
+					int minY = TNumericLimits<int>::Max();
 					for (size_t b = 0; b < scratch->sorted.Num(); ++b)
 					{
 						if (!packedFlag[b] && b != candidate)
@@ -539,7 +538,7 @@ namespace mu
 					bool hole =
 						// Vertical unfillable gap
 						(
-						(minY != std::numeric_limits<int>::max())
+						(minY != TNumericLimits<int>::Max())
 							&&
 							(currentLevel + candidateSizeY) < bestY
 							&&
@@ -548,7 +547,7 @@ namespace mu
 						||
 						// Horizontal unfillable gap
 						(
-						(minX != std::numeric_limits<int>::max())
+						(minX != TNumericLimits<int>::Max())
 							&&
 							(currentX + candidateSizeX) < *maxX
 							&&
@@ -633,7 +632,7 @@ namespace mu
 			fits = false;
 		}
 
-		*maxY = ceilPow2(*maxY);
+		*maxY = FGenericPlatformMath::RoundUpToPowerOfTwo(*maxY);
 
 		return fits;
 	}
@@ -694,8 +693,8 @@ namespace mu
         // in case we cannot pack everything, we will grow Y with the current horizon algorithm.
 		if (pSourceLayout->GetLayoutPackingStrategy() == EPackStrategy::RESIZABLE_LAYOUT)
 		{
-			maxX = ceilPow2(maxX);
-			maxY = ceilPow2(maxY);
+			maxX = FGenericPlatformMath::RoundUpToPowerOfTwo(maxX);
+			maxY = FGenericPlatformMath::RoundUpToPowerOfTwo(maxY);
 
 			while (maxX*maxY < area)
 			{
@@ -779,8 +778,8 @@ namespace mu
 				}
 			}
 
-			maxX = ceilPow2(maxX);
-			maxY = ceilPow2(maxY);
+			maxX = FGenericPlatformMath::RoundUpToPowerOfTwo(maxX);
+			maxY = FGenericPlatformMath::RoundUpToPowerOfTwo(maxY);
 
 			// Grow until the area is big enough to fit all blocks or the size is equal to the max layout size.
 			while (maxX*maxY < area && (maxX < layoutSizeX || maxY < layoutSizeY))
diff --git a/Engine/Plugins/Experimental/Mutable/Source/MutableRuntime/Private/MuR/OpMeshSubtract.h b/Engine/Plugins/Experimental/Mutable/Source/MutableRuntime/Private/MuR/OpMeshSubtract.h
deleted file mode 100644
diff --git a/Engine/Plugins/Experimental/Mutable/Source/MutableRuntime/Private/MuR/Operations.cpp b/Engine/Plugins/Experimental/Mutable/Source/MutableRuntime/Private/MuR/Operations.cpp
--- a/Engine/Plugins/Experimental/Mutable/Source/MutableRuntime/Private/MuR/Operations.cpp
+++ b/Engine/Plugins/Experimental/Mutable/Source/MutableRuntime/Private/MuR/Operations.cpp
@@ -12,7 +12,7 @@ namespace mu
 	//---------------------------------------------------------------------------------------------
 	//---------------------------------------------------------------------------------------------
 	//---------------------------------------------------------------------------------------------
-    #define OP_DESC_COUNT		size_t(OP_TYPE::COUNT)
+    #define OP_DESC_COUNT		int32(OP_TYPE::COUNT)
 
     // clang-format off
     static const OP_DESC s_opDescs[OP_DESC_COUNT] =
@@ -109,7 +109,6 @@ namespace mu
 		{ DT_MESH,		true,	{ 0, 0, 0, 0, 0, 0, 0, 0, 0, 0, 0, 0, 0, 0, 0, 0, 0, 0, 0, 0, 0, 0, 0, 0, 0, 0, 0, 0, 0 }		},	// ME_INTERPOLATE
 		{ DT_MESH,		true,	{ 0, 0, 0, 0, 0, 0, 0, 0, 0, 0, 0, 0, 0, 0, 0, 0, 0, 0, 0, 0, 0, 0, 0, 0, 0, 0, 0, 0, 0 }		},	// ME_MASKCLIPMESH
 		{ DT_MESH,		true,	{ 0, 0, 0, 0, 0, 0, 0, 0, 0, 0, 0, 0, 0, 0, 0, 0, 0, 0, 0, 0, 0, 0, 0, 0, 0, 0, 0, 0, 0 }		},	// ME_MASKDIFF
-		{ DT_MESH,		true,	{ 0, 0, 0, 0, 0, 0, 0, 0, 0, 0, 0, 0, 0, 0, 0, 0, 0, 0, 0, 0, 0, 0, 0, 0, 0, 0, 0, 0, 0 }		},	// ME_SUBTRACT
 		{ DT_MESH,		true,	{ 0, 0, 0, 0, 0, 0, 0, 0, 0, 0, 0, 0, 0, 0, 0, 0, 0, 0, 0, 0, 0, 0, 0, 0, 0, 0, 0, 0, 0 }		},	// ME_REMOVEMASK
 		{ DT_MESH,		true,	{ 0, 0, 0, 0, 0, 0, 0, 0, 0, 0, 0, 0, 0, 0, 0, 0, 0, 0, 0, 0, 0, 0, 0, 0, 0, 0, 0, 0, 0 }		},	// ME_FORMAT
 		{ DT_MESH,		true,	{ 0, 0, 0, 0, 0, 0, 0, 0, 0, 0, 0, 0, 0, 0, 0, 0, 0, 0, 0, 0, 0, 0, 0, 0, 0, 0, 0, 0, 0 }		},	// ME_EXTRACTLAYOUTBLOCK
@@ -360,11 +359,6 @@ namespace mu
             f(&op.args.MeshMaskDiff.fragment );
             break;
 
-        case OP_TYPE::ME_SUBTRACT:
-            f(&op.args.MeshSubtract.a );
-            f(&op.args.MeshSubtract.b );
-            break;
-
         case OP_TYPE::ME_EXTRACTFACEGROUP:
             f(&op.args.MeshExtractFaceGroup.source );
             break;
@@ -933,14 +927,6 @@ namespace mu
             break;
         }
 
-        case OP_TYPE::ME_SUBTRACT:
-        {
-			OP::MeshSubtractArgs args = program.GetOpArgs<OP::MeshSubtractArgs>(at);
-            f(args.a );
-            f(args.b );
-            break;
-        }
-
         case OP_TYPE::ME_REMOVEMASK:
         {
             const uint8_t* data = program.GetOpArgsPointer(at);
diff --git a/Engine/Plugins/Experimental/Mutable/Source/MutableRuntime/Private/MuR/Operations.h b/Engine/Plugins/Experimental/Mutable/Source/MutableRuntime/Private/MuR/Operations.h
--- a/Engine/Plugins/Experimental/Mutable/Source/MutableRuntime/Private/MuR/Operations.h
+++ b/Engine/Plugins/Experimental/Mutable/Source/MutableRuntime/Private/MuR/Operations.h
@@ -256,10 +256,6 @@ namespace mu
         //! Create a new mask mesh selecting all the faces of a source that match another mesh.
         ME_MASKDIFF,
 
-        //! Remove from a source mesh all the faces in common with another mesh
-		//! \TODO: Deprecated?
-		ME_SUBTRACT,
-
         //! Remove all the geometry selected by a mask.
         ME_REMOVEMASK,
 
@@ -873,12 +869,6 @@ namespace mu
             ADDRESS fragment;
         };
 
-        struct MeshSubtractArgs
-        {
-            ADDRESS a;
-            ADDRESS b;
-        };
-
         struct MeshFormatArgs
         {
             ADDRESS source;
@@ -1145,7 +1135,6 @@ namespace mu
             MeshMergeArgs MeshMerge;
             MeshInterpolateArgs MeshInterpolate;
             MeshMaskDiffArgs MeshMaskDiff;
-            MeshSubtractArgs MeshSubtract;
             MeshExtractFaceGroupArgs MeshExtractFaceGroup;
 			MeshClipMorphPlaneArgs MeshClipMorphPlane;
             MeshClipWithMeshArgs MeshClipWithMesh;
diff --git a/Engine/Plugins/Experimental/Mutable/Source/MutableRuntime/Private/MuR/SerialisationPrivate.h b/Engine/Plugins/Experimental/Mutable/Source/MutableRuntime/Private/MuR/SerialisationPrivate.h
--- a/Engine/Plugins/Experimental/Mutable/Source/MutableRuntime/Private/MuR/SerialisationPrivate.h
+++ b/Engine/Plugins/Experimental/Mutable/Source/MutableRuntime/Private/MuR/SerialisationPrivate.h
@@ -6,10 +6,12 @@
 
 #include "Math/Quat.h"
 #include "Math/Vector.h"
+#include "Math/IntVector.h"
 #include "Math/Transform.h"
 #include "MuR/Ptr.h"
 #include "MuR/RefCounted.h"
 #include "MuR/MutableMemory.h"
+#include "MuR/MutableMath.h"
 #include "MuR/Types.h"
 
 #include <string>
@@ -18,6 +20,7 @@
 namespace mu
 {
 	typedef std::string string;
+	
 
     //!
     class MUTABLERUNTIME_API InputMemoryStream : public InputStream
@@ -306,16 +309,35 @@ namespace mu
 		}                                                                               \
 	}
 
-	MUTABLE_DEFINE_POD_VECTOR_SERIALISABLE(float)
-	MUTABLE_DEFINE_POD_VECTOR_SERIALISABLE(double)
-	MUTABLE_DEFINE_POD_VECTOR_SERIALISABLE(uint8)
-	MUTABLE_DEFINE_POD_VECTOR_SERIALISABLE(uint16)
-	MUTABLE_DEFINE_POD_VECTOR_SERIALISABLE(uint32)
-	MUTABLE_DEFINE_POD_VECTOR_SERIALISABLE(uint64)
-	MUTABLE_DEFINE_POD_VECTOR_SERIALISABLE(int8)
-	MUTABLE_DEFINE_POD_VECTOR_SERIALISABLE(int16)
-	MUTABLE_DEFINE_POD_VECTOR_SERIALISABLE(int32)
-	MUTABLE_DEFINE_POD_VECTOR_SERIALISABLE(int64)
+	MUTABLE_DEFINE_POD_VECTOR_SERIALISABLE(float);
+	MUTABLE_DEFINE_POD_VECTOR_SERIALISABLE(double);
+	MUTABLE_DEFINE_POD_VECTOR_SERIALISABLE(uint8);
+	MUTABLE_DEFINE_POD_VECTOR_SERIALISABLE(uint16);
+	MUTABLE_DEFINE_POD_VECTOR_SERIALISABLE(uint32);
+	MUTABLE_DEFINE_POD_VECTOR_SERIALISABLE(uint64);
+	MUTABLE_DEFINE_POD_VECTOR_SERIALISABLE(int8);
+	MUTABLE_DEFINE_POD_VECTOR_SERIALISABLE(int16);
+	MUTABLE_DEFINE_POD_VECTOR_SERIALISABLE(int32);
+	MUTABLE_DEFINE_POD_VECTOR_SERIALISABLE(int64);
+
+	MUTABLE_DEFINE_POD_SERIALISABLE(vec2f);
+	MUTABLE_DEFINE_POD_SERIALISABLE(vec3f);
+	MUTABLE_DEFINE_POD_SERIALISABLE(mat3f);
+	MUTABLE_DEFINE_POD_SERIALISABLE(mat4f);
+	MUTABLE_DEFINE_POD_SERIALISABLE(vec2<int>);
+	MUTABLE_DEFINE_POD_VECTOR_SERIALISABLE(vec2f);
+	MUTABLE_DEFINE_POD_VECTOR_SERIALISABLE(vec3f);
+	MUTABLE_DEFINE_POD_VECTOR_SERIALISABLE(mat3f);
+	MUTABLE_DEFINE_POD_VECTOR_SERIALISABLE(mat4f);
+	MUTABLE_DEFINE_POD_VECTOR_SERIALISABLE(vec2<int>);
+
+	// Unreal POD Serializables
+	MUTABLE_DEFINE_POD_SERIALISABLE(FUintVector2);
+	MUTABLE_DEFINE_POD_SERIALISABLE(UE::Math::TIntVector2<uint16>);
+	MUTABLE_DEFINE_POD_SERIALISABLE(UE::Math::TIntVector2<int16>);
+	MUTABLE_DEFINE_POD_VECTOR_SERIALISABLE(FUintVector2);
+	MUTABLE_DEFINE_POD_VECTOR_SERIALISABLE(UE::Math::TIntVector2<uint16>);
+	MUTABLE_DEFINE_POD_VECTOR_SERIALISABLE(UE::Math::TIntVector2<int16>);
 
 
 	//---------------------------------------------------------------------------------------------
diff --git a/Engine/Plugins/Experimental/Mutable/Source/MutableRuntime/Private/MuR/System.cpp b/Engine/Plugins/Experimental/Mutable/Source/MutableRuntime/Private/MuR/System.cpp
--- a/Engine/Plugins/Experimental/Mutable/Source/MutableRuntime/Private/MuR/System.cpp
+++ b/Engine/Plugins/Experimental/Mutable/Source/MutableRuntime/Private/MuR/System.cpp
@@ -38,27 +38,6 @@ namespace mu
 {
     static_assert( sizeof(mat4f) == 64, "UNEXPECTED_STRUCT_PACKING" );
 
-    //!
-    static std::atomic<mu::Error> s_error(mu::Error::None);
-
-
-    //---------------------------------------------------------------------------------------------
-    //!
-    //---------------------------------------------------------------------------------------------
-    Error GetError()
-    {
-        return s_error;
-    }
-
-
-    //---------------------------------------------------------------------------------------------
-    //!
-    //---------------------------------------------------------------------------------------------
-    void ClearError()
-    {
-        s_error = Error::None;
-    }
-
 
     //---------------------------------------------------------------------------------------------
     //---------------------------------------------------------------------------------------------
@@ -370,7 +349,7 @@ namespace mu
 				if (GetOpDataType(opType) == DT_IMAGE)
 				{
 					int8 executionOptions = 0;
-					CodeRunner Runner(m_pD->m_pSettings, m_pD, pLiveInstance->m_pModel, pLiveInstance->m_pOldParameters.get(), at, System::AllLODs, executionOptions, 0, FScheduledOp::EType::ImageDesc);
+					CodeRunner Runner(m_pD->m_pSettings, m_pD, EExecutionStrategy::MinimizeMemory, pLiveInstance->m_pModel, pLiveInstance->m_pOldParameters.get(), at, System::AllLODs, executionOptions, 0, FScheduledOp::EType::ImageDesc);
 					Runner.Run();
 					Runner.GetImageDescResult(OutDesc);
 				}
@@ -670,7 +649,7 @@ namespace mu
 	void System::Private::RunCode(const TSharedPtr<const Model>& InModel,
 		const Parameters* InParameters, OP::ADDRESS InCodeRoot, uint32 InLODs, uint8 executionOptions, int32 InImageLOD)
 	{
-		CodeRunner Runner(m_pSettings, this, InModel, InParameters, InCodeRoot, InLODs, 
+		CodeRunner Runner(m_pSettings, this, EExecutionStrategy::MinimizeMemory, InModel, InParameters, InCodeRoot, InLODs,
 			executionOptions, InImageLOD, FScheduledOp::EType::Full);
 		Runner.Run();
 		bUnrecoverableError = Runner.bUnrecoverableError;
diff --git a/Engine/Plugins/Experimental/Mutable/Source/MutableRuntime/Private/MuR/SystemPrivate.h b/Engine/Plugins/Experimental/Mutable/Source/MutableRuntime/Private/MuR/SystemPrivate.h
--- a/Engine/Plugins/Experimental/Mutable/Source/MutableRuntime/Private/MuR/SystemPrivate.h
+++ b/Engine/Plugins/Experimental/Mutable/Source/MutableRuntime/Private/MuR/SystemPrivate.h
@@ -158,7 +158,7 @@ namespace mu
 		//! may have more steges to schedule children that are optional for execution, etc.
 		uint8 Stage : 7;
 
-		//! Type of calculation we are requestinf for this operation.
+		//! Type of calculation we are requesting for this operation.
 		enum class EType : uint8
 		{
 			//! Execute the operation to calculate the full result
diff --git a/Engine/Plugins/Experimental/Mutable/Source/MutableRuntime/Public/MuR/Image.h b/Engine/Plugins/Experimental/Mutable/Source/MutableRuntime/Public/MuR/Image.h
--- a/Engine/Plugins/Experimental/Mutable/Source/MutableRuntime/Public/MuR/Image.h
+++ b/Engine/Plugins/Experimental/Mutable/Source/MutableRuntime/Public/MuR/Image.h
@@ -7,7 +7,6 @@
 #include "MuR/Ptr.h"
 #include "MuR/RefCounted.h"
 #include "MuR/Serialisation.h"
-#include "MuR/SerialisationPrivate.h"
 #include "Containers/Array.h"
 #include "HAL/LowLevelMemTracker.h"
 #include "HAL/Platform.h"
@@ -23,7 +22,7 @@ namespace mu
 	typedef Ptr<const Image> ImagePtrConst;
 
 	//!
-	using FImageSize = vec2<uint16>;
+	using FImageSize = UE::Math::TIntVector2<uint16>;
 	using FImageRect = box<vec2<uint16>>;
 
 	//! Pixel formats supported by the images.
@@ -74,7 +73,6 @@ namespace mu
 
         IF_COUNT
 	};
-	MUTABLE_DEFINE_ENUM_SERIALISABLE(EImageFormat);
 
 	struct MUTABLERUNTIME_API FImageDesc
 	{
@@ -317,38 +315,10 @@ namespace mu
 
 
 		//!
-		void Serialise(OutputArchive& arch) const
-		{
-			uint32 ver = 3;
-			arch << ver;
-
-			arch << m_size;
-			arch << m_lods;
-			arch << (uint8)m_format;
-			arch << m_data;
-
-			// Remove non-persistent flags.
-			uint8 flags = m_flags & ~IF_HAS_RELEVANCY_MAP;
-			arch << flags;
-		}
+		void Serialise(OutputArchive& arch) const;
 
 		//!
-		void Unserialise(InputArchive& arch)
-		{
-			uint32 ver;
-			arch >> ver;
-			check(ver == 3);
-
-			arch >> m_size;
-			arch >> m_lods;
-
-			uint8 format;
-			arch >> format;
-			m_format = (EImageFormat)format;
-			arch >> m_data;
-
-			arch >> m_flags;
-		}
+		void Unserialise(InputArchive& arch);
 
 		//!
 		inline bool operator==(const Image& o) const
diff --git a/Engine/Plugins/Experimental/Mutable/Source/MutableRuntime/Public/MuR/Layout.h b/Engine/Plugins/Experimental/Mutable/Source/MutableRuntime/Public/MuR/Layout.h
--- a/Engine/Plugins/Experimental/Mutable/Source/MutableRuntime/Public/MuR/Layout.h
+++ b/Engine/Plugins/Experimental/Mutable/Source/MutableRuntime/Public/MuR/Layout.h
@@ -2,14 +2,13 @@
 
 #pragma once
 
-#include "MuR/MutableMath.h"
 #include "MuR/Ptr.h"
 #include "MuR/RefCounted.h"
 #include "MuR/Serialisation.h"
 #include "MuR/SerialisationPrivate.h"
 #include "Containers/Array.h"
-#include "Math/UnrealMathSSE.h"
 #include "Misc/AssertionMacros.h"
+#include "Math/IntVector.h"
 
 namespace mu
 {
@@ -156,7 +155,7 @@ namespace mu
 		//-----------------------------------------------------------------------------------------
 		struct FBlock
 		{
-			FBlock(vec2<uint16> min = vec2<uint16>(), vec2<uint16> size = vec2<uint16>())
+			FBlock(UE::Math::TIntVector2<uint16> min = UE::Math::TIntVector2<uint16>(), UE::Math::TIntVector2<uint16> size = UE::Math::TIntVector2<uint16>())
 			{
 				m_min = min;
 				m_size = size;
@@ -165,8 +164,8 @@ namespace mu
 				bUseSymmetry = false;
 			}
 
-			vec2<uint16> m_min;
-			vec2<uint16> m_size;
+			UE::Math::TIntVector2<uint16> m_min;
+			UE::Math::TIntVector2<uint16> m_size;
 
 			//! Absolute id used to control merging of various layouts
 			int32 m_id;
@@ -230,9 +229,9 @@ namespace mu
 
 
 		//!
-		vec2<uint16> m_size;
+		UE::Math::TIntVector2<uint16> m_size;
 
-		vec2<uint16> m_maxsize;
+		UE::Math::TIntVector2<uint16> m_maxsize;
 
 		//!
 		TArray<FBlock> m_blocks;
diff --git a/Engine/Plugins/Experimental/Mutable/Source/MutableRuntime/Public/MuR/Mesh.h b/Engine/Plugins/Experimental/Mutable/Source/MutableRuntime/Public/MuR/Mesh.h
--- a/Engine/Plugins/Experimental/Mutable/Source/MutableRuntime/Public/MuR/Mesh.h
+++ b/Engine/Plugins/Experimental/Mutable/Source/MutableRuntime/Public/MuR/Mesh.h
@@ -11,7 +11,6 @@
 #include "Misc/EnumClassFlags.h"
 #include "MuR/Layout.h"
 #include "MuR/MeshBufferSet.h"
-#include "MuR/MutableMath.h"
 #include "MuR/PhysicsBody.h"
 #include "MuR/Ptr.h"
 #include "MuR/RefCounted.h"
@@ -595,7 +594,7 @@ namespace mu
 		bool HasFace( const Mesh& other, int otherFaceIndex, const VERTEX_MATCH_MAP& vertexMap ) const;
 
 		//! Compare the vertex attributes to check if they match.
-		vec3<uint32> GetFaceVertexIndices(int f) const;
+		UE::Math::TIntVector3<uint32> GetFaceVertexIndices(int f) const;
 
 		//! Return true if the given mesh has the same vertex and index formats, and in the same
 		//! buffer structure.
diff --git a/Engine/Plugins/Experimental/Mutable/Source/MutableRuntime/Public/MuR/Model.h b/Engine/Plugins/Experimental/Mutable/Source/MutableRuntime/Public/MuR/Model.h
--- a/Engine/Plugins/Experimental/Mutable/Source/MutableRuntime/Public/MuR/Model.h
+++ b/Engine/Plugins/Experimental/Mutable/Source/MutableRuntime/Public/MuR/Model.h
@@ -13,7 +13,7 @@
 //! data structures. Compiled models are not necessarily compatible when the runtime is updated,
 //! so this version number can be used externally to verify this. It is not used internally, and
 //! serializing models from different versions than this runtime will probably result in a crash.
-#define MUTABLE_COMPILED_MODEL_CODE_VERSION		uint32( 58 )
+#define MUTABLE_COMPILED_MODEL_CODE_VERSION		uint32( 59 )
 #define MUTABLE_PARAMETERS_VERSION              uint32( 1 )
 
 
diff --git a/Engine/Plugins/Experimental/Mutable/Source/MutableRuntime/Public/MuR/MutableMath.h b/Engine/Plugins/Experimental/Mutable/Source/MutableRuntime/Public/MuR/MutableMath.h
new file mode 100644
--- /dev/null
+++ b/Engine/Plugins/Experimental/Mutable/Source/MutableRuntime/Public/MuR/MutableMath.h
@@ -0,0 +1,1456 @@
+// Copyright Epic Games, Inc. All Rights Reserved.
+
+#pragma once
+
+#include "MuR/Types.h"
+#include "Math/Matrix.h"
+#include "Math/IntVector.h"
+
+
+namespace mu
+{
+
+    //---------------------------------------------------------------------------------------------
+    //!
+    //---------------------------------------------------------------------------------------------
+    template<class T>
+    T clamp( const T mi, const T ma, const T v )
+    {
+        return FMath::Max( mi, FMath::Min( ma, v ) );
+    }
+
+
+	//---------------------------------------------------------------------------------------------
+	// 16-bit floating point number support.
+	//---------------------------------------------------------------------------------------------
+    typedef uint16 float16;
+
+	//---------------------------------------------------------------------------------------------
+    inline uint32 halfToFloatI( float16 y )
+	{
+		int s = (y >> 15) & 0x00000001;
+		int e = (y >> 10) & 0x0000001f;
+		int f =  y        & 0x000003ff;
+
+		if (e == 0)
+		{
+			if (f == 0)
+			{
+				return s << 31;
+			}
+			else
+			{
+				while (!(f & 0x00000400))
+				{
+					f <<= 1;
+					e -=  1;
+				}
+				e += 1;
+				f &= ~0x00000400;
+			}
+		}
+		else if (e == 31)
+		{
+			if (f == 0)
+			{
+				return (s << 31) | 0x7f800000;
+			}
+			else
+			{
+				return (s << 31) | 0x7f800000 | (f << 13);
+			}
+		}
+
+		e = e + (127 - 15);
+		f = f << 13;
+
+		return ((s << 31) | (e << 23) | f);
+	}
+
+
+	//---------------------------------------------------------------------------------------------
+	inline float halfToFloat( float16 y)
+	{
+		union
+		{
+			float f;
+            uint32 i;
+		} v;
+
+		v.i = halfToFloatI(y);
+		return v.f;
+	}
+
+
+	//---------------------------------------------------------------------------------------------
+    inline float16 floatToHalfI( uint32 i )
+	{
+		int s =  (i >> 16) & 0x00008000;
+		int e = ((i >> 23) & 0x000000ff) - (127 - 15);
+		int f =   i        & 0x007fffff;
+
+		if (e <= 0)
+		{
+			if (e < -10)
+			{
+				if (s)
+				{
+					return 0x8000;
+				}
+				else
+				{
+				   return 0;
+				}
+			}
+			f = (f | 0x00800000) >> (1 - e);
+			return float16( s | (f >> 13) );
+		}
+		else if (e == 0xff - (127 - 15))
+		{
+			if (f == 0)
+			{
+				return float16(s | 0x7c00 );
+			}
+			else
+			{
+				f >>= 13;
+				return float16(s | 0x7c00 | f | (f == 0));
+			}
+		}
+		else
+		{
+			if (e > 30)
+			{
+				return float16( s | 0x7c00 );
+			}
+			return float16(s | (e << 10) | (f >> 13) );
+		}
+	}
+
+
+	//---------------------------------------------------------------------------------------------
+	inline float16 floatToHalf(float i)
+	{
+		union
+		{
+			float f;
+            uint32 i;
+		} v;
+
+		v.f = i;
+		return floatToHalfI(v.i);
+	}
+
+
+
+	//---------------------------------------------------------------------------------------------
+	//! Generic vector template.
+	//---------------------------------------------------------------------------------------------
+	template <class SCALAR, int DIM>
+    class vec
+	{
+	protected:
+
+		SCALAR m[ DIM ];
+
+	public:
+
+		//-----------------------------------------------------------------------------------------
+		typedef vec<SCALAR,DIM> V;
+
+        //-----------------------------------------------------------------------------------------
+        inline vec()
+        {
+            for( int i=0; i<DIM; ++i )
+            {
+                m[i]=0;
+            }
+        }
+
+
+        inline vec( const V& o )
+        {
+            for( int i=0; i<DIM; ++i )
+            {
+                m[i]=o.m[i];
+            }
+        }
+
+		//-----------------------------------------------------------------------------------------
+		static inline int GetDim()
+		{
+			return DIM;
+		}
+
+		//-----------------------------------------------------------------------------------------
+		inline SCALAR& operator[] (int i)
+		{
+			check( i>=0 && i <DIM );
+			return m[i];
+		}
+
+		//-----------------------------------------------------------------------------------------
+		inline const SCALAR& operator[] (int i) const
+		{
+			check( i>=0 && i <DIM );
+			return m[i];
+		}
+
+		//-----------------------------------------------------------------------------------------
+		inline bool operator==( const V& other ) const
+		{
+			bool res = true;
+			for( int i=0; i<DIM; ++i )
+			{
+				res &= m[i] == other.m[i];
+			}
+			return res;
+		}
+
+		//-----------------------------------------------------------------------------------------
+		inline bool operator!=( const V& other ) const
+		{
+			bool res = false;
+			for( int i=0; i<DIM; ++i )
+			{
+				res |= m[i] != other.m[i];
+			}
+			return res;
+		}
+
+		//-----------------------------------------------------------------------------------------
+		inline V operator+( const V& other ) const
+		{
+			V res;
+			for( int i=0; i<DIM; ++i )
+			{
+				res.m[i] = m[i] + other.m[i];
+			}
+			return res;
+		}
+
+		//-----------------------------------------------------------------------------------------
+		inline void operator+=( const V& other )
+		{
+			for( int i=0; i<DIM; ++i )
+			{
+				m[i] += other.m[i];
+			}
+		}
+
+        //-----------------------------------------------------------------------------------------
+        inline V operator-( const V& other ) const
+        {
+            V res;
+            for( int i=0; i<DIM; ++i )
+            {
+                res.m[i] = m[i] - other.m[i];
+            }
+            return res;
+        }
+
+        //-----------------------------------------------------------------------------------------
+        inline V operator-( SCALAR v ) const
+        {
+            V res;
+            for( int i=0; i<DIM; ++i )
+            {
+                res.m[i] = m[i] - v;
+            }
+            return res;
+        }
+
+		//-----------------------------------------------------------------------------------------
+		inline void operator-=( const V& other )
+		{
+			for( int i=0; i<DIM; ++i )
+			{
+                m[i] -= other.m[i];
+			}
+		}
+
+		//-----------------------------------------------------------------------------------------
+		inline V operator*( const V& other ) const
+		{
+			V res;
+			for( int i=0; i<DIM; ++i )
+			{
+				res.m[i] = m[i] * other.m[i];
+			}
+			return res;
+		}
+
+		//-----------------------------------------------------------------------------------------
+		inline V operator*( SCALAR s ) const
+		{
+			V res;
+			for( int i=0; i<DIM; ++i )
+			{
+				res.m[i] = m[i] * s;
+			}
+			return res;
+		}
+
+		//-----------------------------------------------------------------------------------------
+		inline V operator/( const V& other ) const
+		{
+			V res;
+			for( int i=0; i<DIM; ++i )
+			{
+				res.m[i] = m[i] / other.m[i];
+			}
+			return res;
+		}
+
+		//-----------------------------------------------------------------------------------------
+		inline V operator/( SCALAR s ) const
+		{
+			V res;
+			for( int i=0; i<DIM; ++i )
+			{
+				res.m[i] = m[i] / s;
+			}
+			return res;
+		}		
+
+        //-----------------------------------------------------------------------------------------
+        inline void operator*=( SCALAR s )
+        {
+            for( int i=0; i<DIM; ++i )
+            {
+                m[i] *= s;
+            }
+        }
+
+        //-----------------------------------------------------------------------------------------
+        inline void operator/=( SCALAR s )
+        {
+            for( int i=0; i<DIM; ++i )
+            {
+                m[i] /= s;
+            }
+        }
+
+        //-----------------------------------------------------------------------------------------
+		//! Per component min
+        //-----------------------------------------------------------------------------------------
+		inline static V min(const V& a, const V& b)
+		{
+			V r;
+			for (size_t i = 0; i < DIM; ++i)
+			{
+				r.m[i] = FMath::Min(a.m[i], b.m[i]);
+			}
+			return r;
+		}
+
+		inline static V max(const V& a, const V& b)
+		{
+			V r;
+			for (size_t i = 0; i < DIM; ++i)
+			{
+				r.m[i] = FMath::Max(a.m[i], b.m[i]);
+			}
+			return r;
+		}
+
+		//-----------------------------------------------------------------------------------------
+		//! Strict ordering of vectors.
+		//-----------------------------------------------------------------------------------------
+		inline bool
+		operator<(const V& other) const
+        {
+			bool equal = true;
+			bool res = false;
+			for( int i=0; equal && i<DIM; ++i )
+			{
+				if( m[i] < other.m[i] )
+				{
+					equal = false;
+					res = true;
+				}
+				else if( m[i] > other.m[i] )
+				{
+					equal = false;
+					res = false;
+				}
+			}
+			return res;
+		}
+
+
+        //-----------------------------------------------------------------------------------------
+        inline bool AllSmallerThan( const V& other ) const
+        {
+            bool res = true;
+            for( int i=0; i<DIM; ++i )
+            {
+                res &= m[i] < other.m[i];
+            }
+            return res;
+        }
+
+
+        //-----------------------------------------------------------------------------------------
+        inline bool AllSmallerOrEqualThan( const V& other ) const
+        {
+            bool res = true;
+            for( int i=0; i<DIM; ++i )
+            {
+                res &= m[i] <= other.m[i];
+            }
+            return res;
+        }
+
+		//-----------------------------------------------------------------------------------------
+		inline bool AllGreaterOrEqualThan( const V& other ) const
+		{
+			bool res = true;
+			for( int i=0; i<DIM; ++i )
+			{
+				res &= m[i] >= other.m[i];
+			}
+			return res;
+		}
+
+		//-----------------------------------------------------------------------------------------
+		inline bool AlmostEqual
+			(
+				const V& other,
+				SCALAR tol = std::numeric_limits<SCALAR>::min()
+			) const
+		{
+			bool res = true;
+			for( int i=0; i<DIM; ++i )
+			{
+				res &= fabs(m[i]-other.m[i])<tol;
+			}
+			return res;
+		}
+	};
+
+    //!
+    template <class SCALAR, int DIM>
+    inline vec<SCALAR,DIM> clamp( const vec<SCALAR,DIM>& v, SCALAR mi, SCALAR  ma )
+    {
+        vec<SCALAR,DIM> res;
+        for( int i=0; i<DIM; ++i )
+        {
+            res[i] = clamp( mi, ma, v[i] );
+        }
+        return res;
+    }
+
+    //!
+    template <class SCALAR, int DIM>
+    inline vec<SCALAR,DIM> lerp( const vec<SCALAR,DIM>& f0, const vec<SCALAR,DIM>& f1, float t )
+    {
+        const float s = 1.0f - t;
+        return f0 * s + f1 * t;
+    }
+
+	//!
+	template <class _SCALAR>
+	class vec2 : public vec<_SCALAR,2>
+	{
+	public:
+        typedef _SCALAR SCALAR;
+
+        //-----------------------------------------------------------------------------------------
+		inline vec2()
+		{
+			// Elements are cleared in the parent constructor
+		}
+
+        //-----------------------------------------------------------------------------------------
+        inline vec2( SCALAR x, SCALAR y )
+        {
+            this->m[0] = x;
+            this->m[1] = y;
+        }
+
+        //-----------------------------------------------------------------------------------------
+        inline vec2( const vec2& ) = default;
+
+		//-----------------------------------------------------------------------------------------
+		template<class SCALAR2>
+        inline vec2( const vec<SCALAR2,2>& v )
+        {
+			this->m[0] = v[0];
+			this->m[1] = v[1];
+		}
+
+        //-----------------------------------------------------------------------------------------
+        inline SCALAR& x()
+        {
+            return this->m[0];
+        }
+
+        //-----------------------------------------------------------------------------------------
+        inline SCALAR& y()
+        {
+            return this->m[1];
+        }
+
+        //-----------------------------------------------------------------------------------------
+        inline const SCALAR& x() const
+        {
+            return this->m[0];
+        }
+
+        //-----------------------------------------------------------------------------------------
+        inline const SCALAR& y() const
+        {
+            return this->m[1];
+        }
+	};
+
+    //!
+    typedef vec2<float> vec2f;
+
+	//!
+	template <class SCALAR>
+	class vec3 : public vec<SCALAR,3>
+	{
+	public:
+
+		//-----------------------------------------------------------------------------------------
+		inline vec3()
+		{
+			// Elements are cleared in the parent constructor
+		}
+
+		//-----------------------------------------------------------------------------------------
+		inline vec3( SCALAR x, SCALAR y, SCALAR z )
+		{
+			this->m[0] = x;
+			this->m[1] = y;
+			this->m[2] = z;
+		}
+
+		//-----------------------------------------------------------------------------------------
+		inline vec3(const vec<SCALAR, 3>& other)
+		{
+			this->m[0] = other[0];
+			this->m[1] = other[1];
+			this->m[2] = other[2];
+		}				
+
+		//-----------------------------------------------------------------------------------------
+		inline vec3(const FVector3f& other)
+		{
+			this->m[0] = other[0];
+			this->m[1] = other[1];
+			this->m[2] = other[2];
+		}
+
+		//-----------------------------------------------------------------------------------------
+		inline vec3(const FVector4f& other)
+		{
+			this->m[0] = other[0];
+			this->m[1] = other[1];
+			this->m[2] = other[2];
+		}
+
+        //-----------------------------------------------------------------------------------------
+        inline void Set( SCALAR x, SCALAR y, SCALAR z )
+        {
+            this->m[0] = x;
+            this->m[1] = y;
+            this->m[2] = z;
+        }
+
+
+        //-----------------------------------------------------------------------------------------
+        inline SCALAR& x()
+        {
+            return this->m[0];
+        }
+
+        //-----------------------------------------------------------------------------------------
+        inline SCALAR& y()
+        {
+            return this->m[1];
+        }
+
+        //-----------------------------------------------------------------------------------------
+        inline SCALAR& z()
+        {
+            return this->m[2];
+        }
+
+        //-----------------------------------------------------------------------------------------
+        inline const SCALAR& x() const
+        {
+            return this->m[0];
+        }
+
+        //-----------------------------------------------------------------------------------------
+        inline const SCALAR& y() const
+        {
+            return this->m[1];
+        }
+
+        //-----------------------------------------------------------------------------------------
+        inline const SCALAR& z() const
+        {
+            return this->m[2];
+        }
+
+        //-----------------------------------------------------------------------------------------
+        inline vec2<SCALAR> xy() const
+        {
+            return vec2<SCALAR>( this->m[0], this->m[1] );
+        }
+
+	};
+
+	//!
+	typedef vec3<float> vec3f;
+
+
+	//---------------------------------------------------------------------------------------------
+	//! Cross product of 3D vectors
+	//---------------------------------------------------------------------------------------------
+	template <class SCALAR>
+	inline vec3<SCALAR> cross( const vec3<SCALAR>& a, const vec3<SCALAR>& b )
+	{
+		return vec3<SCALAR>
+				(
+					a[1] * b[2] - a[2] * b[1],
+					a[2] * b[0] - a[0] * b[2],
+					a[0] * b[1] - a[1] * b[0]
+				);
+	}
+
+
+	//---------------------------------------------------------------------------------------------
+	//! Dot product of 3D vectors
+	//---------------------------------------------------------------------------------------------
+	template <class SCALAR, int DIM >
+	inline SCALAR dot( const vec<SCALAR, DIM>& a, const vec<SCALAR,DIM>& b )
+	{
+		SCALAR res = 0;
+
+		for ( int i=0; i<DIM; ++i )
+		{
+			res += a[i]*b[i];
+		}
+
+		return res;
+	}
+
+
+	//-----------------------------------------------------------------------------------------
+	//! Length of a vector
+	//-----------------------------------------------------------------------------------------
+	template <class SCALAR, int DIM>
+	inline SCALAR length( const vec<SCALAR,DIM>& v )
+	{
+		SCALAR l = dot( v, v );
+		l = sqrtf(l);
+		return l;
+	}
+
+
+	//---------------------------------------------------------------------------------------------
+	//! Normalise a vector
+	//---------------------------------------------------------------------------------------------
+	template <class SCALAR, int DIM>
+	inline vec<SCALAR,DIM> normalise( const vec<SCALAR,DIM>& a )
+	{
+		float l = length( a );
+		if ( l > std::numeric_limits<SCALAR>::epsilon() )
+		{
+			return a/l;
+		}
+		return a;
+	}
+
+
+	//---------------------------------------------------------------------------------------------
+	//! Normalise a vector, assuming that its length won't be null
+	//---------------------------------------------------------------------------------------------
+	template <class SCALAR, int DIM>
+	inline vec<SCALAR,DIM> normalise_unsafe( const vec<SCALAR,DIM>& a )
+	{
+		float l = length( a );
+		return a/l;
+	}
+
+	//---------------------------------------------------------------------------------------------
+	//! Point to plane distance
+	//---------------------------------------------------------------------------------------------
+	inline float point_to_plane_dist(const FVector3f& point, const FVector3f& plane_origin, const FVector3f& plane_normal)
+	{
+		return FVector3f::DotProduct((point - plane_origin), plane_normal);
+	}
+
+	//---------------------------------------------------------------------------------------------
+	//! Ray plane intersection
+	//---------------------------------------------------------------------------------------------
+	inline FVector3f ray_plane_intersection(const FVector3f& ray_start, const FVector3f& ray_end, const FVector3f& plane_origin, const FVector3f&plane_normal)
+	{
+		return ray_start + (ray_end - ray_start) *	(FVector3f::DotProduct((plane_origin - ray_start), plane_normal) / FVector3f::DotProduct((ray_end - ray_start), plane_normal));
+	}
+
+    //constexpr float point_in_plane_eps = 0.0001f;
+
+    inline bool point_in_face(const FVector3f& InVtx, const FVector3f& normal, FVector3f vertices[3], int& out_intersected_vert, int& out_intersected_edge_v0, int& out_intersected_edge_v1, float point_in_plane_eps = 0.0001f)
+    {
+		FVector3f SidePlaneNormal;
+		FVector3f Side;
+        bool intersected_edge[3] = { false, false, false };
+        int intersected_edge_count = 0;
+
+        for (int x = 0; x < 3; ++x)
+        {
+            // Create plane perpendicular to both this side and the polygon's normal.
+            Side = vertices[x] - vertices[(x - 1 < 0) ? 3 - 1 : x - 1];
+            SidePlaneNormal = FVector3f::CrossProduct(Side, normal.GetSafeNormal());
+            SidePlaneNormal.Normalize();
+
+            float dist = point_to_plane_dist(InVtx, vertices[x], SidePlaneNormal);
+
+            if (dist > point_in_plane_eps)
+            {
+                return false; // If point is not behind all the planes created by this polys edges, it's outside the poly.
+            }
+            else if (dist > -point_in_plane_eps)
+            {
+                intersected_edge[x] = true;
+                intersected_edge_count++;
+            }
+        }
+
+        out_intersected_vert = -1;
+        out_intersected_edge_v0 = -1;
+        out_intersected_edge_v1 = -1;
+
+        if (intersected_edge_count > 0)
+        {
+            if (intersected_edge_count == 1) // Only one edge intersected
+            {
+                if (intersected_edge[0])
+                {
+                    out_intersected_edge_v0 = 0;
+                    out_intersected_edge_v1 = 2;
+                }
+                else if (intersected_edge[1])
+                {
+                    out_intersected_edge_v0 = 0;
+                    out_intersected_edge_v1 = 1;
+                }
+                else if (intersected_edge[2])
+                {
+                    out_intersected_edge_v0 = 1;
+                    out_intersected_edge_v1 = 2;
+                }
+            }
+            else if (intersected_edge_count == 2) // Two edges intersected means it actually hit a vertex
+            {
+                if (intersected_edge[0] && intersected_edge[1])
+                {
+                    out_intersected_vert = 0;
+                }
+                else if (intersected_edge[1] && intersected_edge[2])
+                {
+                    out_intersected_vert = 1;
+                }
+                else if (intersected_edge[2] && intersected_edge[0])
+                {
+                    out_intersected_vert = 2;
+                }
+            }
+        }
+
+        return true;
+    }
+
+    //---------------------------------------------------------------------------------------------
+    inline bool rayIntersectsFace( const FVector3f& ray_start, const FVector3f& ray_end,
+                                   const FVector3f& v0, const FVector3f& v1, const FVector3f& v2,
+								   FVector3f& out_intersection,
+                                   int& out_intersected_vert, int& out_intersected_edge_v0,
+                                   int& out_intersected_edge_v1, float point_in_face_epsilon = 0.0001f)
+    {
+		FVector3f normal = FVector3f::CrossProduct(v1 - v0, v2 - v0);
+
+        // Does the ray cross the plane?
+        const float DistStart = point_to_plane_dist(ray_start, v0, normal);
+        const float DistEnd = point_to_plane_dist(ray_end, v0, normal);
+
+        if ((DistStart < 0 && DistEnd < 0) || (DistStart > 0 && DistEnd > 0))
+        {
+            return false;
+        }
+
+        // Get the intersection of the line and the plane.
+        out_intersection = ray_plane_intersection(ray_start, ray_end, v0, normal);
+
+        // \todo: review the equality comparison. Is it doing anything without a tolerance?
+        if (out_intersection == ray_start || out_intersection == ray_end)
+        {
+            return false;
+        }
+
+        // Check if the intersection point is actually on the poly.
+		FVector3f vertices[3] = { v0, v1, v2 };
+
+        return point_in_face( out_intersection, normal, vertices,
+                              out_intersected_vert,
+                              out_intersected_edge_v0, out_intersected_edge_v1, point_in_face_epsilon);
+    }
+
+
+    //---------------------------------------------------------------------------------------------
+    inline bool ray_plane_intersection( const FVector3f& ray_start, const FVector3f& ray_end,
+                                        const FVector3f& plane_origin, const FVector3f&plane_normal,
+                                        FVector3f& out_intersection, float& out_t )
+    {
+        float den = FVector3f::DotProduct(ray_end - ray_start, plane_normal);
+
+        if (den == 0.f)
+        {
+            return false;
+        }
+
+        out_t = FVector3f::DotProduct((plane_origin - ray_start), plane_normal) / den;
+        out_intersection = ray_start + (ray_end - ray_start) * out_t;
+
+        return true;
+    }
+
+
+    //---------------------------------------------------------------------------------------------
+    inline bool rayIntersectsFace2(const FVector3f& ray_start, const FVector3f& ray_end,
+                                   const FVector3f& v0, const FVector3f& v1, const FVector3f& v2,
+		FVector3f& out_intersection,
+                                   int& out_intersected_vert,
+                                   int& out_intersected_edge_v0, int& out_intersected_edge_v1,
+                                   float& out_t )
+    {
+		FVector3f normal = FVector3f::CrossProduct(v1 - v0, v2 - v0);
+
+        // Does the ray cross the plane?
+        const float DistStart = point_to_plane_dist(ray_start, v0, normal);
+        const float DistEnd = point_to_plane_dist(ray_end, v0, normal);
+
+        if ((DistStart < 0 && DistEnd < 0) || (DistStart > 0 && DistEnd > 0))
+        {
+            return false;
+        }
+
+        // Get the intersection of the line and the plane.
+        if (!ray_plane_intersection(ray_start, ray_end, v0, normal, out_intersection, out_t))
+        {
+            return false;
+        }
+
+        // \todo: review the equality comparison. Is it doing anything without a tolerance?
+        if (out_intersection == ray_start || out_intersection == ray_end)
+        {
+            return false;
+        }
+
+        // Check if the intersection point is actually on the poly.
+		FVector3f vertices[3] = { v0, v1, v2 };
+
+        return point_in_face( out_intersection, normal, vertices, out_intersected_vert,
+                              out_intersected_edge_v0, out_intersected_edge_v1 );
+    }
+
+	//---------------------------------------------------------------------------------------------
+	//! Approximated inverse square root
+	//---------------------------------------------------------------------------------------------
+	inline float rsqrt_approx( float number )
+	{
+		union MUTABLE_FLOAT_INT
+		{
+            int32 i;
+			float y;
+		} fi;
+		float x2;
+		const float threehalfs = 1.5F;
+
+		x2 = number * 0.5F;
+		fi.y  = number;
+		fi.i  = 0x5f3759df - ( fi.i >> 1 );
+		fi.y  = fi.y * ( threehalfs - ( x2 * fi.y * fi.y ) );
+		return fi.y;
+	}
+
+
+	//---------------------------------------------------------------------------------------------
+	//! Normalise a vector, approximately
+	//---------------------------------------------------------------------------------------------
+	template <class SCALAR, int DIM>
+	inline vec<SCALAR,DIM> normalise_approx( const vec<SCALAR,DIM>& a )
+	{
+		return a * rsqrt_approx( dot(a,a) );
+	}
+
+
+	//!
+    template <class VECTOR>
+    class box
+	{
+	public:
+
+		VECTOR min;
+		VECTOR size;
+
+	public:
+
+        //-----------------------------------------------------------------------------------------
+        static box<VECTOR> FromMinSize(VECTOR m, VECTOR s)
+        {
+            box<VECTOR> result;
+            result.min = m;
+            result.size = s;
+            return result;
+        }
+
+        //-----------------------------------------------------------------------------------------
+        inline bool Contains( const VECTOR& v ) const
+        {
+            return v.AllGreaterOrEqualThan(min) && v.AllSmallerThan(min+size);
+        }
+
+        //-----------------------------------------------------------------------------------------
+        inline bool ContainsInclusive( const VECTOR& v ) const
+        {
+            return v.AllGreaterOrEqualThan(min) && v.AllSmallerOrEqualThan(min+size);
+        }
+
+        //-----------------------------------------------------------------------------------------
+        //! It is inclusive: If the borders touch it is considered an intersection.
+        //-----------------------------------------------------------------------------------------
+        inline bool Intersects(const box<VECTOR>& other) const
+        {
+            bool result = true;
+            for (int d = 0; d < VECTOR::GetDim(); ++d)
+            {
+                result = result &&
+                         FMath::Abs((min[d] * 2 + size[d]) - (other.min[d] * 2 + other.size[d])) <=
+                             (size[d] + other.size[d]);
+            }
+
+            return result;
+        }
+
+        //-----------------------------------------------------------------------------------------
+        //! It is exclusive: If the borders touch it is not considered an intersection.
+        //-----------------------------------------------------------------------------------------
+        inline bool IntersectsExclusive(const box<VECTOR>& other) const
+        {
+            bool result = true;
+            for (int d = 0; d < VECTOR::GetDim(); ++d)
+            {
+                result = result &&
+                         FMath::Abs((min[d] * 2 + size[d]) - (other.min[d] * 2 + other.size[d])) <
+                             (size[d] + other.size[d]);
+            }
+
+            return result;
+        }
+
+        //-----------------------------------------------------------------------------------------
+        //! If the boxes don't intersect, at least one size of the returned box will be negative.
+        //-----------------------------------------------------------------------------------------
+        inline box<VECTOR> Intersect(const box<VECTOR>& other) const
+        {
+            box<VECTOR> result;
+            for (int d = 0; d < VECTOR::GetDim(); ++d)
+            {
+                result.min[d] = FMath::Max(min[d], other.min[d]);
+                result.size[d] =
+                    FMath::Min(typename VECTOR::SCALAR(min[d] + size[d] - result.min[d]),
+                             typename VECTOR::SCALAR(other.min[d] + other.size[d] - result.min[d]));
+            }
+            return result;
+        }
+
+        //-----------------------------------------------------------------------------------------
+		VECTOR Homogenize( const VECTOR& v ) const
+		{
+			return (v-min)/size;
+		}
+
+		//-----------------------------------------------------------------------------------------
+		inline void Bound(const VECTOR& v)
+		{
+			for (int d = 0; d < VECTOR::GetDim(); ++d)
+			{
+				if (v[d] < min[d])
+				{
+					size[d] += min[d] - v[d];
+					min[d] = v[d];
+				}
+				else if (v[d] - min[d] > size[d])
+				{
+					size[d] = v[d] - min[d];
+				}
+			}
+		}
+
+		//-----------------------------------------------------------------------------------------
+		inline void Bound3(const VECTOR& v)
+		{
+			for (int d = 0; d < 3; ++d)
+			{
+				if (v[d] < min[d])
+				{
+					size[d] += min[d] - v[d];
+					min[d] = v[d];
+				}
+				else if (v[d] - min[d] > size[d])
+				{
+					size[d] = v[d] - min[d];
+				}
+			}
+		}
+
+
+		//-----------------------------------------------------------------------------------------
+		inline void Bound(const box<VECTOR>& other)
+		{
+			Bound(other.min);
+			Bound(other.min + other.size);
+		}
+
+		//-----------------------------------------------------------------------------------------
+		inline void ShrinkToHalf()
+		{
+			min /= 2;
+			size /= 2;
+		}
+
+		//-----------------------------------------------------------------------------------------
+		inline bool IsEmpty() const
+		{
+			return size[0] <= 0 || size[1] <= 0;
+		}
+
+		//-----------------------------------------------------------------------------------------
+		//! Strict ordering of boxes.
+		//-----------------------------------------------------------------------------------------
+        inline bool operator<( const box<VECTOR>& other ) const
+        {
+            if( size < other.size )
+            {
+                return true;
+            }
+            else if( size == other.size )
+            {
+                return min < other.min;
+            }
+
+            return false;
+        }
+
+		//-----------------------------------------------------------------------------------------
+		inline bool operator==(const box<VECTOR>& other) const
+		{
+			return (size == other.size)
+				&&
+				(min == other.min);
+		}
+	};
+
+
+
+
+	//---------------------------------------------------------------------------------------------
+	//! Generic square matrix template.
+	//---------------------------------------------------------------------------------------------
+	template <class SCALAR, int DIM>
+    class mat
+	{
+	protected:
+
+		//! Rows
+		vec<SCALAR,DIM> m[ DIM ];
+
+	public:
+
+		//-----------------------------------------------------------------------------------------
+		typedef vec<SCALAR,DIM> V;
+		typedef mat<SCALAR,DIM> M;
+
+		//-----------------------------------------------------------------------------------------
+		static inline int GetDim()
+		{
+			return DIM;
+		}
+
+		//-----------------------------------------------------------------------------------------
+		//! 0 matrix
+		//-----------------------------------------------------------------------------------------
+		mat()
+		{
+		}
+
+		//-----------------------------------------------------------------------------------------
+		//! Build from a float array
+		//-----------------------------------------------------------------------------------------
+		explicit mat( const float* p )
+		{
+			for (int j=0; j<DIM; ++j)
+			{
+				for (int i=0; i<DIM; ++i)
+				{
+					m[j][i] = *p;
+					++p;
+				}
+			}
+		}
+
+		//-----------------------------------------------------------------------------------------
+		//! Initialise from another matrix with different size, filling the unknown with identites.
+		//-----------------------------------------------------------------------------------------
+		template<int DIM2>
+		explicit mat( const mat<SCALAR,DIM2>& other )
+		{
+			for (int i=0; i<FMath::Min(DIM,DIM2); ++i)
+			{
+				for (int j=0; j<FMath::Min(DIM,DIM2) ;++j)
+				{
+					m[i][j] = other[i][j];
+				}
+			}
+
+			for (int i=DIM2; i<DIM; ++i)
+			{
+				m[i][i] = 1;
+			}
+		}
+
+		//-----------------------------------------------------------------------------------------
+		//! Make a scale matrix.
+		//-----------------------------------------------------------------------------------------
+		template<int DIM2>
+		inline static M Scale( const vec<SCALAR,DIM2>& s )
+		{
+			M r;
+			for (int i=0; i<FMath::Min(DIM,DIM2); ++i)
+			{
+				r.m[i][i] = s[i];
+			}
+
+			for (int i=DIM2; i<DIM; ++i)
+			{
+				r.m[i][i] = 1;
+			}
+
+			return r;
+		}
+
+		//-----------------------------------------------------------------------------------------
+		//! Make a translation matrix.
+		//-----------------------------------------------------------------------------------------
+		template<int DIM2>
+		static inline M Translate( const vec<SCALAR,DIM2>& pos )
+		{
+			M r = Identity();
+
+			for (int i=0; i<FMath::Min(DIM,DIM2); ++i)
+			{
+				r.m[i][DIM-1] = pos[i];
+			}
+
+			return r;
+		}
+
+        //-----------------------------------------------------------------------------------------
+        //! Return an identity matrix.
+        //-----------------------------------------------------------------------------------------
+        static inline M Identity()
+        {
+            M res;
+            for (int i=0; i<DIM;++i)
+            {
+                res[i][i] = 1;
+            }
+            return res;
+        }
+
+        //-----------------------------------------------------------------------------------------
+        //! Return an identity matrix.
+        //-----------------------------------------------------------------------------------------
+        inline void SetIdentity()
+        {
+            *this = M::Identity();
+        }
+
+		//-----------------------------------------------------------------------------------------
+		//! Return a row
+		//-----------------------------------------------------------------------------------------
+        inline V& operator[] (int i)
+		{
+			check( i>=0 && i <DIM );
+			return m[i];
+		}
+
+		//-----------------------------------------------------------------------------------------
+		//! Return a row
+		//-----------------------------------------------------------------------------------------
+		inline const V& GetRow(int i) const
+		{
+			check( i>=0 && i <DIM );
+			return m[i];
+		}
+
+		//-----------------------------------------------------------------------------------------
+		//! Return a column
+		//-----------------------------------------------------------------------------------------
+		inline V GetColumn(int i) const
+		{
+			check( i>=0 && i <DIM );
+			V r;
+			for ( int j=0; j<DIM; ++j )
+			{
+				r[j] = m[j][i];
+			}
+			return r;
+		}
+
+		//-----------------------------------------------------------------------------------------
+        inline const V& operator[] (int i) const
+		{
+			check( i>=0 && i <DIM );
+			return m[i];
+		}
+
+		//-----------------------------------------------------------------------------------------
+		inline bool operator==( const M& other ) const
+		{
+			bool res = true;
+			for( int i=0; i<DIM; ++i )
+			{
+				res &= m[i] == other.m[i];
+			}
+			return res;
+		}
+
+		//-----------------------------------------------------------------------------------------
+		inline bool operator!=( const M& other ) const
+		{
+			bool res = false;
+			for( int i=0; i<DIM; ++i )
+			{
+				res |= m[i] != other.m[i];
+			}
+			return res;
+		}
+
+		//-----------------------------------------------------------------------------------------
+		inline V operator+( const M& other ) const
+		{
+			V res;
+			for( int i=0; i<DIM; ++i )
+			{
+				res.m[i] = m[i] + other.m[i];
+			}
+			return res;
+		}
+
+		//-----------------------------------------------------------------------------------------
+		inline void operator+=( const M& other )
+		{
+			for( int i=0; i<DIM; ++i )
+			{
+				m[i] += other.m[i];
+			}
+		}
+
+		//-----------------------------------------------------------------------------------------
+		inline V operator-( const M& other ) const
+		{
+			V res;
+			for( int i=0; i<DIM; ++i )
+			{
+				res.m[i] = m[i] - other.m[i];
+			}
+			return res;
+		}
+
+		//-----------------------------------------------------------------------------------------
+		inline void operator-=( const M& other )
+		{
+			for( int i=0; i<DIM; ++i )
+			{
+				m[i] -= other.m[i];
+			}
+		}
+
+        //-----------------------------------------------------------------------------------------
+        M operator*( const mat<SCALAR,DIM>& m2 ) const
+        {
+            M r;
+
+            for ( int i=0; i<DIM; ++i )
+            {
+                for ( int j=0; j<DIM; ++j )
+                {
+                    r.m[i][j] = dot( m[i], m2.GetColumn(j) );
+                }
+            }
+
+            return r;
+        }
+
+		//-----------------------------------------------------------------------------------------
+		V Transform(const V& v) const
+		{
+			V r;
+
+			for (int i = 0; i < DIM; ++i)
+			{
+				r[i] = dot(m[i], v);
+			}
+
+			return r;
+		}
+
+		//-----------------------------------------------------------------------------------------
+		V operator*(const V& v) const
+		{
+			return Transform(v);
+		}
+
+        //-----------------------------------------------------------------------------------------
+        SCALAR GetDeterminant() const
+        {
+            SCALAR result=0;
+
+            for ( int i=0; i<DIM; ++i )
+            {
+                SCALAR accum=1;
+                for ( int j=0; j<DIM; ++j )
+                {
+                    accum *= m[(i+j)%DIM][j];
+                }
+                result+=accum;
+
+                accum=1;
+                for ( int j=0; j<DIM; ++j )
+                {
+                    accum *= m[(DIM+i-j)%DIM][j];
+                }
+                result-=accum;
+            }
+
+            return result;
+        }
+
+
+		//-----------------------------------------------------------------------------------------
+		//! Strict ordering of matrices.
+		//-----------------------------------------------------------------------------------------
+		inline bool operator<( const M& other ) const
+		{
+			bool equal = true;
+			bool res = false;
+			for( int i=0; equal && i<DIM; ++i )
+			{
+				if( m[i] < other.m[i] )
+				{
+					equal = false;
+					res = true;
+				}
+				else if( m[i] > other.m[i] )
+				{
+					equal = false;
+					res = false;
+				}
+			}
+			return res;
+		}
+
+
+		//-----------------------------------------------------------------------------------------
+		inline bool AlmostEqual
+			(
+				const V& other,
+				SCALAR tol = std::numeric_limits<SCALAR>::min()
+			) const
+		{
+			bool res = true;
+			for( int i=0; i<DIM; ++i )
+			{
+				res &= m[i].AlmostEqual(other.m[i],tol);
+			}
+			return res;
+		}
+
+	};
+
+
+	//!
+	template <class SCALAR>
+	class mat3 : public mat<SCALAR,3>
+	{
+	public:
+
+		mat3()
+		{
+		}
+
+        mat3( SCALAR a0, SCALAR a1, SCALAR a2,
+              SCALAR a3, SCALAR a4, SCALAR a5,
+              SCALAR a6, SCALAR a7, SCALAR a8 )
+        {
+            this->m[0][0] = a0;
+            this->m[0][1] = a1;
+            this->m[0][2] = a2;
+            this->m[1][0] = a3;
+            this->m[1][1] = a4;
+            this->m[1][2] = a5;
+            this->m[2][0] = a6;
+            this->m[2][1] = a7;
+            this->m[2][2] = a8;
+        }
+
+        mat3( vec3<SCALAR> r0, vec3<SCALAR> r1, vec3<SCALAR> r2 )
+        {
+            this->m[0] = r0;
+            this->m[1] = r1;
+            this->m[2] = r2;
+        }
+
+
+	};
+
+
+	//!
+	typedef mat3<float> mat3f;
+    typedef mat<float,4> mat4f;
+
+	//---------------------------------------------------------------------------------------------
+	// Temp while transitioning to UE math.
+	//---------------------------------------------------------------------------------------------
+	inline FMatrix44f ToUnreal(const mat4f& m)
+	{
+		FMatrix44f ueMat;
+		for (int i = 0; i < 4; ++i)
+			for (int j = 0; j < 4; ++j)
+				ueMat.M[i][j] = m[i][j];
+		return ueMat;
+	}
+
+
+	inline FVector3f ToUnreal(const vec3f& v)
+	{
+		return FVector3f(v.x(), v.y(), v.z());
+	}
+
+
+	inline vec3f FromUnreal(const FVector3f& v)
+	{
+		return vec3f(v.X, v.Y, v.Z);
+	}
+
+
+}
+
diff --git a/Engine/Plugins/Experimental/Mutable/Source/MutableRuntime/Public/MuR/System.h b/Engine/Plugins/Experimental/Mutable/Source/MutableRuntime/Public/MuR/System.h
--- a/Engine/Plugins/Experimental/Mutable/Source/MutableRuntime/Public/MuR/System.h
+++ b/Engine/Plugins/Experimental/Mutable/Source/MutableRuntime/Public/MuR/System.h
@@ -26,6 +26,25 @@
 #else
 #endif
 
+#include "System.generated.h"
+
+
+/** Despite being an UEnum, this is not always version-serialized (in MutableTools).
+ * Beware of changing the enum options or order.
+ */
+UENUM()
+enum class ETextureCompressionStrategy : uint8
+{
+	/** Don't change the generated format. */
+	None,
+
+	/** If a texture depends on run-time parameters for an object state, don't compress. */
+	DontCompressRuntime,
+
+	/** Never compress the textures for this state. */
+	NeverCompress
+};
+
 
 namespace mu
 {
@@ -41,35 +60,21 @@ namespace mu
 
 
 
-    //! List of critical errors that may happen during execution of mutable code
-    enum class Error
-    {
-        //! No error happened.
-        None = 0,
-
-        //! A memory allocation failed
-        OutOfMemory,
-
-
-        //! A file was missing or the data was corrupted.
-        StreamingError,
-
-        //! The necessary functionality is not supported in this version of Mutable.
-        Unsupported,
+	/** */
+	enum class EExecutionStrategy : uint8
+	{
+		/** Undefined. */
+		None = 0,
 
-        //! Utility value with the number of error types.
-        Count
-    };
+		/** Always try to run operations that reduce working memory first. */
+		MinimizeMemory,
 
-    //! \brief Get an error code for the most critical error that happened since mutable
-    //! initialisation.
-    //! \ingroup runtime
-	MUTABLERUNTIME_API extern Error GetError();
+		/** Always try to run operations that unlock more operations first. */
+		MaximizeConcurrency,
 
-    //! \brief Remove the last error code, setting it to Error::None. This should obnly be used
-    //! in case of possible recovery from one of the critical errors, which is very unlikely.
-    //! \ingroup runtime
-	MUTABLERUNTIME_API extern void ClearError();
+		/** Utility value with the number of error types. */
+		Count
+	};
 
 
     //! \brief Interface to request external images used as parameters.
diff --git a/Engine/Plugins/Experimental/Mutable/Source/MutableTools/Private/MuT/ASTOpImageCompose.cpp b/Engine/Plugins/Experimental/Mutable/Source/MutableTools/Private/MuT/ASTOpImageCompose.cpp
--- a/Engine/Plugins/Experimental/Mutable/Source/MutableTools/Private/MuT/ASTOpImageCompose.cpp
+++ b/Engine/Plugins/Experimental/Mutable/Source/MutableTools/Private/MuT/ASTOpImageCompose.cpp
@@ -237,7 +237,7 @@ mu::Ptr<ASTOp> ASTOpImageCompose::OptimiseSemantic(const FModelOptimizationOptio
 		// Constant single-block full layout?
 		if (pLayout->GetBlockCount() == 1
 			&&
-			pLayout->m_blocks[0].m_min == vec2<uint16>(0, 0)
+			pLayout->m_blocks[0].m_min == UE::Math::TIntVector2<uint16>(0, 0)
 			&&
 			pLayout->m_blocks[0].m_size == pLayout->m_size
 			&&
diff --git a/Engine/Plugins/Experimental/Mutable/Source/MutableTools/Private/MuT/ASTOpImageLayer.cpp b/Engine/Plugins/Experimental/Mutable/Source/MutableTools/Private/MuT/ASTOpImageLayer.cpp
--- a/Engine/Plugins/Experimental/Mutable/Source/MutableTools/Private/MuT/ASTOpImageLayer.cpp
+++ b/Engine/Plugins/Experimental/Mutable/Source/MutableTools/Private/MuT/ASTOpImageLayer.cpp
@@ -512,7 +512,8 @@ namespace mu
 					Ptr<ASTOpImagePatch> patch = new ASTOpImagePatch();
 					patch->base = baseAt;
 					patch->patch = newLayer;
-					patch->location = maskUsage.min;
+					patch->location[0] = maskUsage.min[0];
+					patch->location[1] = maskUsage.min[1];
 					at = patch;
 				}
 			}
diff --git a/Engine/Plugins/Experimental/Mutable/Source/MutableTools/Private/MuT/ASTOpImagePatch.h b/Engine/Plugins/Experimental/Mutable/Source/MutableTools/Private/MuT/ASTOpImagePatch.h
--- a/Engine/Plugins/Experimental/Mutable/Source/MutableTools/Private/MuT/ASTOpImagePatch.h
+++ b/Engine/Plugins/Experimental/Mutable/Source/MutableTools/Private/MuT/ASTOpImagePatch.h
@@ -20,7 +20,7 @@ struct FProgram;
 
 		ASTChild base;
 		ASTChild patch;
-		vec2<uint16> location;
+		UE::Math::TIntVector2<uint16> location;
 
 	public:
 
diff --git a/Engine/Plugins/Experimental/Mutable/Source/MutableTools/Private/MuT/ASTOpMeshFormat.cpp b/Engine/Plugins/Experimental/Mutable/Source/MutableTools/Private/MuT/ASTOpMeshFormat.cpp
--- a/Engine/Plugins/Experimental/Mutable/Source/MutableTools/Private/MuT/ASTOpMeshFormat.cpp
+++ b/Engine/Plugins/Experimental/Mutable/Source/MutableTools/Private/MuT/ASTOpMeshFormat.cpp
@@ -387,7 +387,6 @@ mu::Ptr<ASTOp> Sink_MeshFormatAST::Visit(const mu::Ptr<ASTOp>& at, const ASTOpMe
 		break;
 
 	case OP_TYPE::ME_DIFFERENCE:
-	case OP_TYPE::ME_SUBTRACT:
 
 	default:
 		if (at != m_initialSource)
diff --git a/Engine/Plugins/Experimental/Mutable/Source/MutableTools/Private/MuT/CodeGenerator.cpp b/Engine/Plugins/Experimental/Mutable/Source/MutableTools/Private/MuT/CodeGenerator.cpp
--- a/Engine/Plugins/Experimental/Mutable/Source/MutableTools/Private/MuT/CodeGenerator.cpp
+++ b/Engine/Plugins/Experimental/Mutable/Source/MutableTools/Private/MuT/CodeGenerator.cpp
@@ -70,7 +70,6 @@
 #include "MuT/NodeObject.h"
 #include "MuT/NodeObjectGroupPrivate.h"
 #include "MuT/NodeObjectNewPrivate.h"
-#include "MuT/NodeObjectStatePrivate.h"
 #include "MuT/NodePatchImagePrivate.h"
 #include "MuT/NodePatchMesh.h"
 #include "MuT/NodePrivate.h"
@@ -1209,8 +1208,9 @@ namespace mu
 
                     // Generate the image
                     IMAGE_STATE newState;
-                    newState.m_imageSize = desc.m_size;
-                    newState.m_imageRect = box<vec2<int>>::FromMinSize( vec2<int>(0,0), desc.m_size );
+					newState.m_imageSize[0] = desc.m_size[0];
+					newState.m_imageSize[1] = desc.m_size[1];
+					newState.m_imageRect = box<UE::Math::TIntVector2<int>>::FromMinSize(UE::Math::TIntVector2<int>(0,0), UE::Math::TIntVector2<int>(desc.m_size) );
                     newState.m_layoutBlockId = -1;
                     newState.m_pLayout = nullptr;
                     m_imageState.Add( newState );
@@ -1350,7 +1350,7 @@ namespace mu
                         for( int b=0; b<pLayout->GetBlockCount(); ++b )
                         {
                             // Block in layout grid units
-                            box< vec2<int> > rectInCells;
+                            box< UE::Math::TIntVector2<int32> > rectInCells;
                             pLayout->GetBlock
                                 (
                                     b,
@@ -1359,7 +1359,7 @@ namespace mu
                                 );
 
                             // Transform to pixels
-                            box< vec2<int> > rect = rectInCells;
+                            box< UE::Math::TIntVector2<int32> > rect = rectInCells;
                             rect.min[0] *= blockSizeX;
                             rect.min[1] *= blockSizeY;
                             rect.size[0] *= blockSizeX;
@@ -1367,7 +1367,7 @@ namespace mu
 
                             // Generate the image
                             IMAGE_STATE newState;
-                            newState.m_imageSize = desc.m_size;
+                            newState.m_imageSize = FIntVector2(desc.m_size);
                             newState.m_imageRect = rect;
                             newState.m_layoutBlockId = pLayout->m_blocks[b].m_id;
                             newState.m_pLayout = pLayout;
@@ -1395,8 +1395,6 @@ namespace mu
                             // Compose layout operation
                             Ptr<ASTOpImageCompose> composeOp = new ASTOpImageCompose();
                             composeOp->Layout = meshResults.layoutOps[ layout ];
-                            // Direct compression support
-                            //composeOp.args.ImageCompose.base = GenerateImageUncompressed( imageAd );
                             composeOp->Base = imageAd;
 							EImageFormat baseFormat = imageAd->GetImageDesc().m_format;
                             composeOp->BlockImage = GenerateImageFormat( blockAd, baseFormat );
@@ -1451,7 +1449,7 @@ namespace mu
                                         for (int b=0; b<pExtendLayout->GetBlockCount(); ++b)
                                         {
                                             // Block in layout grid units
-                                            box< vec2<int> > blockRect;
+                                            box< UE::Math::TIntVector2<int32> > blockRect;
                                             pExtendLayout->GetBlock
                                                 (
                                                     b,
@@ -1460,7 +1458,7 @@ namespace mu
                                                 );
 
                                             // Transform to pixels
-                                            box< vec2<int> > rect = blockRect;
+                                            box< UE::Math::TIntVector2<int32> > rect = blockRect;
                                             rect.min[0] = (rect.min[0] * extendDesc.m_size[0]) / extlayout[0];
                                             rect.min[1] = (rect.min[1] * extendDesc.m_size[1]) / extlayout[1];
                                             rect.size[0] = (rect.size[0] * extendDesc.m_size[0]) / extlayout[0];
@@ -1468,7 +1466,7 @@ namespace mu
 
                                             // Generate the image block
                                             IMAGE_STATE newState;
-                                            newState.m_imageSize = extendDesc.m_size;
+                                            newState.m_imageSize = FIntVector2(extendDesc.m_size);
                                             newState.m_imageRect = rect;
                                             newState.m_layoutBlockId = pExtendLayout->m_blocks[b].m_id;
                                             newState.m_pLayout = pExtendLayout;
@@ -1834,22 +1832,6 @@ namespace mu
     }
 
 
-    //---------------------------------------------------------------------------------------------
-    Ptr<ASTOp> CodeGenerator::Visit( const NodeObjectState::Private& node )
-    {
-        // Generate the source object
-        Ptr<ASTOp> source;
-
-        if (node.m_pSource)
-        {
-            // Generate the base object where we will add the state.
-            source = Generate( node.m_pSource );
-        }
-
-        return source;
-    }
-
-
     //---------------------------------------------------------------------------------------------
     Ptr<ASTOp> CodeGenerator::Visit( const NodeObjectNew::Private& node )
     {
diff --git a/Engine/Plugins/Experimental/Mutable/Source/MutableTools/Private/MuT/CodeGenerator.h b/Engine/Plugins/Experimental/Mutable/Source/MutableTools/Private/MuT/CodeGenerator.h
--- a/Engine/Plugins/Experimental/Mutable/Source/MutableTools/Private/MuT/CodeGenerator.h
+++ b/Engine/Plugins/Experimental/Mutable/Source/MutableTools/Private/MuT/CodeGenerator.h
@@ -31,7 +31,6 @@
 #include "MuT/NodeMeshTable.h"
 #include "MuT/NodeObjectGroup.h"
 #include "MuT/NodeObjectNew.h"
-#include "MuT/NodeObjectState.h"
 #include "MuT/NodePatchImage.h"
 #include "MuT/NodeProjector.h"
 #include "MuT/NodeString.h"
@@ -92,7 +91,6 @@ namespace mu
 	class NodeMeshMakeMorph;
 	class NodeMeshMorph;
 	class NodeMeshReshape;
-	class NodeMeshSubtract;
 	class NodeMeshSwitch;
 	class NodeMeshTransform;
 	class NodeMeshVariation;
@@ -107,7 +105,7 @@ namespace mu
 	class NodeScalarVariation;
 	class NodeStringConstant;
 	class NodeStringParameter;
-	struct OBJECT_STATE;
+	struct FObjectState;
 	struct FProgram;
 
 
@@ -122,7 +120,6 @@ namespace mu
                           public Visitor<NodeComponentEdit::Private, Ptr<ASTOp>, true>,
                           public Visitor<NodeLOD::Private, Ptr<ASTOp>, true>,
                           public Visitor<NodeObjectNew::Private, Ptr<ASTOp>, true>,
-                          public Visitor<NodeObjectState::Private, Ptr<ASTOp>, true>,
                           public Visitor<NodeObjectGroup::Private, Ptr<ASTOp>, true>,
                           public Visitor<NodePatchImage::Private, Ptr<ASTOp>, true>
     {
@@ -143,7 +140,6 @@ namespace mu
         Ptr<ASTOp> Visit( const NodeComponentEdit::Private& ) override;
         Ptr<ASTOp> Visit( const NodeLOD::Private& ) override;
         Ptr<ASTOp> Visit( const NodeObjectNew::Private& ) override;
-        Ptr<ASTOp> Visit( const NodeObjectState::Private& ) override;
         Ptr<ASTOp> Visit( const NodeObjectGroup::Private& ) override;
         Ptr<ASTOp> Visit( const NodePatchImage::Private& ) override;
 
@@ -192,8 +188,8 @@ namespace mu
             // This reference has to be the smart pointer to avoid memory aliasing, keeping
             // processed nodes alive.
             NodePtrConst pNode;
-            vec2<int> imageSize;
-            box< vec2<int> > imageRect;
+			UE::Math::TIntVector2<int> imageSize;
+            box< UE::Math::TIntVector2<int> > imageRect;
             int state = -1;
 			TArray<mu::string> activeTags;
 			TArray<LayoutPtrConst> overrideLayouts;
@@ -219,7 +215,7 @@ namespace mu
         int m_currentStateIndex = -1;
 
         //! After the entire code generation this contains the information about all the states
-        typedef TArray< std::pair<OBJECT_STATE, Ptr<ASTOp>> > StateList;
+        typedef TArray< std::pair<FObjectState, Ptr<ASTOp>> > StateList;
         StateList m_states;
 
     private:
@@ -241,8 +237,8 @@ namespace mu
         //! image that we are generating.
         struct IMAGE_STATE
         {
-            vec2<int> m_imageSize;
-            box< vec2<int> > m_imageRect;
+            UE::Math::TIntVector2<int32> m_imageSize;
+            box< UE::Math::TIntVector2<int32> > m_imageRect;
             int32 m_layoutBlockId;
             LayoutPtrConst m_pLayout;
         };
@@ -535,7 +531,6 @@ namespace mu
         void GenerateMesh_Fragment(const FMeshGenerationOptions&, FMeshGenerationResult&, const NodeMeshFragment* );
         void GenerateMesh_Interpolate(const FMeshGenerationOptions&, FMeshGenerationResult&, const NodeMeshInterpolate* );
         void GenerateMesh_Switch(const FMeshGenerationOptions&, FMeshGenerationResult&, const NodeMeshSwitch* );
-        void GenerateMesh_Subtract(const FMeshGenerationOptions&, FMeshGenerationResult&, const NodeMeshSubtract* );
         void GenerateMesh_Transform(const FMeshGenerationOptions&, FMeshGenerationResult&, const NodeMeshTransform* );
         void GenerateMesh_ClipMorphPlane(const FMeshGenerationOptions&, FMeshGenerationResult&, const NodeMeshClipMorphPlane* );
         void GenerateMesh_ClipWithMesh(const FMeshGenerationOptions&, FMeshGenerationResult&, const NodeMeshClipWithMesh* );
diff --git a/Engine/Plugins/Experimental/Mutable/Source/MutableTools/Private/MuT/CodeGenerator_Color.cpp b/Engine/Plugins/Experimental/Mutable/Source/MutableTools/Private/MuT/CodeGenerator_Color.cpp
--- a/Engine/Plugins/Experimental/Mutable/Source/MutableTools/Private/MuT/CodeGenerator_Color.cpp
+++ b/Engine/Plugins/Experimental/Mutable/Source/MutableTools/Private/MuT/CodeGenerator_Color.cpp
@@ -310,10 +310,12 @@ namespace mu
 			// We take whatever size will be produced
 			FImageDesc desc = CalculateImageDesc(*pSource->GetBasePrivate());
 			IMAGE_STATE newState;
-			newState.m_imageSize = desc.m_size;
+			newState.m_imageSize[0] = desc.m_size[0];
+			newState.m_imageSize[1] = desc.m_size[1];
 			newState.m_imageRect.min[0] = 0;
 			newState.m_imageRect.min[1] = 0;
-			newState.m_imageRect.size = desc.m_size;
+			newState.m_imageRect.size[0] = desc.m_size[0];
+			newState.m_imageRect.size[1] = desc.m_size[1];
 			m_imageState.Add(newState);
 
 			// Generate
diff --git a/Engine/Plugins/Experimental/Mutable/Source/MutableTools/Private/MuT/CodeGenerator_FirstPass.cpp b/Engine/Plugins/Experimental/Mutable/Source/MutableTools/Private/MuT/CodeGenerator_FirstPass.cpp
--- a/Engine/Plugins/Experimental/Mutable/Source/MutableTools/Private/MuT/CodeGenerator_FirstPass.cpp
+++ b/Engine/Plugins/Experimental/Mutable/Source/MutableTools/Private/MuT/CodeGenerator_FirstPass.cpp
@@ -25,7 +25,6 @@
 #include "MuT/NodeObject.h"
 #include "MuT/NodeObjectGroupPrivate.h"
 #include "MuT/NodeObjectNewPrivate.h"
-#include "MuT/NodeObjectStatePrivate.h"
 #include "MuT/NodePrivate.h"
 #include "MuT/NodeSurface.h"
 #include "MuT/NodeSurfaceEditPrivate.h"
@@ -152,7 +151,7 @@ namespace mu
 
         if ( m_states.IsEmpty() )
         {
-            OBJECT_STATE data;
+            FObjectState data;
             data.m_name = "Default";
             m_states.Emplace( data, root );
         }
@@ -634,28 +633,6 @@ namespace mu
  	}
 
 
-    //---------------------------------------------------------------------------------------------
-    Ptr<ASTOp> FirstPassGenerator::Visit( const NodeObjectState::Private& node )
-    {
-        // Generate the source object
-        Ptr<ASTOp> source;
-
-        if (node.m_pSource)
-        {
-            // Process the source
-            node.m_pSource->GetBasePrivate()->Accept(*this);
-
-            if (node.m_pRoot)
-            {
-                // Remember the new state with the new root.
-                m_states.Emplace( node.m_state, node.m_pRoot->GetBasePrivate() );
-            }
-        }
-
-        return source;
-    }
-
-
 	//---------------------------------------------------------------------------------------------
     Ptr<ASTOp> FirstPassGenerator::Visit(const NodePatchMesh::Private&)
 	{
diff --git a/Engine/Plugins/Experimental/Mutable/Source/MutableTools/Private/MuT/CodeGenerator_FirstPass.h b/Engine/Plugins/Experimental/Mutable/Source/MutableTools/Private/MuT/CodeGenerator_FirstPass.h
--- a/Engine/Plugins/Experimental/Mutable/Source/MutableTools/Private/MuT/CodeGenerator_FirstPass.h
+++ b/Engine/Plugins/Experimental/Mutable/Source/MutableTools/Private/MuT/CodeGenerator_FirstPass.h
@@ -17,7 +17,6 @@
 #include "MuT/NodeModifierMeshClipWithMesh.h"
 #include "MuT/NodeObjectGroup.h"
 #include "MuT/NodeObjectNew.h"
-#include "MuT/NodeObjectState.h"
 #include "MuT/NodePatchMesh.h"
 #include "MuT/NodeSurfaceEdit.h"
 #include "MuT/NodeSurfaceNew.h"
@@ -31,7 +30,7 @@
 namespace mu
 {
 	class Layout;
-	struct OBJECT_STATE;
+	struct FObjectState;
 
 	//---------------------------------------------------------------------------------------------
 	//! First pass of the code generation process.
@@ -50,7 +49,6 @@ namespace mu
         public Visitor<NodeLOD::Private, Ptr<ASTOp>, true>,
         public Visitor<NodeObjectNew::Private, Ptr<ASTOp>, true>,
         public Visitor<NodeObjectGroup::Private, Ptr<ASTOp>, true>,
-        public Visitor<NodeObjectState::Private, Ptr<ASTOp>, true>,
         public Visitor<NodePatchMesh::Private, Ptr<ASTOp>, true>,
 
         public Visitor<NodeModifierMeshClipMorphPlane::Private, Ptr<ASTOp>, true>,
@@ -182,7 +180,7 @@ namespace mu
 
         //! Accumulate the model states found while generating code, with their generated root
         //! nodes.
-        typedef TArray< TPair<OBJECT_STATE, const Node::Private*> > StateList;
+        typedef TArray< TPair<FObjectState, const Node::Private*> > StateList;
         StateList m_states;
 
 	public:
@@ -195,7 +193,6 @@ namespace mu
         Ptr<ASTOp> Visit(const NodeLOD::Private&) override;
         Ptr<ASTOp> Visit(const NodeObjectNew::Private& node) override;
         Ptr<ASTOp> Visit(const NodeObjectGroup::Private&) override;
-        Ptr<ASTOp> Visit(const NodeObjectState::Private&) override;
         Ptr<ASTOp> Visit(const NodeModifierMeshClipMorphPlane::Private&) override;
         Ptr<ASTOp> Visit(const NodeModifierMeshClipWithMesh::Private&) override;
         Ptr<ASTOp> Visit(const NodeModifierMeshClipDeform::Private&) override;
diff --git a/Engine/Plugins/Experimental/Mutable/Source/MutableTools/Private/MuT/CodeGenerator_Image.cpp b/Engine/Plugins/Experimental/Mutable/Source/MutableTools/Private/MuT/CodeGenerator_Image.cpp
--- a/Engine/Plugins/Experimental/Mutable/Source/MutableTools/Private/MuT/CodeGenerator_Image.cpp
+++ b/Engine/Plugins/Experimental/Mutable/Source/MutableTools/Private/MuT/CodeGenerator_Image.cpp
@@ -136,7 +136,7 @@ namespace mu
 			IMAGE_STATE newState;
 			newState.m_imageSize[0] = desc.m_size[0] ? desc.m_size[0] : 256;
 			newState.m_imageSize[1] = desc.m_size[1] ? desc.m_size[1] : 256;
-			newState.m_imageRect.size = desc.m_size;
+			newState.m_imageRect.size = FIntVector2(desc.m_size);
 			newState.m_imageRect.min[0] = 0;
 			newState.m_imageRect.min[1] = 0;
 			newState.m_layoutBlockId = -1;
@@ -225,7 +225,7 @@ namespace mu
             m_pErrorLog->GetPrivate()->Add( "Constant image not set.", ELMT_WARNING, node.m_errorContext );
         }
 
-        vec2<int> imageSize( pImage->GetSizeX(), pImage->GetSizeY() );
+		FIntVector2 imageSize( pImage->GetSizeX(), pImage->GetSizeY() );
 
         // The constant image size may be different than the parent rect we are generating
         // In that case we need to crop the proportional part and the code generator will
@@ -234,9 +234,11 @@ namespace mu
 
         // Order of the operations is important: multiply first to avoid losing precision.
         // It will not overflow since image sizes are limited to 16 bit
-		vec2<int> RectDivisor = vec2<int>::max(vec2<int>(1,1),m_imageState.Last().m_imageSize);
-        cropRect.min = ( m_imageState.Last().m_imageRect.min * imageSize ) / RectDivisor;
-        cropRect.size = ( m_imageState.Last().m_imageRect.size * imageSize ) / RectDivisor;
+		FIntVector2 RectDivisor = FIntVector2( FMath::Max(1, m_imageState.Last().m_imageSize[0]), FMath::Max(1, m_imageState.Last().m_imageSize[1]));
+		cropRect.min[0] = (m_imageState.Last().m_imageRect.min[0] * imageSize[0]) / RectDivisor[0];
+		cropRect.min[1] = (m_imageState.Last().m_imageRect.min[1] * imageSize[1]) / RectDivisor[1];
+		cropRect.size[0] = (m_imageState.Last().m_imageRect.size[0] * imageSize[0]) / RectDivisor[0];
+		cropRect.size[1] = (m_imageState.Last().m_imageRect.size[1] * imageSize[1]) / RectDivisor[1];
 
         //check( cropRect.size[0]>0 && cropRect.size[1]>0 );
         cropRect.size[0] = FMath::Max( cropRect.size[0], 1 );
@@ -1547,17 +1549,16 @@ namespace mu
 
             // We take whatever size will be produced
             FImageDesc desc = CalculateImageDesc( *node.m_pImage->GetBasePrivate() );
-            newState.m_imageSize = desc.m_size;
-            newState.m_imageRect.min[0] = 0;
-            newState.m_imageRect.min[1] = 0;
-            newState.m_imageRect.size = desc.m_size;
+            newState.m_imageSize = FIntVector2(desc.m_size);
+            newState.m_imageRect.min = FIntVector2(0,0);
+            newState.m_imageRect.size = FIntVector2(desc.m_size);
             newState.m_layoutBlockId = -1;
             m_imageState.Add( newState );
 
             // Generate
             op->image = Generate( node.m_pImage.get() );
-			op->SourceSizeX = desc.m_size.x();
-			op->SourceSizeY = desc.m_size.y();
+			op->SourceSizeX = desc.m_size[0];
+			op->SourceSizeY = desc.m_size[1];
 
             // Restore rect
             m_imageState.Pop();
diff --git a/Engine/Plugins/Experimental/Mutable/Source/MutableTools/Private/MuT/CodeGenerator_Mesh.cpp b/Engine/Plugins/Experimental/Mutable/Source/MutableTools/Private/MuT/CodeGenerator_Mesh.cpp
--- a/Engine/Plugins/Experimental/Mutable/Source/MutableTools/Private/MuT/CodeGenerator_Mesh.cpp
+++ b/Engine/Plugins/Experimental/Mutable/Source/MutableTools/Private/MuT/CodeGenerator_Mesh.cpp
@@ -63,7 +63,6 @@
 #include "MuT/NodeMeshMorphPrivate.h"
 #include "MuT/NodeMeshReshape.h"
 #include "MuT/NodeMeshReshapePrivate.h"
-#include "MuT/NodeMeshSubtract.h"
 #include "MuT/NodeMeshSwitch.h"
 #include "MuT/NodeMeshSwitchPrivate.h"
 #include "MuT/NodeMeshTable.h"
@@ -319,7 +318,6 @@ class Node;
 		case NodeMesh::EType::Fragment: GenerateMesh_Fragment(InOptions, OutResult, static_cast<const NodeMeshFragment*>(Node)); break;
 		case NodeMesh::EType::Interpolate: GenerateMesh_Interpolate(InOptions, OutResult, static_cast<const NodeMeshInterpolate*>(Node)); break;
 		case NodeMesh::EType::Switch: GenerateMesh_Switch(InOptions, OutResult, static_cast<const NodeMeshSwitch*>(Node)); break;
-		case NodeMesh::EType::Subtract: GenerateMesh_Subtract(InOptions, OutResult, static_cast<const NodeMeshSubtract*>(Node)); break;
 		case NodeMesh::EType::Transform: GenerateMesh_Transform(InOptions, OutResult, static_cast<const NodeMeshTransform*>(Node)); break;
 		case NodeMesh::EType::ClipMorphPlane: GenerateMesh_ClipMorphPlane(InOptions, OutResult, static_cast<const NodeMeshClipMorphPlane*>(Node)); break;
 		case NodeMesh::EType::ClipWithMesh: GenerateMesh_ClipWithMesh(InOptions, OutResult, static_cast<const NodeMeshClipWithMesh*>(Node)); break;
@@ -1103,15 +1101,6 @@ class Node;
     }
 
 
-    //---------------------------------------------------------------------------------------------
-    void CodeGenerator::GenerateMesh_Subtract(const FMeshGenerationOptions& InOptions, FMeshGenerationResult& OutResult,
-		const NodeMeshSubtract* subs )
-    {
-		// This node is deprecated.
-		check(false);        
-    }
-
-
     //---------------------------------------------------------------------------------------------
     void CodeGenerator::GenerateMesh_Transform(const FMeshGenerationOptions& InOptions, FMeshGenerationResult& OutResult,
                                                 const NodeMeshTransform* trans )
diff --git a/Engine/Plugins/Experimental/Mutable/Source/MutableTools/Private/MuT/CodeGenerator_Scalar.cpp b/Engine/Plugins/Experimental/Mutable/Source/MutableTools/Private/MuT/CodeGenerator_Scalar.cpp
--- a/Engine/Plugins/Experimental/Mutable/Source/MutableTools/Private/MuT/CodeGenerator_Scalar.cpp
+++ b/Engine/Plugins/Experimental/Mutable/Source/MutableTools/Private/MuT/CodeGenerator_Scalar.cpp
@@ -167,10 +167,10 @@ namespace mu
 					// We take whatever size will be produced
 					IMAGE_STATE newState;
 					FImageDesc desc = CalculateImageDesc(*node.m_additionalImages[a]->GetBasePrivate());
-					newState.m_imageSize = desc.m_size;
+					newState.m_imageSize = FIntVector2(desc.m_size);
 					newState.m_imageRect.min[0] = 0;
 					newState.m_imageRect.min[1] = 0;
-					newState.m_imageRect.size = desc.m_size;
+					newState.m_imageRect.size = FIntVector2(desc.m_size);
 					newState.m_layoutBlockId = -1;
 					m_imageState.Add(newState);
 
diff --git a/Engine/Plugins/Experimental/Mutable/Source/MutableTools/Private/MuT/CodeOptimiser.h b/Engine/Plugins/Experimental/Mutable/Source/MutableTools/Private/MuT/CodeOptimiser.h
--- a/Engine/Plugins/Experimental/Mutable/Source/MutableTools/Private/MuT/CodeOptimiser.h
+++ b/Engine/Plugins/Experimental/Mutable/Source/MutableTools/Private/MuT/CodeOptimiser.h
@@ -215,7 +215,7 @@ namespace mu
     {
     public:
 
-        RuntimeTextureCompressionRemoverAST( STATE_COMPILATION_DATA* state );
+        RuntimeTextureCompressionRemoverAST( STATE_COMPILATION_DATA* state, bool bInAlwaysUncompress );
 
     protected:
 
@@ -224,6 +224,7 @@ namespace mu
     private:
 
         RuntimeParameterVisitorAST m_hasRuntimeParamVisitor;
+		bool bAlwaysUncompress = false;
 
     };
 
diff --git a/Engine/Plugins/Experimental/Mutable/Source/MutableTools/Private/MuT/CodeOptimiser_Parameters.cpp b/Engine/Plugins/Experimental/Mutable/Source/MutableTools/Private/MuT/CodeOptimiser_Parameters.cpp
--- a/Engine/Plugins/Experimental/Mutable/Source/MutableTools/Private/MuT/CodeOptimiser_Parameters.cpp
+++ b/Engine/Plugins/Experimental/Mutable/Source/MutableTools/Private/MuT/CodeOptimiser_Parameters.cpp
@@ -1901,9 +1901,11 @@ namespace mu
     //---------------------------------------------------------------------------------------------
     //---------------------------------------------------------------------------------------------
     RuntimeTextureCompressionRemoverAST::RuntimeTextureCompressionRemoverAST(
-            STATE_COMPILATION_DATA* pState
-            )
+		STATE_COMPILATION_DATA* pState,
+		bool bInAlwaysUncompress
+	)
         : m_hasRuntimeParamVisitor(pState)
+		, bAlwaysUncompress(bInAlwaysUncompress)
     {
         Traverse( pState->root );
     }
@@ -1914,7 +1916,7 @@ namespace mu
         OP_TYPE type = at->GetOpType();
         processChildren = GetOpDataType(type)==DT_INSTANCE;
 
-        // TODO: Finer grained: what if the runtime parameter just selectes between compressed
+        // TODO: Finer grained: what if the runtime parameter just selects between compressed
         // textures? We don't want them uncompressed.
         if( type==OP_TYPE::IN_ADDIMAGE )
         {
@@ -1924,7 +1926,7 @@ namespace mu
             // Does it have a runtime parameter in its subtree?
             bool hasRuntimeParameter = m_hasRuntimeParamVisitor.HasAny( imageAt );
 
-            if (hasRuntimeParameter)
+            if (bAlwaysUncompress || hasRuntimeParameter)
             {
                 FImageDesc imageDesc = imageAt->GetImageDesc( true );
 
@@ -1998,21 +2000,35 @@ namespace mu
          for ( size_t s=0; s<m_states.size(); ++s )
          {
             // Remove the unnecessary lods
-            if (m_states[s].optimisationFlags.m_onlyFirstLOD)
+            if (m_states[s].nodeState.m_optimisation.bOnlyFirstLOD)
             {
-				LODCountReducerAST(m_states[s].root, m_states[s].optimisationFlags.m_firstLOD + 1 + m_states[s].optimisationFlags.m_numExtraLODsToBuildAfterFirstLOD);
+				LODCountReducerAST(m_states[s].root, m_states[s].nodeState.m_optimisation.FirstLOD + 1 + m_states[s].nodeState.m_optimisation.NumExtraLODsToBuildAfterFirstLOD);
             }
 
+			// Apply texture compression strategy
+			switch (m_states[s].nodeState.m_optimisation.TextureCompressionStrategy)
+			{
+			case ETextureCompressionStrategy::DontCompressRuntime:
+			{
+				MUTABLE_CPUPROFILER_SCOPE(RuntimeTextureCompressionRemover);
+				RuntimeTextureCompressionRemoverAST r(&m_states[s], false);
+				break;
+			}
+
+			case ETextureCompressionStrategy::NeverCompress:
+			{
+				MUTABLE_CPUPROFILER_SCOPE(RuntimeTextureCompressionRemover);
+				RuntimeTextureCompressionRemoverAST r(&m_states[s], true);
+				break;
+			}
+
+			default:
+				break;
+			}
+
             // If a state has no runtime parameters, skip its optimisation alltogether
             if (m_states[s].nodeState.m_runtimeParams.Num())
             {
-                // Remove unnecessary image compression
-                if (m_states[s].optimisationFlags.m_avoidRuntimeCompression)
-                {
-                    MUTABLE_CPUPROFILER_SCOPE(RuntimeTextureCompressionRemover);
-                    RuntimeTextureCompressionRemoverAST r(&m_states[s]);
-                }
-
                 // Promote the intructions that depend on runtime parameters, and sink new
                 // format instructions.
                 bool modified = true;
diff --git a/Engine/Plugins/Experimental/Mutable/Source/MutableTools/Private/MuT/CodeOptimiser_Semantic.cpp b/Engine/Plugins/Experimental/Mutable/Source/MutableTools/Private/MuT/CodeOptimiser_Semantic.cpp
--- a/Engine/Plugins/Experimental/Mutable/Source/MutableTools/Private/MuT/CodeOptimiser_Semantic.cpp
+++ b/Engine/Plugins/Experimental/Mutable/Source/MutableTools/Private/MuT/CodeOptimiser_Semantic.cpp
@@ -441,10 +441,12 @@ namespace mu
 
             auto rectOp = typedPatch->patch.child();
             ASTOp::FGetImageDescContext context;
-            auto patchDesc = rectOp->GetImageDesc( false, &context );
-            box<vec2<int16_t>> patchBox;
-            patchBox.min = typedPatch->location;
-            patchBox.size = patchDesc.m_size;
+            FImageDesc patchDesc = rectOp->GetImageDesc( false, &context );
+            box<vec2<int16>> patchBox;
+			patchBox.min[0] = typedPatch->location[0];
+			patchBox.min[1] = typedPatch->location[1];
+			patchBox.size[0] = patchDesc.m_size[0];
+			patchBox.size[1] = patchDesc.m_size[1];
 
             box<vec2<int16_t>> cropBox;
             cropBox.min[0] = currentCropOp->op.args.ImageCrop.minX;
@@ -464,7 +466,7 @@ namespace mu
                 auto newOp = mu::Clone<ASTOpImagePatch>(at);
                 newOp->base = Visit(newOp->base.child(), currentCropOp);
 
-                box<vec2<int16_t>> ibox = patchBox.Intersect(cropBox);
+                box<vec2<int16>> ibox = patchBox.Intersect(cropBox);
                 check(ibox.size[0]>0 && ibox.size[1]>0);
 
                 Ptr<ASTOpFixed> patchCropOp = mu::Clone<ASTOpFixed>(currentCropOp);
diff --git a/Engine/Plugins/Experimental/Mutable/Source/MutableTools/Private/MuT/Compiler.cpp b/Engine/Plugins/Experimental/Mutable/Source/MutableTools/Private/MuT/Compiler.cpp
--- a/Engine/Plugins/Experimental/Mutable/Source/MutableTools/Private/MuT/Compiler.cpp
+++ b/Engine/Plugins/Experimental/Mutable/Source/MutableTools/Private/MuT/Compiler.cpp
@@ -197,7 +197,6 @@ namespace mu
                 data.nodeState = s.first;
                 data.root = s.second;
                 data.state.m_name = s.first.m_name;
-                data.optimisationFlags = s.first.m_optimisation;
                 states.push_back( data );
             }
 
diff --git a/Engine/Plugins/Experimental/Mutable/Source/MutableTools/Private/MuT/CompilerPrivate.h b/Engine/Plugins/Experimental/Mutable/Source/MutableTools/Private/MuT/CompilerPrivate.h
--- a/Engine/Plugins/Experimental/Mutable/Source/MutableTools/Private/MuT/CompilerPrivate.h
+++ b/Engine/Plugins/Experimental/Mutable/Source/MutableTools/Private/MuT/CompilerPrivate.h
@@ -13,32 +13,25 @@
 
 namespace mu
 {
+	MUTABLE_DEFINE_ENUM_SERIALISABLE(ETextureCompressionStrategy);
 
     //!
-    struct STATE_OPTIMIZATION_OPTIONS
+    struct FStateOptimizationOptions
     {
-        STATE_OPTIMIZATION_OPTIONS()
-        {
-			m_firstLOD = 0;
-            m_onlyFirstLOD = false;
-            m_avoidRuntimeCompression = false;
-			m_numExtraLODsToBuildAfterFirstLOD = 0;
-        }
-
-		uint8 m_firstLOD;
-		bool m_onlyFirstLOD;
-        bool m_avoidRuntimeCompression;
-		int32 m_numExtraLODsToBuildAfterFirstLOD;
+		uint8 FirstLOD = 0;
+		uint8 NumExtraLODsToBuildAfterFirstLOD = 0;
+		bool bOnlyFirstLOD = false;
+		ETextureCompressionStrategy TextureCompressionStrategy = ETextureCompressionStrategy::None;
 
         void Serialise( OutputArchive& arch ) const
         {
-            const int32_t ver = 3;
+            const int32_t ver = 4;
             arch << ver;
 
-			arch << m_firstLOD;
-			arch << m_onlyFirstLOD;
-            arch << m_avoidRuntimeCompression;
-			arch << m_numExtraLODsToBuildAfterFirstLOD;
+			arch << FirstLOD;
+			arch << bOnlyFirstLOD;
+            arch << TextureCompressionStrategy;
+			arch << NumExtraLODsToBuildAfterFirstLOD;
         }
 
 
@@ -46,29 +39,46 @@ namespace mu
         {
             int32_t ver = 0;
             arch >> ver;
-			check(ver <= 3);
+			check(ver <= 4);
 
-			if (ver >= 3)
+			if (ver >= 2)
 			{
-				arch >> m_numExtraLODsToBuildAfterFirstLOD;
+				arch >> FirstLOD;
 			}
 			else
 			{
-				m_numExtraLODsToBuildAfterFirstLOD = 0;
+				FirstLOD = 0;
 			}
 
-			if (ver >= 2)
+            arch >> bOnlyFirstLOD;
+
+			if (ver >= 4)
 			{
-				arch >> m_firstLOD;
+				arch >> TextureCompressionStrategy;
 			}
 			else
 			{
-				m_firstLOD = 0;
+				bool bAvoidRuntimeCompression;
+				arch >> bAvoidRuntimeCompression;
+				TextureCompressionStrategy = bAvoidRuntimeCompression ? ETextureCompressionStrategy::DontCompressRuntime : ETextureCompressionStrategy::None;
 			}
 
-            arch >> m_onlyFirstLOD;
-            arch >> m_avoidRuntimeCompression;
-        }
+			if (ver == 3)
+			{
+				int32 OldNumExtraLODsToBuildAfterFirstLOD;
+				arch >> OldNumExtraLODsToBuildAfterFirstLOD;
+				NumExtraLODsToBuildAfterFirstLOD = OldNumExtraLODsToBuildAfterFirstLOD;
+			}
+			else if (ver >= 4)
+			{
+				arch >> NumExtraLODsToBuildAfterFirstLOD;
+			}
+			else
+			{
+				NumExtraLODsToBuildAfterFirstLOD = 0;
+			}
+
+		}
     };
 
 
@@ -94,13 +104,13 @@ namespace mu
 
 
     //! Information about an object state in the source data
-    struct OBJECT_STATE
+    struct FObjectState
     {
         //! Name used to identify the state from the code and user interface.
         string m_name;
 
         //! GPU Optimisation options
-        STATE_OPTIMIZATION_OPTIONS m_optimisation;
+		FStateOptimizationOptions m_optimisation;
 
         //! List of names of the runtime parameters in this state
         TArray<string> m_runtimeParams;
@@ -132,10 +142,10 @@ namespace mu
     //!
     struct STATE_COMPILATION_DATA
     {
-        OBJECT_STATE nodeState;
+        FObjectState nodeState;
         Ptr<ASTOp> root;
         FProgram::FState state;
-        STATE_OPTIMIZATION_OPTIONS optimisationFlags;
+		//FStateOptimizationOptions optimisationFlags;
 
         //! List of instructions that need to be cached to efficiently update this state
         TArray<Ptr<ASTOp>> m_updateCache;
diff --git a/Engine/Plugins/Experimental/Mutable/Source/MutableTools/Private/MuT/ErrorLog.cpp b/Engine/Plugins/Experimental/Mutable/Source/MutableTools/Private/MuT/ErrorLog.cpp
--- a/Engine/Plugins/Experimental/Mutable/Source/MutableTools/Private/MuT/ErrorLog.cpp
+++ b/Engine/Plugins/Experimental/Mutable/Source/MutableTools/Private/MuT/ErrorLog.cpp
@@ -284,7 +284,6 @@ namespace mu
 		TEXT("ME_INTERPOLATE   "),
 		TEXT("ME_MASKCLIPMESH  "),
 		TEXT("ME_MASKDIFF      "),
-        TEXT("ME_SUBTRACT      "),
         TEXT("ME_REMOVEMASK    "),
         TEXT("ME_FORMAT        "),
         TEXT("ME_EXTRACTLABLOCK"),
diff --git a/Engine/Plugins/Experimental/Mutable/Source/MutableTools/Private/MuT/NodeMesh.cpp b/Engine/Plugins/Experimental/Mutable/Source/MutableTools/Private/MuT/NodeMesh.cpp
--- a/Engine/Plugins/Experimental/Mutable/Source/MutableTools/Private/MuT/NodeMesh.cpp
+++ b/Engine/Plugins/Experimental/Mutable/Source/MutableTools/Private/MuT/NodeMesh.cpp
@@ -17,7 +17,6 @@
 #include "MuT/NodeMeshMakeMorph.h"
 #include "MuT/NodeMeshMorph.h"
 #include "MuT/NodeMeshReshape.h"
-#include "MuT/NodeMeshSubtract.h"
 #include "MuT/NodeMeshSwitch.h"
 #include "MuT/NodeMeshTable.h"
 #include "MuT/NodeMeshTangents.h"
@@ -61,7 +60,6 @@ namespace mu
 		case EType::Constant:		return NodeMeshConstant::StaticUnserialise( arch ); break;
 		case EType::Interpolate:	return NodeMeshInterpolate::StaticUnserialise( arch ); break;
 		case EType::Table:			return NodeMeshTable::StaticUnserialise( arch ); break;
-		case EType::Subtract:		return NodeMeshSubtract::StaticUnserialise( arch ); break;
 		case EType::Format:			return NodeMeshFormat::StaticUnserialise( arch ); break;
 		case EType::Tangents:		return NodeMeshTangents::StaticUnserialise( arch ); break;
 		case EType::Morph:			return NodeMeshMorph::StaticUnserialise( arch ); break;
diff --git a/Engine/Plugins/Experimental/Mutable/Source/MutableTools/Private/MuT/NodeMeshSubtract.cpp b/Engine/Plugins/Experimental/Mutable/Source/MutableTools/Private/MuT/NodeMeshSubtract.cpp
deleted file mode 100644
diff --git a/Engine/Plugins/Experimental/Mutable/Source/MutableTools/Private/MuT/NodeMeshSubtractPrivate.h b/Engine/Plugins/Experimental/Mutable/Source/MutableTools/Private/MuT/NodeMeshSubtractPrivate.h
deleted file mode 100644
diff --git a/Engine/Plugins/Experimental/Mutable/Source/MutableTools/Private/MuT/NodeObject.cpp b/Engine/Plugins/Experimental/Mutable/Source/MutableTools/Private/MuT/NodeObject.cpp
--- a/Engine/Plugins/Experimental/Mutable/Source/MutableTools/Private/MuT/NodeObject.cpp
+++ b/Engine/Plugins/Experimental/Mutable/Source/MutableTools/Private/MuT/NodeObject.cpp
@@ -8,9 +8,6 @@
 #include "MuR/SerialisationPrivate.h"
 #include "MuT/NodeObjectGroup.h"
 #include "MuT/NodeObjectNew.h"
-#include "MuT/NodeObjectState.h"
-
-#include <stdint.h>
 
 
 namespace mu
@@ -47,7 +44,7 @@ namespace mu
 		{
 		case 0 :  return NodeObjectNew::StaticUnserialise( arch ); break;
 		case 1 :  return NodeObjectGroup::StaticUnserialise( arch ); break;
-        case 2 :  return NodeObjectState::StaticUnserialise( arch ); break;
+        //case 2 :  return NodeObjectState::StaticUnserialise( arch ); break;
         //case 3 :  return NodeObjectTransform::StaticUnserialise( arch ); break;
         default : check(false);
 		}
diff --git a/Engine/Plugins/Experimental/Mutable/Source/MutableTools/Private/MuT/NodeObjectNew.cpp b/Engine/Plugins/Experimental/Mutable/Source/MutableTools/Private/MuT/NodeObjectNew.cpp
--- a/Engine/Plugins/Experimental/Mutable/Source/MutableTools/Private/MuT/NodeObjectNew.cpp
+++ b/Engine/Plugins/Experimental/Mutable/Source/MutableTools/Private/MuT/NodeObjectNew.cpp
@@ -317,14 +317,19 @@ namespace mu
 
 
     //---------------------------------------------------------------------------------------------
-    void NodeObjectNew::SetStateProperties( int s, bool avoidRuntimeCompression, bool onlyFirstLOD, int firstLOD, int numExtraLODsToBuildAfterFirstLOD )
+    void NodeObjectNew::SetStateProperties( int32 StateIndex, 
+		ETextureCompressionStrategy TextureCompressionStrategy, 
+		bool bOnlyFirstLOD, 
+		uint8 FirstLOD, 
+		uint8 NumExtraLODsToBuildAfterFirstLOD )
     {
-        check( s>=0 && s<GetStateCount() );
+        check(StateIndex >=0 && StateIndex<GetStateCount() );
 
-        m_pD->m_states[s].m_optimisation.m_avoidRuntimeCompression = avoidRuntimeCompression;
-        m_pD->m_states[s].m_optimisation.m_onlyFirstLOD = onlyFirstLOD;
-		m_pD->m_states[s].m_optimisation.m_firstLOD = firstLOD;
-		m_pD->m_states[s].m_optimisation.m_numExtraLODsToBuildAfterFirstLOD = numExtraLODsToBuildAfterFirstLOD;
+		FStateOptimizationOptions& Data = m_pD->m_states[StateIndex].m_optimisation;
+        Data.TextureCompressionStrategy = TextureCompressionStrategy;
+		Data.bOnlyFirstLOD = bOnlyFirstLOD;
+		Data.FirstLOD = FirstLOD;
+		Data.NumExtraLODsToBuildAfterFirstLOD = NumExtraLODsToBuildAfterFirstLOD;
     }
 
 }
diff --git a/Engine/Plugins/Experimental/Mutable/Source/MutableTools/Private/MuT/NodeObjectNewPrivate.h b/Engine/Plugins/Experimental/Mutable/Source/MutableTools/Private/MuT/NodeObjectNewPrivate.h
--- a/Engine/Plugins/Experimental/Mutable/Source/MutableTools/Private/MuT/NodeObjectNewPrivate.h
+++ b/Engine/Plugins/Experimental/Mutable/Source/MutableTools/Private/MuT/NodeObjectNewPrivate.h
@@ -33,7 +33,7 @@ namespace mu
 
 
 		//! List of states
-        TArray<OBJECT_STATE> m_states;
+        TArray<FObjectState> m_states;
 
 
 		//!
diff --git a/Engine/Plugins/Experimental/Mutable/Source/MutableTools/Private/MuT/NodeObjectState.cpp b/Engine/Plugins/Experimental/Mutable/Source/MutableTools/Private/MuT/NodeObjectState.cpp
deleted file mode 100644
diff --git a/Engine/Plugins/Experimental/Mutable/Source/MutableTools/Private/MuT/NodeObjectStatePrivate.h b/Engine/Plugins/Experimental/Mutable/Source/MutableTools/Private/MuT/NodeObjectStatePrivate.h
deleted file mode 100644
diff --git a/Engine/Plugins/Experimental/Mutable/Source/MutableTools/Public/MuT/NodeMesh.h b/Engine/Plugins/Experimental/Mutable/Source/MutableTools/Public/MuT/NodeMesh.h
--- a/Engine/Plugins/Experimental/Mutable/Source/MutableTools/Public/MuT/NodeMesh.h
+++ b/Engine/Plugins/Experimental/Mutable/Source/MutableTools/Public/MuT/NodeMesh.h
@@ -28,7 +28,7 @@ namespace mu
 			Constant = 0,
 			Interpolate = 2,
 			Table = 4,
-			Subtract = 6,
+			Subtract_Deprecated = 6,
 			Format = 7,
 			Tangents = 8,
 			Morph = 9,
diff --git a/Engine/Plugins/Experimental/Mutable/Source/MutableTools/Public/MuT/NodeMeshSubtract.h b/Engine/Plugins/Experimental/Mutable/Source/MutableTools/Public/MuT/NodeMeshSubtract.h
deleted file mode 100644
diff --git a/Engine/Plugins/Experimental/Mutable/Source/MutableTools/Public/MuT/NodeObject.h b/Engine/Plugins/Experimental/Mutable/Source/MutableTools/Public/MuT/NodeObject.h
--- a/Engine/Plugins/Experimental/Mutable/Source/MutableTools/Public/MuT/NodeObject.h
+++ b/Engine/Plugins/Experimental/Mutable/Source/MutableTools/Public/MuT/NodeObject.h
@@ -28,7 +28,6 @@ namespace mu
 		{
 			New = 0,
 			Group = 1,
-			State = 2,
 
 			None
 		};
diff --git a/Engine/Plugins/Experimental/Mutable/Source/MutableTools/Public/MuT/NodeObjectNew.h b/Engine/Plugins/Experimental/Mutable/Source/MutableTools/Public/MuT/NodeObjectNew.h
--- a/Engine/Plugins/Experimental/Mutable/Source/MutableTools/Public/MuT/NodeObjectNew.h
+++ b/Engine/Plugins/Experimental/Mutable/Source/MutableTools/Public/MuT/NodeObjectNew.h
@@ -5,20 +5,12 @@
 #include "MuR/RefCounted.h"
 #include "MuT/Node.h"
 #include "MuT/NodeObject.h"
+#include "MuR/System.h"
 
 
 namespace mu
 {
-
-	// Forward definitions
-	class NodeObjectNew;
-	typedef Ptr<NodeObjectNew> NodeObjectNewPtr;
-	typedef Ptr<const NodeObjectNew> NodeObjectNewPtrConst;
-
 	class NodeLOD;
-	typedef Ptr<NodeLOD> NodeLODPtr;
-	typedef Ptr<const NodeLOD> NodeLODPtrConst;
-
 
 	//! Node that creates a new object by setting its levels-of-detail and its children.
 	//! \ingroup model
@@ -34,7 +26,7 @@ namespace mu
 
 		void SerialiseWrapper(OutputArchive& arch) const override;
 		static void Serialise( const NodeObjectNew* pNode, OutputArchive& arch );
-		static NodeObjectNewPtr StaticUnserialise( InputArchive& arch );
+		static Ptr<NodeObjectNew> StaticUnserialise( InputArchive& arch );
 
 
 		//-----------------------------------------------------------------------------------------
@@ -67,8 +59,8 @@ namespace mu
 		void SetLODCount( int );
 
 		//! Get a level of detail node from the object.
-		NodeLODPtr GetLOD( int index ) const;
-		void SetLOD( int index, NodeLODPtr );
+		Ptr<NodeLOD> GetLOD( int index ) const;
+		void SetLOD( int index, Ptr<NodeLOD>);
 
 		//! Get the number of child objects
 		int GetChildCount() const;
@@ -102,7 +94,11 @@ namespace mu
 		void RemoveStateParam( int s, const char* param );
 
         //! Set the optimisation properties of a state
-        void SetStateProperties( int s, bool avoidRuntimeCompression, bool onlyFirstLOD, int firstLOD, int numExtraLODsToBuildAfterFirstLOD);
+		void SetStateProperties(int32 StateIndex,
+			ETextureCompressionStrategy TextureCompressionStrategy,
+			bool bOnlyFirstLOD,
+			uint8 FirstLOD,
+			uint8 NumExtraLODsToBuildAfterFirstLOD);
 
 		//-----------------------------------------------------------------------------------------
 		// Interface pattern
diff --git a/Engine/Plugins/Experimental/Mutable/Source/MutableTools/Public/MuT/NodeObjectState.h b/Engine/Plugins/Experimental/Mutable/Source/MutableTools/Public/MuT/NodeObjectState.h
deleted file mode 100644
