[Mutable] Fix missing AutoConsoleVariable include.

#rnx

diff --git a/Engine/Plugins/Experimental/Mutable/Source/CustomizableObject/Classes/MuCO/CustomizableObjectSystem.h b/Engine/Plugins/Experimental/Mutable/Source/CustomizableObject/Classes/MuCO/CustomizableObjectSystem.h
--- a/Engine/Plugins/Experimental/Mutable/Source/CustomizableObject/Classes/MuCO/CustomizableObjectSystem.h
+++ b/Engine/Plugins/Experimental/Mutable/Source/CustomizableObject/Classes/MuCO/CustomizableObjectSystem.h
@@ -2,6 +2,7 @@
 
 #pragma once
 
+#include "HAL/IConsoleManager.h"
 #include "Engine/StreamableManager.h"
 #include "AssetRegistry/AssetData.h"
 
