[mutable] Added states view capability to move around states and parameters by dragging and dropping them.
- States can not be dragged over other states in order to change their order.
- State parameters can now be reordered inside owning state.
- Fixed minor parameter count reporting bug (MTBL-986).
- Fixed some typos.

#rnx

diff --git a/Engine/Plugins/Experimental/Mutable/Source/CustomizableObjectEditor/Private/MuCOE/SCustomizableObjectNodeObjectStatesView.cpp b/Engine/Plugins/Experimental/Mutable/Source/CustomizableObjectEditor/Private/MuCOE/SCustomizableObjectNodeObjectStatesView.cpp
--- a/Engine/Plugins/Experimental/Mutable/Source/CustomizableObjectEditor/Private/MuCOE/SCustomizableObjectNodeObjectStatesView.cpp
+++ b/Engine/Plugins/Experimental/Mutable/Source/CustomizableObjectEditor/Private/MuCOE/SCustomizableObjectNodeObjectStatesView.cpp
@@ -5,6 +5,11 @@
 #include "Framework/Views/TableViewMetadata.h"
 #include "MuCOE/CustomizableObjectGraph.h"
 #include "MuCOE/Nodes/CustomizableObjectNodeObject.h"
+#include "SlotBase.h"
+#include "Styling/AppStyle.h"
+#include "Styling/SlateStyleMacros.h"
+#include "Templates/Casts.h"
+#include "Widgets/Images/SImage.h"
 #include "Widgets/Input/SButton.h"
 #include "Widgets/Input/SCheckBox.h"
 #include "Widgets/Input/SComboButton.h"
@@ -18,131 +23,171 @@ struct FSlateBrush;
 #define LOCTEXT_NAMESPACE "CustomizableObjectNodeObjectStatesView"
 
 
-void SCustomizableObjectNodeObjectSatesView::Construct(const FArguments& InArgs)
+TSharedRef<FDragAndDropOpWithWidget> FDragAndDropOpWithWidget::New(int32 InSlotIndexBeingDragged, SVerticalBox::FSlot* InSlotBeingDragged,
+	TSharedPtr<SWidget> InWidgetToShow)
+{
+	TSharedRef<FDragAndDropOpWithWidget> Operation = MakeShareable(new FDragAndDropOpWithWidget);
+	{
+		Operation->MouseCursor = EMouseCursor::GrabHandClosed;
+		Operation->SlotIndexBeingDragged = InSlotIndexBeingDragged;
+		Operation->SlotBeingDragged = InSlotBeingDragged;
+		Operation->WidgetToShow = InWidgetToShow;
+	}
+
+	Operation->Construct();
+	
+	return Operation;
+}
+
+TSharedPtr<SWidget> FDragAndDropOpWithWidget::GetDefaultDecorator() const
+{
+	return SNew(SBorder)
+	.BorderImage(FAppStyle::GetBrush("ContentBrowser.AssetDragDropTooltipBackground"))
+	.Content()
+	[
+		WidgetToShow.ToSharedRef()
+	];
+}
+
+
+
+void SCustomizableObjectState::Construct(const FArguments& InArgs)
 {
 	Node = InArgs._Node;
+	StateIndex = InArgs._StateIndex;
+	
+	SAssignNew(VerticalSlots, SVerticalBox);
+	
+	// Images and Runtime parameters widgets are stored in arrays to control the visibility
+	SAssignNew(RuntimeParametersWidget,SCustomizableObjectRuntimeParameterList)
+	.Node(Node)
+	.StateIndex(StateIndex);
 
-	if (Node)
-	{
-		SAssignNew(VerticalSlots, SVerticalBox);
+	RuntimeParametersWidget->SetCollapsed(true);
+	RuntimeParametersWidget->SetVisibility(GetCollapsed());
 
-		for (int32 i = 0; i < Node->States.Num(); i++)
-		{
-			// Images and Runtime parameters widgets are stored in arrays to control the visibility
-			RuntimeParametersListWidgets.Add(SNew(SCustomizableObjectRuntimeParameterList).Node(Node).StateIndex(i));
-			RuntimeParametersListWidgets[i]->SetCollapsed(true);
-			RuntimeParametersListWidgets[i]->SetVisibility(GetCollapsed(i));
+	SAssignNew(CollapsedArrow,SImage)
+	.Image(GetExpressionPreviewArrow());
 
-			CollapsedArrows.Add(SNew(SImage).Image(GetExpressionPreviewArrow(i)));
+	VerticalSlots->AddSlot()
+	.AutoHeight()
+	[
+		SNew(SVerticalBox)
+		+ SVerticalBox::Slot()
+		.Padding(0.0f, 0.0f, 3.0f, 5.0f)
+		[
+			SNew(SHorizontalBox)
 
-			VerticalSlots->AddSlot()
-			.AutoHeight()
+			+ SHorizontalBox::Slot()
+			.AutoWidth()
 			[
-				SNew(SVerticalBox)
-				+ SVerticalBox::Slot()
-				.Padding(0.0f, 0.0f, 3.0f, 5.0f)
+				SNew(SBox)
+				.HAlign(HAlign_Center)
+				.VAlign(VAlign_Center)
+				.WidthOverride(16.0f)
 				[
-					SNew(SHorizontalBox)
-
-					// State variable label
-					+ SHorizontalBox::Slot()
-					.HAlign(EHorizontalAlignment::HAlign_Left)
-					.AutoWidth()
-					.Padding(0.0f, 0.0f, 3.0f, 0.0f)
-					[
-						SNew(STextBlock)
-						.Text(FText::FromString("State:"))
-					]
-					
-					// State name
-					+ SHorizontalBox::Slot()
-					.HAlign(EHorizontalAlignment::HAlign_Left)
-					.AutoWidth()
-					.Padding(2.0f, 0.0f, 0.0f, 0.0f)
-					[
-						SNew(STextBlock)
-						.Text(FText::FromString(Node->States[i].Name))
-					]
+					SNew(SImage)
+					.Image(FCoreStyle::Get().GetBrush("VerticalBoxDragIndicatorShort"))
 				]
+			]
+			
+			// State variable label
+			+ SHorizontalBox::Slot()
+			.HAlign(EHorizontalAlignment::HAlign_Left)
+			.AutoWidth()
+			.Padding(2.0f, 0.0f, 3.0f, 0.0f)
+			[
+				SNew(STextBlock)
+				.Text(FText::FromString("State:"))
+			]
+			
+			// State name
+			+ SHorizontalBox::Slot()
+			.HAlign(EHorizontalAlignment::HAlign_Left)
+			.AutoWidth()
+			.Padding(2.0f, 0.0f, 0.0f, 0.0f)
+			[
+				SNew(STextBlock)
+				.Text(FText::FromString(Node->States[StateIndex].Name))
+			]
+		]
 
-				+ SVerticalBox::Slot()
-				.AutoHeight()
-				.Padding(0.0f, 0.0f, 0.0f, 2.0f)
+		+ SVerticalBox::Slot()
+		.AutoHeight()
+		.Padding(0.0f, 0.0f, 0.0f, 2.0f)
+		[
+			SNew(SVerticalBox)
+			+ SVerticalBox::Slot()
+			.HAlign(EHorizontalAlignment::HAlign_Left)
+			.Padding(10.0f, 0.0f, 3.0f, 3.0f)
+			.AutoHeight()
+			[
+				SNew(SHorizontalBox)
+
+				// Collapsed Arrow checkbox
+				+SHorizontalBox::Slot()
+				.HAlign(EHorizontalAlignment::HAlign_Left)
+				.AutoWidth()
+				.Padding(0.0f, 0.0f, 3.0f, 0.0f)
 				[
-					SNew(SVerticalBox)
-					+ SVerticalBox::Slot()
-					.HAlign(EHorizontalAlignment::HAlign_Left)
-					.Padding(10.0f, 0.0f, 3.0f, 3.0f)
-					.AutoHeight()
+					SNew(SCheckBox)
+					.OnCheckStateChanged(this, &SCustomizableObjectState::OnCollapseChanged)
+					.IsChecked(ECheckBoxState::Unchecked)
+					.Cursor(EMouseCursor::Default)
+					.Style(FAppStyle::Get(), "Graph.Node.AdvancedView")
 					[
 						SNew(SHorizontalBox)
-
-						// Collapsed Arrow checkbox
-						+SHorizontalBox::Slot()
-						.HAlign(EHorizontalAlignment::HAlign_Left)
-						.AutoWidth()
-						.Padding(0.0f, 0.0f, 3.0f, 0.0f)
-						[
-							SNew(SCheckBox)
-							.OnCheckStateChanged(this, &SCustomizableObjectNodeObjectSatesView::OnCollapseChanged, i)
-							.IsChecked(ECheckBoxState::Unchecked)
-							.Cursor(EMouseCursor::Default)
-							.Style(FAppStyle::Get(), "Graph.Node.AdvancedView")
-							[
-								SNew(SHorizontalBox)
-								+ SHorizontalBox::Slot()
-								.VAlign(VAlign_Center)
-								.HAlign(HAlign_Center)
-								[
-									CollapsedArrows[i].ToSharedRef()
-								]
-							]
-						]
-						
-						// Number of runtime parameters
 						+ SHorizontalBox::Slot()
-						.HAlign(EHorizontalAlignment::HAlign_Left)
-						.AutoWidth()
-						.Padding(0.0f, 0.0f, 3.0f, 0.0f)
+						.VAlign(VAlign_Center)
+						.HAlign(HAlign_Center)
 						[
-							SNew(STextBlock)
-							.Text(FText::FromString(FString::Printf(TEXT("Runtime Parameters:  %d"), Node->States[i].RuntimeParameters.Num())))
+							CollapsedArrow.ToSharedRef()
 						]
+					]
+				]
+				
+				// Number of runtime parameters
+				+ SHorizontalBox::Slot()
+				.HAlign(EHorizontalAlignment::HAlign_Left)
+				.AutoWidth()
+				.Padding(0.0f, 0.0f, 3.0f, 0.0f)
+				[
+					SNew(STextBlock)
+					.Text(this,&SCustomizableObjectState::GetStateParameterCountText)
+				]
 
+				+ SHorizontalBox::Slot()
+				.HAlign(EHorizontalAlignment::HAlign_Left)
+				.AutoWidth()
+				.Padding(0.0f, 0.0f, 3.0f, 0.0f)
+				[
+					SNew(SButton)
+					.OnClicked(this, &SCustomizableObjectState::OnAddRuntimeParameterPressed)
+					.ToolTipText(LOCTEXT("AddRuntimeParameter", "Add Runtime Parameter"))
+					[
+						SNew(SHorizontalBox)
 						+ SHorizontalBox::Slot()
-						.HAlign(EHorizontalAlignment::HAlign_Left)
-						.AutoWidth()
-						.Padding(0.0f, 0.0f, 3.0f, 0.0f)
+						.VAlign(VAlign_Center)
+						.HAlign(HAlign_Center)
 						[
-							SNew(SButton)
-							.OnClicked(this, &SCustomizableObjectNodeObjectSatesView::OnAddRuntimeParameterPressed,i)
-							.ToolTipText(LOCTEXT("AddRuntimeParameter", "Add Runtime Parameter"))
-							[
-								SNew(SHorizontalBox)
-								+ SHorizontalBox::Slot()
-								.VAlign(VAlign_Center)
-								.HAlign(HAlign_Center)
-								[
-									SNew(SImage)
-									.Image(FAppStyle::GetBrush(TEXT("Plus")))
-								]
-							]
+							SNew(SImage)
+							.Image(FAppStyle::GetBrush(TEXT("Plus")))
 						]
 					]
-					
-					// Runtime parameters widget
-					+ SVerticalBox::Slot()
-					.HAlign(EHorizontalAlignment::HAlign_Left)
-					.Padding(20.0f, 0.0f, 0.0f, 0.0f)
-					.AutoHeight()
-					[
-						RuntimeParametersListWidgets[i].ToSharedRef()
-					]
 				]
-			];
-		}
-	}
-
+			]
+			
+			// Runtime parameters widget
+			+ SVerticalBox::Slot()
+			.HAlign(EHorizontalAlignment::HAlign_Left)
+			.Padding(20.0f, 0.0f, 0.0f, 0.0f)
+			.AutoHeight()
+			[
+				RuntimeParametersWidget.ToSharedRef()
+			]
+		]
+	];
+	
 	// Add the widget to the child slot
 	ChildSlot
 	[
@@ -150,39 +195,182 @@ void SCustomizableObjectNodeObjectSatesView::Construct(const FArguments& InArgs)
 	];
 }
 
-
-void SCustomizableObjectNodeObjectSatesView::OnCollapseChanged(const ECheckBoxState NewCheckedState, int32 StateIndex)
+void SCustomizableObjectState::OnCollapseChanged(const ECheckBoxState NewCheckedState)
 {
 	bool bCollapse = (NewCheckedState != ECheckBoxState::Checked);
 
-	RuntimeParametersListWidgets[StateIndex]->SetCollapsed(bCollapse);
-	RuntimeParametersListWidgets[StateIndex]->SetVisibility(GetCollapsed(StateIndex));
-	CollapsedArrows[StateIndex]->SetImage(GetExpressionPreviewArrow(StateIndex));
+	RuntimeParametersWidget->SetCollapsed(bCollapse);
+	RuntimeParametersWidget->SetVisibility(GetCollapsed());
+	CollapsedArrow->SetImage(GetExpressionPreviewArrow());
 }
 
-
-EVisibility SCustomizableObjectNodeObjectSatesView::GetCollapsed(int32 StateIndex)
+EVisibility SCustomizableObjectState::GetCollapsed()
 {
-	return RuntimeParametersListWidgets[StateIndex]->IsCollapsed() ? EVisibility::Collapsed : EVisibility::Visible;
+	return RuntimeParametersWidget->IsCollapsed() ? EVisibility::Collapsed : EVisibility::Visible;
 }
 
-
-const FSlateBrush* SCustomizableObjectNodeObjectSatesView::GetExpressionPreviewArrow(int32 StateIndex) const
+const FSlateBrush* SCustomizableObjectState::GetExpressionPreviewArrow() const
 {
-	return FAppStyle::GetBrush(RuntimeParametersListWidgets[StateIndex]->IsCollapsed() ? TEXT("SurfaceDetails.PanUPositive") : TEXT("SurfaceDetails.PanVPositive"));
+	return FAppStyle::GetBrush(RuntimeParametersWidget->IsCollapsed() ? TEXT("SurfaceDetails.PanUPositive") : TEXT("SurfaceDetails.PanVPositive"));
 }
 
+FText SCustomizableObjectState::GetStateParameterCountText() const
+{
+	return FText::FromString(FString::Printf(TEXT("Runtime Parameters:  %d"), Node->States[StateIndex].RuntimeParameters.Num()));
+}
 
-FReply SCustomizableObjectNodeObjectSatesView::OnAddRuntimeParameterPressed(int32 StateIndex)
+FReply SCustomizableObjectState::OnAddRuntimeParameterPressed()
 {
 	Node->States[StateIndex].RuntimeParameters.Add("NONE");
-	RuntimeParametersListWidgets[StateIndex]->BuildList();
+	RuntimeParametersWidget->BuildList();
 
 	return FReply::Handled();
 }
 
+void SCustomizableObjectState::UpdateStateIndex(int32 NewStateIndex)
+{
+	StateIndex = NewStateIndex;
+	RuntimeParametersWidget->UpdateStateIndex(StateIndex);
+}
 
-// Widget for a list ofruntime parameters --------------------------------------------------------------------------
+
+void SCustomizableObjectNodeObjectSatesView::Construct(const FArguments& InArgs)
+{
+	Node = InArgs._Node;
+
+	if (Node)
+	{
+		SAssignNew(VerticalSlots, SDragAndDropVerticalBox)
+		.OnDragDetected(this,&SCustomizableObjectNodeObjectSatesView::OnStateDragDetected)
+		.OnCanAcceptDrop(this,&SCustomizableObjectNodeObjectSatesView::OnCanAcceptStateDrop)
+		.OnAcceptDrop(this,&SCustomizableObjectNodeObjectSatesView::OnAcceptStateDrop);
+		
+		for (int32 i = 0; i < Node->States.Num(); i++)
+		{
+			VerticalSlots->AddSlot()
+			.AutoHeight()
+			[
+				 SNew(SCustomizableObjectState)
+				.Node(Node)
+				.StateIndex(i)
+			];
+		}
+
+		// Add the widget to the child slot
+		ChildSlot
+		[
+			VerticalSlots.ToSharedRef()
+		];
+	}
+}
+
+/** drag and drop of states */
+
+
+FReply SCustomizableObjectNodeObjectSatesView::OnStateDragDetected(const FGeometry& Geometry,
+	const FPointerEvent& PointerEvent, int SlotBeingDraggedIndex, SVerticalBox::FSlot* Slot)
+{
+	if ( Slot)
+	{
+		// Widget being displayed during the drag and drop
+		TSharedPtr<STextBlock> WidgetToDisplay =
+			SNew(STextBlock)
+			.Text(FText::FromString( Node->States[SlotBeingDraggedIndex].Name));
+		
+		return FReply::Handled().BeginDragDrop( FDragAndDropOpWithWidget::New(SlotBeingDraggedIndex,Slot, WidgetToDisplay));
+	}
+
+	return FReply::Unhandled();
+}
+
+TOptional<SDragAndDropVerticalBox::EItemDropZone> SCustomizableObjectNodeObjectSatesView::OnCanAcceptStateDrop(
+	const FDragDropEvent& DragDropEvent, SDragAndDropVerticalBox::EItemDropZone ItemDropZone, SVerticalBox::FSlot* Slot)
+{
+	const TSharedPtr<FDragAndDropOpWithWidget> Operation = DragDropEvent.GetOperationAs<FDragAndDropOpWithWidget>();
+	if (Operation.IsValid())		
+	{
+		// only move states over same parent widget 
+		if (Operation->SlotBeingDragged->GetOwnerWidget()->GetId() == VerticalSlots->GetId())
+		{
+			// And if the value to be dragged over is different than the origin one
+			if (Operation->SlotBeingDragged != Slot)
+			{
+				return ItemDropZone;
+			}
+		}
+	}
+
+	// Operation will not succeed
+	return TOptional<SDragAndDropVerticalBox::EItemDropZone>();
+}
+
+FReply SCustomizableObjectNodeObjectSatesView::OnAcceptStateDrop(const FDragDropEvent& DragDropEvent,
+	SDragAndDropVerticalBox::EItemDropZone ItemDropZone, int NewIndex, SVerticalBox::FSlot* Slot)
+{
+	// Move the data around to match the drag and drop operation performed
+	TSharedPtr<FDragAndDropOpWithWidget> Operation = DragDropEvent.GetOperationAs<FDragAndDropOpWithWidget>();
+	if (Operation.IsValid())
+	{
+		// Apply drop operation on the states array to later update the states each state slate targets (Tick)
+		{
+			const FCustomizableObjectState DraggedParamTemp = Node->States[Operation->SlotIndexBeingDragged];
+			Node->States.RemoveAt(Operation->SlotIndexBeingDragged);
+
+			// Array structure changed, update new index accordingly
+			if (ItemDropZone == SDragAndDropVerticalBox::EItemDropZone::BelowItem &&
+				Operation->SlotIndexBeingDragged - NewIndex >= 1)
+			{
+				NewIndex -= 1;
+			}
+
+			Node->States.Insert(DraggedParamTemp,NewIndex);
+		}
+
+		// Set flag so on next tick all the slates get updated state indices
+		bWasStateDropPerformed = true;
+		return FReply::Handled();
+	}
+
+	return FReply::Unhandled();
+}
+
+
+void SCustomizableObjectNodeObjectSatesView::UpdateStatesIndex()
+{
+	FChildren* Children = VerticalSlots->GetChildren();
+	int32 NewSlateIndex = 0;
+	for (int32 ChildIndex = 0 ; ChildIndex < Children->Num(); ChildIndex++)
+	{
+		TSharedRef<SWidget> ChildSlate = Children->GetChildAt(ChildIndex);
+		SCustomizableObjectState* CustomizableObjectState =
+			StaticCast<SCustomizableObjectState*>(ChildSlate.ToSharedPtr().Get());
+
+		// Do not expect for all children to be SCustomizableObjectState
+		if (CustomizableObjectState)
+		{
+			CustomizableObjectState->UpdateStateIndex(NewSlateIndex);
+			NewSlateIndex++;
+		}
+	}
+}
+
+/** end of drag and drop */
+
+void SCustomizableObjectNodeObjectSatesView::Tick(const FGeometry& AllottedGeometry, const double InCurrentTime,
+												  const float InDeltaTime)
+{
+	SCompoundWidget::Tick(AllottedGeometry, InCurrentTime, InDeltaTime);
+
+	// After a drop operation has been performed update the indices Update the state index for each of the state slates
+	// so it matches the index each have as children of the slate
+	if (bWasStateDropPerformed )
+	{
+		UpdateStatesIndex();
+		bWasStateDropPerformed = false;
+	}
+}
+
+// Widget for a list of runtime parameters --------------------------------------------------------------------------
 
 
 void SCustomizableObjectRuntimeParameterList::Construct(const FArguments& InArgs)
@@ -196,7 +384,10 @@ void SCustomizableObjectRuntimeParameterList::Construct(const FArguments& InArgs
 
 		if (CustomizableObject)
 		{
-			SAssignNew(VerticalSlots, SVerticalBox);
+			SAssignNew(VerticalSlots, SDragAndDropVerticalBox)
+			.OnDragDetected(this,&SCustomizableObjectRuntimeParameterList::OnParamDragDetected)
+			.OnCanAcceptDrop(this,&SCustomizableObjectRuntimeParameterList::OnCanAcceptParamDrop)
+			.OnAcceptDrop(this,&SCustomizableObjectRuntimeParameterList::OnAcceptParamDrop);
 
 			if (VerticalSlots.IsValid())
 			{
@@ -209,7 +400,6 @@ void SCustomizableObjectRuntimeParameterList::Construct(const FArguments& InArgs
 			}
 		}
 	}
-
 }
 
 
@@ -222,35 +412,52 @@ void SCustomizableObjectRuntimeParameterList::BuildList()
 		for (int32 i = 0; i < Node->States[StateIndex].RuntimeParameters.Num(); ++i)
 		{
 			VerticalSlots->AddSlot()
-				.Padding(0.0f, 2.0f, 0.0f, 2.0f)
-				[
-					SNew(SHorizontalBox)
-					+ SHorizontalBox::Slot()
-				.AutoWidth()
-				[
-					SNew(SCustomizableObjectRuntimeParameter)
-					.Node(Node)
-				.StateIndex(StateIndex)
-				.RuntimeParameterIndex(i)
-				]
-			+ SHorizontalBox::Slot()
+			.Padding(0.0f, 2.0f, 0.0f, 2.0f)
+			[
+				SNew(SHorizontalBox)
+				
+				+ SHorizontalBox::Slot()
 				.AutoWidth()
 				[
 					SNew(SButton)
 					.OnClicked(this, &SCustomizableObjectRuntimeParameterList::OnDeleteRuntimeParameter, i)
-				.ToolTipText(LOCTEXT("RemoveRuntimeParameter", "Remove Runtime Parameter"))
-				[
-					SNew(SHorizontalBox)
-					+ SHorizontalBox::Slot()
-				.VAlign(VAlign_Center)
-				.HAlign(HAlign_Center)
-				[
-					SNew(SImage)
-					.Image(FAppStyle::GetBrush(TEXT("Cross")))
+					.ToolTipText(LOCTEXT("RemoveRuntimeParameter", "Remove Runtime Parameter"))
+					[
+						SNew(SHorizontalBox)
+						+ SHorizontalBox::Slot()
+						.VAlign(VAlign_Center)
+						.HAlign(HAlign_Center)
+						[
+							SNew(SImage)
+							.Image(FAppStyle::GetBrush(TEXT("Cross")))
+						]
+					]
 				]
+
+				+ SHorizontalBox::Slot()
+				.AutoWidth()
+				.Padding(6,0,0,0)
+				[
+					SNew(SBox)
+					.HAlign(HAlign_Right)
+					.VAlign(VAlign_Center)
+					.WidthOverride(16.0f)
+					[
+						SNew(SImage)
+						.Image(FCoreStyle::Get().GetBrush("VerticalBoxDragIndicatorShort"))
+					]
 				]
+				
+				+ SHorizontalBox::Slot()
+				.AutoWidth()
+				.Padding(4,0)
+				[
+					SNew(SCustomizableObjectRuntimeParameter)
+					.Node(Node)
+					.StateIndex(StateIndex)
+					.RuntimeParameterIndex(i)
 				]
-				];
+			];
 		}
 	}
 }
@@ -264,6 +471,78 @@ FReply SCustomizableObjectRuntimeParameterList::OnDeleteRuntimeParameter(int32 P
 	return FReply::Handled();
 }
 
+void SCustomizableObjectRuntimeParameterList::UpdateStateIndex(int32 NewStateIndex)
+{
+	StateIndex = NewStateIndex;
+	BuildList();
+}
+
+/** drag and drop of parameters */
+
+
+FReply SCustomizableObjectRuntimeParameterList::OnParamDragDetected(const FGeometry& Geometry,
+                                                                    const FPointerEvent& PointerEvent, int SlotBeingDraggedIndex, SVerticalBox::FSlot* Slot)
+{
+	if ( Slot)
+	{
+		// Widget being displayed during the drag and drop
+		TSharedPtr<STextBlock> WidgetToDisplay =
+			SNew(STextBlock)
+			.Text(FText::FromString( Node->States[StateIndex].RuntimeParameters[SlotBeingDraggedIndex]));
+		
+		return FReply::Handled().BeginDragDrop( FDragAndDropOpWithWidget::New(SlotBeingDraggedIndex,Slot, WidgetToDisplay));
+	}
+
+	return FReply::Unhandled();
+}
+
+TOptional<SDragAndDropVerticalBox::EItemDropZone> SCustomizableObjectRuntimeParameterList::OnCanAcceptParamDrop(
+	const FDragDropEvent& DragDropEvent, SDragAndDropVerticalBox::EItemDropZone ItemDropZone, SVerticalBox::FSlot* Slot)
+{
+	const TSharedPtr<FDragAndDropOpWithWidget> Operation = DragDropEvent.GetOperationAs<FDragAndDropOpWithWidget>();
+	if (Operation.IsValid())		
+	{
+		// only move parameters over same state
+		if (Operation->SlotBeingDragged->GetOwnerWidget()->GetId() == VerticalSlots->GetId())
+		{
+			// And if the value to be dragged over is different than the origin one
+			if (Operation->SlotBeingDragged != Slot)
+			{
+				return ItemDropZone;
+			}
+		}
+	}
+
+	return TOptional<SDragAndDropVerticalBox::EItemDropZone>();
+}
+
+FReply SCustomizableObjectRuntimeParameterList::OnAcceptParamDrop(const FDragDropEvent& DragDropEvent,
+	SDragAndDropVerticalBox::EItemDropZone ItemDropZone, int NewIndex, SVerticalBox::FSlot* Slot)
+{
+	// Move the data around to later rebuild the parameter UI objects with the re-structured data
+	TSharedPtr<FDragAndDropOpWithWidget> Operation = DragDropEvent.GetOperationAs<FDragAndDropOpWithWidget>();
+	if (Operation.IsValid())
+	{
+		const FString DraggedParamTemp = Node->States[StateIndex].RuntimeParameters[Operation->SlotIndexBeingDragged];
+		Node->States[StateIndex].RuntimeParameters.RemoveAt(Operation->SlotIndexBeingDragged);
+		
+		// Array structure changed, update new index accordingly
+		if (ItemDropZone == SDragAndDropVerticalBox::EItemDropZone::BelowItem &&
+			Operation->SlotIndexBeingDragged - NewIndex > 0)
+		{
+			NewIndex -= 1;
+		}
+	
+		Node->States[StateIndex].RuntimeParameters.Insert(DraggedParamTemp,NewIndex);
+	}
+	
+	// DO not perform the actual slate drop (only update the data and build the list again)
+	BuildList();
+	return FReply::Unhandled();
+}
+
+/** end of drag and drop */
+
 
 // Widget for each Runtime parameter -------------------------------------------------------------------------------
 
@@ -334,7 +613,7 @@ void SCustomizableObjectRuntimeParameter::Construct(const FArguments& InArgs)
 
 				+ SHorizontalBox::Slot()
 				.AutoWidth()
-				.Padding(2.0f, 4.0f, 2.0f, 0.0f)
+				.Padding(1.0f, 4.0f, 2.0f, 0.0f)
 				[
 					SNew(STextBlock)
 					.Text(FText::FromString(FString::Printf(TEXT("Parameter %d:"), RuntimeParameterIndex)))
@@ -357,7 +636,6 @@ void SCustomizableObjectRuntimeParameter::Construct(const FArguments& InArgs)
 	}
 }
 
-
 TSharedRef<SWidget> SCustomizableObjectRuntimeParameter::GetComboButtonContent()
 {
 	SearchItem.Reset();
@@ -366,18 +644,18 @@ TSharedRef<SWidget> SCustomizableObjectRuntimeParameter::GetComboButtonContent()
 	SAssignNew(RowNameComboListView, SListView<TSharedPtr<FString> >)
 		.ListItemsSource(&ListViewOptions)
 		.OnSelectionChanged(this, &SCustomizableObjectRuntimeParameter::OnComboButtonSelectionChanged)
-		.OnGenerateRow(this, &SCustomizableObjectRuntimeParameter::RowNameComboButtonGenarateWidget)
+		.OnGenerateRow(this, &SCustomizableObjectRuntimeParameter::RowNameComboButtonGenerateWidget)
 		.SelectionMode(ESelectionMode::Single);
 
-	// Serachbox Init
+	// SearchBox Init
 	SearchBoxWidget = SNew(SSearchBox)
 		.OnTextChanged(this, &SCustomizableObjectRuntimeParameter::OnSearchBoxFilterTextChanged)
 		.OnTextCommitted(this, &SCustomizableObjectRuntimeParameter::OnSearchBoxFilterTextCommitted);
 	
-	//Setting the focus to the Searchbox when the Comob button is oppen
+	//Setting the focus to the SearchBox when the Combo button is open
 	ComboButton->SetMenuContentWidgetToFocus(SearchBoxWidget);
 
-	// Creating a widget that gives navigation to the searchbox and the listview
+	// Creating a widget that gives navigation to the SearchBox and the listview
 	return  SNew(SVerticalBox)
 			+SVerticalBox::Slot()
 			.AutoHeight()
@@ -408,7 +686,7 @@ void SCustomizableObjectRuntimeParameter::OnComboButtonSelectionChanged(TSharedP
 }
 
 
-TSharedRef<ITableRow> SCustomizableObjectRuntimeParameter::RowNameComboButtonGenarateWidget(TSharedPtr<FString> InItem, const TSharedRef<STableViewBase>& OwnerTable)
+TSharedRef<ITableRow> SCustomizableObjectRuntimeParameter::RowNameComboButtonGenerateWidget(TSharedPtr<FString> InItem, const TSharedRef<STableViewBase>& OwnerTable)
 {
 	//This is needed because the filter made in the function OnSearchBocFilterTextChanged Only works for the rendered items
 	const EVisibility WidgetVisibility = IsItemVisible(InItem) ? EVisibility::Visible : EVisibility::Collapsed;
@@ -439,7 +717,16 @@ FText SCustomizableObjectRuntimeParameter::GetRowNameComboButtonContentText() co
 
 FText SCustomizableObjectRuntimeParameter::GetCurrentItemLabel() const
 {
-	return FText::FromString(Node->States[StateIndex].RuntimeParameters[RuntimeParameterIndex]);
+	// Due to the drag and drop operation the parameter index may point to an out of range entry. Hold a temp value
+	// until the data gets set up after the drop operation
+	if (Node->States[StateIndex].RuntimeParameters.IsValidIndex(RuntimeParameterIndex))
+	{
+		return FText::FromString(Node->States[StateIndex].RuntimeParameters[RuntimeParameterIndex]);
+	}
+	else
+	{
+		return LOCTEXT("None", "None");
+	}
 }
 
 
@@ -447,7 +734,7 @@ void SCustomizableObjectRuntimeParameter::OnSearchBoxFilterTextChanged(const FTe
 {
 	SearchItem = InText.ToString();
 
-	//This filter is just aplied to the items that are rendered of the ListView
+	//This filter is just applied to the items that are rendered of the ListView
 	for (int32 i = 0; i < ListViewOptions.Num(); i++)
 	{
 		TSharedPtr<ITableRow> Row = RowNameComboListView->WidgetFromItem(ListViewOptions[i]);
diff --git a/Engine/Plugins/Experimental/Mutable/Source/CustomizableObjectEditor/Private/MuCOE/SCustomizableObjectNodeObjectStatesView.h b/Engine/Plugins/Experimental/Mutable/Source/CustomizableObjectEditor/Private/MuCOE/SCustomizableObjectNodeObjectStatesView.h
--- a/Engine/Plugins/Experimental/Mutable/Source/CustomizableObjectEditor/Private/MuCOE/SCustomizableObjectNodeObjectStatesView.h
+++ b/Engine/Plugins/Experimental/Mutable/Source/CustomizableObjectEditor/Private/MuCOE/SCustomizableObjectNodeObjectStatesView.h
@@ -3,6 +3,7 @@
 #pragma once
 
 #include "Widgets/SCompoundWidget.h"
+#include "Widgets/SBoxPanel.h"
 
 enum class ECheckBoxState : uint8;
 namespace ESelectInfo { enum Type : int; }
@@ -20,6 +21,24 @@ class UCustomizableObjectNodeObject;
 struct FSlateBrush;
 template <typename OptionType> class SComboBox;
 
+/** Object representing a drag and drop operation for states and parameters */
+class FDragAndDropOpWithWidget : public FDragAndDropVerticalBoxOp
+{
+public:
+	DRAG_DROP_OPERATOR_TYPE(FDragAndDropOpWithWidget, FDragAndDropVerticalBoxOp)
+
+	/** Returns a new Drag and drop operation already initialized and ready to be used */
+	static TSharedRef<FDragAndDropOpWithWidget> New(int32 InSlotIndexBeingDragged, SVerticalBox::FSlot* InSlotBeingDragged, TSharedPtr<SWidget> InWidgetToShow);
+	
+	// FromFDragAndDropVerticalBoxOp
+	virtual TSharedPtr<SWidget> GetDefaultDecorator() const override;
+
+private:
+	/** Widget object that will be shown when dragging the object around */
+	TSharedPtr<SWidget> WidgetToShow;
+};
+
+
 class SCustomizableObjectRuntimeParameter : public SCompoundWidget
 {
 public:
@@ -43,16 +62,16 @@ private:
 	/** Callback for the Combo Button selection */
 	void OnComboButtonSelectionChanged(TSharedPtr<FString> SelectedItem, ESelectInfo::Type SelectInfo);
 
-	/** Generates the lables of the list view for the combo button */
-	TSharedRef<ITableRow> RowNameComboButtonGenarateWidget(TSharedPtr<FString> InItem, const TSharedRef<STableViewBase>& OwnerTable);
+	/** Generates the labels of the list view for the combo button */
+	TSharedRef<ITableRow> RowNameComboButtonGenerateWidget(TSharedPtr<FString> InItem, const TSharedRef<STableViewBase>& OwnerTable);
 
 	/** Generate the Combo Button selected label */
 	FText GetRowNameComboButtonContentText() const;
 
-	/** Callback for the OnTextChanged of the Searchbox */
+	/** Callback for the OnTextChanged of the SearchBox */
 	void OnSearchBoxFilterTextChanged(const FText& InText);
 
-	/** Callback for the OnTextCommited of the Searchbox */
+	/** Callback for the OnTextCommited of the SearchBox */
 	void OnSearchBoxFilterTextCommitted(const FText& InText, ETextCommit::Type CommitInfo);
 
 	// Returns true if the item should be visible in the Combo button
@@ -61,13 +80,13 @@ private:
 private:
 
 	/** Node with all the information */
-	class UCustomizableObjectNodeObject* Node;
+	class UCustomizableObjectNodeObject* Node = nullptr;
 
 	/** Index to identify which parameter this widget is modifying */
-	int32 StateIndex;
+	int32 StateIndex = -1;
 
 	/** Index to identify which parameter this widget is modifying */
-	int32 RuntimeParameterIndex;
+	int32 RuntimeParameterIndex = -1;
 
 	/** Combobox for runtime parameter options */
 	TSharedPtr<SComboBox<TSharedPtr<FString>>> RuntimeParamCombobox;
@@ -85,13 +104,13 @@ private:
 	/** Search box of the ComboButton */
 	TSharedPtr<SSearchBox> SearchBoxWidget;
 
-	/** Stores the input of the Searchboc widget*/
+	/** Stores the input of the Search box widget*/
 	FString SearchItem;
 };
 
 
-/* Widget that represents List of Runtime Parameters Widgets */
 
+/* Widget that represents List of Runtime Parameters Widgets */
 class SCustomizableObjectRuntimeParameterList : public SCompoundWidget
 {
 public:
@@ -102,64 +121,159 @@ public:
 	SLATE_END_ARGS()
 
 	void Construct(const FArguments& InArgs);
+
 	void BuildList();
 
-	/** Removes a Runtime parameter and rebuilds the runtimeparameter widgets */
+	/** Removes a Runtime parameter and rebuilds the runtimeParameter widgets */
 	FReply OnDeleteRuntimeParameter(int32 ParameterIndex);
 
-	bool IsCollapsed() { return bCollapsed; }
-	void SetCollapsed(bool Collapse ) { bCollapsed = Collapse; }
-	
+	/** Tells the caller if the parameter list is collapsed (hidden contents) or not.*/
+	bool IsCollapsed() const { return bCollapsed; }
+
+	/** Set the collapsed state of the list of parameters.
+	 * @param bShouldBeCollapsed Determines if the contents of the list should be hidden (true) or visible (false)
+	 */
+	void SetCollapsed(bool bShouldBeCollapsed ) { bCollapsed = bShouldBeCollapsed; }
+
+	/** Update the state this parameter list is looking at in order to draw it's parameters. Updating this value will
+	 * rebuild the list of parameters.
+	 */
+	void UpdateStateIndex(int32 NewStateIndex);
+
 private:
+	
+	// Parameter drag and drop callbacks
+
+	/** Method invoked each time a new drag operation is found to be happening using one parameter slate part of the parameters list */
+	FReply OnParamDragDetected(const FGeometry& Geometry, const FPointerEvent& PointerEvent, int SlotBeingDraggedIndex, SVerticalBox::FSlot* Slot);
+
+	/** Method invoked each time the dragged parameter intersects with another slate. It determines if the drop can be performed over the
+	 * intersected slate object.
+	 */
+	TOptional<SDragAndDropVerticalBox::EItemDropZone> OnCanAcceptParamDrop(const FDragDropEvent& DragDropEvent, SDragAndDropVerticalBox::EItemDropZone ItemDropZone, SVerticalBox::FSlot* Slot);
+
+	/** Method invoked if the drop operation has been verified by OnCanAcceptParamDrop.
+	 * @note It will perform such drop by updating the
+	 * runtime parameters array (from the set StateIndex) and then rebuild the list of parameters.
+	 */
+	FReply OnAcceptParamDrop(const FDragDropEvent& DragDropEvent, SDragAndDropVerticalBox::EItemDropZone ItemDropZone, int NewIndex, SVerticalBox::FSlot* Slot);
 
-	/** Node with all the information */
-	class UCustomizableObjectNodeObject* Node;
 
+	/** Node with all the information */
+	class UCustomizableObjectNodeObject* Node = nullptr;
+	
 	/** Index to identify which parameter this widget is modifying */
-	int32 StateIndex;
+	int32 StateIndex = -1;
 
 	/** Bool that determines if a RuntimeParameter should be collapsed or not*/
-	bool bCollapsed;
+	bool bCollapsed = false;
 
 	/** Vertical box widget for the RuntimeParameter Widgets  */
-	TSharedPtr<SVerticalBox> VerticalSlots;
+	TSharedPtr<SDragAndDropVerticalBox> VerticalSlots;
 };
 
-class SCustomizableObjectNodeObjectSatesView : public SCompoundWidget
+
+/** Structure representing a mutable state. */
+class SCustomizableObjectState : public SCompoundWidget
 {
 public:
 
-	SLATE_BEGIN_ARGS(SCustomizableObjectNodeObjectSatesView) {}
-	SLATE_ARGUMENT(UCustomizableObjectNodeObject*, Node)
+	SLATE_BEGIN_ARGS(SCustomizableObjectState) {}
+	SLATE_ARGUMENT(class UCustomizableObjectNodeObject*, Node)
+	SLATE_ARGUMENT(int32, StateIndex)
 	SLATE_END_ARGS()
 
 	void Construct(const FArguments& InArgs);
 
 	/** Callback for the collapsing arrow checkbox */
-	void OnCollapseChanged(const ECheckBoxState NewCheckedState, int32 StateIndex);
+	void OnCollapseChanged(ECheckBoxState NewCheckedState);
 
 	/** Gets if a runtime parameter widget is collapsed */
-	EVisibility GetCollapsed(int32 StateIndex);
+	EVisibility GetCollapsed();
 
 	/** Return the brush for the collapsed arrow */
-	const FSlateBrush* GetExpressionPreviewArrow(int32 StateIndex) const;
+	const FSlateBrush* GetExpressionPreviewArrow() const;
+
+	/** Return a text with the amount of parameters the targeted state has. It is formatted for UI usage*/
+	FText GetStateParameterCountText() const;
 
-	/** Adds a new runtime parameter and rebuild the runtimeparameterListWidgets */
-	FReply OnAddRuntimeParameterPressed(int32 StateIndex);
+	/** Adds a new runtime parameter and rebuild the runtimeparameterList Widgets */
+	FReply OnAddRuntimeParameterPressed();
 
+	/** Update the state index this slate object is using for drawing it's content. Will also affect child structures of this
+	 * slate*/
+	void UpdateStateIndex(int32 NewStateIndex);
+
+	
 private:
 
-	/** Pointer to the current Node */
-	UCustomizableObjectNodeObject* Node;
+	/** Node with all the information */
+	class UCustomizableObjectNodeObject* Node = nullptr;
+
+	/** Index to identify which parameter this widget is modifying */
+	int32 StateIndex = -1;
 
 	/** Vertical Boxes to store the widget of each state */
 	TSharedPtr<SVerticalBox> VerticalSlots;
+	
+	/** Array to edit widgets when needed */
+	TSharedPtr<SCustomizableObjectRuntimeParameterList> RuntimeParametersWidget;
 
-	/** Arrays to edit widgets when needed */
-	TArray<TSharedPtr<SCustomizableObjectRuntimeParameterList>> RuntimeParametersListWidgets;
+	/** SImage to control the collapsing buttons */
+	TSharedPtr<SImage> CollapsedArrow;
+	
+};
 
-	/** Array of SImages to control the collapsing buttons */
-	TArray<TSharedPtr<SImage>> CollapsedArrows;
 
+
+class SCustomizableObjectNodeObjectSatesView : public SCompoundWidget
+{
+public:
+
+	SLATE_BEGIN_ARGS(SCustomizableObjectNodeObjectSatesView) {}
+	SLATE_ARGUMENT(UCustomizableObjectNodeObject*, Node)
+	SLATE_END_ARGS()
+
+	void Construct(const FArguments& InArgs);
+	
+	virtual void Tick(const FGeometry& AllottedGeometry, const double InCurrentTime, const float InDeltaTime) override;
+	
+private:
+	
+	/** Method invoked each time a new drag operation is found to be happening using one state slate part of the parameters list */
+	FReply OnStateDragDetected(const FGeometry& Geometry, const FPointerEvent& PointerEvent, int SlotBeingDraggedIndex, SVerticalBox::FSlot* Slot);
+
+	/** Method invoked each time the dragged state intersects with another slate. It determines if the drop can be performed over the
+	 * intersected slate object.
+	 */
+	TOptional<SDragAndDropVerticalBox::EItemDropZone> OnCanAcceptStateDrop(const FDragDropEvent& DragDropEvent, SDragAndDropVerticalBox::EItemDropZone ItemDropZone, SVerticalBox::FSlot* Slot);
+
+	/** Method invoked if the drop operation has been verified by OnCanAcceptParamDrop.
+	 * @note It will perform such drop by updating the
+	 * state array structure to later update their respective targeted state indices based on the new position they will
+	 * hold as children of this slate.
+	 */
+	FReply OnAcceptStateDrop(const FDragDropEvent& DragDropEvent, SDragAndDropVerticalBox::EItemDropZone ItemDropZone, int NewIndex, SVerticalBox::FSlot* Slot);
+	
+	/** Updates each state slate internal stateIndex to match their children index inside the VerticalBox.
+	 * @note Method that works together with the drag and drop operation. */
+	void UpdateStatesIndex();
+
+private:
+
+	/** Variable keeping track if a state drag and drop has been performed. If so it will correct the slates of the
+	 * child SCustomizableObjectState to match the new UI structure.
+	 */
+	bool bWasStateDropPerformed = false;
+	
+	/** Pointer to the current Node */
+	UCustomizableObjectNodeObject* Node = nullptr;
+
+	/** Vertical Boxes to store the widget of each state */
+	TSharedPtr<SDragAndDropVerticalBox> VerticalSlots;
 };
 
+
+
+
+
