[Mutable] Fix possible deadlock with image raster mesh operations.

#rnx

diff --git a/Engine/Plugins/Experimental/Mutable/Source/MutableRuntime/Private/MuR/CodeRunner.cpp b/Engine/Plugins/Experimental/Mutable/Source/MutableRuntime/Private/MuR/CodeRunner.cpp
--- a/Engine/Plugins/Experimental/Mutable/Source/MutableRuntime/Private/MuR/CodeRunner.cpp
+++ b/Engine/Plugins/Experimental/Mutable/Source/MutableRuntime/Private/MuR/CodeRunner.cpp
@@ -3879,22 +3879,13 @@ namespace mu
                 break;
 			case 1:
 			{
-                MUTABLE_CPUPROFILER_SCOPE(IM_RASTERMESH_1)
+				MUTABLE_CPUPROFILER_SCOPE(IM_RASTERMESH_1)
 
-				Ptr<const Mesh> pMesh = GetMemory().GetMesh(FScheduledOp::FromOpAndOptions(args.mesh, item, 0));
-				Ptr<const Projector> pProjector = GetMemory().GetProjector(FScheduledOp::FromOpAndOptions(args.projector, item, 0));
-
-				if (!pMesh)
-				{
-					// Stop execution.
-					check(false);
-					GetMemory().SetImage(item, nullptr);
-					break;
-				}
-
-				// If no image we are generating a flat mesh UV raster. This is the final stage in this case.
+				// If no image, we are generating a flat mesh UV raster. This is the final stage in this case.
 				if (!args.image)
 				{
+					Ptr<const Mesh> pMesh = GetMemory().GetMesh(FScheduledOp::FromOpAndOptions(args.mesh, item, 0));
+
 					uint16 SizeX = args.sizeX;
 					uint16 SizeY = args.sizeY;
 
@@ -3909,16 +3900,22 @@ namespace mu
 
                     // Flat mesh UV raster
 					Ptr<Image> ResultImage = new Image(SizeX, SizeY, 1, EImageFormat::IF_L_UBYTE);
-                    ImageRasterMesh( pMesh.get(), ResultImage.get(), args.blockIndex );
+					if (pMesh)
+					{
+						ImageRasterMesh(pMesh.get(), ResultImage.get(), args.blockIndex);
+					}
 
 					// Stop execution.
 					GetMemory().SetImage(item, ResultImage);
 					break;
 				}
 
+				Ptr<const Mesh> pMesh = GetMemory().GetMesh(FScheduledOp::FromOpAndOptions(args.mesh, item, 0));
+
 				const int32 MipsToSkip = item.ExecutionOptions;
 				int32 ProjectionMip = MipsToSkip;
 
+				Ptr<const Projector> pProjector = GetMemory().GetProjector(FScheduledOp::FromOpAndOptions(args.projector, item, 0));
 				if (pProjector)
 				{
 					if (static_cast<EMinFilterMethod>(args.MinFilterMethod) == EMinFilterMethod::TotalAreaHeuristic)
@@ -3927,10 +3924,19 @@ namespace mu
 							FMath::Max(args.sizeX >> MipsToSkip, 1),
 							FMath::Max(args.sizeY >> MipsToSkip, 1));
 						FVector2f SourceImageSizeF = FVector2f(args.SourceSizeX, args.SourceSizeY);
-
-						ProjectionMip = ComputeProjectedFootprintBestMip(pMesh.get(), pProjector->m_value, TargetImageSizeF, SourceImageSizeF);
+						
+						if (pMesh)
+						{
+							ProjectionMip = ComputeProjectedFootprintBestMip(pMesh.get(), pProjector->m_value, TargetImageSizeF, SourceImageSizeF);
+						}
 					}
 				}
+	
+				// Set mesh and projector back to the cache to prevent possible code re-execution.
+				// this needs to be done before the operation is re-scheduled otherwise the execution may deadlock.	
+				// TODO: Review if this is a good way of reusing data between stages.
+				GetMemory().SetMesh(FScheduledOp::FromOpAndOptions(args.mesh, item, 0), pMesh);
+				GetMemory().SetProjector(FScheduledOp::FromOpAndOptions(args.projector, item, 0), pProjector);
 
 				AddOp(FScheduledOp(item.At, item, 2, ProjectionMip),
 					FScheduledOp::FromOpAndOptions(args.mesh, item, 0),
@@ -3939,10 +3945,6 @@ namespace mu
 					FScheduledOp(args.mask, item),
 					FScheduledOp::FromOpAndOptions(args.angleFadeProperties, item, 0));
 
-				// Set mesh and projector back to the cache to prevent possible code re-execution.
-				GetMemory().SetMesh(FScheduledOp::FromOpAndOptions(args.mesh, item, 0), pMesh);
-				GetMemory().SetProjector(FScheduledOp::FromOpAndOptions(args.projector, item, 0), pProjector);
-
 				break;
 			}
             case 2:
@@ -5260,6 +5262,7 @@ namespace mu
 
 		OP_TYPE type = pModel->GetPrivate()->m_program.GetOpType(item.At);
 		//UE_LOG(LogMutableCore, Log, TEXT("Running :%5d , %d, of type %d "), item.At, item.Stage, type);
+
 		switch ( type )
         {
         case OP_TYPE::NONE:
@@ -5350,7 +5353,6 @@ namespace mu
         }
     }
 
-
 	//---------------------------------------------------------------------------------------------
 	void CodeRunner::RunCodeImageDesc(const FScheduledOp& item, const Parameters* pParams, const Model* pModel,  uint32 lodMask )
 	{
@@ -6050,7 +6052,5 @@ namespace mu
 			break;
 		}
 	}
-
-
 }
 
