#rnx

diff --git a/Engine/Plugins/Experimental/Mutable/Source/CustomizableObject/Private/MuCO/CustomizableObjectSystem.cpp b/Engine/Plugins/Experimental/Mutable/Source/CustomizableObject/Private/MuCO/CustomizableObjectSystem.cpp
--- a/Engine/Plugins/Experimental/Mutable/Source/CustomizableObject/Private/MuCO/CustomizableObjectSystem.cpp
+++ b/Engine/Plugins/Experimental/Mutable/Source/CustomizableObject/Private/MuCO/CustomizableObjectSystem.cpp
@@ -2250,12 +2250,10 @@ namespace impl
 
 		if (SystemPrivateData->EnableReuseInstanceTextures)
 		{
-			if (Operation->bNeverStream)
-			{
-				bReuseInstanceTextures = StateData ? StateData->bReuseInstanceTextures : false;
-				bReuseInstanceTextures |= CandidateInstancePrivateData->HasCOInstanceFlags(ReuseTextures);
-			}
-			else
+			bReuseInstanceTextures = StateData ? StateData->bReuseInstanceTextures : false;
+			bReuseInstanceTextures |= CandidateInstancePrivateData->HasCOInstanceFlags(ReuseTextures);
+			
+			if (bReuseInstanceTextures && !Operation->bNeverStream)
 			{
 				UE_LOG(LogMutable, Warning, TEXT("Instance texture reuse requires that the current Mutable state is in non-streaming mode. Change it in the Mutable graph base node in the state definition."));
 				bReuseInstanceTextures = false;
