[mutable] fixed toggle option of Group Nodes (of child Customizable Objects)

#rnx

diff --git a/Engine/Plugins/Experimental/Mutable/Source/CustomizableObjectEditor/Private/MuCOE/GenerateMutableSource/GenerateMutableSource.cpp b/Engine/Plugins/Experimental/Mutable/Source/CustomizableObjectEditor/Private/MuCOE/GenerateMutableSource/GenerateMutableSource.cpp
--- a/Engine/Plugins/Experimental/Mutable/Source/CustomizableObjectEditor/Private/MuCOE/GenerateMutableSource/GenerateMutableSource.cpp
+++ b/Engine/Plugins/Experimental/Mutable/Source/CustomizableObjectEditor/Private/MuCOE/GenerateMutableSource/GenerateMutableSource.cpp
@@ -1092,6 +1092,8 @@ mu::NodeObjectPtr GenerateMutableSource(const UEdGraphPin * Pin, FMutableGraphGe
 					BooleanParam.ParamUIMetadata,
 					EMutableParameterType::Bool);
 
+				ParameterUIDataBoolean.ParamUIMetadata.ExtraInformation = ParameterUIData.ParamUIMetadata.ExtraInformation;
+
 				GenerationContext.ParameterUIDataMap.Add(BooleanParam.Name, ParameterUIDataBoolean);
 			}
 		}
