[mutable] Fixed layout preview for multiple UV channels

#rnx

diff --git a/Engine/Plugins/Experimental/Mutable/Source/CustomizableObjectEditor/Private/MuCOE/SCustomizableObjectNodeLayoutBlocksEditor.cpp b/Engine/Plugins/Experimental/Mutable/Source/CustomizableObjectEditor/Private/MuCOE/SCustomizableObjectNodeLayoutBlocksEditor.cpp
--- a/Engine/Plugins/Experimental/Mutable/Source/CustomizableObjectEditor/Private/MuCOE/SCustomizableObjectNodeLayoutBlocksEditor.cpp
+++ b/Engine/Plugins/Experimental/Mutable/Source/CustomizableObjectEditor/Private/MuCOE/SCustomizableObjectNodeLayoutBlocksEditor.cpp
@@ -74,7 +74,7 @@ void SCustomizableObjectNodeLayoutBlocksEditor::SetCurrentLayout(UCustomizableOb
 
 	if (CurrentLayout)
 	{
-		CurrentLayout->GetUVChannel(UVs);
+		CurrentLayout->GetUVChannel(UVs, CurrentLayout->GetUVChannel());
 
 		UnassignedUVs = TArray<FVector2f>();
 		
