#rnx

diff --git a/Engine/Plugins/Experimental/Mutable/Source/CustomizableObject/Private/MuCO/CustomizableObject.cpp b/Engine/Plugins/Experimental/Mutable/Source/CustomizableObject/Private/MuCO/CustomizableObject.cpp
--- a/Engine/Plugins/Experimental/Mutable/Source/CustomizableObject/Private/MuCO/CustomizableObject.cpp
+++ b/Engine/Plugins/Experimental/Mutable/Source/CustomizableObject/Private/MuCO/CustomizableObject.cpp
@@ -1098,8 +1098,10 @@ int32 UCustomizableObject::GetComponentCount() const
 	{
 		UE_LOG(LogMutable, Warning,
 		       TEXT(
-			       "You are trying to get the component count of a non compiled CO. This will always return 0 as value."
-		       ));
+			       "You are trying to get the component count of the non-compiled CO %s. This will always return 0 as value."
+		       ), *GetName()
+		);
+
 		return 0;
 	}
 	
@@ -1137,7 +1139,7 @@ EMutableParameterType UCustomizableObject::GetParameterTypeByName(const FString&
 		return ParameterProperties[Index].Type;
 	}
 
-	UE_LOG(LogMutable, Warning, TEXT("Name '%s' does not exist in ParameterProperties lookup table at GetParameterTypeByName."), *Name);
+	UE_LOG(LogMutable, Warning, TEXT("Name '%s' does not exist in ParameterProperties lookup table at GetParameterTypeByName at CO %s."), *Name, *GetName());
 
 	for (int32 ParamIndex = 0; ParamIndex < ParameterProperties.Num(); ++ParamIndex)
 	{
@@ -1147,7 +1149,7 @@ EMutableParameterType UCustomizableObject::GetParameterTypeByName(const FString&
 		}
 	}
 
-	UE_LOG(LogMutable, Warning, TEXT("Name '%s' does not exist in ParameterProperties at GetParameterTypeByName."), *Name);
+	UE_LOG(LogMutable, Warning, TEXT("Name '%s' does not exist in ParameterProperties at GetParameterTypeByName at CO %s."), *Name, *GetName());
 
 	return EMutableParameterType::None;
 }
@@ -1163,7 +1165,7 @@ const FString & UCustomizableObject::GetParameterName(int32 ParamIndex) const
 	}
 	else
 	{
-		UE_LOG(LogMutable, Warning, TEXT("Index [%d] out of ParameterProperties bounds at GetParameterName."), ParamIndex);
+		UE_LOG(LogMutable, Warning, TEXT("Index [%d] out of ParameterProperties bounds at GetParameterName at CO %s."), ParamIndex, *GetName());
 	}
 
 	return s_EmptyString;
@@ -1261,7 +1263,7 @@ int UCustomizableObject::GetParameterDescriptionCount(int32 ParamIndex) const
 	}
 	else
 	{
-		UE_LOG(LogMutable, Warning, TEXT("Index [%d] out of ParameterProperties bounds at GetParameterDescriptionCount."), ParamIndex);
+		UE_LOG(LogMutable, Warning, TEXT("Index [%d] out of ParameterProperties bounds at GetParameterDescriptionCount at CO %s."), ParamIndex, *GetName());
 	}
 
 	return 0;
@@ -1282,7 +1284,7 @@ int32 UCustomizableObject::GetIntParameterNumOptions(int32 ParamIndex) const
 	}
 	else
 	{
-		UE_LOG(LogMutable, Warning, TEXT("Index [%d] out of ParameterProperties bounds at GetIntParameterNumOptions."), ParamIndex);
+		UE_LOG(LogMutable, Warning, TEXT("Index [%d] out of ParameterProperties bounds at GetIntParameterNumOptions at CO %s."), ParamIndex, *GetName());
 	}
 
 	return 0;
@@ -1299,12 +1301,12 @@ const FString& UCustomizableObject::GetIntParameterAvailableOption(int32 ParamIn
 		}
 		else
 		{
-			UE_LOG(LogMutable, Warning, TEXT("Index [%d] out of IntParameterNumOptions bounds at GetIntParameterAvailableOption."), K);
+			UE_LOG(LogMutable, Warning, TEXT("Index [%d] out of IntParameterNumOptions bounds at GetIntParameterAvailableOption at CO %s."), K, *GetName());
 		}
 	}
 	else
 	{
-		UE_LOG(LogMutable, Warning, TEXT("Index [%d] out of ParameterProperties bounds at GetIntParameterAvailableOption."), ParamIndex);
+		UE_LOG(LogMutable, Warning, TEXT("Index [%d] out of ParameterProperties bounds at GetIntParameterAvailableOption at CO %s."), ParamIndex, *GetName());
 	}
 
 	return s_EmptyString;
@@ -1346,12 +1348,12 @@ int32 UCustomizableObject::FindIntParameterValue(int32 ParamIndex, const FString
 		}
 		else
 		{
-			UE_LOG(LogMutable, Warning, TEXT("No possible values for parameter with index [%d] at FindIntParameterValue."), ParamIndex);
+			UE_LOG(LogMutable, Warning, TEXT("No possible values for parameter with index [%d] at FindIntParameterValue at CO %s."), ParamIndex, *GetName());
 		}
 	}
 	else
 	{
-		UE_LOG(LogMutable, Warning, TEXT("Index [%d] out of ParameterProperties bounds at FindIntParameterValue."), ParamIndex);
+		UE_LOG(LogMutable, Warning, TEXT("Index [%d] out of ParameterProperties bounds at FindIntParameterValue at CO %s."), ParamIndex, *GetName());
 	}
 	return MinValueIndex;
 }
@@ -1373,7 +1375,7 @@ FString UCustomizableObject::FindIntParameterValueName(int32 ParamIndex, int32 P
 	}
 	else
 	{
-		UE_LOG(LogMutable, Warning, TEXT("Index [%d] out of ParameterProperties bounds at FindIntParameterValueName."), ParamIndex);
+		UE_LOG(LogMutable, Warning, TEXT("Index [%d] out of ParameterProperties bounds at FindIntParameterValueName at CO %s."), ParamIndex, *GetName());
 	}
 
 	return FString();
diff --git a/Engine/Plugins/Experimental/Mutable/Source/CustomizableObjectEditor/Private/MuCOE/CustomizableObjectCompileRunnable.cpp b/Engine/Plugins/Experimental/Mutable/Source/CustomizableObjectEditor/Private/MuCOE/CustomizableObjectCompileRunnable.cpp
--- a/Engine/Plugins/Experimental/Mutable/Source/CustomizableObjectEditor/Private/MuCOE/CustomizableObjectCompileRunnable.cpp
+++ b/Engine/Plugins/Experimental/Mutable/Source/CustomizableObjectEditor/Private/MuCOE/CustomizableObjectCompileRunnable.cpp
@@ -178,7 +178,7 @@ FCustomizableObjectSaveDDRunnable::FCustomizableObjectSaveDDRunnable(UCustomizab
 	{
 		// We will be saving all compilation data in two separate files, write CO Data
 		FolderPath = CustomizableObject->GetCompiledDataFolderPath(!InOptions.bIsCooking);
-		CompildeDataFullFileName = FolderPath + CustomizableObject->GetCompiledDataFileName(true, InOptions.TargetPlatform);
+		CompileDataFullFileName = FolderPath + CustomizableObject->GetCompiledDataFileName(true, InOptions.TargetPlatform);
 		StreamableDataFullFileName = FolderPath + CustomizableObject->GetCompiledDataFileName(false, InOptions.TargetPlatform);
 
 		// Serialize Customizable Object's data
@@ -213,10 +213,10 @@ uint32 FCustomizableObjectSaveDDRunnable::Run()
 
 		// Delete files...
 		bool bFilesDeleted = true;
-		if (FileManager.FileExists(*CompildeDataFullFileName)
-			&& !FileManager.Delete(*CompildeDataFullFileName, true, false, true))
+		if (FileManager.FileExists(*CompileDataFullFileName)
+			&& !FileManager.Delete(*CompileDataFullFileName, true, false, true))
 		{
-			UE_LOG(LogMutable, Error, TEXT("Failed to delete compiled data in file [%s]."), *CompildeDataFullFileName);
+			UE_LOG(LogMutable, Error, TEXT("Failed to delete compiled data in file [%s]."), *CompileDataFullFileName);
 			bFilesDeleted = false;
 		}
 
@@ -231,7 +231,7 @@ uint32 FCustomizableObjectSaveDDRunnable::Run()
 		if (bFilesDeleted)
 		{
 			// Create file writers...
-			FArchive* ModelMemoryWriter = FileManager.CreateFileWriter(*CompildeDataFullFileName);
+			FArchive* ModelMemoryWriter = FileManager.CreateFileWriter(*CompileDataFullFileName);
 			FArchive* StreamableMemoryWriter = FileManager.CreateFileWriter(*StreamableDataFullFileName);
 			check(ModelMemoryWriter);
 			check(StreamableMemoryWriter);
diff --git a/Engine/Plugins/Experimental/Mutable/Source/CustomizableObjectEditor/Private/MuCOE/CustomizableObjectCompileRunnable.h b/Engine/Plugins/Experimental/Mutable/Source/CustomizableObjectEditor/Private/MuCOE/CustomizableObjectCompileRunnable.h
--- a/Engine/Plugins/Experimental/Mutable/Source/CustomizableObjectEditor/Private/MuCOE/CustomizableObjectCompileRunnable.h
+++ b/Engine/Plugins/Experimental/Mutable/Source/CustomizableObjectEditor/Private/MuCOE/CustomizableObjectCompileRunnable.h
@@ -103,7 +103,7 @@ private:
 
 	// Paths used to save files to disk
 	FString FolderPath;
-	FString CompildeDataFullFileName;
+	FString CompileDataFullFileName;
 	FString StreamableDataFullFileName;
 
 	TSharedPtr<mu::Model, ESPMode::ThreadSafe> Model;
diff --git a/Engine/Plugins/Experimental/Mutable/Source/CustomizableObjectEditor/Private/MuCOE/CustomizableObjectCompiler.cpp b/Engine/Plugins/Experimental/Mutable/Source/CustomizableObjectEditor/Private/MuCOE/CustomizableObjectCompiler.cpp
--- a/Engine/Plugins/Experimental/Mutable/Source/CustomizableObjectEditor/Private/MuCOE/CustomizableObjectCompiler.cpp
+++ b/Engine/Plugins/Experimental/Mutable/Source/CustomizableObjectEditor/Private/MuCOE/CustomizableObjectCompiler.cpp
@@ -49,7 +49,7 @@ bool FCustomizableObjectCompiler::Tick()
 
 	if (CompileTask.IsValid() && CompileTask->IsCompleted())
 	{
-		UE_LOG(LogMutable, Log, TEXT("PROFILE: [ %16.8f ] Finishing Compilation task."), FPlatformTime::Seconds());
+		UE_LOG(LogMutable, Log, TEXT("PROFILE: [ %16.8f ] Finishing Compilation task for Object %s."), FPlatformTime::Seconds(), *CurrentObject->GetName());
 
 		FinishCompilation();
 
diff --git a/Engine/Plugins/Experimental/Mutable/Source/CustomizableObjectEditor/Private/MuCOE/GenerateMutableSource/GenerateMutableSourceMesh.cpp b/Engine/Plugins/Experimental/Mutable/Source/CustomizableObjectEditor/Private/MuCOE/GenerateMutableSource/GenerateMutableSourceMesh.cpp
--- a/Engine/Plugins/Experimental/Mutable/Source/CustomizableObjectEditor/Private/MuCOE/GenerateMutableSource/GenerateMutableSourceMesh.cpp
+++ b/Engine/Plugins/Experimental/Mutable/Source/CustomizableObjectEditor/Private/MuCOE/GenerateMutableSource/GenerateMutableSourceMesh.cpp
@@ -542,7 +542,8 @@ mu::MeshPtr ConvertSkeletalMeshToMutable(USkeletalMesh* InSkeletalMesh, int LOD,
 				if (InSkeletonRefSkeleton.FindRawBoneIndex(BoneInfo.Name) == INDEX_NONE)
 				{
 					bIsSkeletonMissingBones = true;
-					UE_LOG(LogMutable, Warning, TEXT("SkeletalMesh [%s] uses bone [%s] not present in skeleton [%s]."),
+					UE_LOG(LogMutable, Warning, TEXT("In object [%s] SkeletalMesh [%s] uses bone [%s] not present in skeleton [%s]."),
+						*GenerationContext.Object->GetName(),
 						*InSkeletalMesh->GetName(),
 						*BoneInfo.ExportName,
 						*InSkeleton->GetName());
@@ -1311,7 +1312,8 @@ mu::MeshPtr ConvertSkeletalMeshToMutable(USkeletalMesh* InSkeletalMesh, int LOD,
 			if (BoneMapIndex >= NumRequiredBones)
 			{
 				// This can only happen with invalid data.
-				UE_LOG(LogMutable, Warning, TEXT("ConvertSkeletalMeshToMutable: Source mesh [%s] section uses a bone [%s] that is not in its RequiredBones list."),
+				UE_LOG(LogMutable, Warning, TEXT("In object [%s] ConvertSkeletalMeshToMutable: Source mesh [%s] section uses a bone [%s] that is not in its RequiredBones list."),
+					*GenerationContext.Object->GetName(),
 					*SkeletalMesh->GetName(),
 					*BoneName);
 				MutableSkeleton->SetBoneCount(BoneMapIndex + 1);
diff --git a/Engine/Plugins/Experimental/Mutable/Source/CustomizableObjectEditor/Private/MuCOE/GenerateMutableSource/GenerateMutableSourceSurface.cpp b/Engine/Plugins/Experimental/Mutable/Source/CustomizableObjectEditor/Private/MuCOE/GenerateMutableSource/GenerateMutableSourceSurface.cpp
--- a/Engine/Plugins/Experimental/Mutable/Source/CustomizableObjectEditor/Private/MuCOE/GenerateMutableSource/GenerateMutableSourceSurface.cpp
+++ b/Engine/Plugins/Experimental/Mutable/Source/CustomizableObjectEditor/Private/MuCOE/GenerateMutableSource/GenerateMutableSourceSurface.cpp
@@ -83,7 +83,8 @@ void LayoutOperationUVNormalizedWarning(FMutableGraphGenerationContext& Generati
 	}
 }
 
-void SetSurfaceFormat( mu::FMeshBufferSet& OutVertexBufferFormat, mu::FMeshBufferSet& OutIndexBufferFormat, const FMutableGraphMeshGenerationData& MeshData, 
+void SetSurfaceFormat( FMutableGraphGenerationContext& GenerationContext,
+					   mu::FMeshBufferSet& OutVertexBufferFormat, mu::FMeshBufferSet& OutIndexBufferFormat, const FMutableGraphMeshGenerationData& MeshData, 
 					   bool bWithExtraBoneInfluences, bool bWithRealTimeMorphs, bool bWithClothing, bool bWith16BitWeights )
 {
 	// Limit skinning weights if necessary
@@ -92,7 +93,7 @@ void SetSurfaceFormat( mu::FMeshBufferSet& OutVertexBufferFormat, mu::FMeshBuffe
 	
 	if (MutableBonesPerVertex != MeshData.MaxNumBonesPerVertex)
 	{
-		UE_LOG(LogMutable, Log, TEXT("Mesh bone number adjusted from %d to %d."), MeshData.MaxNumBonesPerVertex, MutableBonesPerVertex);
+		UE_LOG(LogMutable, Log, TEXT("In object [%s] Mesh bone number adjusted from %d to %d."), *GenerationContext.Object->GetName(), MeshData.MaxNumBonesPerVertex, MutableBonesPerVertex);
 	}
 
 	int MutableBufferCount = MUTABLE_VERTEXBUFFER_TEXCOORDS + 1;
@@ -352,7 +353,7 @@ mu::NodeSurfacePtr GenerateMutableSourceSurface(const UEdGraphPin * Pin, FMutabl
 
 				mu::NodeMeshFormatPtr MeshFormatNode = new mu::NodeMeshFormat();
 				MeshFormatNode->SetSource(MeshNode.get());
-				SetSurfaceFormat( 
+				SetSurfaceFormat( GenerationContext,
 						MeshFormatNode->GetVertexBuffers(), MeshFormatNode->GetIndexBuffers(), MeshData,
 						GenerationContext.Options.bExtraBoneInfluencesEnabled,
 						GenerationContext.Options.bRealTimeMorphTargetsEnabled, 
@@ -697,7 +698,10 @@ mu::NodeSurfacePtr GenerateMutableSourceSurface(const UEdGraphPin * Pin, FMutabl
 									// Format not supported by Mutable, use RBGA_UBYTE as default.
 									mutableFormat = mu::EImageFormat::IF_RGBA_UBYTE;
 
-									UE_LOG(LogMutable, Warning, TEXT("Unexpected image format [%s]."), *FormatWithoutPrefix);
+									const FString UnexpectedImageFormatMsg = FString::Printf(TEXT("In object [%s] Unexpected image format [%s], RGBA_UBYTE will be used instead."), *GenerationContext.Object->GetName(), *FormatWithoutPrefix);
+									const FText UnexpectedImageFormatText = FText::FromString(UnexpectedImageFormatMsg);
+									GenerationContext.Compiler->CompilerLog(UnexpectedImageFormatText, Node);
+									UE_LOG(LogMutable, Warning, TEXT("%s"), *UnexpectedImageFormatMsg);
 								}
 
 								FormatImage->SetFormat(mutableFormat, mutableFormatIfAlpha);
@@ -986,7 +990,7 @@ mu::NodeSurfacePtr GenerateMutableSourceSurface(const UEdGraphPin * Pin, FMutabl
 			if (AddMeshNode)
 			{
 				mu::NodeMeshFormatPtr MeshFormat = new mu::NodeMeshFormat();
-				SetSurfaceFormat( 
+				SetSurfaceFormat( GenerationContext,
 						MeshFormat->GetVertexBuffers(), MeshFormat->GetIndexBuffers(), MeshData,
 						GenerationContext.Options.bExtraBoneInfluencesEnabled,
 						GenerationContext.Options.bRealTimeMorphTargetsEnabled, 
