[mutable] Removed dedicated slate exposing states and runtime parameters and enabled UProperty with runtime parameters.
- Removed SCustomizableObjectNodeObjectStatesView from project (deprecated)
- Unhided RuntimeParameters array part of FCustomizableObjectState to enable state parameter's update/addition and removal.

diff --git a/Engine/Plugins/Experimental/Mutable/Source/CustomizableObjectEditor/Classes/MuCOE/Nodes/CustomizableObjectNodeObject.h b/Engine/Plugins/Experimental/Mutable/Source/CustomizableObjectEditor/Classes/MuCOE/Nodes/CustomizableObjectNodeObject.h
--- a/Engine/Plugins/Experimental/Mutable/Source/CustomizableObjectEditor/Classes/MuCOE/Nodes/CustomizableObjectNodeObject.h
+++ b/Engine/Plugins/Experimental/Mutable/Source/CustomizableObjectEditor/Classes/MuCOE/Nodes/CustomizableObjectNodeObject.h
@@ -131,7 +131,7 @@ public:
 	UPROPERTY(EditAnywhere, Category = CustomizableObject, meta = (ClampMin = "1"))
 	int32 NumMeshComponents = 1;
 
-	UPROPERTY(EditAnywhere, Category=CustomizableObject)
+	UPROPERTY(EditAnywhere, Category=CustomizableObject ,meta = (TitleProperty = "Name"))
 	TArray<FCustomizableObjectState> States;
 
 	UPROPERTY(EditAnywhere, Category = AttachedToExternalObject)
@@ -143,7 +143,7 @@ public:
 	UPROPERTY()
 	FGuid Identifier;
 
-    // Soft references SkeletalMeshes found in the provoius compilation.
+    // Soft references SkeletalMeshes found in the previous compilation.
     // Only populated if the node is the root.
     UPROPERTY()
     TArray<FSoftObjectPath> ReferencedSkeletalMeshes;
diff --git a/Engine/Plugins/Experimental/Mutable/Source/CustomizableObjectEditor/Private/MuCOE/Nodes/CustomizableObjectNodeObjectDetails.cpp b/Engine/Plugins/Experimental/Mutable/Source/CustomizableObjectEditor/Private/MuCOE/Nodes/CustomizableObjectNodeObjectDetails.cpp
--- a/Engine/Plugins/Experimental/Mutable/Source/CustomizableObjectEditor/Private/MuCOE/Nodes/CustomizableObjectNodeObjectDetails.cpp
+++ b/Engine/Plugins/Experimental/Mutable/Source/CustomizableObjectEditor/Private/MuCOE/Nodes/CustomizableObjectNodeObjectDetails.cpp
@@ -8,7 +8,6 @@
 #include "MuCOE/Nodes/CustomizableObjectNodeObject.h"
 #include "MuCOE/Nodes/CustomizableObjectNodeObjectGroup.h"
 #include "MuCOE/SCustomizableObjectNodeObjectRTMorphTargetOverride.h"
-#include "MuCOE/SCustomizableObjectNodeObjectStatesView.h"
 #include "PropertyCustomizationHelpers.h"
 #include "ScopedTransaction.h"
 #include "Widgets/Input/STextComboBox.h"
@@ -44,20 +43,6 @@ void FCustomizableObjectNodeObjectDetails::CustomizeDetails( IDetailLayoutBuilde
 	{
 		TSharedRef<IPropertyHandle> StatesProperty = DetailBuilder.GetProperty("States");
 		StatesProperty->SetOnPropertyValueChanged(FSimpleDelegate::CreateSP(this, &FCustomizableObjectNodeObjectDetails::OnStatesPropertyChanged));
-		
-		// Hinding all the runtime parameters arrays
-		for (int i = 0; i < Node->States.Num(); ++i)
-		{
-			FString VariablePath = FString::Printf(TEXT("States[%d].RuntimeParameters"), i);
-			DetailBuilder.HideProperty(*VariablePath);
-		}
-
-		//Adding StetesView Widget
-		RuntimeParameters.AddCustomRow(LOCTEXT("StatesRuntimeParameters", "Runtime Parameters"))
-		[
-			SAssignNew(StatesViewWidget,SCustomizableObjectNodeObjectSatesView)
-			.Node(Node)
-		];
 
 		TSharedRef<IPropertyHandle> ParentObjectProperty = DetailBuilder.GetProperty("ParentObject");
 
diff --git a/Engine/Plugins/Experimental/Mutable/Source/CustomizableObjectEditor/Private/MuCOE/Nodes/CustomizableObjectNodeObjectDetails.h b/Engine/Plugins/Experimental/Mutable/Source/CustomizableObjectEditor/Private/MuCOE/Nodes/CustomizableObjectNodeObjectDetails.h
--- a/Engine/Plugins/Experimental/Mutable/Source/CustomizableObjectEditor/Private/MuCOE/Nodes/CustomizableObjectNodeObjectDetails.h
+++ b/Engine/Plugins/Experimental/Mutable/Source/CustomizableObjectEditor/Private/MuCOE/Nodes/CustomizableObjectNodeObjectDetails.h
@@ -39,7 +39,4 @@ private:
 
 	FIntPoint GetGridSize() const;
 	TArray<FIntRect> GetBlocks() const;
-
-	// Widget to create the runtime parameters view
-	TSharedPtr< class SCustomizableObjectNodeObjectSatesView > StatesViewWidget;
 };
diff --git a/Engine/Plugins/Experimental/Mutable/Source/CustomizableObjectEditor/Private/MuCOE/SCustomizableObjectNodeObjectStatesView.cpp b/Engine/Plugins/Experimental/Mutable/Source/CustomizableObjectEditor/Private/MuCOE/SCustomizableObjectNodeObjectStatesView.cpp
deleted file mode 100644
diff --git a/Engine/Plugins/Experimental/Mutable/Source/CustomizableObjectEditor/Private/MuCOE/SCustomizableObjectNodeObjectStatesView.h b/Engine/Plugins/Experimental/Mutable/Source/CustomizableObjectEditor/Private/MuCOE/SCustomizableObjectNodeObjectStatesView.h
deleted file mode 100644
