[Mutable] Fix Core not releasing root image operations results.

#rnx

diff --git a/Engine/Plugins/Experimental/Mutable/Source/MutableRuntime/Private/MuR/CodeRunner.cpp b/Engine/Plugins/Experimental/Mutable/Source/MutableRuntime/Private/MuR/CodeRunner.cpp
--- a/Engine/Plugins/Experimental/Mutable/Source/MutableRuntime/Private/MuR/CodeRunner.cpp
+++ b/Engine/Plugins/Experimental/Mutable/Source/MutableRuntime/Private/MuR/CodeRunner.cpp
@@ -94,8 +94,11 @@ namespace mu
 		ScheduledStagePerOp.resize(m_pModel->GetPrivate()->m_program.m_opAddress.Num());
 
 		// We will read this in the end, so make sure we keep it.
-		m_pSystem->m_memory->IncreaseHitCount(FCacheAddress(at, 0, executionOptions));
-
+   		if (Type == FScheduledOp::EType::Full)
+   		{
+			m_pSystem->m_memory->IncreaseHitCount(FCacheAddress(at, 0, executionOptions));
+		}
+    	
 		FProgram& program = m_pModel->GetPrivate()->m_program;
 		m_romPendingOps.SetNum(program.m_roms.Num());
 
