[Mutable] Fix wrong CVar mutable.StreamingCacheMemory description.

#rnx

diff --git a/Engine/Plugins/Experimental/Mutable/Source/CustomizableObject/Classes/MuCO/CustomizableObjectSystem.h b/Engine/Plugins/Experimental/Mutable/Source/CustomizableObject/Classes/MuCO/CustomizableObjectSystem.h
--- a/Engine/Plugins/Experimental/Mutable/Source/CustomizableObject/Classes/MuCO/CustomizableObjectSystem.h
+++ b/Engine/Plugins/Experimental/Mutable/Source/CustomizableObject/Classes/MuCO/CustomizableObjectSystem.h
@@ -22,16 +22,6 @@ class UTexture2D;
 struct FFrame;
 struct FGuid;
 
-// This sets the amount of memory in bytes used to keep streaming data after using it, to reduce the streaming load.
-// High values use more memory, but save object construction time.
-// Setting it to 0 was the original behaviour, and keeps all the data loaded for a mutable operation until the
-// next operation.
-// This can be overwritten with the console variable b.MutableStreamingMemory
-#if !PLATFORM_DESKTOP
-	#define MUTABLE_STREAMING_CACHE				(3 * 1024 * 1024)
-#else
-	#define MUTABLE_STREAMING_CACHE				(12 * 1024 * 1024)
-#endif
 
 // Split StreamedBulkData into chunks smaller than MUTABLE_STREAMED_DATA_MAXCHUNKSIZE
 #define MUTABLE_STREAMED_DATA_MAXCHUNKSIZE		(512 * 1024 * 1024)
diff --git a/Engine/Plugins/Experimental/Mutable/Source/CustomizableObject/Private/MuCO/CustomizableObjectSystem.cpp b/Engine/Plugins/Experimental/Mutable/Source/CustomizableObject/Private/MuCO/CustomizableObjectSystem.cpp
--- a/Engine/Plugins/Experimental/Mutable/Source/CustomizableObject/Private/MuCO/CustomizableObjectSystem.cpp
+++ b/Engine/Plugins/Experimental/Mutable/Source/CustomizableObject/Private/MuCO/CustomizableObjectSystem.cpp
@@ -76,11 +76,14 @@ DECLARE_CYCLE_STAT(TEXT("MutableTask"), STAT_MutableTask, STATGROUP_Game);
 
 UCustomizableObjectSystem* FCustomizableObjectSystemPrivate::SSystem = nullptr;
 
-
 static TAutoConsoleVariable<int32> CVarStreamingMemory(
 	TEXT("mutable.StreamingCacheMemory"),
-	-1,
-	TEXT("If different than 0, limit the amount of memory (in KB) to use to cache streaming data when building characters. 0 means no limit, -1 means use default (40Mb in PC and 20Mb in consoles)."),
+#if !PLATFORM_DESKTOP
+	(3 * 1024),
+#else
+	(12 * 1024),
+#endif
+	TEXT("Limit the amount of memory (in KB) to use to cache streaming data when building characters. More memory reduces the object construction time. Defaults: Desktop = 12288, Others = 3072"),
 	ECVF_Scalability);
 
 
@@ -312,14 +315,15 @@ void UCustomizableObjectSystem::InitSystem()
 	SET_DWORD_STAT(STAT_MutableInstanceBuildTimeAvrg, 0);
 #endif
 
+	Private->LastStreamingMemorySize = CVarClearStreamingCacheOnUpdateEnd.GetValueOnGameThread();
+
 	mu::SettingsPtr pSettings = new mu::Settings;
 	check(pSettings);
 	pSettings->SetProfile(false);
-	pSettings->SetStreamingCache(MUTABLE_STREAMING_CACHE);
+	pSettings->SetStreamingCache(Private->LastStreamingMemorySize);
 	Private->MutableSystem = new mu::System(pSettings);
 	check(Private->MutableSystem);
 
-	Private->LastStreamingMemorySize = MUTABLE_STREAMING_CACHE;
 	Private->Streamer = new FUnrealMutableModelBulkStreamer();
 	check(Private->Streamer != nullptr);
 	Private->MutableSystem->SetStreamingInterface(Private->Streamer);
@@ -1113,8 +1117,7 @@ void FCustomizableObjectSystemPrivate::UpdateStreamingLimit()
 	// This must run on game thread, and when the mutable thread is not running
 	check(IsInGameThread());
 
-	int32 VarValue = CVarStreamingMemory.GetValueOnGameThread();
-	uint64_t MemoryBytes = VarValue < 0 ? MUTABLE_STREAMING_CACHE : uint64_t(VarValue) * 1024;
+	const uint64 MemoryBytes = CVarStreamingMemory.GetValueOnGameThread() * 1024;
 	if (MemoryBytes != LastStreamingMemorySize)
 	{
 		LastStreamingMemorySize = MemoryBytes;
diff --git a/Engine/Plugins/Experimental/Mutable/Source/CustomizableObject/Private/MuCO/CustomizableObjectSystemPrivate.h b/Engine/Plugins/Experimental/Mutable/Source/CustomizableObject/Private/MuCO/CustomizableObjectSystemPrivate.h
--- a/Engine/Plugins/Experimental/Mutable/Source/CustomizableObject/Private/MuCO/CustomizableObjectSystemPrivate.h
+++ b/Engine/Plugins/Experimental/Mutable/Source/CustomizableObject/Private/MuCO/CustomizableObjectSystemPrivate.h
@@ -536,7 +536,7 @@ public:
 	mu::SystemPtr MutableSystem;
 
 	// Store the last streaming memory size in bytes, to change it when it is safe.
-	uint64_t LastStreamingMemorySize = 0;
+	uint64 LastStreamingMemorySize = 0;
 
 	// This object is responsible for streaming data to the MutableSystem.
 	// Non-owned reference
