[Mutable] Fix TextureCompressionStrategy not loading in packaged builds.

#rnx

diff --git a/Engine/Plugins/Experimental/Mutable/Source/CustomizableObject/Public/MuCO/CustomizableObjectUIData.h b/Engine/Plugins/Experimental/Mutable/Source/CustomizableObject/Public/MuCO/CustomizableObjectUIData.h
--- a/Engine/Plugins/Experimental/Mutable/Source/CustomizableObject/Public/MuCO/CustomizableObjectUIData.h
+++ b/Engine/Plugins/Experimental/Mutable/Source/CustomizableObject/Public/MuCO/CustomizableObjectUIData.h
@@ -263,8 +263,8 @@ struct FParameterUIData
 	bool bDontCompressRuntimeTextures_DEPRECATED = false; // Only useful for state metadata
 
 	/** Not really relevant for UI, but apparently bDontCompressRuntimeTextures_DEPRECATED was used to decided some texture properties at runtime.
-	* TODO: Try to handle the NeverStream without going through this.
-	*/
+	* TODO: Try to handle the NeverStream without going through this. */
+	UPROPERTY()
 	ETextureCompressionStrategy TextureCompressionStrategy = ETextureCompressionStrategy::None;
 
 	/** In this mode instances and their temp data will be reused between updates. It will be much faster but spend as much as ten times the memory.
