[Mutable] Correctly handle backwards compatible code.

#rnx

diff --git a/Engine/Plugins/Experimental/Mutable/Source/CustomizableObject/Public/MuCO/CustomizableObjectCustomVersion.h b/Engine/Plugins/Experimental/Mutable/Source/CustomizableObject/Public/MuCO/CustomizableObjectCustomVersion.h
--- a/Engine/Plugins/Experimental/Mutable/Source/CustomizableObject/Public/MuCO/CustomizableObjectCustomVersion.h
+++ b/Engine/Plugins/Experimental/Mutable/Source/CustomizableObject/Public/MuCO/CustomizableObjectCustomVersion.h
@@ -66,6 +66,8 @@ struct CUSTOMIZABLEOBJECT_API FCustomizableObjectCustomVersion
 
 		GroupProjectorImagePinRemoved,
 
+		StateTextureCompressionStrategyEnum,
+
 		// -----<new versions can be added above this line>-------------------------------------------------
 		VersionPlusOne,
 		LatestVersion = VersionPlusOne - 1
diff --git a/Engine/Plugins/Experimental/Mutable/Source/CustomizableObject/Public/MuCO/CustomizableObjectUIData.h b/Engine/Plugins/Experimental/Mutable/Source/CustomizableObject/Public/MuCO/CustomizableObjectUIData.h
--- a/Engine/Plugins/Experimental/Mutable/Source/CustomizableObject/Public/MuCO/CustomizableObjectUIData.h
+++ b/Engine/Plugins/Experimental/Mutable/Source/CustomizableObject/Public/MuCO/CustomizableObjectUIData.h
@@ -258,10 +258,6 @@ struct FParameterUIData
 	UPROPERTY(BlueprintReadWrite, Category = UI)
 	ECustomizableObjectGroupType IntegerParameterGroupType;
 
-	/** Replaced by TextureCompressionStrategy */
-	UPROPERTY()
-	bool bDontCompressRuntimeTextures_DEPRECATED = false; // Only useful for state metadata
-
 	/** Not really relevant for UI, but apparently bDontCompressRuntimeTextures_DEPRECATED was used to decided some texture properties at runtime.
 	* TODO: Try to handle the NeverStream without going through this. */
 	UPROPERTY()
diff --git a/Engine/Plugins/Experimental/Mutable/Source/CustomizableObjectEditor/Classes/MuCOE/Nodes/CustomizableObjectNodeObject.h b/Engine/Plugins/Experimental/Mutable/Source/CustomizableObjectEditor/Classes/MuCOE/Nodes/CustomizableObjectNodeObject.h
--- a/Engine/Plugins/Experimental/Mutable/Source/CustomizableObjectEditor/Classes/MuCOE/Nodes/CustomizableObjectNodeObject.h
+++ b/Engine/Plugins/Experimental/Mutable/Source/CustomizableObjectEditor/Classes/MuCOE/Nodes/CustomizableObjectNodeObject.h
@@ -28,10 +28,6 @@ struct CUSTOMIZABLEOBJECTEDITOR_API FCustomizableObjectState
 	UPROPERTY(EditAnywhere, Category=CustomizableObject)
 	TArray<FString> RuntimeParameters;
 
-	/** This is now TextureCompressionStrategy.  */
-	UPROPERTY()
-	bool bDontCompressRuntimeTextures_DEPRECATED = false;
-		
 	/** Special treatment of texture compression for this state. */
 	UPROPERTY(EditAnywhere, Category = CustomizableObject)
 	ETextureCompressionStrategy TextureCompressionStrategy = ETextureCompressionStrategy::None;
@@ -58,6 +54,12 @@ struct CUSTOMIZABLEOBJECTEDITOR_API FCustomizableObjectState
 
 	UPROPERTY(EditAnywhere, Category = UI, meta = (DisplayName = "State UI Metadata"))
 	FMutableParamUIMetadata StateUIMetadata;
+
+	// Deprecated
+	
+	/** This is now TextureCompressionStrategy.  */
+	UPROPERTY()
+	bool bDontCompressRuntimeTextures_DEPRECATED = false;
 };
 
 
@@ -159,7 +161,7 @@ public:
 
 	// UObject interface.
 	void PostEditChangeProperty(FPropertyChangedEvent& PropertyChangedEvent) override;
-	void PostLoad() override;
+	void BackwardsCompatibleFixup() override;
 
 	// EdGraphNode interface
 	FText GetNodeTitle(ENodeTitleType::Type TitleType) const override;
diff --git a/Engine/Plugins/Experimental/Mutable/Source/CustomizableObjectEditor/Private/MuCOE/Nodes/CustomizableObjectNodeObject.cpp b/Engine/Plugins/Experimental/Mutable/Source/CustomizableObjectEditor/Private/MuCOE/Nodes/CustomizableObjectNodeObject.cpp
--- a/Engine/Plugins/Experimental/Mutable/Source/CustomizableObjectEditor/Private/MuCOE/Nodes/CustomizableObjectNodeObject.cpp
+++ b/Engine/Plugins/Experimental/Mutable/Source/CustomizableObjectEditor/Private/MuCOE/Nodes/CustomizableObjectNodeObject.cpp
@@ -34,18 +34,23 @@ UCustomizableObjectNodeObject::UCustomizableObjectNodeObject()
 }
 
 
-void UCustomizableObjectNodeObject::PostLoad()
+void UCustomizableObjectNodeObject::BackwardsCompatibleFixup()
 {
-	Super::PostLoad();
+	Super::BackwardsCompatibleFixup();
 
-	for (FCustomizableObjectState& State : States)
+	const int32 CustomizableObjectCustomVersion = GetLinkerCustomVersion(FCustomizableObjectCustomVersion::GUID);
+	
+	if (CustomizableObjectCustomVersion < FCustomizableObjectCustomVersion::StateTextureCompressionStrategyEnum)
 	{
-		if (State.TextureCompressionStrategy== ETextureCompressionStrategy::None
-			&&
-			State.bDontCompressRuntimeTextures_DEPRECATED)
+		for (FCustomizableObjectState& State : States)
 		{
-			State.bDontCompressRuntimeTextures_DEPRECATED = false;
-			State.TextureCompressionStrategy = ETextureCompressionStrategy::DontCompressRuntime;
+			if (State.TextureCompressionStrategy== ETextureCompressionStrategy::None
+				&&
+				State.bDontCompressRuntimeTextures_DEPRECATED)
+			{
+				State.bDontCompressRuntimeTextures_DEPRECATED = false;
+				State.TextureCompressionStrategy = ETextureCompressionStrategy::DontCompressRuntime;
+			}
 		}
 	}
 }
