[Mutable] Fix PARAMETER_VALUE issues.

Converted the union to a TVarian. Improvements:
- Fixes the comparision issues (now we are only compraing the bytes that belong to the type).
- Makes code safer (types must always match)
- Text no longer has a size limit.

#rnx

diff --git a/Engine/Plugins/Experimental/Mutable/Source/CustomizableObject/Classes/MuCO/CustomizableObject.h b/Engine/Plugins/Experimental/Mutable/Source/CustomizableObject/Classes/MuCO/CustomizableObject.h
--- a/Engine/Plugins/Experimental/Mutable/Source/CustomizableObject/Classes/MuCO/CustomizableObject.h
+++ b/Engine/Plugins/Experimental/Mutable/Source/CustomizableObject/Classes/MuCO/CustomizableObject.h
@@ -1365,7 +1365,7 @@ private:
 	// This is a manual version number for the binary blobs in this asset.
 	// Increasing it invalidates all the previously compiled models.
 	// Warning: If while merging code both versions have changed, take the highest+1.
-	static const int32 CurrentSupportedVersion = 374;
+	static const int32 CurrentSupportedVersion = 375;
 
 public:
 
diff --git a/Engine/Plugins/Experimental/Mutable/Source/CustomizableObject/Private/MuCO/CustomizableObject.cpp b/Engine/Plugins/Experimental/Mutable/Source/CustomizableObject/Private/MuCO/CustomizableObject.cpp
--- a/Engine/Plugins/Experimental/Mutable/Source/CustomizableObject/Private/MuCO/CustomizableObject.cpp
+++ b/Engine/Plugins/Experimental/Mutable/Source/CustomizableObject/Private/MuCO/CustomizableObject.cpp
@@ -23,6 +23,7 @@
 #include "MuCO/MutableProjectorTypeUtils.h"
 #include "MuCO/UnrealMutableModelDiskStreamer.h"
 #include "MuR/ModelPrivate.h"
+#include "MuR/ParametersPrivate.h"
 #include "PhysicsEngine/PhysicsAsset.h"
 #include "Serialization/MemoryReader.h"
 #include "Serialization/MemoryWriter.h"
@@ -1462,7 +1463,7 @@ float UCustomizableObject::GetFloatParameterDefaultValue(const FString& InParame
 	}
 
 	const TSharedPtr<const mu::Model> MutableModel = PrivateData->GetModel();
-	Value = MutableModel->GetPrivate()->m_program.m_parameters[ParameterIndex].m_defaultValue.m_float;
+	Value = MutableModel->GetPrivate()->m_program.m_parameters[ParameterIndex].m_defaultValue.Get<mu::ParamFloatType>();
 	return Value;
 }
 
@@ -1478,7 +1479,7 @@ int32 UCustomizableObject::GetIntParameterDefaultValue(const FString& InParamete
 	}
 
 	const TSharedPtr<const mu::Model> MutableModel = PrivateData->GetModel();
-	Value = MutableModel->GetPrivate()->m_program.m_parameters[ParameterIndex].m_defaultValue.m_int;
+	Value = MutableModel->GetPrivate()->m_program.m_parameters[ParameterIndex].m_defaultValue.Get<mu::ParamIntType>();
 	return Value;
 }
 
@@ -1494,7 +1495,7 @@ bool UCustomizableObject::GetBoolParameterDefaultValue(const FString& InParamete
 	}
 	
 	const TSharedPtr<const mu::Model> MutableModel = PrivateData->GetModel();
-	Value = MutableModel->GetPrivate()->m_program.m_parameters[ParameterIndex].m_defaultValue.m_bool;
+	Value = MutableModel->GetPrivate()->m_program.m_parameters[ParameterIndex].m_defaultValue.Get<mu::ParamBoolType>();
 	return Value;
 }
 
@@ -1510,9 +1511,9 @@ FLinearColor UCustomizableObject::GetColorParameterDefaultValue(const FString& I
 	}
 	
 	const mu::Model::Private* ModelPrivate = PrivateData->GetModel()->GetPrivate();
-	Value.R = ModelPrivate->m_program.m_parameters[ParameterIndex].m_defaultValue.m_colour[0];
-	Value.G = ModelPrivate->m_program.m_parameters[ParameterIndex].m_defaultValue.m_colour[1];
-	Value.B = ModelPrivate->m_program.m_parameters[ParameterIndex].m_defaultValue.m_colour[2];
+	Value.R = ModelPrivate->m_program.m_parameters[ParameterIndex].m_defaultValue.Get<mu::ParamColorType>()[0];
+	Value.G = ModelPrivate->m_program.m_parameters[ParameterIndex].m_defaultValue.Get<mu::ParamColorType>()[1];
+	Value.B = ModelPrivate->m_program.m_parameters[ParameterIndex].m_defaultValue.Get<mu::ParamColorType>()[2];
 	Value.A = 1.0f;
 	
 	return Value;
@@ -1531,7 +1532,7 @@ void UCustomizableObject::GetProjectorParameterDefaultValue(const FString& InPar
 	
 	const TSharedPtr<const mu::Model> MutableModel = PrivateData->GetModel();
 	
-	mu::FProjector MutableProjector = MutableModel->GetPrivate()->m_program.m_parameters[ParameterIndex].m_defaultValue.m_projector;
+	mu::FProjector MutableProjector = MutableModel->GetPrivate()->m_program.m_parameters[ParameterIndex].m_defaultValue.Get<mu::ParamProjectorType>();
 	OutType = ProjectorUtils::GetEquivalentProjectorType(MutableProjector.type);
 	OutPos = MutableProjector.position;
 	OutDirection = MutableProjector.direction;
@@ -1540,6 +1541,7 @@ void UCustomizableObject::GetProjectorParameterDefaultValue(const FString& InPar
 	OutAngle = MutableProjector.projectionAngle;
 }
 
+
 FCustomizableObjectProjector UCustomizableObject::GetProjectorParameterDefaultValue(const FString& InParameterName) const 
 {
 	FCustomizableObjectProjector Value;
@@ -1552,7 +1554,7 @@ FCustomizableObjectProjector UCustomizableObject::GetProjectorParameterDefaultVa
 	
 	const TSharedPtr<const mu::Model> MutableModel = PrivateData->GetModel();
 
-	mu::FProjector MutableProjector = MutableModel->GetPrivate()->m_program.m_parameters[ParameterIndex].m_defaultValue.m_projector;
+	mu::FProjector MutableProjector = MutableModel->GetPrivate()->m_program.m_parameters[ParameterIndex].m_defaultValue.Get<mu::ParamProjectorType>();
 	Value.Angle = MutableProjector.projectionAngle;
 	Value.Direction = MutableProjector.direction;
 	Value.Position = MutableProjector.position;
@@ -1575,7 +1577,7 @@ uint64 UCustomizableObject::GetTextureParameterDefaultValue(const FString& InPar
 	}
 	
 	const TSharedPtr<const mu::Model> MutableModel = PrivateData->GetModel();
-	Value = MutableModel->GetPrivate()->m_program.m_parameters[ParameterIndex].m_defaultValue.m_image;
+	Value = MutableModel->GetPrivate()->m_program.m_parameters[ParameterIndex].m_defaultValue.Get<mu::ParamImageType>();
 	return Value;
 }
 
diff --git a/Engine/Plugins/Experimental/Mutable/Source/MutableRuntime/Private/MuR/CodeRunner.cpp b/Engine/Plugins/Experimental/Mutable/Source/MutableRuntime/Private/MuR/CodeRunner.cpp
--- a/Engine/Plugins/Experimental/Mutable/Source/MutableRuntime/Private/MuR/CodeRunner.cpp
+++ b/Engine/Plugins/Experimental/Mutable/Source/MutableRuntime/Private/MuR/CodeRunner.cpp
@@ -5065,7 +5065,7 @@ namespace mu
             pResult->m_value = pParams->GetPrivate()->GetProjectorValue(args.variable,index);
 
             // The type cannot be changed, take it from the default value
-            const FProjector& def = program.m_parameters[args.variable].m_defaultValue.m_projector;
+            const FProjector& def = program.m_parameters[args.variable].m_defaultValue.Get<ParamProjectorType>();
             pResult->m_value.type = def.type;
 
             GetMemory().SetProjector( item, pResult );
diff --git a/Engine/Plugins/Experimental/Mutable/Source/MutableRuntime/Private/MuR/Model.cpp b/Engine/Plugins/Experimental/Mutable/Source/MutableRuntime/Private/MuR/Model.cpp
--- a/Engine/Plugins/Experimental/Mutable/Source/MutableRuntime/Private/MuR/Model.cpp
+++ b/Engine/Plugins/Experimental/Mutable/Source/MutableRuntime/Private/MuR/Model.cpp
@@ -391,8 +391,7 @@ namespace mu
                         {
 //							float m[16];
 //							pOld->GetProjectorValue( p, m );
-                            pRes->GetPrivate()->m_values[thisP].m_projector =
-                                    pOld->GetPrivate()->m_values[p].m_projector;
+                            pRes->GetPrivate()->m_values[thisP].Set<ParamProjectorType>(pOld->GetPrivate()->m_values[p].Get<ParamProjectorType>());
                             break;
                         }
 
@@ -538,7 +537,7 @@ namespace mu
             {
             case PARAMETER_TYPE::T_BOOL:
                 dataSize = 1;
-                parameterValuesBlob.Add( pParams->GetPrivate()->m_values[param].m_bool ? 1 : 0 );
+                parameterValuesBlob.Add( pParams->GetPrivate()->m_values[param].Get<ParamBoolType>() ? 1 : 0 );
                 pos += dataSize;
 
                 // Multi-values
@@ -548,7 +547,7 @@ namespace mu
                     pos = AddMultiValueKeys( parameterValuesBlob, pos, multi );
                     for(const auto& v: multi)
                     {
-                        parameterValuesBlob.Add( v.Value.m_bool ? 1 : 0 );
+                        parameterValuesBlob.Add( v.Value.Get<ParamBoolType>() ? 1 : 0 );
                         pos += dataSize;
                     }
                 }
@@ -557,7 +556,7 @@ namespace mu
             case PARAMETER_TYPE::T_INT:
                 dataSize = sizeof(int32);
                 parameterValuesBlob.SetNum( pos+dataSize );
-                FMemory::Memcpy( &parameterValuesBlob[pos], &pParams->GetPrivate()->m_values[param].m_int, dataSize );
+                FMemory::Memcpy( &parameterValuesBlob[pos], &pParams->GetPrivate()->m_values[param].Get<ParamIntType>(), dataSize );
                 pos += dataSize;
 
                 // Multi-values
@@ -568,7 +567,7 @@ namespace mu
                     parameterValuesBlob.SetNum( pos+multi.Num()*dataSize );
                     for(const auto& v: multi)
                     {
-						FMemory::Memcpy( &parameterValuesBlob[pos], &v.Value.m_int, dataSize );
+						FMemory::Memcpy( &parameterValuesBlob[pos], &v.Value.Get<ParamIntType>(), dataSize );
                         pos += dataSize;
                     }
                 }
@@ -578,7 +577,7 @@ namespace mu
             case PARAMETER_TYPE::T_FLOAT:
                 dataSize = sizeof(float);
                 parameterValuesBlob.SetNum( pos+dataSize );
-				FMemory::Memcpy( &parameterValuesBlob[pos], &pParams->GetPrivate()->m_values[param].m_float, dataSize );
+				FMemory::Memcpy( &parameterValuesBlob[pos], &pParams->GetPrivate()->m_values[param].Get<ParamFloatType>(), dataSize );
                 pos += dataSize;
 
                 // Multi-values
@@ -589,7 +588,7 @@ namespace mu
                     parameterValuesBlob.SetNum( pos+multi.Num()*dataSize );
                     for(const auto& v: multi)
                     {
-						FMemory::Memcpy( &parameterValuesBlob[pos], &v.Value.m_float,dataSize );
+						FMemory::Memcpy( &parameterValuesBlob[pos], &v.Value.Get<ParamFloatType>(),dataSize );
                         pos += dataSize;
                     }
                 }
@@ -599,7 +598,7 @@ namespace mu
             case PARAMETER_TYPE::T_COLOUR:
                 dataSize = 3*sizeof(float);
                 parameterValuesBlob.SetNum( pos+dataSize );
-				FMemory::Memcpy( &parameterValuesBlob[pos], &pParams->GetPrivate()->m_values[param].m_colour, dataSize );
+				FMemory::Memcpy( &parameterValuesBlob[pos], &pParams->GetPrivate()->m_values[param].Get<ParamFloatType>(), dataSize );
                 pos += dataSize;
 
                 // Multi-values
@@ -610,7 +609,7 @@ namespace mu
                     parameterValuesBlob.SetNum( pos+multi.Num()*dataSize );
                     for(const auto& v: multi)
                     {
-						FMemory::Memcpy( &parameterValuesBlob[pos], &v.Value.m_colour,dataSize );
+						FMemory::Memcpy( &parameterValuesBlob[pos], &v.Value.Get<ParamColorType>(),dataSize );
                         pos += dataSize;
                     }
                 }
@@ -623,7 +622,7 @@ namespace mu
 
                 // \todo: padding will be random?
                 parameterValuesBlob.SetNum( pos+dataSize );
-				FMemory::Memcpy( &parameterValuesBlob[pos], &pParams->GetPrivate()->m_values[param].m_projector, dataSize );
+				FMemory::Memcpy( &parameterValuesBlob[pos], &pParams->GetPrivate()->m_values[param].Get<ParamProjectorType>(), dataSize );
 				pos += dataSize;
 
                 // Multi-values
@@ -634,7 +633,7 @@ namespace mu
                     parameterValuesBlob.SetNum( pos+multi.Num()*dataSize );
                     for(const auto& v: multi)
                     {
-						FMemory::Memcpy( &parameterValuesBlob[pos], &v.Value.m_projector,dataSize );
+						FMemory::Memcpy( &parameterValuesBlob[pos], &v.Value.Get<ParamProjectorType>(),dataSize );
                         pos += dataSize;
                     }
                 }
@@ -644,7 +643,7 @@ namespace mu
                 dataSize = sizeof(EXTERNAL_IMAGE_ID);
                 parameterValuesBlob.SetNum( pos+dataSize );
 				FMemory::Memcpy( &parameterValuesBlob[pos],
-                        &pParams->GetPrivate()->m_values[param].m_image,
+                        &pParams->GetPrivate()->m_values[param].Get<ParamImageType>(),
                         dataSize );
 				pos += dataSize;
 
@@ -656,7 +655,7 @@ namespace mu
                     parameterValuesBlob.SetNum( pos+multi.Num()*dataSize );
                     for(const auto& v: multi)
                     {
-						FMemory::Memcpy( &parameterValuesBlob[pos], &v.Value.m_image,dataSize );
+						FMemory::Memcpy( &parameterValuesBlob[pos], &v.Value.Get<ParamImageType>(),dataSize );
                         pos += dataSize;
                     }
                 }
diff --git a/Engine/Plugins/Experimental/Mutable/Source/MutableRuntime/Private/MuR/Parameters.cpp b/Engine/Plugins/Experimental/Mutable/Source/MutableRuntime/Private/MuR/Parameters.cpp
--- a/Engine/Plugins/Experimental/Mutable/Source/MutableRuntime/Private/MuR/Parameters.cpp
+++ b/Engine/Plugins/Experimental/Mutable/Source/MutableRuntime/Private/MuR/Parameters.cpp
@@ -338,7 +338,7 @@ namespace mu
         if (!pos)
         {
             // Return the single value
-            return m_pD->m_values[index].m_bool;
+            return m_pD->m_values[index].Get<ParamBoolType>();
         }
 
         // Multivalue case
@@ -350,12 +350,12 @@ namespace mu
 			const PARAMETER_VALUE* it = m.Find(pos->m_pD->m_values);
             if (it)
             {
-                return it->m_bool;
+                return it->Get<ParamBoolType>();
             }
         }
 
         // Multivalue parameter, but no multivalue set. Return single value.
-        return m_pD->m_values[index].m_bool;
+        return m_pD->m_values[index].Get<ParamBoolType>();
 	}
 
 
@@ -385,7 +385,7 @@ namespace mu
                 m_pD->m_multiValues[index].Empty();
             }
 
-            m_pD->m_values[index].m_bool = value;
+            m_pD->m_values[index].Set<ParamBoolType>(value);
         }
 
         // Multivalue case
@@ -400,7 +400,7 @@ namespace mu
 
 			TMap< TArray<int32_t>, PARAMETER_VALUE >& m = m_pD->m_multiValues[index];
 			PARAMETER_VALUE& it = m.FindOrAdd(pos->m_pD->m_values);
-            it.m_bool = value;
+            it.Set<ParamBoolType>(value);
         }
 	}
 
@@ -499,7 +499,7 @@ namespace mu
         if (!pos)
         {
             // Return the single value
-            return m_pD->m_values[index].m_int;
+            return m_pD->m_values[index].Get<ParamIntType>();
         }
 
         // Multivalue case
@@ -511,12 +511,12 @@ namespace mu
 			const PARAMETER_VALUE* it = m.Find(pos->m_pD->m_values);
             if (it)
             {
-                return it->m_int;
+                return it->Get<ParamIntType>();
             }
         }
 
         // Multivalue parameter, but no multivalue set. Return single value.
-        return m_pD->m_values[index].m_int;
+        return m_pD->m_values[index].Get<ParamIntType>();
     }
 
 
@@ -546,7 +546,7 @@ namespace mu
                 m_pD->m_multiValues[index].Empty();
             }
 
-            m_pD->m_values[index].m_int = value;
+            m_pD->m_values[index].Set<ParamIntType>(value);
         }
 
         // Multivalue case
@@ -561,7 +561,7 @@ namespace mu
 
 			TMap< TArray<int32_t>, PARAMETER_VALUE >& m = m_pD->m_multiValues[index];
 			PARAMETER_VALUE& it = m.FindOrAdd(pos->m_pD->m_values);
-            it.m_int = value;
+            it.Set<ParamIntType>(value);
         }
     }
 
@@ -587,7 +587,7 @@ namespace mu
         if (!pos)
         {
             // Return the single value
-            return m_pD->m_values[index].m_float;
+            return m_pD->m_values[index].Get<ParamFloatType>();
         }
 
         // Multivalue case
@@ -599,12 +599,12 @@ namespace mu
 			const PARAMETER_VALUE* it = m.Find(pos->m_pD->m_values);
             if (it)
             {
-                return it->m_float;
+                return it->Get<ParamFloatType>();
             }
         }
 
         // Multivalue parameter, but no multivalue set. Return single value.
-        return m_pD->m_values[index].m_float;
+        return m_pD->m_values[index].Get<ParamFloatType>();
     }
 
 
@@ -634,7 +634,7 @@ namespace mu
                 m_pD->m_multiValues[index].Empty();
             }
 
-            m_pD->m_values[index].m_float = value;
+            m_pD->m_values[index].Set<ParamFloatType>(value);
         }
 
         // Multivalue case
@@ -649,7 +649,7 @@ namespace mu
 
 			TMap< TArray<int32_t>, PARAMETER_VALUE >& m = m_pD->m_multiValues[index];
 			PARAMETER_VALUE& it = m.FindOrAdd(pos->m_pD->m_values);
-            it.m_float = value;
+            it.Set<ParamFloatType>(value);
         }
     }
 
@@ -675,9 +675,9 @@ namespace mu
         if (!pos)
         {
             // Return the single value
-            if (pR) *pR = m_pD->m_values[index].m_colour[0];
-            if (pG) *pG = m_pD->m_values[index].m_colour[1];
-            if (pB) *pB = m_pD->m_values[index].m_colour[2];
+            if (pR) *pR = m_pD->m_values[index].Get<ParamColorType>()[0];
+            if (pG) *pG = m_pD->m_values[index].Get<ParamColorType>()[1];
+            if (pB) *pB = m_pD->m_values[index].Get<ParamColorType>()[2];
             return;
         }
 
@@ -690,17 +690,17 @@ namespace mu
 			const PARAMETER_VALUE* it = m.Find(pos->m_pD->m_values);
             if (it)
             {
-                if (pR) *pR = it->m_colour[0];
-                if (pG) *pG = it->m_colour[1];
-                if (pB) *pB = it->m_colour[2];
+                if (pR) *pR = it->Get<ParamColorType>()[0];
+                if (pG) *pG = it->Get<ParamColorType>()[1];
+                if (pB) *pB = it->Get<ParamColorType>()[2];
                 return;
             }
         }
 
         // Multivalue parameter, but no multivalue set. Return single value.
-        if (pR) *pR = m_pD->m_values[index].m_colour[0];
-        if (pG) *pG = m_pD->m_values[index].m_colour[1];
-        if (pB) *pB = m_pD->m_values[index].m_colour[2];
+        if (pR) *pR = m_pD->m_values[index].Get<ParamColorType>()[0];
+        if (pG) *pG = m_pD->m_values[index].Get<ParamColorType>()[1];
+        if (pB) *pB = m_pD->m_values[index].Get<ParamColorType>()[2];
         return;
     }
 
@@ -731,9 +731,12 @@ namespace mu
                 m_pD->m_multiValues[index].Empty();
             }
 
-            m_pD->m_values[index].m_colour[0] = r;
-            m_pD->m_values[index].m_colour[1] = g;
-            m_pD->m_values[index].m_colour[2] = b;
+        	ParamColorType Value;
+            Value[0] = r;
+            Value[1] = g;
+            Value[2] = b;
+        	
+        	m_pD->m_values[index].Set<ParamColorType>(Value);
         }
 
         // Multivalue case
@@ -748,9 +751,13 @@ namespace mu
 
 			TMap< TArray<int32_t>, PARAMETER_VALUE >& m = m_pD->m_multiValues[index];
 			PARAMETER_VALUE& it = m.FindOrAdd(pos->m_pD->m_values);
-            it.m_colour[0] = r;
-            it.m_colour[1] = g;
-            it.m_colour[2] = b;
+
+        	ParamColorType Value;
+            Value[0] = r;
+            Value[1] = g;
+            Value[2] = b;
+
+        	it.Set<ParamColorType>(Value);
         }
     }
 
@@ -764,7 +771,7 @@ namespace mu
 		// Single value case
 		if (!pos)
 		{
-			return m_pD->m_values[index].m_image;
+			return m_pD->m_values[index].Get<ParamImageType>();
 		}
 
 		// Multivalue case
@@ -776,12 +783,12 @@ namespace mu
 			const PARAMETER_VALUE* it = m.Find(pos->m_pD->m_values);
 			if (it)
 			{
-				return it->m_image;
+				return it->Get<ParamImageType>();
 			}
 		}
 
 		// Multivalue parameter, but no multivalue set. Return single value.
-		return m_pD->m_values[index].m_image;
+		return m_pD->m_values[index].Get<ParamImageType>();
 
     }
 
@@ -795,7 +802,7 @@ namespace mu
 		// Single value case
 		if (!pos)
 		{
-			m_pD->m_values[index].m_image = id;
+			m_pD->m_values[index].Set<ParamImageType>(id);
 		}
 
 		// Multivalue case
@@ -810,7 +817,7 @@ namespace mu
 
 			TMap< TArray<int32_t>, PARAMETER_VALUE >& m = m_pD->m_multiValues[index];
 			PARAMETER_VALUE& it = m.FindOrAdd(pos->m_pD->m_values);
-			it.m_image = id;
+			it.Set<ParamImageType>(id);
 		}
     }
 
@@ -832,7 +839,7 @@ namespace mu
         if ( !pos )
         {
             // Return the single value
-            return m_pD->m_values[index].m_text;
+            return m_pD->m_values[index].Get<ParamStringType>().c_str();
         }
 
         // Multivalue case
@@ -844,12 +851,12 @@ namespace mu
 			const PARAMETER_VALUE* it = m.Find( pos->m_pD->m_values );
             if ( it )
             {
-                return it->m_text;
+                return it->Get<ParamStringType>().c_str();
             }
         }
 
         // Multivalue parameter, but no multivalue set. Return single value.
-        return m_pD->m_values[index].m_text;
+        return m_pD->m_values[index].Get<ParamStringType>().c_str();
     }
 
 
@@ -868,8 +875,6 @@ namespace mu
             return;
         }
 
-        int32 len = FMath::Min( MUTABLE_MAX_STRING_PARAM_LENGTH, int32( strnlen( value, MUTABLE_MAX_STRING_PARAM_LENGTH ) ) );
-
         // Single value case
         if ( !pos )
         {
@@ -879,8 +884,7 @@ namespace mu
                 m_pD->m_multiValues[index].Empty();
             }
 
-			FMemory::Memcpy( m_pD->m_values[index].m_text, value, len );
-            m_pD->m_values[index].m_text[len] = 0;
+        	m_pD->m_values[index].Set<ParamStringType>(value);
         }
 
         // Multivalue case
@@ -895,8 +899,7 @@ namespace mu
 
 			TMap< TArray<int32_t>, PARAMETER_VALUE >& m = m_pD->m_multiValues[index];
 			PARAMETER_VALUE& it = m.FindOrAdd( pos->m_pD->m_values );
-			FMemory::Memcpy( it.m_text, value, len );
-            it.m_text[len] = 0;
+        	it.Set<ParamStringType>(value);
         }
     }
 
@@ -927,7 +930,7 @@ namespace mu
         if (!pos)
         {
             // Return the single value
-            result = &m_values[index].m_projector;
+            result = &m_values[index].Get<ParamProjectorType>();
         }
 
         // Multivalue case
@@ -940,7 +943,7 @@ namespace mu
                 const PARAMETER_VALUE* it = m_multiValues[index].Find(pos->m_pD->m_values);
                 if (it)
                 {
-                    result = &it->m_projector;
+                    result = &it->Get<ParamProjectorType>();
                 }
             }
         }
@@ -948,7 +951,7 @@ namespace mu
         // Multivalue parameter, but no multivalue set. Return single value.
         if (!result)
         {
-            result = &m_values[index].m_projector;
+            result = &m_values[index].Get<ParamProjectorType>();
         }
 
         return *result;
@@ -1025,7 +1028,7 @@ namespace mu
                 m_pD->m_multiValues[ParameterIndex].Empty();
             }
 
-            result = &m_pD->m_values[ParameterIndex].m_projector;
+            result = &m_pD->m_values[ParameterIndex].Get<ParamProjectorType>();
         }
 
         // Multivalue case
@@ -1040,7 +1043,7 @@ namespace mu
 
 			TMap< TArray<int32_t>, PARAMETER_VALUE >& m = m_pD->m_multiValues[ParameterIndex];
 			PARAMETER_VALUE& it = m.FindOrAdd(RangePosition->m_pD->m_values);
-            result = &it.m_projector;
+            result = &it.Get<ParamProjectorType>();
         }
 
         check( result );
@@ -1050,7 +1053,7 @@ namespace mu
         if (m_pD->m_pModel)
         {
             const FProgram& program = m_pD->m_pModel->GetPrivate()->m_program;
-            result->type = program.m_parameters[ParameterIndex].m_defaultValue.m_projector.type;
+            result->type = program.m_parameters[ParameterIndex].m_defaultValue.Get<ParamProjectorType>().type;
         }
 
         result->position = pos;
diff --git a/Engine/Plugins/Experimental/Mutable/Source/MutableRuntime/Private/MuR/ParametersPrivate.h b/Engine/Plugins/Experimental/Mutable/Source/MutableRuntime/Private/MuR/ParametersPrivate.h
--- a/Engine/Plugins/Experimental/Mutable/Source/MutableRuntime/Private/MuR/ParametersPrivate.h
+++ b/Engine/Plugins/Experimental/Mutable/Source/MutableRuntime/Private/MuR/ParametersPrivate.h
@@ -2,8 +2,10 @@
 
 #pragma once
 
-#include "MuR/Parameters.h"
+#include "Misc/TVariant.h"
+#include "Math/MathFwd.h"
 
+#include "MuR/Parameters.h"
 #include "MuR/SerialisationPrivate.h"
 #include "MuR/Operations.h"
 #include "MuR/ImagePrivate.h"
@@ -265,40 +267,38 @@ namespace mu
             arch >> defaultValue;
         }
 	};
+	
 
+	using ParamBoolType = bool;
+	using ParamIntType = int32;
+	using ParamFloatType = float;
+	using ParamColorType = FVector3f;
+	using ParamProjectorType = FProjector;
+	using ParamImageType = EXTERNAL_IMAGE_ID;
+	using ParamStringType = string;
 
-    union PARAMETER_VALUE
-    {
-        PARAMETER_VALUE()
-			: m_projector()
-        {
-            memset( this, 0, sizeof(PARAMETER_VALUE) );
-        }
-
-        bool m_bool;
-
-        int m_int;
-
-        float m_float;
-
-        float m_colour[3];
-
-        FProjector m_projector;
-
-        EXTERNAL_IMAGE_ID m_image;
+	
+	using PARAMETER_VALUE = TVariant<ParamBoolType, ParamIntType, ParamFloatType, ParamColorType, ParamProjectorType, ParamImageType, ParamStringType>;
 
-        char m_text[MUTABLE_MAX_STRING_PARAM_LENGTH+1];
 
-        //!
-        bool operator==( const PARAMETER_VALUE& other ) const
-        {
-            return FMemory::Memcmp(this, &other, sizeof(PARAMETER_VALUE))==0;
-        }
-
-    };
+	// TVariant currently does not support this operator. Once supported remove it.
+	inline bool operator==(const PARAMETER_VALUE& ValueA, const PARAMETER_VALUE& ValueB)
+	{
+		const int32 IndexA = ValueA.GetIndex();
+		const int32 IndexB = ValueB.GetIndex();
 
-    MUTABLE_DEFINE_POD_SERIALISABLE( PARAMETER_VALUE )
+		if (IndexA != IndexB)
+		{
+			return false;
+		}
 
+		return Visit([&ValueB](const auto& StoredValueA)
+		{
+			using Type = typename TDecay<decltype(StoredValueA)>::Type;
+			return StoredValueA == ValueB.Get<Type>();
+		}, ValueA);
+	}
+	
 
     struct FParameterDesc
     {
@@ -364,7 +364,7 @@ namespace mu
         //!
         void Serialise( OutputArchive& arch ) const
         {
-            const int32 ver = 5;
+            const int32 ver = 6;
             arch << ver;
 
 			arch << m_name;
@@ -381,7 +381,7 @@ namespace mu
         {
             int32 ver;
             arch >> ver;
-            check( ver == 5 );
+            check( ver == 6 );
 
 			arch >> m_name;
             arch >> m_uid;
@@ -476,7 +476,7 @@ namespace mu
         //!
         void Serialise( OutputArchive& arch ) const
         {
-            const uint32 ver = 1;
+            const uint32 ver = 2;
             arch << ver;
 
             arch << m_values;
@@ -488,10 +488,10 @@ namespace mu
         {
             uint32 ver;
             arch >> ver;
-            check( ver == 1 );
-
-            arch >> m_values;
-            arch >> m_multiValues;
+			check(ver == 2);
+        	
+			arch >> m_values;
+			arch >> m_multiValues;
         }
 
         //!
diff --git a/Engine/Plugins/Experimental/Mutable/Source/MutableRuntime/Private/MuR/SerialisationPrivate.h b/Engine/Plugins/Experimental/Mutable/Source/MutableRuntime/Private/MuR/SerialisationPrivate.h
--- a/Engine/Plugins/Experimental/Mutable/Source/MutableRuntime/Private/MuR/SerialisationPrivate.h
+++ b/Engine/Plugins/Experimental/Mutable/Source/MutableRuntime/Private/MuR/SerialisationPrivate.h
@@ -2,8 +2,10 @@
 
 #pragma once
 
+#include "Containers/StaticArray.h"
 #include "MuR/Serialisation.h"
 
+#include "Misc/TVariant.h"
 #include "Math/Quat.h"
 #include "Math/Vector.h"
 #include "Math/IntVector.h"
@@ -180,6 +182,61 @@ namespace mu
     MUTABLE_DEFINE_POD_SERIALISABLE( uint32 );
     MUTABLE_DEFINE_POD_SERIALISABLE( uint64 );
 
+	//---------------------------------------------------------------------------------------------
+	// TVariant custom serialize. Based on the default serialization.
+	//---------------------------------------------------------------------------------------------
+	template <typename... Ts>
+	void operator<<(OutputArchive& Ar, const TVariant<Ts...>& Variant)
+	{
+		const uint8 Index = static_cast<uint8>(Variant.GetIndex());
+		Ar << Index;
+		
+		Visit([&Ar](auto& StoredValue)
+		{
+			Ar << StoredValue;
+		}, Variant);
+	}
+
+
+	template <typename T, typename VariantType>
+	struct TVariantLoadFromInputArchiveCaller
+	{
+		/** Default construct the type and load it from the FArchive */
+		static void Load(InputArchive& Ar, VariantType& OutVariant)
+		{
+			OutVariant.template Emplace<T>();
+			Ar >> OutVariant.template Get<T>();
+		}
+	};
+
+	
+	template <typename... Ts>
+	struct TVariantLoadFromInputArchiveLookup
+	{
+		using VariantType = TVariant<Ts...>;
+		static_assert((std::is_default_constructible<Ts>::value && ...), "Each type in TVariant template parameter pack must be default constructible in order to use FArchive serialization");
+
+		/** Load the type at the specified index from the FArchive and emplace it into the TVariant */
+		static void Load(SIZE_T TypeIndex, InputArchive& Ar, VariantType& OutVariant)
+		{
+			static constexpr void(*Loaders[])(InputArchive&, VariantType&) = { &TVariantLoadFromInputArchiveCaller<Ts, VariantType>::Load... };
+			check(TypeIndex < UE_ARRAY_COUNT(Loaders));
+			Loaders[TypeIndex](Ar, OutVariant);
+		}
+	};
+
+	
+	template <typename... Ts>
+	void operator>>(InputArchive& Ar, TVariant<Ts...>& Variant)
+	{
+		uint8 Index;
+		Ar >> Index;
+		check(Index < sizeof...(Ts));
+
+		TVariantLoadFromInputArchiveLookup<Ts...>::Load(static_cast<SIZE_T>(Index), Ar, Variant);
+	}
+	
+
 	//---------------------------------------------------------------------------------------------
 	// Bool size is not a standard
 	//---------------------------------------------------------------------------------------------
@@ -257,6 +314,25 @@ namespace mu
 		}
 	}
 
+	
+	//---------------------------------------------------------------------------------------------
+	template<typename T, uint32 Size> void operator<<(OutputArchive& arch, const TStaticArray<T, Size>& v)
+	{
+		for (uint32 i = 0; i < Size; ++i)
+		{
+			arch << v[i];
+		}
+	}
+
+	template<typename T, uint32 Size> void operator>>(InputArchive& arch, TStaticArray<T, Size>& v)
+	{
+		for (uint32 i = 0; i < Size; ++i)
+		{
+			arch >> v[i];
+		}
+	}
+	
+
 	//---------------------------------------------------------------------------------------------
 	template< typename K, typename T >
 	inline void operator<< (OutputArchive& arch, const TMap<K, T>& v)
diff --git a/Engine/Plugins/Experimental/Mutable/Source/MutableRuntime/Public/MuR/Model.h b/Engine/Plugins/Experimental/Mutable/Source/MutableRuntime/Public/MuR/Model.h
--- a/Engine/Plugins/Experimental/Mutable/Source/MutableRuntime/Public/MuR/Model.h
+++ b/Engine/Plugins/Experimental/Mutable/Source/MutableRuntime/Public/MuR/Model.h
@@ -13,8 +13,8 @@
 //! data structures. Compiled models are not necessarily compatible when the runtime is updated,
 //! so this version number can be used externally to verify this. It is not used internally, and
 //! serializing models from different versions than this runtime will probably result in a crash.
-#define MUTABLE_COMPILED_MODEL_CODE_VERSION		uint32( 59 )
-#define MUTABLE_PARAMETERS_VERSION              uint32( 1 )
+#define MUTABLE_COMPILED_MODEL_CODE_VERSION		uint32( 60 )
+#define MUTABLE_PARAMETERS_VERSION              uint32( 2 )
 
 
 namespace mu
diff --git a/Engine/Plugins/Experimental/Mutable/Source/MutableRuntime/Public/MuR/Parameters.h b/Engine/Plugins/Experimental/Mutable/Source/MutableRuntime/Public/MuR/Parameters.h
--- a/Engine/Plugins/Experimental/Mutable/Source/MutableRuntime/Public/MuR/Parameters.h
+++ b/Engine/Plugins/Experimental/Mutable/Source/MutableRuntime/Public/MuR/Parameters.h
@@ -9,8 +9,6 @@
 #include "Templates/SharedPointer.h"
 
 
-#define MUTABLE_MAX_STRING_PARAM_LENGTH     40
-
 namespace mu
 {
 
diff --git a/Engine/Plugins/Experimental/Mutable/Source/MutableTools/Private/MuT/CodeGenerator.cpp b/Engine/Plugins/Experimental/Mutable/Source/MutableTools/Private/MuT/CodeGenerator.cpp
--- a/Engine/Plugins/Experimental/Mutable/Source/MutableTools/Private/MuT/CodeGenerator.cpp
+++ b/Engine/Plugins/Experimental/Mutable/Source/MutableTools/Private/MuT/CodeGenerator.cpp
@@ -324,7 +324,7 @@ namespace mu
             param.m_name = pTable->GetName();
         }
         param.m_type = PARAMETER_TYPE::T_INT;
-        param.m_defaultValue.m_int = 0;
+        param.m_defaultValue.Set<ParamIntType>(0);
 
         // Add the possible values
         {
@@ -346,7 +346,7 @@ namespace mu
 				nullValue.m_value = -1;
 				nullValue.m_name = "None";
 				param.m_possibleValues.Add(nullValue);
-				param.m_defaultValue.m_int = nullValue.m_value;
+				param.m_defaultValue.Set<ParamIntType>(nullValue.m_value);
 			}
 
             // Add every row
@@ -366,7 +366,7 @@ namespace mu
                 // The first row is the default one
                 if (i==0)
                 {
-                    param.m_defaultValue.m_int = value.m_value;
+                    param.m_defaultValue.Set<ParamIntType>(value.m_value);
                 }
             }
         }
diff --git a/Engine/Plugins/Experimental/Mutable/Source/MutableTools/Private/MuT/CodeGenerator_Bool.cpp b/Engine/Plugins/Experimental/Mutable/Source/MutableTools/Private/MuT/CodeGenerator_Bool.cpp
--- a/Engine/Plugins/Experimental/Mutable/Source/MutableTools/Private/MuT/CodeGenerator_Bool.cpp
+++ b/Engine/Plugins/Experimental/Mutable/Source/MutableTools/Private/MuT/CodeGenerator_Bool.cpp
@@ -100,7 +100,7 @@ namespace mu
 			param.m_name = node.m_name;
 			param.m_uid = node.m_uid;
 			param.m_type = PARAMETER_TYPE::T_BOOL;
-			param.m_defaultValue.m_bool = node.m_defaultValue;
+			param.m_defaultValue.Set<ParamBoolType>(node.m_defaultValue);
 
 			op = new ASTOpParameter();
 			op->type = OP_TYPE::BO_PARAMETER;
diff --git a/Engine/Plugins/Experimental/Mutable/Source/MutableTools/Private/MuT/CodeGenerator_Color.cpp b/Engine/Plugins/Experimental/Mutable/Source/MutableTools/Private/MuT/CodeGenerator_Color.cpp
--- a/Engine/Plugins/Experimental/Mutable/Source/MutableTools/Private/MuT/CodeGenerator_Color.cpp
+++ b/Engine/Plugins/Experimental/Mutable/Source/MutableTools/Private/MuT/CodeGenerator_Color.cpp
@@ -148,10 +148,14 @@ namespace mu
 			param.m_name = node.m_name;
 			param.m_uid = node.m_uid;
 			param.m_type = PARAMETER_TYPE::T_COLOUR;
-			param.m_defaultValue.m_colour[0] = node.m_defaultValue[0];
-			param.m_defaultValue.m_colour[1] = node.m_defaultValue[1];
-			param.m_defaultValue.m_colour[2] = node.m_defaultValue[2];
 
+			ParamColorType Value;
+			Value[0] =  node.m_defaultValue[0];
+			Value[1] = node.m_defaultValue[1];
+			Value[2] = node.m_defaultValue[2];
+
+			param.m_defaultValue.Set<ParamColorType>(Value);
+			
 			op = new ASTOpParameter();
 			op->type = OP_TYPE::CO_PARAMETER;
 			op->parameter = param;
diff --git a/Engine/Plugins/Experimental/Mutable/Source/MutableTools/Private/MuT/CodeGenerator_FirstPass.cpp b/Engine/Plugins/Experimental/Mutable/Source/MutableTools/Private/MuT/CodeGenerator_FirstPass.cpp
--- a/Engine/Plugins/Experimental/Mutable/Source/MutableTools/Private/MuT/CodeGenerator_FirstPass.cpp
+++ b/Engine/Plugins/Experimental/Mutable/Source/MutableTools/Private/MuT/CodeGenerator_FirstPass.cpp
@@ -529,7 +529,7 @@ namespace mu
             op->parameter.m_name = node.m_name;
             op->parameter.m_uid = node.m_uid;
             op->parameter.m_type = PARAMETER_TYPE::T_INT;
-            op->parameter.m_defaultValue.m_int = -1;
+            op->parameter.m_defaultValue.Set<ParamIntType>(-1);
 
             if ( node.m_type==NodeObjectGroup::CS_ONE_OR_NONE )
             {
@@ -537,7 +537,7 @@ namespace mu
                 nullValue.m_value = -1;
                 nullValue.m_name = "None";
                 op->parameter.m_possibleValues.Add( nullValue );
-                op->parameter.m_defaultValue.m_int = nullValue.m_value;
+                op->parameter.m_defaultValue.Set<ParamIntType>(nullValue.m_value);
             }
 
             enumOp = op;
@@ -562,7 +562,7 @@ namespace mu
                         op->parameter.m_name = pChildNode->GetName();
                         op->parameter.m_uid = pChildNode->GetUid();
                         op->parameter.m_type = PARAMETER_TYPE::T_BOOL;
-                        op->parameter.m_defaultValue.m_bool = false;
+                        op->parameter.m_defaultValue.Set<ParamBoolType>(false);
 
                         paramOp = op;
                         break;
@@ -588,9 +588,9 @@ namespace mu
                         enumOp->parameter.m_possibleValues.Add( value );
 
                         // Set as default if none has been set.
-                        if ( enumOp->parameter.m_defaultValue.m_int == -1 )
+                        if ( enumOp->parameter.m_defaultValue.Get<ParamIntType>() == -1 )
                         {
-                            enumOp->parameter.m_defaultValue.m_int = (int)t;
+                            enumOp->parameter.m_defaultValue.Set<ParamIntType>(t);
                         }
 
                         check(enumOp);
diff --git a/Engine/Plugins/Experimental/Mutable/Source/MutableTools/Private/MuT/CodeGenerator_Image.cpp b/Engine/Plugins/Experimental/Mutable/Source/MutableTools/Private/MuT/CodeGenerator_Image.cpp
--- a/Engine/Plugins/Experimental/Mutable/Source/MutableTools/Private/MuT/CodeGenerator_Image.cpp
+++ b/Engine/Plugins/Experimental/Mutable/Source/MutableTools/Private/MuT/CodeGenerator_Image.cpp
@@ -276,6 +276,7 @@ namespace mu
 			op->parameter.m_name = node.m_name;
 			op->parameter.m_uid = node.m_uid;
             op->parameter.m_type = PARAMETER_TYPE::T_IMAGE;
+        	op->parameter.m_defaultValue.Set<ParamImageType>(0);
 
 			// Generate the code for the ranges
 			for (int32 a = 0; a < node.m_ranges.Num(); ++a)
diff --git a/Engine/Plugins/Experimental/Mutable/Source/MutableTools/Private/MuT/CodeGenerator_Projector.cpp b/Engine/Plugins/Experimental/Mutable/Source/MutableTools/Private/MuT/CodeGenerator_Projector.cpp
--- a/Engine/Plugins/Experimental/Mutable/Source/MutableTools/Private/MuT/CodeGenerator_Projector.cpp
+++ b/Engine/Plugins/Experimental/Mutable/Source/MutableTools/Private/MuT/CodeGenerator_Projector.cpp
@@ -119,7 +119,7 @@ namespace mu
 			param.m_type = PARAMETER_TYPE::T_PROJECTOR;
 
 			FProjector p = ProjectorFromNode(node);
-			param.m_defaultValue.m_projector = p;
+			param.m_defaultValue.Set<ParamProjectorType>(p);
 
 			op = new ASTOpParameter();
 			op->type = OP_TYPE::PR_PARAMETER;
@@ -140,7 +140,7 @@ namespace mu
 			op = it->second;
 		}
 
-		result.type = op->parameter.m_defaultValue.m_projector.type;
+		result.type = op->parameter.m_defaultValue.Get<ParamProjectorType>().type;
 		result.op = op;
 	}
 
diff --git a/Engine/Plugins/Experimental/Mutable/Source/MutableTools/Private/MuT/CodeGenerator_Scalar.cpp b/Engine/Plugins/Experimental/Mutable/Source/MutableTools/Private/MuT/CodeGenerator_Scalar.cpp
--- a/Engine/Plugins/Experimental/Mutable/Source/MutableTools/Private/MuT/CodeGenerator_Scalar.cpp
+++ b/Engine/Plugins/Experimental/Mutable/Source/MutableTools/Private/MuT/CodeGenerator_Scalar.cpp
@@ -143,7 +143,7 @@ namespace mu
 			param.m_name = node.m_name;
 			param.m_uid = node.m_uid;
 			param.m_type = PARAMETER_TYPE::T_FLOAT;
-			param.m_defaultValue.m_float = node.m_defaultValue;
+			param.m_defaultValue.Set<ParamFloatType>(node.m_defaultValue);
 			param.m_detailedType = node.m_detailedType;
 
 			op = new ASTOpParameter();
@@ -209,7 +209,7 @@ namespace mu
 			param.m_name = node.m_name;
 			param.m_uid = node.m_uid;
 			param.m_type = PARAMETER_TYPE::T_INT;
-			param.m_defaultValue.m_int = node.m_defaultValue;
+			param.m_defaultValue.Set<ParamIntType>(node.m_defaultValue);
 			param.m_detailedType = node.m_detailedType;
 
 			param.m_possibleValues.SetNum(node.m_options.Num());
diff --git a/Engine/Plugins/Experimental/Mutable/Source/MutableTools/Private/MuT/CodeGenerator_String.cpp b/Engine/Plugins/Experimental/Mutable/Source/MutableTools/Private/MuT/CodeGenerator_String.cpp
--- a/Engine/Plugins/Experimental/Mutable/Source/MutableTools/Private/MuT/CodeGenerator_String.cpp
+++ b/Engine/Plugins/Experimental/Mutable/Source/MutableTools/Private/MuT/CodeGenerator_String.cpp
@@ -93,11 +93,7 @@ namespace mu
 			param.m_uid = node.m_uid;
 			param.m_type = PARAMETER_TYPE::T_FLOAT;
 			param.m_detailedType = node.m_detailedType;
-			// param.m_defaultValue.m_text = node.m_defaultValue;
-			int len = FMath::Min(MUTABLE_MAX_STRING_PARAM_LENGTH, int(node.m_defaultValue.size()));
-			FMemory::Memcpy(param.m_defaultValue.m_text, node.m_defaultValue.c_str(), len);
-			param.m_defaultValue.m_text[len] = 0;
-
+			param.m_defaultValue.Set<ParamStringType>(node.m_defaultValue);
 
 			op = new ASTOpParameter();
 			op->type = OP_TYPE::ST_PARAMETER;
