#rnx

diff --git a/Engine/Plugins/Experimental/Mutable/Source/CustomizableObject/Private/MuCO/CustomizableObjectSystem.cpp b/Engine/Plugins/Experimental/Mutable/Source/CustomizableObject/Private/MuCO/CustomizableObjectSystem.cpp
--- a/Engine/Plugins/Experimental/Mutable/Source/CustomizableObject/Private/MuCO/CustomizableObjectSystem.cpp
+++ b/Engine/Plugins/Experimental/Mutable/Source/CustomizableObject/Private/MuCO/CustomizableObjectSystem.cpp
@@ -28,6 +28,7 @@
 #include "TextureResource.h"
 #include "UObject/UObjectIterator.h"
 #include "Widgets/Notifications/SNotificationList.h"
+#include "ContentStreaming.h"
 
 #if WITH_EDITOR
 #include "Editor.h"
@@ -822,15 +823,24 @@ void FCustomizableObjectSystemPrivate::InitUpdateSkeletalMesh(UCustomizableObjec
 	MUTABLE_CPUPROFILER_SCOPE(FCustomizableObjectSystemPrivate::InitUpdateSkeletalMesh);
 
 	check(IsInGameThread());
-	
+
 	const FString CurrentState = Instance.GetCurrentState();
 	const FParameterUIData* State = Instance.GetCustomizableObject()->StateUIDataMap.Find(CurrentState);
 
 	// \TODO: This should be controllable independently
-	const bool bNeverStream = State ? State->TextureCompressionStrategy!=ETextureCompressionStrategy::None : false;
-	const bool bUseMipmapStreaming = !bNeverStream;
+	const bool bNeverStream = State ? State->TextureCompressionStrategy != ETextureCompressionStrategy::None : false;
+	bool bUseMipmapStreaming = !bNeverStream;
 	int32 MipsToSkip = 0; // 0 means all mips
 
+#if PLATFORM_SUPPORTS_TEXTURE_STREAMING
+	if (!IStreamingManager::Get().IsTextureStreamingEnabled())
+	{
+		bUseMipmapStreaming = false;
+	}
+#else
+	bUseMipmapStreaming = false;
+#endif
+
 	if (bUseMipmapStreaming && EnableMutableProgressiveMipStreaming)
 	{
 		MipsToSkip = 255; // This means skip all possible mips until only UTexture::GetStaticMinTextureResidentMipCount() are left
