[Mutable] Fix model deserialization crash.

- Changed some sizeof types to their typedef.

#rnx

diff --git a/Engine/Plugins/Experimental/Mutable/Source/MutableRuntime/Private/MuR/Model.cpp b/Engine/Plugins/Experimental/Mutable/Source/MutableRuntime/Private/MuR/Model.cpp
--- a/Engine/Plugins/Experimental/Mutable/Source/MutableRuntime/Private/MuR/Model.cpp
+++ b/Engine/Plugins/Experimental/Mutable/Source/MutableRuntime/Private/MuR/Model.cpp
@@ -554,7 +554,7 @@ namespace mu
                 break;
 
             case PARAMETER_TYPE::T_INT:
-                dataSize = sizeof(int32);
+                dataSize = sizeof(ParamIntType);
                 parameterValuesBlob.SetNum( pos+dataSize );
                 FMemory::Memcpy( &parameterValuesBlob[pos], &pParams->GetPrivate()->m_values[param].Get<ParamIntType>(), dataSize );
                 pos += dataSize;
@@ -575,7 +575,7 @@ namespace mu
 
             //! Floating point value in the range of 0.0 to 1.0
             case PARAMETER_TYPE::T_FLOAT:
-                dataSize = sizeof(float);
+                dataSize = sizeof(ParamFloatType);
                 parameterValuesBlob.SetNum( pos+dataSize );
 				FMemory::Memcpy( &parameterValuesBlob[pos], &pParams->GetPrivate()->m_values[param].Get<ParamFloatType>(), dataSize );
                 pos += dataSize;
@@ -596,9 +596,9 @@ namespace mu
 
             //! Floating point RGBA colour, with each channel ranging from 0.0 to 1.0
             case PARAMETER_TYPE::T_COLOUR:
-                dataSize = 3*sizeof(float);
+                dataSize = sizeof(ParamColorType);
                 parameterValuesBlob.SetNum( pos+dataSize );
-				FMemory::Memcpy( &parameterValuesBlob[pos], &pParams->GetPrivate()->m_values[param].Get<ParamFloatType>(), dataSize );
+				FMemory::Memcpy( &parameterValuesBlob[pos], &pParams->GetPrivate()->m_values[param].Get<ParamColorType>(), dataSize );
                 pos += dataSize;
 
                 // Multi-values
@@ -618,7 +618,7 @@ namespace mu
             //! 3D Projector type, defining a position, scale and orientation. Basically used for
             //! projected decals.
             case PARAMETER_TYPE::T_PROJECTOR:
-                dataSize = sizeof(FProjector);
+                dataSize = sizeof(ParamProjectorType);
 
                 // \todo: padding will be random?
                 parameterValuesBlob.SetNum( pos+dataSize );
@@ -640,7 +640,7 @@ namespace mu
                 break;
 
             case PARAMETER_TYPE::T_IMAGE:
-                dataSize = sizeof(EXTERNAL_IMAGE_ID);
+                dataSize = sizeof(ParamImageType);
                 parameterValuesBlob.SetNum( pos+dataSize );
 				FMemory::Memcpy( &parameterValuesBlob[pos],
                         &pParams->GetPrivate()->m_values[param].Get<ParamImageType>(),
