[mutable] Apply fix for variations set unconditionally.
#rnx

diff --git a/Engine/Plugins/Experimental/Mutable/Source/MutableTools/Private/MuT/CodeGenerator_SecondPass.cpp b/Engine/Plugins/Experimental/Mutable/Source/MutableTools/Private/MuT/CodeGenerator_SecondPass.cpp
--- a/Engine/Plugins/Experimental/Mutable/Source/MutableTools/Private/MuT/CodeGenerator_SecondPass.cpp
+++ b/Engine/Plugins/Experimental/Mutable/Source/MutableTools/Private/MuT/CodeGenerator_SecondPass.cpp
@@ -267,9 +267,14 @@ namespace mu
 				negSurf,
 				posTag,
 				negTag);
-			// TODO: Optimise the tag condition here
 
+			if (!tagCondition)
+			{
+				// This tag is unconditionally activated, so there's no condition logic to add
+				continue;
+			}
 
+			// TODO: Optimise the tag condition here
 
 			// If the tag is a constant ...
 			bool isConstant = false;
