#rnx

diff --git a/Engine/Plugins/Experimental/Mutable/Source/CustomizableObject/Private/MuCO/CustomizableObjectSystem.cpp b/Engine/Plugins/Experimental/Mutable/Source/CustomizableObject/Private/MuCO/CustomizableObjectSystem.cpp
--- a/Engine/Plugins/Experimental/Mutable/Source/CustomizableObject/Private/MuCO/CustomizableObjectSystem.cpp
+++ b/Engine/Plugins/Experimental/Mutable/Source/CustomizableObject/Private/MuCO/CustomizableObjectSystem.cpp
@@ -1491,6 +1491,12 @@ namespace impl
 				int32 MaxMipsToSkip = FullLODCount - MinMipsInImage;
 				int32 MipsToSkip = FMath::Min(MaxMipsToSkip, OperationData->MipsToSkip);
 
+				if (!FMath::IsPowerOfTwo(Image.FullImageSizeX) || !FMath::IsPowerOfTwo(Image.FullImageSizeY))
+				{
+					// It doesn't make sense to skip mips as non-power-of-two size textures cannot be streamed anyway
+					MipsToSkip = 0;
+				}
+
 				if (MipsToSkip > 0 && CustomizableObjectSystemPrivateData->EnableSkipGenerateResidentMips != 0 && OperationData->LowPriorityTextures.Find(Image.Name) != INDEX_NONE)
 				{
 					const int32 MipSizeX = FMath::Max(Image.FullImageSizeX >> MipsToSkip, 1);
@@ -1527,7 +1533,7 @@ namespace impl
 		}
 
 	}
-
+	
 
 	void Subtask_Mutable_PrepareTextures(const TSharedPtr<FMutableOperationData>& OperationData)
 	{
