[Mutable] Added CVar to disable Image Cache.

#rnx

diff --git a/Engine/Plugins/Experimental/Mutable/Source/CustomizableObject/Classes/MuCO/CustomizableObjectSystem.h b/Engine/Plugins/Experimental/Mutable/Source/CustomizableObject/Classes/MuCO/CustomizableObjectSystem.h
--- a/Engine/Plugins/Experimental/Mutable/Source/CustomizableObject/Classes/MuCO/CustomizableObjectSystem.h
+++ b/Engine/Plugins/Experimental/Mutable/Source/CustomizableObject/Classes/MuCO/CustomizableObjectSystem.h
@@ -12,6 +12,7 @@
 
 #include "CustomizableObjectSystem.generated.h"
 
+class IConsoleVariable;
 class FCustomizableObjectCompilerBase;
 class ITargetPlatform;
 class SNotificationItem;
@@ -59,6 +60,8 @@ DECLARE_DWORD_ACCUMULATOR_STAT_EXTERN(TEXT("Streaming Cache"), STAT_MutableStrea
 
 extern TAutoConsoleVariable<bool> CVarClearStreamingCacheOnUpdateEnd;
 
+extern TAutoConsoleVariable<bool> CVarEnableImageCache;
+
 
 #if WITH_EDITOR
 
@@ -420,7 +423,8 @@ private:
 	TMap<FString, int64> PlatformMaxChunkSize;
 	
 #endif
-
+	void OnEnableImageCacheChanged(IConsoleVariable* CVar);
+	
 	// Friends
 	friend class FCustomizableObjectSystemPrivate;
 };
diff --git a/Engine/Plugins/Experimental/Mutable/Source/CustomizableObject/Private/MuCO/CustomizableObjectInstance.cpp b/Engine/Plugins/Experimental/Mutable/Source/CustomizableObject/Private/MuCO/CustomizableObjectInstance.cpp
--- a/Engine/Plugins/Experimental/Mutable/Source/CustomizableObject/Private/MuCO/CustomizableObjectInstance.cpp
+++ b/Engine/Plugins/Experimental/Mutable/Source/CustomizableObject/Private/MuCO/CustomizableObjectInstance.cpp
@@ -5412,7 +5412,10 @@ void UCustomizableInstancePrivateData::BuildMaterials(const TSharedPtr<FMutableO
 									}
 
 									// update the model resources cache
-									Cache.Images.Add(ImageCacheKey, Texture);
+									if (CVarEnableImageCache.GetValueOnGameThread())
+									{
+										Cache.Images.Add(ImageCacheKey, Texture);						
+									}
 								}
 								else
 								{
diff --git a/Engine/Plugins/Experimental/Mutable/Source/CustomizableObject/Private/MuCO/CustomizableObjectSystem.cpp b/Engine/Plugins/Experimental/Mutable/Source/CustomizableObject/Private/MuCO/CustomizableObjectSystem.cpp
--- a/Engine/Plugins/Experimental/Mutable/Source/CustomizableObject/Private/MuCO/CustomizableObjectSystem.cpp
+++ b/Engine/Plugins/Experimental/Mutable/Source/CustomizableObject/Private/MuCO/CustomizableObjectSystem.cpp
@@ -95,6 +95,13 @@ TAutoConsoleVariable<bool> CVarClearStreamingCacheOnUpdateEnd(
 	ECVF_Scalability);
 
 
+TAutoConsoleVariable<bool> CVarEnableImageCache(
+	TEXT("mutable.EnableImageCache"),
+	true,
+	TEXT("Enables or disables the Instance Image Cache. Avoids regenerating images when no parameters or mips have changed."),
+	ECVF_Scalability);
+
+
 bool FMutableQueue::IsEmpty() const
 {
 	return Array.Num() == 0;
@@ -350,6 +357,8 @@ void UCustomizableObjectSystem::InitSystem()
 	DefaultInstanceLODManagement = NewObject<UCustomizableInstanceLODManagement>();
 	check(DefaultInstanceLODManagement != nullptr);
 	CurrentInstanceLODManagement = DefaultInstanceLODManagement;
+
+	CVarEnableImageCache->SetOnChangedCallback(FConsoleVariableDelegate::CreateUObject(this, &UCustomizableObjectSystem::OnEnableImageCacheChanged));
 }
 
 
@@ -3026,6 +3035,18 @@ void UCustomizableObjectSystem::TickRecompileCustomizableObjects()
 }
 
 
+void UCustomizableObjectSystem::OnEnableImageCacheChanged(IConsoleVariable* CVar)
+{
+	if (CVar->GetBool())
+	{		
+		for (FMutableResourceCache& Cache : GetPrivate()->ModelResourcesCache)
+		{
+			Cache.Images.Empty();
+		}
+	}	
+}
+
+
 uint64 UCustomizableObjectSystem::GetMaxChunkSizeForPlatform(const ITargetPlatform* TargetPlatform)
 {
 	const FString& PlatformName = TargetPlatform ? TargetPlatform->IniPlatformName() : FPlatformProperties::IniPlatformName();
