[Mutable] Fix clear cache modifying the Disable Image cache CVar.

- Fix compilation error non-editor targets.

#rnx

diff --git a/Engine/Plugins/Experimental/Mutable/Source/CustomizableObject/Private/MuCO/CustomizableObjectSystem.cpp b/Engine/Plugins/Experimental/Mutable/Source/CustomizableObject/Private/MuCO/CustomizableObjectSystem.cpp
--- a/Engine/Plugins/Experimental/Mutable/Source/CustomizableObject/Private/MuCO/CustomizableObjectSystem.cpp
+++ b/Engine/Plugins/Experimental/Mutable/Source/CustomizableObject/Private/MuCO/CustomizableObjectSystem.cpp
@@ -3035,18 +3035,6 @@ void UCustomizableObjectSystem::TickRecompileCustomizableObjects()
 }
 
 
-void UCustomizableObjectSystem::OnEnableImageCacheChanged(IConsoleVariable* CVar)
-{
-	if (CVar->GetBool())
-	{		
-		for (FMutableResourceCache& Cache : GetPrivate()->ModelResourcesCache)
-		{
-			Cache.Images.Empty();
-		}
-	}	
-}
-
-
 uint64 UCustomizableObjectSystem::GetMaxChunkSizeForPlatform(const ITargetPlatform* TargetPlatform)
 {
 	const FString& PlatformName = TargetPlatform ? TargetPlatform->IniPlatformName() : FPlatformProperties::IniPlatformName();
@@ -3083,6 +3071,18 @@ uint64 UCustomizableObjectSystem::GetMaxChunkSizeForPlatform(const ITargetPlatfo
 #endif // WITH_EDITOR
 
 
+void UCustomizableObjectSystem::OnEnableImageCacheChanged(IConsoleVariable* CVar)
+{
+	if (!CVar->GetBool())
+	{		
+		for (FMutableResourceCache& Cache : GetPrivate()->ModelResourcesCache)
+		{
+			Cache.Images.Empty();
+		}
+	}	
+}
+
+
 void UCustomizableObjectSystem::CacheImage(uint64 ImageId)
 {
 	check(GetPrivate() != nullptr);
