[mutable] Fixed UV viewer's combobox. Now parameters are resetted if the CO is recompiled or if the previewd material no longer exists.

#rnx

diff --git a/Engine/Plugins/Experimental/Mutable/Source/CustomizableObjectEditor/Private/MuCOE/CustomizableObjectEditor.cpp b/Engine/Plugins/Experimental/Mutable/Source/CustomizableObjectEditor/Private/MuCOE/CustomizableObjectEditor.cpp
--- a/Engine/Plugins/Experimental/Mutable/Source/CustomizableObjectEditor/Private/MuCOE/CustomizableObjectEditor.cpp
+++ b/Engine/Plugins/Experimental/Mutable/Source/CustomizableObjectEditor/Private/MuCOE/CustomizableObjectEditor.cpp
@@ -1366,8 +1366,9 @@ void FCustomizableObjectEditor::CompileObject()
 	// If any projector selected, unselect it in case the projector node is removed with this CO compile
 	ResetProjectorVisibilityNoUpdate();
 
-	// Unselect NodeMeshClipMorph if selected
+	// Resetting viewport parameters
 	Viewport->SetClipMorphPlaneVisibility(false, nullptr);
+	Viewport->SetDrawDefaultUVMaterial(true);
 
 	UE_LOG(LogMutable, Verbose, TEXT("PROFILE: -----------------------------------------------------------"));
 	UE_LOG(LogMutable, Verbose, TEXT("PROFILE: [ %16.8f ] FCustomizableObjectEditor::CompileObject start."), FPlatformTime::Seconds());
@@ -2645,7 +2646,7 @@ void FCustomizableObjectEditor::OnUpdatePreviewInstance()
 	}
 	
 	Viewport->SetPreviewComponents(PreviewSkeletalMeshComponents);
-	Viewport->SetDrawDefaultUVMaterial();
+	Viewport->SetDrawDefaultUVMaterial(false);
 	ViewportClient->Invalidate();
 	ViewportClient->ReSetAnimation();
 	ViewportClient->SetReferenceMeshMissingWarningMessage(false);
diff --git a/Engine/Plugins/Experimental/Mutable/Source/CustomizableObjectEditor/Private/MuCOE/SCustomizableObjectEditorViewport.cpp b/Engine/Plugins/Experimental/Mutable/Source/CustomizableObjectEditor/Private/MuCOE/SCustomizableObjectEditorViewport.cpp
--- a/Engine/Plugins/Experimental/Mutable/Source/CustomizableObjectEditor/Private/MuCOE/SCustomizableObjectEditorViewport.cpp
+++ b/Engine/Plugins/Experimental/Mutable/Source/CustomizableObjectEditor/Private/MuCOE/SCustomizableObjectEditorViewport.cpp
@@ -655,30 +655,55 @@ bool SCustomizableObjectEditorViewportTabBody::IsCameraModeActive(int Value)
 }
 
 
-void SCustomizableObjectEditorViewportTabBody::SetDrawDefaultUVMaterial()
+void SCustomizableObjectEditorViewportTabBody::SetDrawDefaultUVMaterial(bool bIsCompilation)
 {
 	GenerateUVMaterialOptions();
-
-	FString CurrentDrawnMaterial = LevelViewportClient->GetMaterialToDrawInUVs();
+	bool bMaterialFound = false;
 
 	if (!ArrayUVMaterialOptionString.IsEmpty())
 	{
-		bool bFound = false;
-
-		for (int32 StringIndex = 0; StringIndex < ArrayUVMaterialOptionString.Num(); ++StringIndex)
+		if (!bIsCompilation && SelectedUVMaterial.IsValid())
 		{
-			if (*(ArrayUVMaterialOptionString[StringIndex]) == CurrentDrawnMaterial)
+			// We check if the selected Material still exists after the update
+			FString MaterialName = *SelectedUVMaterial;
+
+			for (int32 MaterialIndex = 0; MaterialIndex < ArrayUVMaterialOptionString.Num(); ++MaterialIndex)
 			{
-				bFound = true;
-				break;
+				if (MaterialName == *ArrayUVMaterialOptionString[MaterialIndex])
+				{
+					bMaterialFound = true;
+					break;
+				}
 			}
 		}
 
-		if (!bFound)
+		if (bIsCompilation || !bMaterialFound)
 		{
 			LevelViewportClient->SetDrawUVOverlayMaterial(*(ArrayUVMaterialOptionString[0]), "0");
+			SelectedUVMaterial = ArrayUVMaterialOptionString[0];
 		}
 	}
+	else
+	{
+		SelectedUVMaterial = nullptr;
+	}
+
+	GenerateUVChannelOptions();
+
+	if (!ArrayUVChannelOptionString.IsEmpty() && (bIsCompilation || !bMaterialFound))
+	{
+		SelectedUVChannel = ArrayUVChannelOptionString[0];
+	}
+	else
+	{
+		SelectedUVChannel = nullptr;
+	}
+
+	if (UVMaterialOptionCombo.IsValid() && UVChannelOptionCombo.IsValid())
+	{
+		UVMaterialOptionCombo->SetSelectedItem(SelectedUVMaterial);
+		UVChannelOptionCombo->SetSelectedItem(SelectedUVChannel);
+	}
 }
 
 
@@ -791,20 +816,6 @@ TSharedRef<SWidget> SCustomizableObjectEditorViewportTabBody::GenerateUVMaterial
 					}
 				}
 			}
-			else
-			{
-				FString CurrentDrawnMaterial = LevelViewportClient->GetMaterialToDrawInUVs();
-
-				for (int32 i = 0; i < ArrayUVMaterialOptionString.Num(); ++i)
-				{
-					if (CurrentDrawnMaterial == *ArrayUVMaterialOptionString[i])
-					{
-						bFound = true;
-						SelectedUVMaterial = ArrayUVMaterialOptionString[i];
-						break;
-					}
-				}
-			}
 
 			if (!bFound)
 			{
@@ -822,7 +833,7 @@ TSharedRef<SWidget> SCustomizableObjectEditorViewportTabBody::GenerateUVMaterial
 			.OnSelectionChanged(this, &SCustomizableObjectEditorViewportTabBody::OnMaterialChanged);
 
 		// Generating an array with all the options of the combobox 
-		GenerateUVChannelOptions(false);
+		GenerateUVChannelOptions();
 
 		// Setting initial selected option
 		if (ArrayUVChannelOptionString.Num())
@@ -977,7 +988,14 @@ void SCustomizableObjectEditorViewportTabBody::OnMaterialChanged(TSharedPtr<FStr
 		SelectedUVMaterial = Selected;
 
 		//We need to update options for the new LOD
-		GenerateUVChannelOptions(true);
+		GenerateUVChannelOptions();
+
+		// Resets the value to the first element of the array
+		if (!ArrayUVChannelOptionString.IsEmpty())
+		{
+			SelectedUVChannel = ArrayUVChannelOptionString[0];
+			UVChannelOptionCombo->SetSelectedItem(SelectedUVChannel);
+		}
 
 		if (LevelViewportClient.IsValid() && SelectedUVChannel.IsValid())
 		{
@@ -987,7 +1005,7 @@ void SCustomizableObjectEditorViewportTabBody::OnMaterialChanged(TSharedPtr<FStr
 }
 
 
-void SCustomizableObjectEditorViewportTabBody::GenerateUVChannelOptions(bool bReset)
+void SCustomizableObjectEditorViewportTabBody::GenerateUVChannelOptions()
 {
 	ArrayUVChannelOptionString.Empty();
 
@@ -1022,13 +1040,6 @@ void SCustomizableObjectEditorViewportTabBody::GenerateUVChannelOptions(bool bRe
 		{
 			ArrayUVChannelOptionString.Add(MakeShareable(new FString(FString::FromInt(UVChan))));
 		}
-
-		// Resets the value to the first element of the array
-		if (bReset)
-		{
-			SelectedUVChannel = ArrayUVChannelOptionString[0];
-			UVChannelOptionCombo->SetSelectedItem(SelectedUVChannel);
-		}
 	}
 }
 
diff --git a/Engine/Plugins/Experimental/Mutable/Source/CustomizableObjectEditor/Private/MuCOE/SCustomizableObjectEditorViewport.h b/Engine/Plugins/Experimental/Mutable/Source/CustomizableObjectEditor/Private/MuCOE/SCustomizableObjectEditorViewport.h
--- a/Engine/Plugins/Experimental/Mutable/Source/CustomizableObjectEditor/Private/MuCOE/SCustomizableObjectEditorViewport.h
+++ b/Engine/Plugins/Experimental/Mutable/Source/CustomizableObjectEditor/Private/MuCOE/SCustomizableObjectEditorViewport.h
@@ -251,8 +251,9 @@ public:
 	/** Returns whether or not the a camera mode is active (Orbital or Free) */
 	bool IsCameraModeActive(int Value);
 
-	/** Sets the first material of the instance as the default section to draw in the UVs Overlay (when there is no material selected in the combobox)*/
-	void SetDrawDefaultUVMaterial();
+	/** Sets the first material of the instance as the default section to draw in the UVs Overlay (when there is no material selected in the combobox)
+		Param bIsCompilation: true if this function is called during a CO compilation, false when is called from an instance update	*/
+	void SetDrawDefaultUVMaterial(bool bIsCompilation);
 
 	/** Sets the bones visibility */
 	void SetShowBones();
@@ -308,7 +309,7 @@ private:
 	void GenerateUVMaterialOptions();
 
 	// Generates the Combobox options for each UV Channel
-	void GenerateUVChannelOptions(bool bReset);
+	void GenerateUVChannelOptions();
 
 	// Selected option of the UV Material ComboBox
 	TSharedPtr<FString> SelectedUVMaterial;
diff --git a/Engine/Plugins/Experimental/Mutable/Source/CustomizableObjectEditor/Public/MuCOE/CustomizableObjectEditorViewportClient.h b/Engine/Plugins/Experimental/Mutable/Source/CustomizableObjectEditor/Public/MuCOE/CustomizableObjectEditorViewportClient.h
--- a/Engine/Plugins/Experimental/Mutable/Source/CustomizableObjectEditor/Public/MuCOE/CustomizableObjectEditorViewportClient.h
+++ b/Engine/Plugins/Experimental/Mutable/Source/CustomizableObjectEditor/Public/MuCOE/CustomizableObjectEditorViewportClient.h
@@ -398,12 +398,6 @@ public:
 	/** Sets camera mode */
 	void SetCameraMode(bool Value);
 
-	/** Getter of the Material name being drawn in the UV viewer */
-	FString GetMaterialToDrawInUVs()
-	{
-		return FullNameMaterialToDrawInUVs;
-	}
-
 	/** Sets the skeletal mesh bones visibility */
 	void SetShowBones();
 
@@ -446,7 +440,6 @@ private:
 
 	// Material index to draw in the uv preview.
 	FString MaterialToDrawInUVs;
-	FString FullNameMaterialToDrawInUVs;
 	int32 MaterialToDrawInUVsLOD;
 	int32 MaterialToDrawInUVsIndex;
 	int32 UVChannelToDrawInUVs;
