[mutable] Initial tiled image generation, disabled by default.
It uses "crop" and "patch" and several additional optimizations have been implemented for these operations.
#rnx

diff --git a/Engine/Plugins/Experimental/Mutable/Source/CustomizableObject/Classes/MuCO/CustomizableObject.h b/Engine/Plugins/Experimental/Mutable/Source/CustomizableObject/Classes/MuCO/CustomizableObject.h
--- a/Engine/Plugins/Experimental/Mutable/Source/CustomizableObject/Classes/MuCO/CustomizableObject.h
+++ b/Engine/Plugins/Experimental/Mutable/Source/CustomizableObject/Classes/MuCO/CustomizableObject.h
@@ -175,7 +175,7 @@ struct FCompilationOptions
 	bool bSkinWeightProfilesEnabled = true;
 
 	// Used to enable physics asset merge.
-	bool bPhysicsAssetMergeEnebled = false;
+	bool bPhysicsAssetMergeEnabled = false;
 
 	// Used to reduce the number of notifications when compiling objects
 	bool bSilentCompilation = true;
@@ -184,6 +184,9 @@ struct FCompilationOptions
 	bool bForceLargeLODBias = false;
 	int32 DebugBias = 0;
 
+	// Control image tiled generation
+	int32 ImageTiling = 0;
+
 };
 
 
@@ -1365,7 +1368,7 @@ private:
 	// This is a manual version number for the binary blobs in this asset.
 	// Increasing it invalidates all the previously compiled models.
 	// Warning: If while merging code both versions have changed, take the highest+1.
-	static const int32 CurrentSupportedVersion = 375;
+	static const int32 CurrentSupportedVersion = 377;
 
 public:
 
diff --git a/Engine/Plugins/Experimental/Mutable/Source/CustomizableObjectEditor/Private/MuCOE/CustomizableObjectCompileRunnable.cpp b/Engine/Plugins/Experimental/Mutable/Source/CustomizableObjectEditor/Private/MuCOE/CustomizableObjectCompileRunnable.cpp
--- a/Engine/Plugins/Experimental/Mutable/Source/CustomizableObjectEditor/Private/MuCOE/CustomizableObjectCompileRunnable.cpp
+++ b/Engine/Plugins/Experimental/Mutable/Source/CustomizableObjectEditor/Private/MuCOE/CustomizableObjectCompileRunnable.cpp
@@ -88,8 +88,11 @@ uint32 FCustomizableObjectCompileRunnable::Run()
 	const int MinRomSize = 128;
 	CompilerOptions->SetDataPackingStrategy(MinRomSize, MinResidentMips);
 
-	/** At object compilation time we don't know if we will want progressive images or not. Assume we will. */
+	// At object compilation time we don't know if we will want progressive images or not. Assume we will. 
+	// TODO: Per-state setting?
 	CompilerOptions->SetEnableProgressiveImages(true);
+	
+	CompilerOptions->SetImageTiling(Options.ImageTiling);
 
 	CompilerOptions->SetTextureLayoutStrategy(bDisableTextureLayout
 		? mu::CompilerOptions::TextureLayoutStrategy::None
diff --git a/Engine/Plugins/Experimental/Mutable/Source/CustomizableObjectEditor/Private/MuCOE/CustomizableObjectCompiler.cpp b/Engine/Plugins/Experimental/Mutable/Source/CustomizableObjectEditor/Private/MuCOE/CustomizableObjectCompiler.cpp
--- a/Engine/Plugins/Experimental/Mutable/Source/CustomizableObjectEditor/Private/MuCOE/CustomizableObjectCompiler.cpp
+++ b/Engine/Plugins/Experimental/Mutable/Source/CustomizableObjectEditor/Private/MuCOE/CustomizableObjectCompiler.cpp
@@ -909,7 +909,7 @@ void FCustomizableObjectCompiler::CompileInternal(UCustomizableObject* Object, c
 	Options.bClothingEnabled = Object->bEnableClothing;
 	Options.b16BitBoneWeightsEnabled = Object->bEnable16BitBoneWeights;
 	Options.bSkinWeightProfilesEnabled = Object->bEnableAltSkinWeightProfiles;
-	Options.bPhysicsAssetMergeEnebled = Object->bEnablePhysicsAssetMerge;
+	Options.bPhysicsAssetMergeEnabled = Object->bEnablePhysicsAssetMerge;
 
 	if (Object->IsLocked() || !UCustomizableObjectSystem::GetInstance()->LockObject(Object))
 	{
diff --git a/Engine/Plugins/Experimental/Mutable/Source/CustomizableObjectEditor/Private/MuCOE/CustomizableObjectEditor.cpp b/Engine/Plugins/Experimental/Mutable/Source/CustomizableObjectEditor/Private/MuCOE/CustomizableObjectEditor.cpp
--- a/Engine/Plugins/Experimental/Mutable/Source/CustomizableObjectEditor/Private/MuCOE/CustomizableObjectEditor.cpp
+++ b/Engine/Plugins/Experimental/Mutable/Source/CustomizableObjectEditor/Private/MuCOE/CustomizableObjectEditor.cpp
@@ -1128,6 +1128,11 @@ TSharedRef<SWidget> FCustomizableObjectEditor::GenerateCompileOptionsMenuContent
 	}
 	MenuBuilder.EndSection();
 
+	if (!CustomizableObject)
+	{
+		return MenuBuilder.MakeWidget();
+	}
+
 	MenuBuilder.BeginSection("Optimization", LOCTEXT("MutableCompileOptimizationHeading", "Optimization"));
 	{
 		// Level
@@ -1140,12 +1145,37 @@ TSharedRef<SWidget> FCustomizableObjectEditor::GenerateCompileOptionsMenuContent
 		CompileOptimizationCombo =
 			SNew(STextComboBox)
 			.OptionsSource(&CompileOptimizationStrings)
-			.InitiallySelectedItem(CompileOptimizationStrings[CustomizableObject ? CustomizableObject->CompileOptions.OptimizationLevel : 0])
+			.InitiallySelectedItem(CompileOptimizationStrings[CustomizableObject->CompileOptions.OptimizationLevel])
 			.OnSelectionChanged(this, &FCustomizableObjectEditor::OnChangeCompileOptimizationLevel)
 			;
 
 		MenuBuilder.AddWidget(CompileOptimizationCombo.ToSharedRef(), LOCTEXT("MutableCompileOptimizationLevel", "Optimization Level"));
 
+		// Image tiling
+		// Unfortunately SNumericDropDown doesn't work with integers at the time of writing.
+		TArray<SNumericDropDown<float>::FNamedValue> TilingOptions;
+		TilingOptions.Add(SNumericDropDown<float>::FNamedValue(0, FText::FromString(TEXT("0")), FText::FromString(TEXT("Disabled"))));
+		TilingOptions.Add(SNumericDropDown<float>::FNamedValue(64, FText::FromString(TEXT("64")), FText::FromString(TEXT("64"))));
+		TilingOptions.Add(SNumericDropDown<float>::FNamedValue(128, FText::FromString(TEXT("128")), FText::FromString(TEXT("128"))));
+		TilingOptions.Add(SNumericDropDown<float>::FNamedValue(256, FText::FromString(TEXT("256")), FText::FromString(TEXT("256"))));
+		TilingOptions.Add(SNumericDropDown<float>::FNamedValue(512, FText::FromString(TEXT("512")), FText::FromString(TEXT("512"))));
+
+		CompileTilingCombo = SNew(SNumericDropDown<float>)
+			.DropDownValues(TilingOptions)
+			.Value_Lambda([this]() 
+				{ 
+					return float(CustomizableObject->CompileOptions.ImageTiling); 
+				})
+			.OnValueChanged_Lambda([this](float Value) 
+				{ 
+					if (CustomizableObject)
+					{
+						CustomizableObject->CompileOptions.ImageTiling = int32(Value);
+						CustomizableObject->Modify();
+					}
+				});
+		MenuBuilder.AddWidget(CompileTilingCombo.ToSharedRef(), LOCTEXT("MutableCompileImageTiling", "Image Tiling"));
+
 		MenuBuilder.AddMenuEntry(FCustomizableObjectEditorCommands::Get().CompileOptions_UseParallelCompilation);
 		MenuBuilder.AddMenuEntry(FCustomizableObjectEditorCommands::Get().CompileOptions_UseDiskCompilation);
 		MenuBuilder.AddMenuEntry(FCustomizableObjectEditorCommands::Get().CompileOptions_EnableTextureCompression);
diff --git a/Engine/Plugins/Experimental/Mutable/Source/CustomizableObjectEditor/Private/MuCOE/CustomizableObjectEditor.h b/Engine/Plugins/Experimental/Mutable/Source/CustomizableObjectEditor/Private/MuCOE/CustomizableObjectEditor.h
--- a/Engine/Plugins/Experimental/Mutable/Source/CustomizableObjectEditor/Private/MuCOE/CustomizableObjectEditor.h
+++ b/Engine/Plugins/Experimental/Mutable/Source/CustomizableObjectEditor/Private/MuCOE/CustomizableObjectEditor.h
@@ -8,6 +8,7 @@
 #include "MuCOE/CustomizableObjectCompiler.h"
 #include "MuCOE/ICustomizableObjectEditor.h"
 #include "TickableEditorObject.h"
+#include "Widgets/Input/SNumericDropDown.h"
 
 #include "CustomizableObjectEditor.generated.h"
 
@@ -289,6 +290,8 @@ private:
 	void ResetCompileOptions();
 	TSharedPtr<STextComboBox> CompileOptimizationCombo;
 	TArray< TSharedPtr<FString> > CompileOptimizationStrings;
+	TSharedPtr<SNumericDropDown<float>> CompileTilingCombo;
+
 	void CompileOptions_UseParallelCompilation_Toggled();
 	bool CompileOptions_UseParallelCompilation_IsChecked();
 	void CompileOptions_UseDiskCompilation_Toggled();
diff --git a/Engine/Plugins/Experimental/Mutable/Source/CustomizableObjectEditor/Private/MuCOE/CustomizableObjectLayout.cpp b/Engine/Plugins/Experimental/Mutable/Source/CustomizableObjectEditor/Private/MuCOE/CustomizableObjectLayout.cpp
--- a/Engine/Plugins/Experimental/Mutable/Source/CustomizableObjectEditor/Private/MuCOE/CustomizableObjectLayout.cpp
+++ b/Engine/Plugins/Experimental/Mutable/Source/CustomizableObjectEditor/Private/MuCOE/CustomizableObjectLayout.cpp
@@ -74,11 +74,11 @@ void UCustomizableObjectLayout::GenerateBlocksFromUVs()
 	if (Node && Mesh)
 	{
 		//Creating a GenerationContext
-		FCustomizableObjectCompiler* Compiler = new FCustomizableObjectCompiler();
+		TUniquePtr<FCustomizableObjectCompiler> Compiler( new FCustomizableObjectCompiler() );
 		UCustomizableObject* Object = Node->GetGraphEditor()->GetCustomizableObject();
 		FCompilationOptions Options = Object->CompileOptions;
 	
-		FMutableGraphGenerationContext GenerationContext(Object, Compiler, Options);
+		FMutableGraphGenerationContext GenerationContext(Object, Compiler.Get(), Options);
 	
 		//Transforming skeletalmesh to mutable mesh
 		mu::MeshPtr	MutableMesh = nullptr;
@@ -102,22 +102,22 @@ void UCustomizableObjectLayout::GenerateBlocksFromUVs()
 	
 		if (MutableMesh)
 		{
-			//Generating blocks with the mutable mesh
+			// Generating blocks with the mutable mesh
 			Layout = mu::NodeLayoutBlocks::GenerateLayoutBlocks(MutableMesh, UVChannel, GridSize.X, GridSize.Y);
 		}
 	
-		delete Compiler;
+		Compiler.Reset();
 	
 		if (Layout)
 		{
 			Blocks.Empty();
 		
-			//Generating the layout blocks with the mutable layout
+			// Generating the layout blocks with the mutable layout
 			for (int i = 0; i < Layout->GetBlockCount(); ++i)
 			{
-				int minX, minY, sizeX, sizeY;
+				uint16 minX, minY, sizeX, sizeY;
 		
-				Layout->GetBlock(i, &minX, &minY, &sizeX, &sizeY);
+				Layout->GetLayout()->GetBlock(i, &minX, &minY, &sizeX, &sizeY);
 		
 				FCustomizableObjectLayoutBlock block;
 				block.Min = FIntPoint(minX, minY);
diff --git a/Engine/Plugins/Experimental/Mutable/Source/CustomizableObjectEditor/Private/MuCOE/GenerateMutableSource/GenerateMutableSourceMesh.cpp b/Engine/Plugins/Experimental/Mutable/Source/CustomizableObjectEditor/Private/MuCOE/GenerateMutableSource/GenerateMutableSourceMesh.cpp
--- a/Engine/Plugins/Experimental/Mutable/Source/CustomizableObjectEditor/Private/MuCOE/GenerateMutableSource/GenerateMutableSourceMesh.cpp
+++ b/Engine/Plugins/Experimental/Mutable/Source/CustomizableObjectEditor/Private/MuCOE/GenerateMutableSource/GenerateMutableSourceMesh.cpp
@@ -378,7 +378,7 @@ mu::MeshPtr ConvertSkeletalMeshToMutable(USkeletalMesh* InSkeletalMesh, int LOD,
 	const EMutableMeshConversionFlags CurrentFlags = GenerationContext.MeshGenerationFlags.Last();
 	const bool bIgnoreSkeleton = EnumHasAnyFlags(CurrentFlags, EMutableMeshConversionFlags::IgnoreSkinning);
 	const bool bIgnorePhysics = EnumHasAnyFlags(CurrentFlags, EMutableMeshConversionFlags::IgnorePhysics) || 
-								!GenerationContext.Options.bPhysicsAssetMergeEnebled;
+								!GenerationContext.Options.bPhysicsAssetMergeEnabled;
 
 	mu::MeshPtr MutableMesh = new mu::Mesh();
 		
diff --git a/Engine/Plugins/Experimental/Mutable/Source/CustomizableObjectEditor/Private/MuCOE/SMutableGraphViewer.cpp b/Engine/Plugins/Experimental/Mutable/Source/CustomizableObjectEditor/Private/MuCOE/SMutableGraphViewer.cpp
--- a/Engine/Plugins/Experimental/Mutable/Source/CustomizableObjectEditor/Private/MuCOE/SMutableGraphViewer.cpp
+++ b/Engine/Plugins/Experimental/Mutable/Source/CustomizableObjectEditor/Private/MuCOE/SMutableGraphViewer.cpp
@@ -14,6 +14,7 @@
 #include "MuT/Streams.h"
 #include "Widgets/Docking/SDockTab.h"
 #include "Widgets/Input/STextComboBox.h"
+#include "Widgets/Input/SNumericDropDown.h"
 #include "Widgets/Views/STreeView.h"
 
 class FExtender;
@@ -300,6 +301,22 @@ TSharedRef<SWidget> SMutableGraphViewer::GenerateCompileOptionsMenuContent()
 			;
 		MenuBuilder.AddWidget(CompileOptimizationCombo.ToSharedRef(), LOCTEXT("MutableCompileOptimizationLevel", "Optimization Level"));
 
+		// Image tiling
+		// Unfortunately SNumericDropDown doesn't work with integers at the time of writing.
+		TArray<SNumericDropDown<float>::FNamedValue> TilingOptions;
+		TilingOptions.Add(SNumericDropDown<float>::FNamedValue(0, FText::FromString(TEXT("0")), FText::FromString(TEXT("Disabled"))));
+		TilingOptions.Add(SNumericDropDown<float>::FNamedValue(64, FText::FromString(TEXT("64")), FText::FromString(TEXT("64"))));
+		TilingOptions.Add(SNumericDropDown<float>::FNamedValue(128, FText::FromString(TEXT("128")), FText::FromString(TEXT("128"))));
+		TilingOptions.Add(SNumericDropDown<float>::FNamedValue(256, FText::FromString(TEXT("256")), FText::FromString(TEXT("256"))));
+		TilingOptions.Add(SNumericDropDown<float>::FNamedValue(512, FText::FromString(TEXT("512")), FText::FromString(TEXT("512"))));
+
+		CompileTilingCombo = SNew(SNumericDropDown<float>)
+			.DropDownValues(TilingOptions)
+			.Value_Lambda([&]() { return float(CompileOptions.ImageTiling); })
+			.OnValueChanged_Lambda([&](float Value) { CompileOptions.ImageTiling = int32(Value); })
+			;
+		MenuBuilder.AddWidget(CompileTilingCombo.ToSharedRef(), LOCTEXT("MutableCompileImageTiling", "Image Tiling"));
+
 		// Parallel compilation
 		MenuBuilder.AddMenuEntry(
 			LOCTEXT("MutableEnableParallelCompilation", "Enable compiling in multiple threads."),
diff --git a/Engine/Plugins/Experimental/Mutable/Source/CustomizableObjectEditor/Private/MuCOE/SMutableGraphViewer.h b/Engine/Plugins/Experimental/Mutable/Source/CustomizableObjectEditor/Private/MuCOE/SMutableGraphViewer.h
--- a/Engine/Plugins/Experimental/Mutable/Source/CustomizableObjectEditor/Private/MuCOE/SMutableGraphViewer.h
+++ b/Engine/Plugins/Experimental/Mutable/Source/CustomizableObjectEditor/Private/MuCOE/SMutableGraphViewer.h
@@ -4,6 +4,7 @@
 
 #include "MuCOE/CustomizableObjectCompiler.h"
 #include "Widgets/SCompoundWidget.h"
+#include "Widgets/Input/SNumericDropDown.h"
 
 class STableViewBase;
 namespace ESelectInfo { enum Type : int; }
@@ -103,6 +104,7 @@ private:
 	TSharedRef<SWidget> GenerateCompileOptionsMenuContent();
 	TSharedPtr<STextComboBox> CompileOptimizationCombo;
 	TArray< TSharedPtr<FString> > CompileOptimizationStrings;
+	TSharedPtr<SNumericDropDown<float>> CompileTilingCombo;
 	void OnChangeCompileOptimizationLevel(TSharedPtr<FString> NewSelection, ESelectInfo::Type SelectInfo);
 
 	/** Callbacks from the tree widget. */
diff --git a/Engine/Plugins/Experimental/Mutable/Source/CustomizableObjectEditor/Private/MuCOE/SMutableLayoutViewer.cpp b/Engine/Plugins/Experimental/Mutable/Source/CustomizableObjectEditor/Private/MuCOE/SMutableLayoutViewer.cpp
--- a/Engine/Plugins/Experimental/Mutable/Source/CustomizableObjectEditor/Private/MuCOE/SMutableLayoutViewer.cpp
+++ b/Engine/Plugins/Experimental/Mutable/Source/CustomizableObjectEditor/Private/MuCOE/SMutableLayoutViewer.cpp
@@ -99,14 +99,18 @@ void SMutableLayoutViewer::Tick(const FGeometry& AllottedGeometry, const double
 		}
 		else
 		{
-			int BlockCount = MutableLayout->GetBlockCount();
+			int32 BlockCount = MutableLayout->GetBlockCount();
 			TArray<FCustomizableObjectLayoutBlock> Blocks;
 			Blocks.SetNum( BlockCount );
-			for (int BlockIndex = 0; BlockIndex < BlockCount; ++BlockIndex)
+			for (int32 BlockIndex = 0; BlockIndex < BlockCount; ++BlockIndex)
 			{
-				FIntPoint Size;
-				MutableLayout->GetBlock(BlockIndex, &Blocks[BlockIndex].Min[0], &Blocks[BlockIndex].Min[1], &Size[0], &Size[1]);
-				Blocks[BlockIndex].Max = Blocks[BlockIndex].Min + Size;
+				UE::Math::TIntVector2<uint16> Min;
+				UE::Math::TIntVector2<uint16> Size;
+				MutableLayout->GetBlock(BlockIndex, &Min[0], &Min[1], &Size[0], &Size[1]);
+				Blocks[BlockIndex].Min[0] = Min[0];
+				Blocks[BlockIndex].Min[1] = Min[1];
+				Blocks[BlockIndex].Max[0] = Min[0] + Size[0];
+				Blocks[BlockIndex].Max[1] = Min[1] + Size[1];
 				Blocks[BlockIndex].Id = FGuid(0,0,0,BlockIndex);
 			}
 
diff --git a/Engine/Plugins/Experimental/Mutable/Source/CustomizableObjectEditor/Private/MuCOE/SMutableObjectViewer.cpp b/Engine/Plugins/Experimental/Mutable/Source/CustomizableObjectEditor/Private/MuCOE/SMutableObjectViewer.cpp
--- a/Engine/Plugins/Experimental/Mutable/Source/CustomizableObjectEditor/Private/MuCOE/SMutableObjectViewer.cpp
+++ b/Engine/Plugins/Experimental/Mutable/Source/CustomizableObjectEditor/Private/MuCOE/SMutableObjectViewer.cpp
@@ -12,6 +12,7 @@
 #include "MuCOE/SMutableGraphViewer.h"
 #include "Widgets/Docking/SDockTab.h"
 #include "Widgets/Input/STextComboBox.h"
+#include "Widgets/Input/SNumericDropDown.h"
 #include "Widgets/Views/STreeView.h"
 
 class FExtender;
@@ -369,6 +370,22 @@ TSharedRef<SWidget> SMutableObjectViewer::GenerateCompileOptionsMenuContent()
 			;
 		MenuBuilder.AddWidget(CompileOptimizationCombo.ToSharedRef(), LOCTEXT("MutableCompileOptimizationLevel", "Optimization Level"));
 
+		// Image tiling
+		// Unfortunately SNumericDropDown doesn't work with integers at the time of writing.
+		TArray<SNumericDropDown<float>::FNamedValue> TilingOptions;
+		TilingOptions.Add(SNumericDropDown<float>::FNamedValue(0, FText::FromString(TEXT("0")), FText::FromString(TEXT("Disabled"))));
+		TilingOptions.Add(SNumericDropDown<float>::FNamedValue(64, FText::FromString(TEXT("64")), FText::FromString(TEXT("64"))));
+		TilingOptions.Add(SNumericDropDown<float>::FNamedValue(128, FText::FromString(TEXT("128")), FText::FromString(TEXT("128"))));
+		TilingOptions.Add(SNumericDropDown<float>::FNamedValue(256, FText::FromString(TEXT("256")), FText::FromString(TEXT("256"))));
+		TilingOptions.Add(SNumericDropDown<float>::FNamedValue(512, FText::FromString(TEXT("512")), FText::FromString(TEXT("512"))));
+
+		CompileTilingCombo = SNew(SNumericDropDown<float>)
+			.DropDownValues(TilingOptions)
+			.Value_Lambda( [&]() { return float(CompileOptions.ImageTiling); })
+			.OnValueChanged_Lambda( [&](float Value) { CompileOptions.ImageTiling = int32(Value); } )
+			;
+		MenuBuilder.AddWidget(CompileTilingCombo.ToSharedRef(), LOCTEXT("MutableCompileImageTiling", "Image Tiling"));
+
 		// Parallel compilation
 		MenuBuilder.AddMenuEntry(
 			LOCTEXT("Generate_MutableEnableParallelCompilation", "Enable compiling in multiple threads."),
diff --git a/Engine/Plugins/Experimental/Mutable/Source/CustomizableObjectEditor/Private/MuCOE/SMutableObjectViewer.h b/Engine/Plugins/Experimental/Mutable/Source/CustomizableObjectEditor/Private/MuCOE/SMutableObjectViewer.h
--- a/Engine/Plugins/Experimental/Mutable/Source/CustomizableObjectEditor/Private/MuCOE/SMutableObjectViewer.h
+++ b/Engine/Plugins/Experimental/Mutable/Source/CustomizableObjectEditor/Private/MuCOE/SMutableObjectViewer.h
@@ -4,6 +4,7 @@
 
 #include "MuCOE/CustomizableObjectCompiler.h"
 #include "Widgets/SCompoundWidget.h"
+#include "Widgets/Input/SNumericDropDown.h"
 
 class STableViewBase;
 namespace ESelectInfo { enum Type : int; }
@@ -67,6 +68,7 @@ private:
 	TSharedRef<SWidget> GenerateCompileOptionsMenuContent();
 	TSharedPtr<STextComboBox> CompileOptimizationCombo;
 	TArray< TSharedPtr<FString> > CompileOptimizationStrings;
+	TSharedPtr<SNumericDropDown<float>> CompileTilingCombo;
 	void OnChangeCompileOptimizationLevel(TSharedPtr<FString> NewSelection, ESelectInfo::Type SelectInfo);
 	void OnChangeDebugPlatform(TSharedPtr<FString> NewSelection, ESelectInfo::Type SelectInfo);
 
diff --git a/Engine/Plugins/Experimental/Mutable/Source/MutableRuntime/Private/MuR/CodeRunner.cpp b/Engine/Plugins/Experimental/Mutable/Source/MutableRuntime/Private/MuR/CodeRunner.cpp
--- a/Engine/Plugins/Experimental/Mutable/Source/MutableRuntime/Private/MuR/CodeRunner.cpp
+++ b/Engine/Plugins/Experimental/Mutable/Source/MutableRuntime/Private/MuR/CodeRunner.cpp
@@ -3671,19 +3671,27 @@ namespace mu
 
 				uint16 SizeX = args.size[0];
 				uint16 SizeY = args.size[1];
+				int32 LODs = args.LODs;
+				
+				// This means all the mip chain
+				if (LODs == 0)
+				{
+					LODs = FMath::CeilLogTwo(FMath::Max(SizeX,SizeY));
+				}
+
 				for (int l=0; l<item.ExecutionOptions; ++l)
 				{
 					SizeX = FMath::Max(uint16(1), FMath::DivideAndRoundUp(SizeX, uint16(2)));
 					SizeY = FMath::Max(uint16(1), FMath::DivideAndRoundUp(SizeY, uint16(2)));
+					--LODs;
 				}
 
-                ImagePtr pA = new Image( SizeX,
-                                         SizeY,
-                                         1,
-                                         (EImageFormat)args.format,
+                ImagePtr pA = new Image( SizeX, SizeY,
+                                         FMath::Max(LODs,1),
+                                         EImageFormat(args.format),
 										 EInitializationType::NotInitialized );
 
-                FillPlainColourImage(pA.get(), c);
+                pA->FillColour(c);
 
                 GetMemory().SetImage( item, pA );
                 break;
@@ -3751,16 +3759,19 @@ namespace mu
 
         case OP_TYPE::IM_PATCH:
         {
+			// TODO: This is optimized for memory-usage but base and patch could be requested at the same time
 			OP::ImagePatchArgs args = pModel->GetPrivate()->m_program.GetOpArgs<OP::ImagePatchArgs>(item.At);
             switch (item.Stage)
             {
-            case 0:
-                AddOp( FScheduledOp( item.At, item, 1),
-					FScheduledOp( args.base, item ),
-					FScheduledOp( args.patch, item ) );
-                break;
+			case 0:
+				AddOp(FScheduledOp(item.At, item, 1), FScheduledOp(args.base, item));
+				break;
 
-            case 1:
+			case 1:
+				AddOp(FScheduledOp(item.At, item, 2), FScheduledOp(args.patch, item));
+				break;
+
+			case 2:
             {
                 MUTABLE_CPUPROFILER_SCOPE(IM_PATCH_1)
 
@@ -3774,35 +3785,19 @@ namespace mu
 					break;
 				}
 
-                box<vec2<int>> rect;
-                rect.min[0] = args.minX;
-                rect.min[1] = args.minY;
-                rect.size[0] = pB->GetSizeX();
-                rect.size[1] = pB->GetSizeY();
-
-				// Apply ther mipmap reduction to the crop rectangle.
+				// Apply the mipmap reduction to the crop rectangle.
 				int32 MipsToSkip = item.ExecutionOptions;
-				while (MipsToSkip > 0 && rect.size[0] > 0 && rect.size[1] > 0)
-				{
-					rect.min/=2;
-					MipsToSkip--;
-				}
+				box<UE::Math::TIntVector2<uint16>> rect;
+				rect.min[0] = args.minX / (1 << MipsToSkip);
+				rect.min[1] = args.minY / (1 << MipsToSkip);
+				rect.size[0] = pB->GetSizeX();
+				rect.size[1] = pB->GetSizeY();
 
-                ImagePtr pResult = pA->Clone();
+                ImagePtr pResult = mu::CloneOrTakeOver(pA.get());
 
 				bool bApplyPatch = !rect.IsEmpty();
 				if (bApplyPatch)
 				{
-					// Resize image if it doesn't fit in the new block size
-					if (pB->GetSizeX() != rect.size[0] ||
-						pB->GetSizeY() != rect.size[1])
-					{
-						MUTABLE_CPUPROFILER_SCOPE(ImagePatchResize_Emergency);
-
-						FImageSize blockSize((uint16)rect.size[0], (uint16)rect.size[1]);
-						pB = ImageResizeLinear(m_pSettings->GetPrivate()->m_imageCompressionQuality,pB.get(), blockSize);
-					}
-
 					// Change the block image format if it doesn't match the composed image
 					// This is usually enforced at object compilation time.
 					if (pResult->GetFormat() != pB->GetFormat())
@@ -3880,24 +3875,35 @@ namespace mu
 						FScheduledOp::FromOpAndOptions(args.mesh, item, 0));
 				}
                 break;
+
 			case 1:
 			{
 				MUTABLE_CPUPROFILER_SCOPE(IM_RASTERMESH_1)
 
+				Ptr<const Mesh> pMesh = GetMemory().GetMesh(FScheduledOp::FromOpAndOptions(args.mesh, item, 0));
+
 				// If no image, we are generating a flat mesh UV raster. This is the final stage in this case.
 				if (!args.image)
 				{
-					Ptr<const Mesh> pMesh = GetMemory().GetMesh(FScheduledOp::FromOpAndOptions(args.mesh, item, 0));
-
 					uint16 SizeX = args.sizeX;
 					uint16 SizeY = args.sizeY;
+					UE::Math::TIntVector2<uint16> CropMin(args.CropMinX, args.CropMinY);
+					UE::Math::TIntVector2<uint16> UncroppedSize(args.UncroppedSizeX, args.UncroppedSizeY);
 
 					// Drop mips while possible
 					int32 MipsToDrop = item.ExecutionOptions;
+					bool bUseCrop = UncroppedSize[0] > 0;
 					while (MipsToDrop && !(SizeX % 2) && !(SizeY % 2))
 					{
-						SizeX = FMath::Max(uint16(1),FMath::DivideAndRoundUp(SizeX, uint16(2)));
-						SizeY = FMath::Max(uint16(1),FMath::DivideAndRoundUp(SizeY, uint16(2)));
+						SizeX = FMath::Max(uint16(1), FMath::DivideAndRoundUp(SizeX, uint16(2)));
+						SizeY = FMath::Max(uint16(1), FMath::DivideAndRoundUp(SizeY, uint16(2)));
+						if (bUseCrop)
+						{
+							CropMin[0] = FMath::DivideAndRoundUp(CropMin[0], uint16(2));
+							CropMin[1] = FMath::DivideAndRoundUp(CropMin[1], uint16(2));
+							UncroppedSize[0] = FMath::Max(uint16(1), FMath::DivideAndRoundUp(UncroppedSize[0], uint16(2)));
+							UncroppedSize[1] = FMath::Max(uint16(1), FMath::DivideAndRoundUp(UncroppedSize[1], uint16(2)));
+						}
 						--MipsToDrop;
 					}
 
@@ -3905,7 +3911,7 @@ namespace mu
 					Ptr<Image> ResultImage = new Image(SizeX, SizeY, 1, EImageFormat::IF_L_UBYTE);
 					if (pMesh)
 					{
-						ImageRasterMesh(pMesh.get(), ResultImage.get(), args.blockIndex);
+						ImageRasterMesh(pMesh.get(), ResultImage.get(), args.blockIndex, CropMin, UncroppedSize);
 					}
 
 					// Stop execution.
@@ -3913,8 +3919,6 @@ namespace mu
 					break;
 				}
 
-				Ptr<const Mesh> pMesh = GetMemory().GetMesh(FScheduledOp::FromOpAndOptions(args.mesh, item, 0));
-
 				const int32 MipsToSkip = item.ExecutionOptions;
 				int32 ProjectionMip = MipsToSkip;
 
@@ -3950,6 +3954,7 @@ namespace mu
 
 				break;
 			}
+
             case 2:
             {
 				MUTABLE_CPUPROFILER_SCOPE(IM_RASTERMESH_2)
@@ -3964,27 +3969,34 @@ namespace mu
 
                 Ptr<const Mesh> pMesh = GetMemory().GetMesh(FScheduledOp::FromOpAndOptions(args.mesh, item, 0));
 
+				if (!pMesh)
+				{
+					check(false);
+					GetMemory().SetImage(item, nullptr);
+					break;
+				}
+
 				uint16 SizeX = args.sizeX;
 				uint16 SizeY = args.sizeY;
+				UE::Math::TIntVector2<uint16> CropMin(args.CropMinX, args.CropMinY);
+				UE::Math::TIntVector2<uint16> UncroppedSize(args.UncroppedSizeX, args.UncroppedSizeY);
 
 				// Drop mips while possible
 				int32 MipsToDrop = item.ExecutionOptions;
+				bool bUseCrop = UncroppedSize[0] > 0;
 				while (MipsToDrop && !(SizeX % 2) && !(SizeY % 2))
 				{
 					SizeX = FMath::Max(uint16(1),FMath::DivideAndRoundUp(SizeX, uint16(2)));
 					SizeY = FMath::Max(uint16(1),FMath::DivideAndRoundUp(SizeY, uint16(2)));
+					if (bUseCrop)
+					{
+						CropMin[0] = FMath::DivideAndRoundUp(CropMin[0], uint16(2));
+						CropMin[1] = FMath::DivideAndRoundUp(CropMin[1], uint16(2));
+						UncroppedSize[0] = FMath::Max(uint16(1), FMath::DivideAndRoundUp(UncroppedSize[0], uint16(2)));
+						UncroppedSize[1] = FMath::Max(uint16(1), FMath::DivideAndRoundUp(UncroppedSize[1], uint16(2)));
+					}
 					--MipsToDrop;
 				}
-	
-				if (!pMesh)
-				{
-                    check(false);
-                    Ptr<Image> ResultImage = new Image(SizeX, SizeY, 1, EImageFormat::IF_L_UBYTE);
-
-					GetMemory().SetImage(item, ResultImage);
-
-					break;
-				}
 
 				// Raster with projection
 				Ptr<const Image> pSource = GetMemory().GetImage(
@@ -4049,6 +4061,7 @@ namespace mu
 								SamplingMethod,
 								FadeStartRad, FadeEndRad,
 								layout, args.blockIndex,
+								CropMin, UncroppedSize,
 								&scratch);
 							break;
 
@@ -4059,6 +4072,7 @@ namespace mu
 								SamplingMethod,
 								FadeStartRad, FadeEndRad,
 								layout, args.blockIndex,
+								CropMin, UncroppedSize,
 								&scratch);
 							break;
 
@@ -4069,6 +4083,7 @@ namespace mu
 								FadeStartRad, FadeEndRad,
 								layout,
 								pProjector->m_value.projectionAngle,
+								CropMin, UncroppedSize,
 								&scratch);
 							break;
 
@@ -4895,7 +4910,7 @@ namespace mu
 
                     if( relBlockIndex >=0 )
                     {
-                        box< vec2<int> > rectInblocks;
+                        box< UE::Math::TIntVector2<uint16> > rectInblocks;
                         pLayout->GetBlock
                                 (
                                     relBlockIndex,
@@ -5935,7 +5950,7 @@ namespace mu
 			OP::ImagePlainColourArgs args = program.GetOpArgs<OP::ImagePlainColourArgs>(item.At);
 			m_heapImageDesc[item.CustomState].m_size[0] = args.size[0];
 			m_heapImageDesc[item.CustomState].m_size[1] = args.size[1];
-			m_heapImageDesc[item.CustomState].m_lods = 1;
+			m_heapImageDesc[item.CustomState].m_lods = args.LODs;
 			m_heapImageDesc[item.CustomState].m_format = args.format;
 			GetMemory().SetValidDesc(item);
 			break;
diff --git a/Engine/Plugins/Experimental/Mutable/Source/MutableRuntime/Private/MuR/CodeRunnerTasked.cpp b/Engine/Plugins/Experimental/Mutable/Source/MutableRuntime/Private/MuR/CodeRunnerTasked.cpp
--- a/Engine/Plugins/Experimental/Mutable/Source/MutableRuntime/Private/MuR/CodeRunnerTasked.cpp
+++ b/Engine/Plugins/Experimental/Mutable/Source/MutableRuntime/Private/MuR/CodeRunnerTasked.cpp
@@ -1111,7 +1111,7 @@ namespace mu
 					case EBlendType::BT_OVERLAY: BufferLayerColour<OverlayChannel>(pNew.get(), m_base.get(), m_col); break;
 					case EBlendType::BT_LIGHTEN: BufferLayerColour<LightenChannel>(pNew.get(), m_base.get(), m_col); break;
 					case EBlendType::BT_MULTIPLY: BufferLayerColour<MultiplyChannel>(pNew.get(), m_base.get(), m_col); break;
-					case EBlendType::BT_BLEND: FillPlainColourImage(pNew.get(), m_col); break;
+					case EBlendType::BT_BLEND: pNew->FillColour(m_col); break;
 					default: check(false);
 					}
 				}
@@ -1778,7 +1778,7 @@ namespace mu
 			ImagePtr pComposed = CloneOrTakeOver<>(m_base.get());
 			pComposed->m_flags = 0;
 
-			box< vec2<int> > rectInblocks;
+			box< UE::Math::TIntVector2<uint16> > rectInblocks;
 			m_layout->GetBlock
 			(
 				relBlockIndex,
@@ -1791,7 +1791,7 @@ namespace mu
 
 			int blockSizeX = m_base->GetSizeX() / grid[0];
 			int blockSizeY = m_base->GetSizeY() / grid[1];
-			box< vec2<int> > rect = rectInblocks;
+			box< UE::Math::TIntVector2<uint16> > rect = rectInblocks;
 			rect.min[0] *= blockSizeX;
 			rect.min[1] *= blockSizeY;
 			rect.size[0] *= blockSizeX;
diff --git a/Engine/Plugins/Experimental/Mutable/Source/MutableRuntime/Private/MuR/Image.cpp b/Engine/Plugins/Experimental/Mutable/Source/MutableRuntime/Private/MuR/Image.cpp
--- a/Engine/Plugins/Experimental/Mutable/Source/MutableRuntime/Private/MuR/Image.cpp
+++ b/Engine/Plugins/Experimental/Mutable/Source/MutableRuntime/Private/MuR/Image.cpp
@@ -1115,5 +1115,115 @@ namespace mu
 	}
 
 
+	//---------------------------------------------------------------------------------------------
+	void Image::FillColour(FVector4f c)
+	{
+		EImageFormat Format = GetFormat();
+		switch (Format)
+		{
+		case EImageFormat::IF_RGB_UBYTE:
+		{
+			// TODO: Optimize: don't write bytes one by one
+			int pixelCount = CalculatePixelCount();
+			uint8* pData = GetData();
+			uint8 r = uint8(FMath::Clamp(255.0f * c[0], 0.0f, 255.0f));
+			uint8 g = uint8(FMath::Clamp(255.0f * c[1], 0.0f, 255.0f));
+			uint8 b = uint8(FMath::Clamp(255.0f * c[2], 0.0f, 255.0f));
+			for (int p = 0; p < pixelCount; ++p)
+			{
+				pData[0] = r;
+				pData[1] = g;
+				pData[2] = b;
+				pData += 3;
+			}
+			break;
+		}
+
+		case EImageFormat::IF_RGBA_UBYTE:
+		{
+			// TODO: Optimize: don't write bytes one by one
+			int pixelCount = CalculatePixelCount();
+			uint8* pData = GetData();
+			uint8 r = uint8(FMath::Clamp(255.0f * c[0], 0.0f, 255.0f));
+			uint8 g = uint8(FMath::Clamp(255.0f * c[1], 0.0f, 255.0f));
+			uint8 b = uint8(FMath::Clamp(255.0f * c[2], 0.0f, 255.0f));
+			uint8 a = uint8(FMath::Clamp(255.0f * c[3], 0.0f, 255.0f));
+			for (int p = 0; p < pixelCount; ++p)
+			{
+				pData[0] = r;
+				pData[1] = g;
+				pData[2] = b;
+				pData[3] = a;
+				pData += 4;
+			}
+			break;
+		}
+
+		case EImageFormat::IF_BGRA_UBYTE:
+		{
+			// TODO: Optimize: don't write bytes one by one
+			int pixelCount = CalculatePixelCount();
+			uint8* pData = GetData();
+			uint8 r = uint8(FMath::Clamp(255.0f * c[0], 0.0f, 255.0f));
+			uint8 g = uint8(FMath::Clamp(255.0f * c[1], 0.0f, 255.0f));
+			uint8 b = uint8(FMath::Clamp(255.0f * c[2], 0.0f, 255.0f));
+			uint8 a = uint8(FMath::Clamp(255.0f * c[3], 0.0f, 255.0f));
+			for (int p = 0; p < pixelCount; ++p)
+			{
+				pData[0] = b;
+				pData[1] = g;
+				pData[2] = r;
+				pData[3] = a;
+				pData += 4;
+			}
+			break;
+		}
+
+		case EImageFormat::IF_L_UBYTE:
+		{
+			int pixelCount = CalculatePixelCount();
+			uint8* pData = GetData();
+			uint8 v = FMath::Min<uint8>(255, FMath::Max<uint8>(0, uint8(255.0f * c[0])));
+			FMemory::Memset(pData, v, pixelCount);
+			break;
+		}
+
+		default:
+		{
+			// Generic case that supports compressed formats.
+			const FImageFormatData& FormatData = GetImageFormatData(Format);
+			Ptr<Image> Block = new Image(FormatData.m_pixelsPerBlockX, FormatData.m_pixelsPerBlockY, 1, EImageFormat::IF_RGBA_UBYTE, EInitializationType::NotInitialized);
+
+			uint32 Pixel = (uint32(FMath::Clamp(255.0f * c[0], 0.0f, 255.0f)) << 0)
+				| (uint32(FMath::Clamp(255.0f * c[1], 0.0f, 255.0f)) << 8)
+				| (uint32(FMath::Clamp(255.0f * c[2], 0.0f, 255.0f)) << 16)
+				| (uint32(FMath::Clamp(255.0f * c[3], 0.0f, 255.0f)) << 24);
+			uint32* BlockData = reinterpret_cast<uint32*>(Block->GetData());
+			for (int32 I = 0; I < FormatData.m_pixelsPerBlockX * FormatData.m_pixelsPerBlockY; ++I)
+			{
+				*BlockData = Pixel;
+				BlockData += 1;
+			}
+			Ptr<Image> Converted = ImagePixelFormat(0, Block.get(), Format);
+
+			int32 DataSize = CalculateDataSize();
+			check(DataSize % FormatData.m_bytesPerBlock == 0);
+			int32 BlockCount = DataSize / FormatData.m_bytesPerBlock;
+			uint8* TargetData = GetData();
+			while (BlockCount)
+			{
+				FMemory::Memcpy(TargetData, Converted->GetData(), FormatData.m_bytesPerBlock);
+				TargetData += FormatData.m_bytesPerBlock;
+				--BlockCount;
+			}
+			break;
+		}
+
+
+		}
+
+	}
+
+
 }
 
diff --git a/Engine/Plugins/Experimental/Mutable/Source/MutableRuntime/Private/MuR/ImagePrivate.h b/Engine/Plugins/Experimental/Mutable/Source/MutableRuntime/Private/MuR/ImagePrivate.h
--- a/Engine/Plugins/Experimental/Mutable/Source/MutableRuntime/Private/MuR/ImagePrivate.h
+++ b/Engine/Plugins/Experimental/Mutable/Source/MutableRuntime/Private/MuR/ImagePrivate.h
@@ -189,82 +189,6 @@ namespace mu
     }
 
 
-    //---------------------------------------------------------------------------------------------
-    //!
-    //---------------------------------------------------------------------------------------------
-    inline void FillPlainColourImage( Image* pImage, FVector4f c )
-    {
-        int pixelCount = pImage->GetSizeX() * pImage->GetSizeY();
-
-        switch (pImage->GetFormat())
-        {
-		case EImageFormat::IF_RGB_UBYTE:
-        {
-            uint8* pData = pImage->GetData();
-            uint8 r = uint8(FMath::Clamp(255.0f * c[0], 0.0f, 255.0f));
-            uint8 g = uint8(FMath::Clamp(255.0f * c[1], 0.0f, 255.0f));
-            uint8 b = uint8(FMath::Clamp(255.0f * c[2], 0.0f, 255.0f));
-            for ( int p = 0; p < pixelCount; ++p )
-            {
-                pData[0] = r;
-                pData[1] = g;
-                pData[2] = b;
-                pData += 3;
-            }
-            break;
-        }
-
-        case EImageFormat::IF_RGBA_UBYTE:
-        {
-            uint8* pData = pImage->GetData();
-			uint8 r = uint8(FMath::Clamp(255.0f * c[0], 0.0f, 255.0f));
-			uint8 g = uint8(FMath::Clamp(255.0f * c[1], 0.0f, 255.0f));
-			uint8 b = uint8(FMath::Clamp(255.0f * c[2], 0.0f, 255.0f));
-			uint8 a = uint8(FMath::Clamp(255.0f * c[3], 0.0f, 255.0f));
-			for ( int p = 0; p < pixelCount; ++p )
-            {
-                pData[0] = r;
-                pData[1] = g;
-                pData[2] = b;
-                pData[3] = a;
-                pData += 4;
-            }
-            break;
-        }
-
-        case EImageFormat::IF_BGRA_UBYTE:
-        {
-            uint8* pData = pImage->GetData();
-			uint8 r = uint8(FMath::Clamp(255.0f * c[0], 0.0f, 255.0f));
-			uint8 g = uint8(FMath::Clamp(255.0f * c[1], 0.0f, 255.0f));
-			uint8 b = uint8(FMath::Clamp(255.0f * c[2], 0.0f, 255.0f));
-			uint8 a = uint8(FMath::Clamp(255.0f * c[3], 0.0f, 255.0f));
-			for ( int p = 0; p < pixelCount; ++p )
-            {
-                pData[0] = b;
-                pData[1] = g;
-                pData[2] = r;
-                pData[3] = a;
-                pData += 4;
-            }
-            break;
-        }
-
-        case EImageFormat::IF_L_UBYTE:
-        {
-            uint8* pData = pImage->GetData();
-            uint8 v = FMath::Min<uint8>( 255, FMath::Max<uint8>( 0, uint8( 255.0f * c[0] ) ) );
-            FMemory::Memset( pData, v, pixelCount );
-            break;
-        }
-
-        default:
-            check(false);
-        }
-
-    }
-
-
     //---------------------------------------------------------------------------------------------
     //! Use with care.
     //---------------------------------------------------------------------------------------------
diff --git a/Engine/Plugins/Experimental/Mutable/Source/MutableRuntime/Private/MuR/Layout.cpp b/Engine/Plugins/Experimental/Mutable/Source/MutableRuntime/Private/MuR/Layout.cpp
--- a/Engine/Plugins/Experimental/Mutable/Source/MutableRuntime/Private/MuR/Layout.cpp
+++ b/Engine/Plugins/Experimental/Mutable/Source/MutableRuntime/Private/MuR/Layout.cpp
@@ -116,7 +116,7 @@ namespace mu {
 
 
 	//---------------------------------------------------------------------------------------------
-	void Layout::GetBlock( int index, int* pMinX, int* pMinY, int* pSizeX, int* pSizeY ) const
+	void Layout::GetBlock( int index, uint16* pMinX, uint16* pMinY, uint16* pSizeX, uint16* pSizeY ) const
 	{
 		check( index >=0 && index < m_blocks.Num() );
 		check( pMinX && pMinY && pSizeX && pSizeY );
diff --git a/Engine/Plugins/Experimental/Mutable/Source/MutableRuntime/Private/MuR/OpImageCompose.h b/Engine/Plugins/Experimental/Mutable/Source/MutableRuntime/Private/MuR/OpImageCompose.h
--- a/Engine/Plugins/Experimental/Mutable/Source/MutableRuntime/Private/MuR/OpImageCompose.h
+++ b/Engine/Plugins/Experimental/Mutable/Source/MutableRuntime/Private/MuR/OpImageCompose.h
@@ -10,7 +10,7 @@
 namespace mu
 {
 
-	inline void ImageCompose( Image* pBase, const Image* pBlock, const box< vec2<int> >& rect )
+	inline void ImageCompose( Image* pBase, const Image* pBlock, const box< UE::Math::TIntVector2<uint16> >& rect )
 	{
 		check(pBase && pBlock);
 		check(pBase != pBlock);
@@ -55,15 +55,20 @@ namespace mu
 		else
 		{
 			// Sizes in blocks of the current mips
-			vec2<int> baseMipSize(pBase->GetSizeX() / finfo.m_pixelsPerBlockX,
+			UE::Math::TIntVector2<uint16> baseMipSize(
+				pBase->GetSizeX() / finfo.m_pixelsPerBlockX,
 				pBase->GetSizeY() / finfo.m_pixelsPerBlockY);
-			vec2<int> blockMipPos(rect.min[0] / finfo.m_pixelsPerBlockX,
+
+			UE::Math::TIntVector2<uint16> blockMipPos(
+				rect.min[0] / finfo.m_pixelsPerBlockX,
 				rect.min[1] / finfo.m_pixelsPerBlockY);
-			vec2<int> blockMipSize(rect.size[0] / finfo.m_pixelsPerBlockX,
+
+			UE::Math::TIntVector2<uint16> blockMipSize(
+				rect.size[0] / finfo.m_pixelsPerBlockX,
 				rect.size[1] / finfo.m_pixelsPerBlockY);
 
 			int doneMips = 0;
-			for (doneMips = 0; doneMips < pBase->GetLODCount() && doneMips < pBlock->GetLODCount(); ++doneMips)
+			for (; doneMips < pBase->GetLODCount() && doneMips < pBlock->GetLODCount(); ++doneMips)
 			{
 				//UE_LOG(LogMutableCore, Warning, "Block Mip Pos : %d, %d", blockMipPos[0], blockMipPos[1]);
 
@@ -90,15 +95,16 @@ namespace mu
 				// pixel format block size.
 				if (blockMipPos[0] % 2 || blockMipPos[1] % 2 || blockMipSize[0] % 2 || blockMipSize[1] % 2)
 				{
+					++doneMips;
 					break;
 				}
 
-				baseMipSize[0] = FMath::DivideAndRoundUp(baseMipSize[0], 2);
-				baseMipSize[1] = FMath::DivideAndRoundUp(baseMipSize[1], 2);
-				blockMipPos[0] = FMath::DivideAndRoundUp(blockMipPos[0], 2);
-				blockMipPos[1] = FMath::DivideAndRoundUp(blockMipPos[1], 2);
-				blockMipSize[0] = FMath::DivideAndRoundUp(blockMipSize[0], 2);
-				blockMipSize[1] = FMath::DivideAndRoundUp(blockMipSize[1], 2);
+				baseMipSize[0] = FMath::DivideAndRoundUp(baseMipSize[0], uint16(2));
+				baseMipSize[1] = FMath::DivideAndRoundUp(baseMipSize[1], uint16(2));
+				blockMipPos[0] = FMath::DivideAndRoundUp(blockMipPos[0], uint16(2));
+				blockMipPos[1] = FMath::DivideAndRoundUp(blockMipPos[1], uint16(2));
+				blockMipSize[0] = FMath::DivideAndRoundUp(blockMipSize[0], uint16(2));
+				blockMipSize[1] = FMath::DivideAndRoundUp(blockMipSize[1], uint16(2));
 			}
 
 			// Adjust the actually valid mips, maybe wasting some of the reserved memory.
diff --git a/Engine/Plugins/Experimental/Mutable/Source/MutableRuntime/Private/MuR/OpImageLayer.h b/Engine/Plugins/Experimental/Mutable/Source/MutableRuntime/Private/MuR/OpImageLayer.h
--- a/Engine/Plugins/Experimental/Mutable/Source/MutableRuntime/Private/MuR/OpImageLayer.h
+++ b/Engine/Plugins/Experimental/Mutable/Source/MutableRuntime/Private/MuR/OpImageLayer.h
@@ -2021,7 +2021,7 @@ namespace mu
 	inline void ImageLayerOnBaseNoAlpha( Image* pBase,
 										 const Image* pMask,
 										 const Image* pBlended,
-										 const box< vec2<int> >& rect )
+										 const box<UE::Math::TIntVector2<uint16>>& rect )
 	{
 		check( pBase->GetSizeX() >= rect.min[0]+rect.size[0] );
 		check( pBase->GetSizeY() >= rect.min[1]+rect.size[1] );
@@ -2663,7 +2663,7 @@ namespace mu
 	inline void ImageBlendOnBaseNoAlpha(Image* pBase,
 		const Image* pMask,
 		const Image* pBlended,
-		const box< vec2<int> >& rect)
+		const box< UE::Math::TIntVector2<uint16> >& rect)
 	{
 		ImageLayerOnBaseNoAlpha< BlendChannel, false>
 			(pBase, pMask, pBlended, rect);
diff --git a/Engine/Plugins/Experimental/Mutable/Source/MutableRuntime/Private/MuR/OpImageProject.cpp b/Engine/Plugins/Experimental/Mutable/Source/MutableRuntime/Private/MuR/OpImageProject.cpp
--- a/Engine/Plugins/Experimental/Mutable/Source/MutableRuntime/Private/MuR/OpImageProject.cpp
+++ b/Engine/Plugins/Experimental/Mutable/Source/MutableRuntime/Private/MuR/OpImageProject.cpp
@@ -486,163 +486,6 @@ namespace
 
 
 
-//-------------------------------------------------------------------------------------------------
-//void ImageRasterProjected_Generic( const Mesh* pMesh,
-//                           Image* pImage,
-//                           const Image* pSource,
-//                           const Image* pMask,
-//                           float fadeStart,
-//                           float fadeEnd,
-//                           int layout,
-//                           int block )
-//{
-//	MUTABLE_CPUPROFILER_SCOPE(ImageRasterProjected_Generic);
-//
-//	EImageFormat format = pImage->GetFormat();
-//    int pixelSize = GetImageFormatData( format ).m_bytesPerBlock;
-//
-//    int sizeX = pImage->GetSizeX();
-//    int sizeY = pImage->GetSizeY();
-//
-//    RasterProjectedPixelProcessor pixelProc( pSource,
-//                                     pImage->GetData(),
-//                                     pMask ? pMask->GetData() : 0,
-//                                     fadeStart, fadeEnd );
-//
-//    // Get the vertices
-//    int vertexCount = pMesh->GetVertexCount();
-//	TArray< RasterVertex<4> > vertices;
-//	vertices.SetNum(vertexCount);
-//
-//    UntypedMeshBufferIteratorConst texIt( pMesh->GetVertexBuffers(), MBS_TEXCOORDS, layout );
-//    UntypedMeshBufferIteratorConst posIt( pMesh->GetVertexBuffers(), MBS_POSITION, 0 );
-//    UntypedMeshBufferIteratorConst norIt( pMesh->GetVertexBuffers(), MBS_NORMAL, 0 );
-//    for ( int v=0; v<vertexCount; ++v )
-//    {
-//        float uv[2]={0.0f,0.0f};
-//        ConvertData( 0, uv, MBF_FLOAT32, texIt.ptr(), texIt.GetFormat() );
-//        ConvertData( 1, uv, MBF_FLOAT32, texIt.ptr(), texIt.GetFormat() );
-//
-//        vertices[v].x = uv[0] * sizeX;
-//        vertices[v].y = uv[1] * sizeY;
-//
-//        FVector3f pos={0.0f,0.0f,0.0f};
-//        ConvertData( 0, &pos[0], MBF_FLOAT32, posIt.ptr(), posIt.GetFormat() );
-//        ConvertData( 1, &pos[0], MBF_FLOAT32, posIt.ptr(), posIt.GetFormat() );
-//        ConvertData( 2, &pos[0], MBF_FLOAT32, posIt.ptr(), posIt.GetFormat() );
-//        vertices[v].interpolators[0] = pos[0];
-//        vertices[v].interpolators[1] = pos[1];
-//        vertices[v].interpolators[2] = pos[2];
-//
-//        float nor=0.0f;
-//        ConvertData( 0, &nor, MBF_FLOAT32, norIt.ptr(), norIt.GetFormat() );
-//        vertices[v].interpolators[3] = nor;
-//
-//        ++texIt;
-//        ++posIt;
-//        ++norIt;
-//    }
-//
-//    // Get the indices
-//    int faceCount = pMesh->GetFaceCount();
-//	TArray<int> indices;
-//	indices.SetNum(faceCount * 3);
-//
-//    UntypedMeshBufferIteratorConst indIt( pMesh->GetIndexBuffers(), MBS_VERTEXINDEX, 0 );
-//    for ( int i=0; i<faceCount*3; ++i )
-//    {
-//        uint32 index=0;
-//        ConvertData( 0, &index, MBF_UINT32, indIt.ptr(), indIt.GetFormat() );
-//
-//        indices[i] = index;
-//        ++indIt;
-//    }
-//
-//    // Get the block per face
-//	TArray<int> blocks;
-//	blocks.SetNum(vertexCount);
-//
-//    UntypedMeshBufferIteratorConst bloIt( pMesh->GetVertexBuffers(), MBS_LAYOUTBLOCK, layout );
-//    if (bloIt.ptr())
-//    {
-//        for ( int i=0; i<vertexCount; ++i )
-//        {
-//            uint16 index=0;
-//            ConvertData( 0, &index, MBF_UINT16, bloIt.ptr(), bloIt.GetFormat() );
-//
-//            blocks[i] = index;
-//            ++bloIt;
-//        }
-//    }
-//
-//    if ( block<0 )
-//    {
-//        // We are rastering the entire layout in an image, so we need to apply the block->layout
-//        // transform to the UVs
-//        const Layout* pLayout = layout<pMesh->GetLayoutCount() ? pMesh->GetLayout( layout ) : nullptr;
-//        if (pLayout && bloIt.ptr())
-//        {
-//			TArray< box< vec2<float> > > transforms;
-//			transforms.SetNum(pLayout->GetBlockCount());
-//            for ( int b=0; b<pLayout->GetBlockCount(); ++b )
-//            {
-//                FIntPoint grid = pLayout->GetGridSize();
-//
-//                box< vec2<int> > blockRect;
-//                pLayout->GetBlock( b, &blockRect.min[0], &blockRect.min[1], &blockRect.size[0], &blockRect.size[1] );
-//
-//                box< vec2<float> > rect;
-//                rect.min[0] = ( (float)blockRect.min[0] ) / (float) grid[0];
-//                rect.min[1] = ( (float)blockRect.min[1] ) / (float) grid[1];
-//                rect.size[0] = ( (float)blockRect.size[0] ) / (float) grid[0];
-//                rect.size[1] = ( (float)blockRect.size[1] ) / (float) grid[1];
-//                transforms[b] = rect;
-//            }
-//
-//            for ( int i=0; i<vertexCount; ++i )
-//            {
-//                int relBlock = pLayout->FindBlock( blocks[i] );
-//                vertices[i].x = vertices[i].x * transforms[ relBlock ].size[0]
-//                        + transforms[ relBlock ].min[0] * sizeX;
-//                vertices[i].y = vertices[i].y * transforms[ relBlock ].size[1]
-//                        + transforms[ relBlock ].min[1] * sizeY;
-//            }
-//        }
-//
-//        // Raster all the faces
-//        for ( int f=0; f<faceCount; ++f )
-//        {
-//            Triangle( pImage->GetData(),
-//                      sizeX, sizeY,
-//                      pixelSize,
-//                      vertices[indices[f*3+0]],
-//                      vertices[indices[f*3+1]],
-//                      vertices[indices[f*3+2]],
-//                      pixelProc,
-//                      false );
-//        }
-//    }
-//    else
-//    {
-//        // Raster only the faces in the selected block
-//        for ( int f=0; f<faceCount; ++f )
-//        {
-//            if ( blocks[ indices[f*3+0] ]==block )
-//            {
-//                Triangle( pImage->GetData(),
-//                          sizeX, sizeY,
-//                          pixelSize,
-//                          vertices[indices[f*3+0]],
-//                          vertices[indices[f*3+1]],
-//                          vertices[indices[f*3+2]],
-//                          pixelProc,
-//                          false );
-//            }
-//        }
-//    }
-//}
-//
-
 //-------------------------------------------------------------------------------------------------
 //! This format is the one we assumee the meshes optimised for planar and cylindrical projection
 //! will have.
@@ -673,11 +516,10 @@ static_assert( sizeof(OPTIMISED_VERTEX_WRAPPING)==36, "UNEXPECTED_STRUCT_SIZE" )
 
 //-------------------------------------------------------------------------------------------------
 template<class PIXEL_PROCESSOR>
-void ImageRasterProjected_Optimised( const Mesh* pMesh,
-                             Image* pImage,
-                             PIXEL_PROCESSOR& pixelProc,
-                             float fadeEnd,
-                             SCRATCH_IMAGE_PROJECT* scratch )
+void ImageRasterProjected_Optimised( const Mesh* pMesh, Image* pImage,
+	PIXEL_PROCESSOR& pixelProc, float fadeEnd,
+	UE::Math::TIntVector2<uint16> CropMin, UE::Math::TIntVector2<uint16> UncroppedSize,
+	SCRATCH_IMAGE_PROJECT* scratch )
 {
 	MUTABLE_CPUPROFILER_SCOPE(ImageRasterProjected_Optimised);
 
@@ -692,6 +534,8 @@ void ImageRasterProjected_Optimised( const Mesh* pMesh,
     int32 sizeX = pImage->GetSizeX();
     int32 sizeY = pImage->GetSizeY();
 
+	bool bUseCropping = UncroppedSize[0] > 0;
+
     // Get the vertices
     int vertexCount = pMesh->GetVertexCount();
 
@@ -704,8 +548,16 @@ void ImageRasterProjected_Optimised( const Mesh* pMesh,
     float fadeEndCos = cosf(fadeEnd);
     for ( int v=0; v<vertexCount; ++v )
     {
-        scratch->vertices[v].x = pVertices[v].uv[0] * sizeX;
-        scratch->vertices[v].y = pVertices[v].uv[1] * sizeY;
+		if (bUseCropping)
+		{
+			scratch->vertices[v].x = pVertices[v].uv[0] * UncroppedSize[0] - CropMin[0];
+			scratch->vertices[v].y = pVertices[v].uv[1] * UncroppedSize[1] - CropMin[1];
+		}
+		else
+		{
+			scratch->vertices[v].x = pVertices[v].uv[0] * sizeX;
+			scratch->vertices[v].y = pVertices[v].uv[1] * sizeY;
+		}
 
         // TODO: No need to copy all. use scratch for the rest only.
         scratch->vertices[v].interpolators[0] = pVertices[v].pos[0];
@@ -747,6 +599,8 @@ void ImageRasterProjected_Optimised( const Mesh* pMesh,
 		});
 
 	// Update the relevancy map
+	// \TODO: fix it to work with cropping.
+	if (!bUseCropping)
 	{
 		MUTABLE_CPUPROFILER_SCOPE(ImageRasterProjected_Optimised_Relevancy);
 
@@ -782,12 +636,11 @@ void ImageRasterProjected_Optimised( const Mesh* pMesh,
 
 //-------------------------------------------------------------------------------------------------
 template<class PIXEL_PROCESSOR>
-void ImageRasterProjected_OptimisedWrapping( const Mesh* pMesh,
-                             Image* pImage,
-                             PIXEL_PROCESSOR& pixelProc,
-                             float fadeEnd,
-                             int block,
-                             SCRATCH_IMAGE_PROJECT* scratch )
+void ImageRasterProjected_OptimisedWrapping( const Mesh* pMesh, Image* pImage, PIXEL_PROCESSOR& pixelProc,
+	float fadeEnd,
+	int block,
+	UE::Math::TIntVector2<uint16> CropMin, UE::Math::TIntVector2<uint16> UncroppedSize,
+	SCRATCH_IMAGE_PROJECT* scratch )
 {
 	MUTABLE_CPUPROFILER_SCOPE(ImageRasterProjected_OptimisedWrapping);
 
@@ -814,8 +667,17 @@ void ImageRasterProjected_OptimisedWrapping( const Mesh* pMesh,
     float fadeEndCos = cosf(fadeEnd);
     for ( int v=0; v<vertexCount; ++v )
     {
-        scratch->vertices[v].x = pVertices[v].uv[0] * sizeX;
-        scratch->vertices[v].y = pVertices[v].uv[1] * sizeY;
+		bool bUseCropping = UncroppedSize[0] > 0;
+		if (bUseCropping)
+		{
+			scratch->vertices[v].x = pVertices[v].uv[0] * UncroppedSize[0] - CropMin[0];
+			scratch->vertices[v].y = pVertices[v].uv[1] * UncroppedSize[1] - CropMin[1];
+		}
+		else
+		{
+			scratch->vertices[v].x = pVertices[v].uv[0] * sizeX;
+			scratch->vertices[v].y = pVertices[v].uv[1] * sizeY;
+		}
 
         // TODO: No need to copy all. use scratch for the rest only.
         scratch->vertices[v].interpolators[0] = pVertices[v].pos[0];
@@ -851,13 +713,11 @@ void ImageRasterProjected_OptimisedWrapping( const Mesh* pMesh,
              !scratch->culledVertex[i2] )
         {
             Triangle( pImage->GetData(),
-                      sizeX, sizeY,
-                      pixelSize,
-                      scratch->vertices[i0],
-                      scratch->vertices[i1],
-                      scratch->vertices[i2],
-                      pixelProc,
-                      false );
+				sizeX, sizeY,
+				pixelSize,
+				scratch->vertices[i0], scratch->vertices[i1], scratch->vertices[i2],
+				pixelProc,
+				false );
         }
     }
 
@@ -873,6 +733,7 @@ void ImageRasterProjectedPlanar( const Mesh* pMesh,
 	ESamplingMethod SamplingMethod,
 	float fadeStart, float fadeEnd,
 	int layout, int block,
+	UE::Math::TIntVector2<uint16> CropMin, UE::Math::TIntVector2<uint16> UncroppedSize,
 	SCRATCH_IMAGE_PROJECT* scratch )
 {
     check( !pMask || pMask->GetSizeX() == pImage->GetSizeX() );
@@ -894,7 +755,7 @@ void ImageRasterProjectedPlanar( const Mesh* pMesh,
 					pMask ? pMask->GetData() : nullptr,
 					bIsRGBFadingEnabled, bIsAlphaFadingEnabled,
 					fadeStart, fadeEnd);
-				ImageRasterProjected_Optimised(pMesh, pImage, pixelProc, fadeEnd, scratch);
+				ImageRasterProjected_Optimised(pMesh, pImage, pixelProc, fadeEnd, CropMin, UncroppedSize, scratch);
 			}
 			else if (SamplingMethod == ESamplingMethod::BiLinear)
 			{
@@ -903,7 +764,7 @@ void ImageRasterProjectedPlanar( const Mesh* pMesh,
 					pMask ? pMask->GetData() : nullptr,
 					bIsRGBFadingEnabled, bIsAlphaFadingEnabled,
 					fadeStart, fadeEnd);
-				ImageRasterProjected_Optimised(pMesh, pImage, pixelProc, fadeEnd, scratch);
+				ImageRasterProjected_Optimised(pMesh, pImage, pixelProc, fadeEnd, CropMin, UncroppedSize, scratch);
 			}
 			else
 			{
@@ -923,7 +784,7 @@ void ImageRasterProjectedPlanar( const Mesh* pMesh,
 						pMask->GetData(),
 						bIsRGBFadingEnabled, bIsAlphaFadingEnabled,
 						fadeStart, fadeEnd);
-					ImageRasterProjected_Optimised(pMesh, pImage, pixelProc, fadeEnd, scratch);
+					ImageRasterProjected_Optimised(pMesh, pImage, pixelProc, fadeEnd, CropMin, UncroppedSize, scratch);
 				}
 				else if (SamplingMethod == ESamplingMethod::BiLinear)
 				{
@@ -932,7 +793,7 @@ void ImageRasterProjectedPlanar( const Mesh* pMesh,
 						pMask->GetData(),
 						bIsRGBFadingEnabled, bIsAlphaFadingEnabled,
 						fadeStart, fadeEnd);
-					ImageRasterProjected_Optimised(pMesh, pImage, pixelProc, fadeEnd, scratch);
+					ImageRasterProjected_Optimised(pMesh, pImage, pixelProc, fadeEnd, CropMin, UncroppedSize, scratch);
 				}
 				else
 				{
@@ -948,7 +809,7 @@ void ImageRasterProjectedPlanar( const Mesh* pMesh,
 						nullptr,
 						bIsRGBFadingEnabled, bIsAlphaFadingEnabled,
 						fadeStart, fadeEnd);
-					ImageRasterProjected_Optimised(pMesh, pImage, pixelProc, fadeEnd, scratch);
+					ImageRasterProjected_Optimised(pMesh, pImage, pixelProc, fadeEnd, CropMin, UncroppedSize, scratch);
 				}
 				else if (SamplingMethod == ESamplingMethod::BiLinear)
 				{
@@ -957,7 +818,7 @@ void ImageRasterProjectedPlanar( const Mesh* pMesh,
 						nullptr,
 						bIsRGBFadingEnabled, bIsAlphaFadingEnabled,
 						fadeStart, fadeEnd);
-					ImageRasterProjected_Optimised(pMesh, pImage, pixelProc, fadeEnd, scratch);
+					ImageRasterProjected_Optimised(pMesh, pImage, pixelProc, fadeEnd, CropMin, UncroppedSize, scratch);
 				}
 				else
 				{
@@ -974,7 +835,7 @@ void ImageRasterProjectedPlanar( const Mesh* pMesh,
 					pMask ? pMask->GetData() : nullptr,
 					bIsRGBFadingEnabled, bIsAlphaFadingEnabled,
 					fadeStart, fadeEnd);
-				ImageRasterProjected_Optimised(pMesh, pImage, pixelProc, fadeEnd, scratch);
+				ImageRasterProjected_Optimised(pMesh, pImage, pixelProc, fadeEnd, CropMin, UncroppedSize, scratch);
 			}
 			else if (SamplingMethod == ESamplingMethod::BiLinear)
 			{
@@ -983,7 +844,7 @@ void ImageRasterProjectedPlanar( const Mesh* pMesh,
 					pMask ? pMask->GetData() : nullptr,
 					bIsRGBFadingEnabled, bIsAlphaFadingEnabled,
 					fadeStart, fadeEnd);
-				ImageRasterProjected_Optimised(pMesh, pImage, pixelProc, fadeEnd, scratch);
+				ImageRasterProjected_Optimised(pMesh, pImage, pixelProc, fadeEnd, CropMin, UncroppedSize, scratch);
 			}
 			else
 			{
@@ -994,39 +855,26 @@ void ImageRasterProjectedPlanar( const Mesh* pMesh,
         else
         {
 			check(false);
-    //        RasterProjectedPixelProcessor pixelProc( pSource,
-    //                                         pImage->GetData(),
-    //                                         pMask ? pMask->GetData() : 0, 
-				//bIsRGBFadingEnabled, bIsAlphaFadingEnabled,
-    //                                         fadeStart, fadeEnd );
-
-    //        ImageRasterProjected_Optimised( pMesh, pImage, pixelProc, fadeEnd, scratch );
         }
     }
 
     else
     {
         check(false);
-        // Generic version
-        //ImageRasterProjected_Generic( pMesh, pImage, pSource, pMask,
-        //                      fadeStart, fadeEnd,
-        //                      layout, block );
     }
 
 }
 
 
 //-------------------------------------------------------------------------------------------------
-void ImageRasterProjectedWrapping( const Mesh* pMesh,
-                                       Image* pImage,
-                                       const Image* pSource,
-                                       const Image* pMask,
-									   bool bIsRGBFadingEnabled, bool bIsAlphaFadingEnabled,
-									   ESamplingMethod SamplingMethod,
-                                       float fadeStart, float fadeEnd,
-                                       int layout,
-                                       int block,
-                                       SCRATCH_IMAGE_PROJECT* scratch )
+void ImageRasterProjectedWrapping( const Mesh* pMesh, Image* pImage,
+	const Image* pSource, const Image* pMask,
+	bool bIsRGBFadingEnabled, bool bIsAlphaFadingEnabled,
+	ESamplingMethod SamplingMethod,
+	float fadeStart, float fadeEnd,
+	int layout, int block,
+	UE::Math::TIntVector2<uint16> CropMin, UE::Math::TIntVector2<uint16> UncroppedSize,
+	SCRATCH_IMAGE_PROJECT* scratch )
 {
     check( !pMask || pMask->GetSizeX() == pImage->GetSizeX() );
     check( !pMask || pMask->GetSizeY() == pImage->GetSizeY() );
@@ -1050,7 +898,7 @@ void ImageRasterProjectedWrapping( const Mesh* pMesh,
 					bIsRGBFadingEnabled, bIsAlphaFadingEnabled,
 					fadeStart, fadeEnd);
 
-				ImageRasterProjected_OptimisedWrapping(pMesh, pImage, pixelProc, fadeEnd, block, scratch);
+				ImageRasterProjected_OptimisedWrapping(pMesh, pImage, pixelProc, fadeEnd, block, CropMin, UncroppedSize, scratch);
 			}
 			else if (SamplingMethod == ESamplingMethod::BiLinear)
 			{
@@ -1060,7 +908,7 @@ void ImageRasterProjectedWrapping( const Mesh* pMesh,
 					bIsRGBFadingEnabled, bIsAlphaFadingEnabled,
 					fadeStart, fadeEnd);
 
-				ImageRasterProjected_OptimisedWrapping(pMesh, pImage, pixelProc, fadeEnd, block, scratch);
+				ImageRasterProjected_OptimisedWrapping(pMesh, pImage, pixelProc, fadeEnd, block, CropMin, UncroppedSize, scratch);
 			}
 			else
 			{
@@ -1080,7 +928,7 @@ void ImageRasterProjectedWrapping( const Mesh* pMesh,
 						pMask->GetData(),
 						bIsRGBFadingEnabled, bIsAlphaFadingEnabled,
 						fadeStart, fadeEnd);
-					ImageRasterProjected_OptimisedWrapping(pMesh, pImage, pixelProc, fadeEnd, block, scratch);
+					ImageRasterProjected_OptimisedWrapping(pMesh, pImage, pixelProc, fadeEnd, block, CropMin, UncroppedSize, scratch);
 				}
 				else if (SamplingMethod == ESamplingMethod::BiLinear)
 				{
@@ -1089,7 +937,7 @@ void ImageRasterProjectedWrapping( const Mesh* pMesh,
 						pMask->GetData(),
 						bIsRGBFadingEnabled, bIsAlphaFadingEnabled,
 						fadeStart, fadeEnd);
-					ImageRasterProjected_OptimisedWrapping(pMesh, pImage, pixelProc, fadeEnd, block, scratch);
+					ImageRasterProjected_OptimisedWrapping(pMesh, pImage, pixelProc, fadeEnd, block, CropMin, UncroppedSize, scratch);
 				}
 				else
 				{
@@ -1105,7 +953,7 @@ void ImageRasterProjectedWrapping( const Mesh* pMesh,
 						nullptr,
 						bIsRGBFadingEnabled, bIsAlphaFadingEnabled,
 						fadeStart, fadeEnd);
-					ImageRasterProjected_OptimisedWrapping(pMesh, pImage, pixelProc, fadeEnd, block, scratch);
+					ImageRasterProjected_OptimisedWrapping(pMesh, pImage, pixelProc, fadeEnd, block, CropMin, UncroppedSize, scratch);
 				}
 				else if (SamplingMethod == ESamplingMethod::BiLinear)
 				{
@@ -1114,7 +962,7 @@ void ImageRasterProjectedWrapping( const Mesh* pMesh,
 						nullptr,
 						bIsRGBFadingEnabled, bIsAlphaFadingEnabled,
 						fadeStart, fadeEnd);
-					ImageRasterProjected_OptimisedWrapping(pMesh, pImage, pixelProc, fadeEnd, block, scratch);
+					ImageRasterProjected_OptimisedWrapping(pMesh, pImage, pixelProc, fadeEnd, block, CropMin, UncroppedSize, scratch);
 				}
 				else
 				{
@@ -1133,7 +981,7 @@ void ImageRasterProjectedWrapping( const Mesh* pMesh,
 					bIsRGBFadingEnabled, bIsAlphaFadingEnabled,
 					fadeStart, fadeEnd);
 
-				ImageRasterProjected_OptimisedWrapping(pMesh, pImage, pixelProc, fadeEnd, block, scratch);
+				ImageRasterProjected_OptimisedWrapping(pMesh, pImage, pixelProc, fadeEnd, block, CropMin, UncroppedSize, scratch);
 			}
 			else if (SamplingMethod == ESamplingMethod::BiLinear)
 			{	
@@ -1143,7 +991,7 @@ void ImageRasterProjectedWrapping( const Mesh* pMesh,
 					bIsRGBFadingEnabled, bIsAlphaFadingEnabled,
 					fadeStart, fadeEnd);
 
-				ImageRasterProjected_OptimisedWrapping(pMesh, pImage, pixelProc, fadeEnd, block, scratch);
+				ImageRasterProjected_OptimisedWrapping(pMesh, pImage, pixelProc, fadeEnd, block, CropMin, UncroppedSize, scratch);
 			}
 			else
 			{
@@ -1154,23 +1002,12 @@ void ImageRasterProjectedWrapping( const Mesh* pMesh,
         else
         {
 			check(false);
-    //        RasterProjectedPixelProcessor pixelProc( pSource,
-				//pImage->GetData(),
-				//pMask ? pMask->GetData() : 0,
-				//bIsRGBFadingEnabled, bIsAlphaFadingEnabled,
-				//fadeStart, fadeEnd );
-
-    //        ImageRasterProjected_OptimisedWrapping( pMesh, pImage, pixelProc, fadeEnd, block, scratch );
-        }
+         }
     }
 
     else
     {
         check(false);
-        // Generic version
-        //ImageRasterProjected_Generic( pMesh, pImage, pSource, pMask,
-        //                      fadeStart, fadeEnd,
-        //                      layout, block );
     }
 
 }
@@ -1256,15 +1093,14 @@ int32 ComputeProjectedFootprintBestMip(
 }
 
 //-------------------------------------------------------------------------------------------------
-void ImageRasterProjectedCylindrical( const Mesh* pMesh,
-                                        Image* pImage,
-                                        const Image* pSource,
-                                        const Image* pMask,
-										bool bIsRGBFadingEnabled, bool bIsAlphaFadingEnabled,
-                                        float fadeStart, float fadeEnd,
-                                        int /*layout*/,
-                                        float projectionAngle,
-                                        SCRATCH_IMAGE_PROJECT* scratch )
+void ImageRasterProjectedCylindrical( const Mesh* pMesh, Image* pImage,
+	const Image* pSource, const Image* pMask,
+	bool bIsRGBFadingEnabled, bool bIsAlphaFadingEnabled,
+	float fadeStart, float fadeEnd,
+	int /*layout*/,
+	float projectionAngle,
+	UE::Math::TIntVector2<uint16> CropMin, UE::Math::TIntVector2<uint16> UncroppedSize,
+	SCRATCH_IMAGE_PROJECT* scratch )
 {
     check( !pMask || pMask->GetSizeX() == pImage->GetSizeX() );
     check( !pMask || pMask->GetSizeY() == pImage->GetSizeY() );
@@ -1284,7 +1120,7 @@ void ImageRasterProjectedCylindrical( const Mesh* pMesh,
                                                        fadeStart, fadeEnd,
                                                                           projectionAngle);
 
-            ImageRasterProjected_Optimised( pMesh, pImage, pixelProc, fadeEnd, scratch );
+            ImageRasterProjected_Optimised( pMesh, pImage, pixelProc, fadeEnd, CropMin, UncroppedSize, scratch );
         }
 
         else if ( pSource->GetFormat() == EImageFormat::IF_RGBA_UBYTE ||
@@ -1297,7 +1133,7 @@ void ImageRasterProjectedCylindrical( const Mesh* pMesh,
                                                            pMask->GetData(),
                                                            fadeStart, fadeEnd,
                                                                                   projectionAngle );
-                ImageRasterProjected_Optimised( pMesh, pImage, pixelProc, fadeEnd, scratch );
+                ImageRasterProjected_Optimised( pMesh, pImage, pixelProc, fadeEnd, CropMin, UncroppedSize, scratch );
             }
             else
             {
@@ -1306,7 +1142,7 @@ void ImageRasterProjectedCylindrical( const Mesh* pMesh,
                                                            nullptr,
                                                            fadeStart, fadeEnd,
                                                                                    projectionAngle );
-                ImageRasterProjected_Optimised( pMesh, pImage, pixelProc, fadeEnd, scratch );
+                ImageRasterProjected_Optimised( pMesh, pImage, pixelProc, fadeEnd, CropMin, UncroppedSize, scratch );
             }
 
         }
@@ -1319,30 +1155,17 @@ void ImageRasterProjectedCylindrical( const Mesh* pMesh,
                                                      fadeStart, fadeEnd,
                                                                          projectionAngle );
 
-            ImageRasterProjected_Optimised( pMesh, pImage, pixelProc, fadeEnd, scratch );
+            ImageRasterProjected_Optimised( pMesh, pImage, pixelProc, fadeEnd, CropMin, UncroppedSize, scratch );
         }
 
         else
         {
             check(false);
-//            RasterCylindricalProjectedPixelProcessor pixelProc( pSource,
-//                                             pImage->GetData(),
-//                                             pMask ? pMask->GetData() : 0,
-//                                             fadeStart, fadeEnd,
-//                                                                projectionAngle );
-
-//            ImageRasterProjected_Optimised( pMesh, pImage, pixelProc, fadeEnd, block, scratch );
         }
     }
     else
     {
         check(false);
-        // Generic version
-        // TODO: It will allocate temp memory.
-//        ImageRasterCylindricalProjected_Generic( pMesh, pImage, pSource, pMask,
-//                              fadeStart, fadeEnd,
-//                                                 projectionAngle,
-//                              layout, block );
     }
 
 }
diff --git a/Engine/Plugins/Experimental/Mutable/Source/MutableRuntime/Private/MuR/OpImageProject.h b/Engine/Plugins/Experimental/Mutable/Source/MutableRuntime/Private/MuR/OpImageProject.h
--- a/Engine/Plugins/Experimental/Mutable/Source/MutableRuntime/Private/MuR/OpImageProject.h
+++ b/Engine/Plugins/Experimental/Mutable/Source/MutableRuntime/Private/MuR/OpImageProject.h
@@ -17,37 +17,32 @@ namespace mu
 		TArray<uint8> culledVertex;
     };
 
-    extern void ImageRasterProjectedPlanar( const Mesh* pMesh,
-                                            Image* pTargetImage,
-                                            const Image* pSource,
-                                            const Image* pMask,
-											bool bIsRGBFadingEnabled, bool bIsAlphaFadingEnabled,
-											ESamplingMethod SamplingMethod,
-											float fadeStart, float fadeEnd,
-                                            int layout,
-                                            int block,
-                                            SCRATCH_IMAGE_PROJECT* scratch );
-
-    extern void ImageRasterProjectedCylindrical( const Mesh* pMesh,
-                                      Image* pTargetImage,
-                                      const Image* pSource,
-                                      const Image* pMask,
-									  bool bIsRGBFadingEnabled, bool bIsAlphaFadingEnabled,
-									  float fadeStart, float fadeEnd,
-                                      int layout,
-                                      float projectionAngle,
-                                      SCRATCH_IMAGE_PROJECT* scratch );
-
-    extern void ImageRasterProjectedWrapping( const Mesh* pMesh,
-                                      Image* pTargetImage,
-                                      const Image* pSource,
-                                      const Image* pMask,
-									  bool bIsRGBFadingEnabled, bool bIsAlphaFadingEnabled,
-									  ESamplingMethod SamplingMethod,
-									  float fadeStart, float fadeEnd,
-									  int layout,
-                                      int block,
-                                      SCRATCH_IMAGE_PROJECT* scratch );
+    extern void ImageRasterProjectedPlanar( const Mesh* pMesh, Image* pTargetImage,
+		const Image* pSource, const Image* pMask,
+		bool bIsRGBFadingEnabled, bool bIsAlphaFadingEnabled,
+		ESamplingMethod SamplingMethod,
+		float fadeStart, float fadeEnd,
+		int layout, int block,
+		UE::Math::TIntVector2<uint16> CropMin, UE::Math::TIntVector2<uint16> UncroppedSize,
+		SCRATCH_IMAGE_PROJECT* scratch );
+
+    extern void ImageRasterProjectedCylindrical( const Mesh* pMesh, Image* pTargetImage,
+		const Image* pSource, const Image* pMask,
+		bool bIsRGBFadingEnabled, bool bIsAlphaFadingEnabled,
+		float fadeStart, float fadeEnd,
+		int layout,
+		float projectionAngle,
+		UE::Math::TIntVector2<uint16> CropMin, UE::Math::TIntVector2<uint16> UncroppedSize,
+		SCRATCH_IMAGE_PROJECT* scratch );
+
+    extern void ImageRasterProjectedWrapping( const Mesh* pMesh, Image* pTargetImage,
+		const Image* pSource, const Image* pMask,
+		bool bIsRGBFadingEnabled, bool bIsAlphaFadingEnabled,
+		ESamplingMethod SamplingMethod,
+		float fadeStart, float fadeEnd,
+		int layout, int block,
+		UE::Math::TIntVector2<uint16> CropMin, UE::Math::TIntVector2<uint16> UncroppedSize,
+		SCRATCH_IMAGE_PROJECT* scratch );
 
 	extern int32 ComputeProjectedFootprintBestMip(
 			const Mesh* pMesh, const FProjector& Projector, const FVector2f& TargetSize, const FVector2f& SourceSize);
diff --git a/Engine/Plugins/Experimental/Mutable/Source/MutableRuntime/Private/MuR/OpImageRasterMesh.h b/Engine/Plugins/Experimental/Mutable/Source/MutableRuntime/Private/MuR/OpImageRasterMesh.h
--- a/Engine/Plugins/Experimental/Mutable/Source/MutableRuntime/Private/MuR/OpImageRasterMesh.h
+++ b/Engine/Plugins/Experimental/Mutable/Source/MutableRuntime/Private/MuR/OpImageRasterMesh.h
@@ -23,7 +23,8 @@ namespace mu
 	};
 
 
-	inline void ImageRasterMesh( const Mesh* pMesh, Image* pImage, int block )
+	inline void ImageRasterMesh( const Mesh* pMesh, Image* pImage, int block,
+		UE::Math::TIntVector2<uint16> CropMin, UE::Math::TIntVector2<uint16> UncroppedSize )
 	{
 		MUTABLE_CPUPROFILER_SCOPE(ImageRasterMesh)
 
@@ -44,8 +45,17 @@ namespace mu
 			ConvertData( 0, uv, MBF_FLOAT32, texIt.ptr(), texIt.GetFormat() );
 			ConvertData( 1, uv, MBF_FLOAT32, texIt.ptr(), texIt.GetFormat() );
 
-			vertices[v].x = uv[0] * sizeX;
-			vertices[v].y = uv[1] * sizeY;
+			bool bUseCropping = UncroppedSize[0] > 0;
+			if (bUseCropping)
+			{
+				vertices[v].x = uv[0] * UncroppedSize[0] - CropMin[0];
+				vertices[v].y = uv[1] * UncroppedSize[1] - CropMin[1];
+			}
+			else
+			{
+				vertices[v].x = uv[0] * sizeX;
+				vertices[v].y = uv[1] * sizeY;
+			}
 			++texIt;
 		}
 
diff --git a/Engine/Plugins/Experimental/Mutable/Source/MutableRuntime/Private/MuR/OpLayoutPack.h b/Engine/Plugins/Experimental/Mutable/Source/MutableRuntime/Private/MuR/OpLayoutPack.h
--- a/Engine/Plugins/Experimental/Mutable/Source/MutableRuntime/Private/MuR/OpLayoutPack.h
+++ b/Engine/Plugins/Experimental/Mutable/Source/MutableRuntime/Private/MuR/OpLayoutPack.h
@@ -132,8 +132,6 @@ namespace mu
 	// }
 	
 
-	//---------------------------------------------------------------------------------------------
-	//! Block comparison function to establish an order
 	//---------------------------------------------------------------------------------------------
 	struct LAY_BLOCK
 	{
@@ -142,7 +140,7 @@ namespace mu
 			index = -1;
 		}
 
-		LAY_BLOCK( int i, vec2<int> s, int p=0, bool sym = false )
+		LAY_BLOCK( int32 i, UE::Math::TIntVector2<uint16> s, int32 p=0, bool sym = false )
 		{
 			index = i;
 			size = s;
@@ -150,9 +148,9 @@ namespace mu
 			useSymmetry = sym;
 		}
 
-		int index;
-		vec2<int> size;
-		int priority;
+		int32 index;
+		UE::Math::TIntVector2<uint16> size;
+		int32 priority;
 		bool useSymmetry;
 	};
 
@@ -202,7 +200,7 @@ namespace mu
     //---------------------------------------------------------------------------------------------
     struct SCRATCH_LAYOUT_PACK
     {
-        TArray< vec2<int> > blocks;
+        TArray< UE::Math::TIntVector2<uint16> > blocks;
 		TArray< LAY_BLOCK > sorted;
 		TArray< vec2<int> > positions;
 		TArray< int > priorities;
@@ -363,14 +361,15 @@ namespace mu
     // }
 
 
-	inline int32 ReductionOperation(int& BlockSize, EReductionMethod ReductionMethod)
+	inline void ReductionOperation(uint16& BlockSize, EReductionMethod ReductionMethod)
 	{
 		if (ReductionMethod == EReductionMethod::UNITARY_REDUCTION)
 		{
-			return BlockSize -= 1;
+			BlockSize -= 1;
+			return;
 		}
 		
-		return BlockSize /= 2;
+		BlockSize /= 2;
 	}
 
 
@@ -467,7 +466,7 @@ namespace mu
 	}
 
 
-	inline bool SetPositions(int bestY,int layoutSizeY, int* maxX, int* maxY, SCRATCH_LAYOUT_PACK* scratch, EPackStrategy packStrategy)
+	inline bool SetPositions(int bestY,int layoutSizeY, uint16* maxX, uint16* maxY, SCRATCH_LAYOUT_PACK* scratch, EPackStrategy packStrategy)
 	{
 		bool fits = true;
 
@@ -524,8 +523,8 @@ namespace mu
 					}
 
 					// Does it make an unfillable hole with the top or side?
-					int minX = TNumericLimits<int>::Max();
-					int minY = TNumericLimits<int>::Max();
+					uint16 minX = TNumericLimits<uint16>::Max();
+					uint16 minY = TNumericLimits<uint16>::Max();
 					for (size_t b = 0; b < scratch->sorted.Num(); ++b)
 					{
 						if (!packedFlag[b] && b != candidate)
@@ -604,7 +603,7 @@ namespace mu
 
 			// Store
 			scratch->positions[scratch->sorted[best].index] = vec2<int>(bestX, bestLevel);
-			*maxY = FMath::Max(*maxY, bestLevel + scratch->sorted[best].size[1]);
+			*maxY = FMath::Max(*maxY, uint16(bestLevel + scratch->sorted[best].size[1]) );
 
 			if (packStrategy == EPackStrategy::FIXED_LAYOUT && *maxY > layoutSizeY)
 			{
@@ -656,12 +655,12 @@ namespace mu
 		bool usePriority = false;
 
         // Look for the maximum block sizes on the layout and the total area
-        int maxX = 0;
-        int maxY = 0;
+		uint16 maxX = 0;
+		uint16 maxY = 0;
         int area = 0;
         for ( int index=0; index<blockCount; ++index )
         {
-            box< vec2<int> > b;
+            box< UE::Math::TIntVector2<uint16> > b;
             pSourceLayout->GetBlock( index, &b.min[0], &b.min[1], &b.size[0], &b.size[1] );
 
 			int p;
diff --git a/Engine/Plugins/Experimental/Mutable/Source/MutableRuntime/Private/MuR/OpMeshApplyLayout.h b/Engine/Plugins/Experimental/Mutable/Source/MutableRuntime/Private/MuR/OpMeshApplyLayout.h
--- a/Engine/Plugins/Experimental/Mutable/Source/MutableRuntime/Private/MuR/OpMeshApplyLayout.h
+++ b/Engine/Plugins/Experimental/Mutable/Source/MutableRuntime/Private/MuR/OpMeshApplyLayout.h
@@ -61,7 +61,7 @@ namespace mu
 		{
 			FIntPoint grid = pLayout->GetGridSize();
 
-			box< vec2<int> > block;
+			box< UE::Math::TIntVector2<uint16> > block;
 			pLayout->GetBlock( b, &block.min[0], &block.min[1], &block.size[0], &block.size[1] );
 
 			Box rect;
diff --git a/Engine/Plugins/Experimental/Mutable/Source/MutableRuntime/Private/MuR/Operations.h b/Engine/Plugins/Experimental/Mutable/Source/MutableRuntime/Private/MuR/Operations.h
--- a/Engine/Plugins/Experimental/Mutable/Source/MutableRuntime/Private/MuR/Operations.h
+++ b/Engine/Plugins/Experimental/Mutable/Source/MutableRuntime/Private/MuR/Operations.h
@@ -657,7 +657,7 @@ namespace mu
             uint16 blockSize[2];
 			EImageFormat format;
 
-            //! If true, generate all the mipmaps except the ones in skipMipmaps
+            //! If true, generate mipmaps
             uint8 generateMipmaps;
 
             //! Mipmaps to generate if mipmaps are to be generated. 0 means all.
@@ -743,6 +743,9 @@ namespace mu
             ADDRESS colour;
 			EImageFormat format;
             uint16 size[2];
+
+			/** Number of mipmaps to generate. 0 means all the chain. */
+			uint8 LODs;
         };
 
         struct ImageGPUArgs
@@ -781,6 +784,8 @@ namespace mu
 			int32 blockIndex;
 			uint16 sizeX, sizeY;
 			uint16 SourceSizeX, SourceSizeY;
+			uint16 CropMinX, CropMinY;
+			uint16 UncroppedSizeX, UncroppedSizeY;
 			uint8 bIsRGBFadingEnabled : 1;
 			uint8 bIsAlphaFadingEnabled : 1;
 			
@@ -1126,7 +1131,6 @@ namespace mu
             ImagePlainColourArgs ImagePlainColour;
             ImageGPUArgs ImageGPU;
             ImageCropArgs ImageCrop;
-//            ImageRasterMeshArgs ImageRasterMesh;
             ImageDisplaceArgs ImageDisplace;
 			ImageInvertArgs ImageInvert;
 
diff --git a/Engine/Plugins/Experimental/Mutable/Source/MutableRuntime/Public/MuR/Image.h b/Engine/Plugins/Experimental/Mutable/Source/MutableRuntime/Public/MuR/Image.h
--- a/Engine/Plugins/Experimental/Mutable/Source/MutableRuntime/Public/MuR/Image.h
+++ b/Engine/Plugins/Experimental/Mutable/Source/MutableRuntime/Public/MuR/Image.h
@@ -385,6 +385,9 @@ namespace mu
 		void GetNonBlackRect(FImageRect& rect) const;
 		void GetNonBlackRect_Reference(FImageRect& rect) const;
 
+		/** Fill the image with a plain colour. */
+		void FillColour(FVector4f c);
+
 	};
 }
 
diff --git a/Engine/Plugins/Experimental/Mutable/Source/MutableRuntime/Public/MuR/Layout.h b/Engine/Plugins/Experimental/Mutable/Source/MutableRuntime/Public/MuR/Layout.h
--- a/Engine/Plugins/Experimental/Mutable/Source/MutableRuntime/Public/MuR/Layout.h
+++ b/Engine/Plugins/Experimental/Mutable/Source/MutableRuntime/Public/MuR/Layout.h
@@ -100,7 +100,7 @@ namespace mu
 		//! (from 0 to the Y size of the grid minus one )
 		//! \param[out] pSizeX will be set to the x size of the block
 		//! \param[out] pSizeY will be set to the y size of the block
-		void GetBlock( int index, int* pMinX, int* pMinY, int* pSizeX, int* pSizeY ) const;
+		void GetBlock( int index, uint16* pMinX, uint16* pMinY, uint16* pSizeX, uint16* pSizeY ) const;
 
 		//! Returns the reduction priority of a block.
 		//! \param index Block to get priority.
diff --git a/Engine/Plugins/Experimental/Mutable/Source/MutableRuntime/Public/MuR/Model.h b/Engine/Plugins/Experimental/Mutable/Source/MutableRuntime/Public/MuR/Model.h
--- a/Engine/Plugins/Experimental/Mutable/Source/MutableRuntime/Public/MuR/Model.h
+++ b/Engine/Plugins/Experimental/Mutable/Source/MutableRuntime/Public/MuR/Model.h
@@ -13,7 +13,7 @@
 //! data structures. Compiled models are not necessarily compatible when the runtime is updated,
 //! so this version number can be used externally to verify this. It is not used internally, and
 //! serializing models from different versions than this runtime will probably result in a crash.
-#define MUTABLE_COMPILED_MODEL_CODE_VERSION		uint32( 60 )
+#define MUTABLE_COMPILED_MODEL_CODE_VERSION		uint32( 62 )
 #define MUTABLE_PARAMETERS_VERSION              uint32( 2 )
 
 
diff --git a/Engine/Plugins/Experimental/Mutable/Source/MutableTools/Private/MuT/AST.cpp b/Engine/Plugins/Experimental/Mutable/Source/MutableTools/Private/MuT/AST.cpp
--- a/Engine/Plugins/Experimental/Mutable/Source/MutableTools/Private/MuT/AST.cpp
+++ b/Engine/Plugins/Experimental/Mutable/Source/MutableTools/Private/MuT/AST.cpp
@@ -1149,7 +1149,7 @@ FImageDesc ASTOpFixed::GetImageDesc( bool returnBestOption, FGetImageDescContext
         res.m_format = EImageFormat(op.args.ImagePlainColour.format);
         res.m_size[0] = op.args.ImagePlainColour.size[0];
         res.m_size[1] = op.args.ImagePlainColour.size[1];
-        res.m_lods = 1;
+		res.m_lods = op.args.ImagePlainColour.LODs;
         check( res.m_format != EImageFormat::IF_NONE );
         break;
 
@@ -1296,6 +1296,12 @@ void ASTOpFixed::GetLayoutBlockSize( int* pBlockX, int* pBlockY )
 		break;
 	}
 
+	case OP_TYPE::IM_CROP:
+	{
+		GetLayoutBlockSize(op.args.ImageCrop.source, pBlockX, pBlockY);
+		break;
+	}
+
     default:
         checkf( false, TEXT("Instruction not supported") );
     }
diff --git a/Engine/Plugins/Experimental/Mutable/Source/MutableTools/Private/MuT/ASTOpImageLayer.cpp b/Engine/Plugins/Experimental/Mutable/Source/MutableTools/Private/MuT/ASTOpImageLayer.cpp
--- a/Engine/Plugins/Experimental/Mutable/Source/MutableTools/Private/MuT/ASTOpImageLayer.cpp
+++ b/Engine/Plugins/Experimental/Mutable/Source/MutableTools/Private/MuT/ASTOpImageLayer.cpp
@@ -320,7 +320,7 @@ namespace mu
 						&&
 						MaskRaster->angleFadeProperties.child() == BlendRaster->angleFadeProperties.child()
 						&&
-						MaskRaster->blockIndex == BlendRaster->blockIndex
+						MaskRaster->BlockIndex == BlendRaster->BlockIndex
 						)
 					{
 						CurrentMask = MaskRaster->image.child();
diff --git a/Engine/Plugins/Experimental/Mutable/Source/MutableTools/Private/MuT/ASTOpImageMipmap.cpp b/Engine/Plugins/Experimental/Mutable/Source/MutableTools/Private/MuT/ASTOpImageMipmap.cpp
--- a/Engine/Plugins/Experimental/Mutable/Source/MutableTools/Private/MuT/ASTOpImageMipmap.cpp
+++ b/Engine/Plugins/Experimental/Mutable/Source/MutableTools/Private/MuT/ASTOpImageMipmap.cpp
@@ -131,6 +131,15 @@ namespace mu
 			break;
 		}
 
+		case OP_TYPE::IM_PLAINCOLOUR:
+		{
+			// Set the mipmap generation on the plaincolour operation
+			Ptr<ASTOpFixed> mop = mu::Clone<ASTOpFixed>(sourceAt);
+			mop->op.args.ImagePlainColour.LODs = Levels;
+			at = mop;
+			break;
+		}
+
 		case OP_TYPE::IM_PIXELFORMAT:
 		{
 			// Swap unless the mipmap operation builds only the tail or is compressed.
@@ -414,12 +423,16 @@ namespace mu
 					Ptr<ASTOpImagePatch> newOp = mu::Clone<ASTOpImagePatch>(at);
 					newOp->base = Visit(newOp->base.child(), newMip.get());
 					newOp->patch = Visit(newOp->patch.child(), newMip.get());
+					newAt = newOp;
 
-					// We need to add a mipmap on top to finish the mipmapping
-					mu::Ptr<ASTOpImageMipmap> topMipOp = mu::Clone<ASTOpImageMipmap>(currentMipmapOp);
-					topMipOp->Source = newOp;
-					topMipOp->bOnlyTail = true;
-					newAt = topMipOp;
+					if (currentMipmapOp->Levels != currentMipmapOp->BlockLevels)
+					{
+						// We need to add a mipmap on top to finish the mipmapping
+						mu::Ptr<ASTOpImageMipmap> topMipOp = mu::Clone<ASTOpImageMipmap>(currentMipmapOp);
+						topMipOp->Source = newOp;
+						topMipOp->bOnlyTail = true;
+						newAt = topMipOp;
+					}
 				}
 				else
 				{
@@ -427,6 +440,7 @@ namespace mu
 					Ptr<ASTOpImagePatch> newOp = mu::Clone<ASTOpImagePatch>(at);
 					newOp->base = Visit(newOp->base.child(), currentMipmapOp);
 					newOp->patch = Visit(newOp->patch.child(), currentMipmapOp);
+					newAt = newOp;
 				}
 			}
 
diff --git a/Engine/Plugins/Experimental/Mutable/Source/MutableTools/Private/MuT/ASTOpImagePatch.cpp b/Engine/Plugins/Experimental/Mutable/Source/MutableTools/Private/MuT/ASTOpImagePatch.cpp
--- a/Engine/Plugins/Experimental/Mutable/Source/MutableTools/Private/MuT/ASTOpImagePatch.cpp
+++ b/Engine/Plugins/Experimental/Mutable/Source/MutableTools/Private/MuT/ASTOpImagePatch.cpp
@@ -142,4 +142,23 @@ namespace mu
 	}
 
 
+	//-------------------------------------------------------------------------------------------------
+	void ASTOpImagePatch::GetLayoutBlockSize(int* pBlockX, int* pBlockY)
+	{
+		// We didn't find any layout yet.
+		*pBlockX = 0;
+		*pBlockY = 0;
+
+		// Try the source
+		if (base)
+		{
+			base->GetLayoutBlockSize( pBlockX, pBlockY );
+		}
+
+		if (patch && *pBlockX == 0 && *pBlockY == 0)
+		{
+			patch->GetLayoutBlockSize(pBlockX, pBlockY);
+		}
+	}
+
 }
diff --git a/Engine/Plugins/Experimental/Mutable/Source/MutableTools/Private/MuT/ASTOpImagePatch.h b/Engine/Plugins/Experimental/Mutable/Source/MutableTools/Private/MuT/ASTOpImagePatch.h
--- a/Engine/Plugins/Experimental/Mutable/Source/MutableTools/Private/MuT/ASTOpImagePatch.h
+++ b/Engine/Plugins/Experimental/Mutable/Source/MutableTools/Private/MuT/ASTOpImagePatch.h
@@ -36,7 +36,7 @@ struct FProgram;
 		void Link(FProgram& program, const FLinkerOptions* Options) override;
 		FImageDesc GetImageDesc(bool returnBestOption, FGetImageDescContext* context) const override;
 		Ptr<ImageSizeExpression> GetImageSizeExpression() const override;
-		//TODO: void GetLayoutBlockSize(int* pBlockX, int* pBlockY) override;
+		void GetLayoutBlockSize(int* pBlockX, int* pBlockY) override;
 		//TODO: bool IsImagePlainConstant(vec4<float>& colour) const override;
 	};
 
diff --git a/Engine/Plugins/Experimental/Mutable/Source/MutableTools/Private/MuT/ASTOpImagePixelFormat.cpp b/Engine/Plugins/Experimental/Mutable/Source/MutableTools/Private/MuT/ASTOpImagePixelFormat.cpp
--- a/Engine/Plugins/Experimental/Mutable/Source/MutableTools/Private/MuT/ASTOpImagePixelFormat.cpp
+++ b/Engine/Plugins/Experimental/Mutable/Source/MutableTools/Private/MuT/ASTOpImagePixelFormat.cpp
@@ -174,7 +174,6 @@ namespace mu
 			break;
 		}
 
-
 		case OP_TYPE::IM_BLANKLAYOUT:
 		{
 			// Just make sure the layout format is the right one and forget the op
@@ -193,6 +192,24 @@ namespace mu
 			break;
 		}
 
+		case OP_TYPE::IM_PLAINCOLOUR:
+		{
+			// Just make sure the format is the right one and forget the op
+			auto nop = mu::Clone<ASTOpFixed>(sourceAt);
+
+			EImageFormat layoutFormat = (EImageFormat)nop->op.args.ImagePlainColour.format;
+			if (FormatIfAlpha != EImageFormat::IF_NONE
+				&&
+				GetImageFormatData(layoutFormat).m_channels > 3)
+			{
+				format = FormatIfAlpha;
+			}
+
+			nop->op.args.ImagePlainColour.format = format;
+			at = nop;
+			break;
+		}
+
 		default:
 		{
 
diff --git a/Engine/Plugins/Experimental/Mutable/Source/MutableTools/Private/MuT/ASTOpImageRasterMesh.cpp b/Engine/Plugins/Experimental/Mutable/Source/MutableTools/Private/MuT/ASTOpImageRasterMesh.cpp
--- a/Engine/Plugins/Experimental/Mutable/Source/MutableTools/Private/MuT/ASTOpImageRasterMesh.cpp
+++ b/Engine/Plugins/Experimental/Mutable/Source/MutableTools/Private/MuT/ASTOpImageRasterMesh.cpp
@@ -26,11 +26,12 @@ namespace mu
 		, mask(this)
 		, projector(this)
 	{
-		blockIndex = -1;
-		sizeX = 0;
-		sizeY = 0;
-		SourceSizeX = 0;
-		SourceSizeY = 0;
+		BlockIndex = -1;
+		SizeX = SizeY = 0;
+		SourceSizeX = SourceSizeY = 0;
+		CropMinX = CropMinY = 0;
+		UncroppedSizeX = UncroppedSizeY = 0;
+
 		bIsRGBFadingEnabled = 1;
 		bIsAlphaFadingEnabled = 1;
 		SamplingMethod = ESamplingMethod::Point;
@@ -56,11 +57,15 @@ namespace mu
 				angleFadeProperties == Other->angleFadeProperties &&
 				mask == Other->mask &&
 				projector == Other->projector &&
-				blockIndex == Other->blockIndex &&
-				sizeX == Other->sizeX &&
-				sizeY == Other->sizeY &&
+				BlockIndex == Other->BlockIndex &&
+				SizeX == Other->SizeX &&
+				SizeY == Other->SizeY &&
 				SourceSizeX == Other->SourceSizeX &&
 				SourceSizeY == Other->SourceSizeY &&
+				CropMinX == Other->CropMinX &&
+				CropMinY == Other->CropMinY &&
+				UncroppedSizeX == Other->UncroppedSizeX &&
+				UncroppedSizeY == Other->UncroppedSizeY &&
 				bIsRGBFadingEnabled == Other->bIsRGBFadingEnabled &&
 				bIsAlphaFadingEnabled == Other->bIsAlphaFadingEnabled &&
 				SamplingMethod == Other->SamplingMethod &&
@@ -79,7 +84,7 @@ namespace mu
 		hash_combine(res, angleFadeProperties.child().get());
 		hash_combine(res, mask.child().get());
 		hash_combine(res, projector.child().get());
-		hash_combine(res, blockIndex);
+		hash_combine(res, BlockIndex);
 		return res;
 	}
 
@@ -93,9 +98,13 @@ namespace mu
 		n->angleFadeProperties = mapChild(angleFadeProperties.child());
 		n->mask = mapChild(mask.child());
 		n->projector = mapChild(projector.child());
-		n->blockIndex = blockIndex;
-		n->sizeX = sizeX;
-		n->sizeY = sizeY;
+		n->BlockIndex = BlockIndex;
+		n->SizeX = SizeX;
+		n->SizeY = SizeY;
+		n->CropMinX = CropMinX;
+		n->CropMinY = CropMinY;
+		n->UncroppedSizeX = UncroppedSizeX;
+		n->UncroppedSizeY = UncroppedSizeY;
 		n->SourceSizeX = SourceSizeX;
 		n->SourceSizeY = SourceSizeY;
 		n->bIsRGBFadingEnabled = bIsRGBFadingEnabled;
@@ -126,11 +135,15 @@ namespace mu
 			OP::ImageRasterMeshArgs args;
 			FMemory::Memzero(&args, sizeof(args));
 
-			args.blockIndex = blockIndex;
-			args.sizeX = sizeX;
-			args.sizeY = sizeY;
+			args.blockIndex = BlockIndex;
+			args.sizeX = SizeX;
+			args.sizeY = SizeY;
 			args.SourceSizeX = SourceSizeX;
 			args.SourceSizeY = SourceSizeY;
+			args.CropMinX = CropMinX;
+			args.CropMinY = CropMinY;
+			args.UncroppedSizeX = UncroppedSizeX;
+			args.UncroppedSizeY = UncroppedSizeY;
 			args.bIsRGBFadingEnabled = bIsRGBFadingEnabled;
 			args.bIsAlphaFadingEnabled = bIsAlphaFadingEnabled;
 			args.SamplingMethod = static_cast<uint8>(SamplingMethod);
@@ -175,8 +188,8 @@ namespace mu
 		if (image)
 		{
 			res = image->GetImageDesc( returnBestOption, context);
-			res.m_size[0] = sizeX;
-			res.m_size[1] = sizeY;
+			res.m_size[0] = SizeX;
+			res.m_size[1] = SizeY;
 		}
 
 
@@ -195,8 +208,8 @@ namespace mu
 	{
 		Ptr<ImageSizeExpression> pRes = new ImageSizeExpression;
 		pRes->type = ImageSizeExpression::ISET_CONSTANT;
-		pRes->size[0] = sizeX ? sizeX : 256;
-		pRes->size[1] = sizeY ? sizeY : 256;
+		pRes->size[0] = SizeX ? SizeX : 256;
+		pRes->size[1] = SizeY ? SizeY : 256;
 		return pRes;
 	}
 
diff --git a/Engine/Plugins/Experimental/Mutable/Source/MutableTools/Private/MuT/ASTOpImageRasterMesh.h b/Engine/Plugins/Experimental/Mutable/Source/MutableTools/Private/MuT/ASTOpImageRasterMesh.h
--- a/Engine/Plugins/Experimental/Mutable/Source/MutableTools/Private/MuT/ASTOpImageRasterMesh.h
+++ b/Engine/Plugins/Experimental/Mutable/Source/MutableTools/Private/MuT/ASTOpImageRasterMesh.h
@@ -23,9 +23,21 @@ namespace mu
 		ASTChild mask;
 		ASTChild projector;
 
-		int32 blockIndex;
-		uint16 sizeX, sizeY;
+		/** If layouts are used, this can indicate a single layout block that we want to raster. */
+		int32 BlockIndex;
+
+		/** Size of the image to generate by rasterization of the mesh. */
+		uint16 SizeX, SizeY;
+
+		/** Expected size of the image that we want to project. */
 		uint16 SourceSizeX, SourceSizeY;
+
+		/** Sub-rect to raster, ignoring all the rest.
+		 * Only valid if any UncroppedSizeX is greater than 0. 
+		 */
+		uint16 CropMinX, CropMinY;
+		uint16 UncroppedSizeX, UncroppedSizeY;
+
 		uint8 bIsRGBFadingEnabled : 1;
 		uint8 bIsAlphaFadingEnabled : 1;
 		ESamplingMethod SamplingMethod;
diff --git a/Engine/Plugins/Experimental/Mutable/Source/MutableTools/Private/MuT/ASTOpImageSwizzle.cpp b/Engine/Plugins/Experimental/Mutable/Source/MutableTools/Private/MuT/ASTOpImageSwizzle.cpp
--- a/Engine/Plugins/Experimental/Mutable/Source/MutableTools/Private/MuT/ASTOpImageSwizzle.cpp
+++ b/Engine/Plugins/Experimental/Mutable/Source/MutableTools/Private/MuT/ASTOpImageSwizzle.cpp
@@ -654,9 +654,13 @@ namespace mu
 							|| Typed->angleFadeProperties.child() != FirstRasterMesh->angleFadeProperties.child()
 							|| Typed->mask.child() != FirstRasterMesh->mask.child()
 							|| Typed->projector.child() != FirstRasterMesh->projector.child()
-							|| Typed->blockIndex != FirstRasterMesh->blockIndex
-							|| Typed->sizeX != FirstRasterMesh->sizeX
-							|| Typed->sizeY != FirstRasterMesh->sizeY
+							|| Typed->BlockIndex != FirstRasterMesh->BlockIndex
+							|| Typed->SizeX != FirstRasterMesh->SizeX
+							|| Typed->SizeY != FirstRasterMesh->SizeY
+							|| Typed->UncroppedSizeX != FirstRasterMesh->UncroppedSizeX
+							|| Typed->UncroppedSizeY != FirstRasterMesh->UncroppedSizeY
+							|| Typed->CropMinX != FirstRasterMesh->CropMinX
+							|| Typed->CropMinY != FirstRasterMesh->CropMinY
 							|| Typed->bIsRGBFadingEnabled != FirstRasterMesh->bIsRGBFadingEnabled
 							|| Typed->bIsAlphaFadingEnabled != FirstRasterMesh->bIsAlphaFadingEnabled
 							)
diff --git a/Engine/Plugins/Experimental/Mutable/Source/MutableTools/Private/MuT/CodeGenerator.cpp b/Engine/Plugins/Experimental/Mutable/Source/MutableTools/Private/MuT/CodeGenerator.cpp
--- a/Engine/Plugins/Experimental/Mutable/Source/MutableTools/Private/MuT/CodeGenerator.cpp
+++ b/Engine/Plugins/Experimental/Mutable/Source/MutableTools/Private/MuT/CodeGenerator.cpp
@@ -6,7 +6,7 @@
 #include "Logging/LogCategory.h"
 #include "Logging/LogMacros.h"
 #include "Math/IntPoint.h"
-#include "Math/UnrealMathSSE.h"
+#include "Math/UnrealMathUtility.h"
 #include "MuR/ImagePrivate.h"
 #include "MuR/Layout.h"
 #include "MuR/MutableTrace.h"
@@ -22,6 +22,7 @@
 #include "MuT/ASTOpImagePixelFormat.h"
 #include "MuT/ASTOpImageLayer.h"
 #include "MuT/ASTOpImageLayerColor.h"
+#include "MuT/ASTOpImagePatch.h"
 #include "MuT/ASTOpInstanceAdd.h"
 #include "MuT/ASTOpMeshBindShape.h"
 #include "MuT/ASTOpMeshClipDeform.h"
@@ -90,114 +91,114 @@
 namespace mu
 {
 
-    //---------------------------------------------------------------------------------------------
-    //---------------------------------------------------------------------------------------------
-    //---------------------------------------------------------------------------------------------
-    CodeGenerator::CodeGenerator( CompilerOptions::Private* options )
-    {
-        m_compilerOptions = options;
+	//---------------------------------------------------------------------------------------------
+	//---------------------------------------------------------------------------------------------
+	//---------------------------------------------------------------------------------------------
+	CodeGenerator::CodeGenerator(CompilerOptions::Private* options)
+	{
+		m_compilerOptions = options;
 
-        // Create the message log
-        m_pErrorLog = new ErrorLog;
+		// Create the message log
+		m_pErrorLog = new ErrorLog;
 
-        // Add the parent at the top of the hierarchy
-        m_currentParents.Add( FParentKey() );
-    }
+		// Add the parent at the top of the hierarchy
+		m_currentParents.Add(FParentKey());
+	}
 
-    //---------------------------------------------------------------------------------------------
-    void CodeGenerator::GenerateRoot(const NodePtrConst pNode )
-    {
-        MUTABLE_CPUPROFILER_SCOPE(Generate);
+	//---------------------------------------------------------------------------------------------
+	void CodeGenerator::GenerateRoot(const NodePtrConst pNode)
+	{
+		MUTABLE_CPUPROFILER_SCOPE(Generate);
 
-//        UE_LOG(LogMutableCore,Warning,"astop %9u#  %9u+  %9u- generate root start", size_t(ASTOp::s_count), size_t(ASTOp::s_created), size_t(ASTOp::s_destroyed) );
+		//        UE_LOG(LogMutableCore,Warning,"astop %9u#  %9u+  %9u- generate root start", size_t(ASTOp::s_count), size_t(ASTOp::s_created), size_t(ASTOp::s_destroyed) );
 
-		// First pass
-        m_firstPass.Generate( m_pErrorLog, pNode->GetBasePrivate(), m_compilerOptions->m_ignoreStates );
+				// First pass
+		m_firstPass.Generate(m_pErrorLog, pNode->GetBasePrivate(), m_compilerOptions->m_ignoreStates);
 
-//        UE_LOG(LogMutableCore,Warning,"astop %9u#  %9u+  %9u- first pass completed", size_t(ASTOp::s_count), size_t(ASTOp::s_created), size_t(ASTOp::s_destroyed) );
+		//        UE_LOG(LogMutableCore,Warning,"astop %9u#  %9u+  %9u- first pass completed", size_t(ASTOp::s_count), size_t(ASTOp::s_created), size_t(ASTOp::s_destroyed) );
 
-		// Second pass
-		SecondPassGenerator secondPass( &m_firstPass, m_compilerOptions );
-		bool success = secondPass.Generate(m_pErrorLog, pNode->GetBasePrivate() );
+				// Second pass
+		SecondPassGenerator secondPass(&m_firstPass, m_compilerOptions);
+		bool success = secondPass.Generate(m_pErrorLog, pNode->GetBasePrivate());
 		if (!success)
 		{
-            return;
+			return;
 		}
 
-//        UE_LOG(LogMutableCore,Warning,"astop %9u#  %9u+  %9u- second pass completed ", size_t(ASTOp::s_count), size_t(ASTOp::s_created), size_t(ASTOp::s_destroyed) );
+		//        UE_LOG(LogMutableCore,Warning,"astop %9u#  %9u+  %9u- second pass completed ", size_t(ASTOp::s_count), size_t(ASTOp::s_created), size_t(ASTOp::s_destroyed) );
 
-        // Main pass for each state
-        {
-            MUTABLE_CPUPROFILER_SCOPE(MainPass);
+				// Main pass for each state
+		{
+			MUTABLE_CPUPROFILER_SCOPE(MainPass);
 
             m_currentStateIndex = 0;
-            for( const auto& s: m_firstPass.m_states )
-            {
-                MUTABLE_CPUPROFILER_SCOPE(MainPassState);
+			for (const TPair<FObjectState, const Node::Private*>& s : m_firstPass.m_states)
+			{
+				MUTABLE_CPUPROFILER_SCOPE(MainPassState);
 
-                Ptr<ASTOp> stateRoot = Generate(pNode);
-                m_states.Emplace( s.Key, stateRoot );
+				Ptr<ASTOp> stateRoot = Generate(pNode);
+				m_states.Emplace(s.Key, stateRoot);
 
-                ++m_currentStateIndex;
-            }
-        }
+				++m_currentStateIndex;
+			}
+		}
 
-//        UE_LOG(LogMutableCore,Warning,"astop %9u#  %9u+  %9u- before cache clear", size_t(ASTOp::s_count), size_t(ASTOp::s_created), size_t(ASTOp::s_destroyed) );
+		//        UE_LOG(LogMutableCore,Warning,"astop %9u#  %9u+  %9u- before cache clear", size_t(ASTOp::s_count), size_t(ASTOp::s_created), size_t(ASTOp::s_destroyed) );
 
-        // Free caches
-        m_compiled.Reset();
-        m_constantMeshes.Empty();
+				// Free caches
+		m_compiled.Reset();
+		m_constantMeshes.Empty();
 		m_generatedLayouts.Empty();
-        m_nodeVariables.clear();
-        m_generatedMeshes.Reset();
-        m_generatedProjectors.Reset();
-        m_generatedRanges.Reset();
-        m_generatedTables.clear();
-        m_firstPass = FirstPassGenerator();
-        m_currentBottomUpState = BOTTOM_UP_STATE();
-        m_imageState.Empty();
-        m_currentParents.Empty();
-        m_currentObject.Empty();
-        m_additionalComponents.clear();
-
-//        UE_LOG(LogMutableCore,Warning,"astop %9u#  %9u+  %9u- exiting generate root", size_t(ASTOp::s_count), size_t(ASTOp::s_created), size_t(ASTOp::s_destroyed) );
+		m_nodeVariables.clear();
+		m_generatedMeshes.Reset();
+		m_generatedProjectors.Reset();
+		m_generatedRanges.Reset();
+		m_generatedTables.clear();
+		m_firstPass = FirstPassGenerator();
+		m_currentBottomUpState = BOTTOM_UP_STATE();
+		m_imageState.Empty();
+		m_currentParents.Empty();
+		m_currentObject.Empty();
+		m_additionalComponents.clear();
+
+		//        UE_LOG(LogMutableCore,Warning,"astop %9u#  %9u+  %9u- exiting generate root", size_t(ASTOp::s_count), size_t(ASTOp::s_created), size_t(ASTOp::s_destroyed) );
 	}
 
 
 	//---------------------------------------------------------------------------------------------
-    Ptr<ASTOp> CodeGenerator::Generate(const NodePtrConst pNode)
+	Ptr<ASTOp> CodeGenerator::Generate(const NodePtrConst pNode)
 	{
-        if (!pNode)
-        {
-            return nullptr;
-        }
+		if (!pNode)
+		{
+			return nullptr;
+		}
 
 		// Clear bottom-up state
 		m_currentBottomUpState.m_address = nullptr;
 
         // Temp by-passes while we remove the visitor pattern
-		if (auto ScalarNode = dynamic_cast<const NodeScalar*>(pNode.get()))
+		if (const NodeScalar* ScalarNode = dynamic_cast<const NodeScalar*>(pNode.get()))
 		{
 			FScalarGenerationResult ScalarResult;
 			GenerateScalar(ScalarResult, ScalarNode);
 			return ScalarResult.op;
 		}
 
-		if (auto ColorNode = dynamic_cast<const NodeColour*>(pNode.get()))
+		if (const NodeColour* ColorNode = dynamic_cast<const NodeColour*>(pNode.get()))
 		{
 			FColorGenerationResult Result;
 			GenerateColor(Result, ColorNode);
 			return Result.op;
 		}
 
-		if (auto ImageNode = dynamic_cast<const NodeImage*>(pNode.get()))
+		if (const NodeImage* ImageNode = dynamic_cast<const NodeImage*>(pNode.get()))
 		{
 			FImageGenerationResult Result;
 			GenerateImage(Result, ImageNode);
 			return Result.op;
 		}
 
-		if (auto MeshNode = dynamic_cast<const NodeMesh*>(pNode.get()))
+		if (const NodeMesh* MeshNode = dynamic_cast<const NodeMesh*>(pNode.get()))
 		{
 			// This should only happen with mesh root nodes
 			FMeshGenerationOptions MeshOptions;
@@ -210,59 +211,59 @@ namespace mu
 			return MeshResult.meshOp;
 		}
 
-		if ( auto projNode = dynamic_cast<const NodeProjector*>(pNode.get()) )
-        {
-            FProjectorGenerationResult ProjResult;
-            GenerateProjector( ProjResult, projNode );
-            return ProjResult.op;
-        }
+		if (const NodeProjector* projNode = dynamic_cast<const NodeProjector*>(pNode.get()))
+		{
+			FProjectorGenerationResult ProjResult;
+			GenerateProjector(ProjResult, projNode);
+			return ProjResult.op;
+		}
 
-        if ( auto surfNode = dynamic_cast<const NodeSurfaceNew*>(pNode.get()) )
-        {
-            // This happens only if we generate a node graph that has a NodeSurfaceNew at the root.
-            FSurfaceGenerationResult surfResult;
-            const TArray<FirstPassGenerator::SURFACE::EDIT> edits;
-            GenerateSurface( surfResult, surfNode, edits );
-            return surfResult.surfaceOp;
-        }
+		if (const NodeSurfaceNew* surfNode = dynamic_cast<const NodeSurfaceNew*>(pNode.get()))
+		{
+			// This happens only if we generate a node graph that has a NodeSurfaceNew at the root.
+			FSurfaceGenerationResult surfResult;
+			const TArray<FirstPassGenerator::SURFACE::EDIT> edits;
+			GenerateSurface(surfResult, surfNode, edits);
+			return surfResult.surfaceOp;
+		}
 
-        else if ( dynamic_cast<const NodeSurfaceVariation*>(pNode.get()) )
-        {
-            // This happens only if we generate a node graph that has a NodeSurfaceVariation at the root.
-            return nullptr;
-        }
+		else if (dynamic_cast<const NodeSurfaceVariation*>(pNode.get()))
+		{
+			// This happens only if we generate a node graph that has a NodeSurfaceVariation at the root.
+			return nullptr;
+		}
 
-        else if ( dynamic_cast<const NodeSurfaceEdit*>(pNode.get()) )
-        {
-            // This happens only if we generate a node graph that has a NodeSurfaceEdit at the root.
-            return nullptr;
-        }
+		else if (dynamic_cast<const NodeSurfaceEdit*>(pNode.get()))
+		{
+			// This happens only if we generate a node graph that has a NodeSurfaceEdit at the root.
+			return nullptr;
+		}
 
-        else if (  dynamic_cast<const NodeModifier*>(pNode.get()) )
-        {
-            // This happens only if we generate a node graph that has a modifier at the root.
-            return nullptr;
-        }
+		else if (dynamic_cast<const NodeModifier*>(pNode.get()))
+		{
+			// This happens only if we generate a node graph that has a modifier at the root.
+			return nullptr;
+		}
 
 
-        Ptr<ASTOp> result;
+		Ptr<ASTOp> result;
 
-        // See if it was already generated
+		// See if it was already generated
 		FVisitedKeyMap key = GetCurrentCacheKey(pNode);
-        VisitedMap::ValueType* it = m_compiled.Find( key );
-        if ( it )
-        {
-            m_currentBottomUpState = *it;
-            result = m_currentBottomUpState.m_address;
-        }
-        else
-        {
+		VisitedMap::ValueType* it = m_compiled.Find(key);
+		if (it)
+		{
+			m_currentBottomUpState = *it;
+			result = m_currentBottomUpState.m_address;
+		}
+		else
+		{
 			result = pNode->GetBasePrivate()->Accept(*this);
 			m_currentBottomUpState.m_address = result;
 			m_compiled.Add(key, m_currentBottomUpState);
 		}
 
-        // debug: expensive check of all code generation
+		// debug: expensive check of all code generation
 //        if (result)
 //        {
 //            ASTOpList roots;
@@ -270,52 +271,52 @@ namespace mu
 //            ASTOp::FullAssert(roots);
 //        }
 
-        return result;
-    }
+		return result;
+	}
 
 
-    //---------------------------------------------------------------------------------------------
-    void CodeGenerator::GenerateRange( FRangeGenerationResult& Result,
-                                      NodeRangePtrConst Untyped)
-    {
-        if (!Untyped)
-        {
+	//---------------------------------------------------------------------------------------------
+	void CodeGenerator::GenerateRange(FRangeGenerationResult& Result,
+		NodeRangePtrConst Untyped)
+	{
+		if (!Untyped)
+		{
 			Result = FRangeGenerationResult();
-            return;
-        }
+			return;
+		}
 
-        // See if it was already generated
+		// See if it was already generated
 		FVisitedKeyMap Key = GetCurrentCacheKey(Untyped);
 		GeneratedRangeMap::ValueType* it = m_generatedRanges.Find(Key);
-        if ( it  )
-        {
+		if (it)
+		{
 			Result = *it;
-            return;
-        }
+			return;
+		}
 
 
-        // Generate for each different type of node
-        if (const NodeRangeFromScalar* FromScalar = dynamic_cast<const NodeRangeFromScalar*>(Untyped.get()) )
-        {
+		// Generate for each different type of node
+		if (const NodeRangeFromScalar* FromScalar = dynamic_cast<const NodeRangeFromScalar*>(Untyped.get()))
+		{
 			Result = FRangeGenerationResult();
 			Result.rangeName = FromScalar->GetName();
 			Result.sizeOp = Generate(FromScalar->GetSize());
-        }
-        else
-        {
-            check(false);
-        }
+		}
+		else
+		{
+			check(false);
+		}
 
 
-        // Cache the result
-        m_generatedRanges.Add( Key, Result);
-    }
+		// Cache the result
+		m_generatedRanges.Add(Key, Result);
+	}
 
-	
-    //---------------------------------------------------------------------------------------------
-    Ptr<ASTOp> CodeGenerator::GenerateTableVariable( TablePtr pTable, const string& strName )
-    {
-        Ptr<ASTOp> result;
+
+	//---------------------------------------------------------------------------------------------
+	Ptr<ASTOp> CodeGenerator::GenerateTableVariable(TablePtr pTable, const string& strName)
+	{
+		Ptr<ASTOp> result;
 
         FParameterDesc param;
         param.m_name = strName;
@@ -326,19 +327,19 @@ namespace mu
         param.m_type = PARAMETER_TYPE::T_INT;
         param.m_defaultValue.Set<ParamIntType>(0);
 
-        // Add the possible values
-        {
-            // See if there is a string column. If there is one, we will use it as names for the
-            // options. Only the first string column will be used.
-            int nameCol = -1;
-            int32 cols = pTable->GetPrivate()->m_columns.Num();
-            for ( int32 c=0; c<cols && nameCol<0; ++c )
-            {
-                if ( pTable->GetPrivate()->m_columns[c].m_type==TCT_STRING )
-                {
-                    nameCol = c;
-                }
-            }
+		// Add the possible values
+		{
+			// See if there is a string column. If there is one, we will use it as names for the
+			// options. Only the first string column will be used.
+			int nameCol = -1;
+			int32 cols = pTable->GetPrivate()->m_columns.Num();
+			for (int32 c = 0; c < cols && nameCol < 0; ++c)
+			{
+				if (pTable->GetPrivate()->m_columns[c].m_type == TCT_STRING)
+				{
+					nameCol = c;
+				}
+			}
 
 			if (pTable->GetPrivate()->m_NoneOption)
 			{
@@ -349,19 +350,19 @@ namespace mu
 				param.m_defaultValue.Set<ParamIntType>(nullValue.m_value);
 			}
 
-            // Add every row
-            int32 rows = pTable->GetPrivate()->m_rows.Num();
-			for (size_t i = 0; i<rows; ++i)
-            {
-                FParameterDesc::INT_VALUE_DESC value;
-                value.m_value = (int16_t)pTable->GetPrivate()->m_rows[i].m_id;
+			// Add every row
+			int32 rows = pTable->GetPrivate()->m_rows.Num();
+			for (size_t i = 0; i < rows; ++i)
+			{
+				FParameterDesc::INT_VALUE_DESC value;
+				value.m_value = (int16_t)pTable->GetPrivate()->m_rows[i].m_id;
 
-                if (nameCol>-1)
-                {
-                    value.m_name = pTable->GetPrivate()->m_rows[i].m_values[nameCol].m_string;
-                }
+				if (nameCol > -1)
+				{
+					value.m_name = pTable->GetPrivate()->m_rows[i].m_values[nameCol].m_string;
+				}
 
-                param.m_possibleValues.Add( value );
+				param.m_possibleValues.Add(value);
 
                 // The first row is the default one
                 if (i==0)
@@ -371,328 +372,385 @@ namespace mu
             }
         }
 
-        Ptr<ASTOpParameter> op = new ASTOpParameter();
-        op->type = OP_TYPE::NU_PARAMETER;
-        op->parameter = param;
+		Ptr<ASTOpParameter> op = new ASTOpParameter();
+		op->type = OP_TYPE::NU_PARAMETER;
+		op->parameter = param;
 
-        return op;
-    }
+		return op;
+	}
 
 
-    //---------------------------------------------------------------------------------------------
-	Ptr<const Layout> CodeGenerator::AddLayout(Ptr<const Layout> SourceLayout )
-    {
+	//---------------------------------------------------------------------------------------------
+	Ptr<const Layout> CodeGenerator::AddLayout(Ptr<const Layout> SourceLayout)
+	{
 		// The layout we are adding must be a source layout, without block ids yet.
 		check(SourceLayout->m_blocks.IsEmpty() || SourceLayout->m_blocks[0].m_id == -1);
-		
-		Ptr<const Layout>* it = m_generatedLayouts.Find(SourceLayout.get() );
 
-        if ( it )
-        {
-            return *it;
-        }
+		Ptr<const Layout>* it = m_generatedLayouts.Find(SourceLayout.get());
+
+		if (it)
+		{
+			return *it;
+		}
 
 		// Assign unique ids to each layout block
 		Ptr<Layout> ClonedLayout = SourceLayout->Clone();
-        for (int32 b=0;b< ClonedLayout->m_blocks.Num();++b)
-        {
+		for (int32 b = 0; b < ClonedLayout->m_blocks.Num(); ++b)
+		{
 			// This is a hard limit due to layout block index data being stored in 16 bit.
-			check(m_absoluteLayoutIndex<65536);
+			check(m_absoluteLayoutIndex < 65536);
 			ClonedLayout->m_blocks[b].m_id = m_absoluteLayoutIndex++;
-        }
+		}
 		check(SourceLayout->m_blocks.Num() == ClonedLayout->m_blocks.Num());
 		check(ClonedLayout->m_blocks.IsEmpty() || ClonedLayout->m_blocks[0].m_id != -1);
 		m_generatedLayouts.Add(SourceLayout.get(), ClonedLayout);
 
-        return ClonedLayout;
-    }
+		return ClonedLayout;
+	}
 
 
-    //---------------------------------------------------------------------------------------------
-    Ptr<ASTOp> CodeGenerator::GenerateImageBlockPatch( Ptr<ASTOp> blockAd,
-                                                        const NodePatchImage* pPatch,
-                                                        Ptr<ASTOp> conditionAd )
-    {
-        // Blend operation
-         Ptr<ASTOp> blendAd;
-         {
-             MUTABLE_CPUPROFILER_SCOPE(PatchBlend);
-
-             Ptr<ASTOpImageLayer> op = new ASTOpImageLayer();
-			 op->blendType = pPatch->GetPrivate()->m_blendType;
-             op->base = blockAd;
-
-             // When we patch from edit nodes, we want to apply it to all the channels.
-             // \todo: since we can choose the patch function, maybe we want to be able to
-             // select this as well.
-             op->Flags = pPatch->GetPrivate()->m_applyToAlpha
-                     ? OP::ImageLayerArgs::F_APPLY_TO_ALPHA
-                     : 0;
-
-             NodeImage* pImage = pPatch->GetPrivate()->m_pImage.get();
-             Ptr<ASTOp> blend;
-             if ( pImage )
-             {
-                 blend = Generate( pImage );
-             }
-             else
-             {
-                 blend = GenerateMissingImageCode( "Blend top image", EImageFormat::IF_RGB_UBYTE,
-                                                   pPatch->GetPrivate()->m_errorContext );
-             }
-             blend = GenerateImageFormat( blend, blockAd->GetImageDesc().m_format );
-             blend = GenerateImageSize
-                 (
-                     blend,
-                     FImageSize( (uint16)m_imageState.Last().m_imageRect.size[0],
-                                 (uint16)m_imageState.Last().m_imageRect.size[1] )
-                 );
-             op->blend = blend;
-
-             NodeImage* pMask = pPatch->GetPrivate()->m_pMask.get();
-             Ptr<ASTOp> mask;
-             if ( pMask )
-             {
-                 mask = Generate( pMask );
-             }
-             else
-             {
-                 // Set the argument default value: affect all pixels.
-                 // TODO: Special operation code without mask
-                 mask = GeneratePlainImageCode( vec3<float>( 1,1,1 ) );
-             }
-             mask = GenerateImageFormat( mask, EImageFormat::IF_L_UBYTE );
-             mask = GenerateImageSize
-                 (
-                     mask,
-                     FImageSize( (uint16)m_imageState.Last().m_imageRect.size[0],
-                                 (uint16)m_imageState.Last().m_imageRect.size[1] )
-                 );
-             op->mask = mask;
-
-             blendAd = op;
-         }
-
-         // Condition to enable this patch
-         if (conditionAd)
-         {
-             Ptr<ASTOp> conditionalAd;
-             {
-                 Ptr<ASTOpConditional> op = new ASTOpConditional();
-                 op->type = OP_TYPE::IM_CONDITIONAL;
-                 op->no = blockAd;
-                 op->yes = blendAd;
-                 op->condition = conditionAd;
-                 conditionalAd = op;
-             }
-
-             blockAd = conditionalAd;
-         }
-         else
-         {
-             blockAd = blendAd;
-         }
-
-         return blockAd;
-    }
+	//---------------------------------------------------------------------------------------------
+	Ptr<ASTOp> CodeGenerator::GenerateImageBlockPatch(Ptr<ASTOp> blockAd,
+		const NodePatchImage* pPatch,
+		Ptr<ASTOp> conditionAd)
+	{
+		// Blend operation
+		Ptr<ASTOp> blendAd;
+		{
+			MUTABLE_CPUPROFILER_SCOPE(PatchBlend);
+
+			Ptr<ASTOpImageLayer> op = new ASTOpImageLayer();
+			op->blendType = pPatch->GetPrivate()->m_blendType;
+			op->base = blockAd;
+
+			// When we patch from edit nodes, we want to apply it to all the channels.
+			// \todo: since we can choose the patch function, maybe we want to be able to
+			// select this as well.
+			op->Flags = pPatch->GetPrivate()->m_applyToAlpha
+				? OP::ImageLayerArgs::F_APPLY_TO_ALPHA
+				: 0;
+
+			NodeImage* pImage = pPatch->GetPrivate()->m_pImage.get();
+			Ptr<ASTOp> blend;
+			if (pImage)
+			{
+				blend = Generate(pImage);
+			}
+			else
+			{
+				blend = GenerateMissingImageCode("Blend top image", EImageFormat::IF_RGB_UBYTE,
+					pPatch->GetPrivate()->m_errorContext);
+			}
+			blend = GenerateImageFormat(blend, blockAd->GetImageDesc().m_format);
+			blend = GenerateImageSize
+			(
+				blend,
+				FImageSize((uint16)m_imageState.Last().m_imageRect.size[0],
+					(uint16)m_imageState.Last().m_imageRect.size[1])
+			);
+			op->blend = blend;
+
+			NodeImage* pMask = pPatch->GetPrivate()->m_pMask.get();
+			Ptr<ASTOp> mask;
+			if (pMask)
+			{
+				mask = Generate(pMask);
+			}
+			else
+			{
+				// Set the argument default value: affect all pixels.
+				// TODO: Special operation code without mask
+				mask = GeneratePlainImageCode(vec3<float>(1, 1, 1));
+			}
+			mask = GenerateImageFormat(mask, EImageFormat::IF_L_UBYTE);
+			mask = GenerateImageSize
+			(
+				mask,
+				FImageSize((uint16)m_imageState.Last().m_imageRect.size[0],
+					(uint16)m_imageState.Last().m_imageRect.size[1])
+			);
+			op->mask = mask;
+
+			blendAd = op;
+		}
 
+		// Condition to enable this patch
+		if (conditionAd)
+		{
+			Ptr<ASTOp> conditionalAd;
+			{
+				Ptr<ASTOpConditional> op = new ASTOpConditional();
+				op->type = OP_TYPE::IM_CONDITIONAL;
+				op->no = blockAd;
+				op->yes = blendAd;
+				op->condition = conditionAd;
+				conditionalAd = op;
+			}
 
-    //---------------------------------------------------------------------------------------------
-    const NodeMeshConstant* FindSourceMesh( const Node* pNode )
-    {
-        const NodeMeshConstant* pResult = nullptr;
+			blockAd = conditionalAd;
+		}
+		else
+		{
+			blockAd = blendAd;
+		}
 
-        if (!pNode)
-        {
-            return pResult;
-        }
+		return blockAd;
+	}
 
-        if ( auto pTypedSN = dynamic_cast<const NodeSurfaceNew*>( pNode ) )
+
+	//---------------------------------------------------------------------------------------------
+	const NodeMeshConstant* FindSourceMesh(const Node* pNode)
+	{
+		const NodeMeshConstant* pResult = nullptr;
+
+		if (!pNode)
+		{
+			return pResult;
+		}
+
+		if (const NodeSurfaceNew* pTypedSN = dynamic_cast<const NodeSurfaceNew*>(pNode))
         {
             // TODO: 0...
             pResult = FindSourceMesh( pTypedSN->GetMesh(0).get() );
         }
-        else if ( auto pTypedSE = dynamic_cast<const NodeSurfaceEdit*>(pNode) )
+ 		else if (const NodeSurfaceEdit* pTypedSE = dynamic_cast<const NodeSurfaceEdit*>(pNode))
 		{
-            NodePatchMesh* pPatch = pTypedSE->GetMesh();
+			NodePatchMesh* pPatch = pTypedSE->GetMesh();
 			if (pPatch)
 			{
 				pResult = FindSourceMesh(pPatch->GetAdd());
 			}
 		}
-        else if ( auto pTypedSV = dynamic_cast<const NodeSurfaceVariation*>(pNode) )
+		else if (const NodeSurfaceVariation* pTypedSV = dynamic_cast<const NodeSurfaceVariation*>(pNode))
 		{
-            if (pTypedSV->GetPrivate()->m_defaultSurfaces.Num())
+			if (pTypedSV->GetPrivate()->m_defaultSurfaces.Num())
 			{
-                pResult = FindSourceMesh(pTypedSV->GetPrivate()->m_defaultSurfaces[0].get());
+				pResult = FindSourceMesh(pTypedSV->GetPrivate()->m_defaultSurfaces[0].get());
 			}
 		}
-        else if ( auto pTypedMI = dynamic_cast<const NodeMeshInterpolate*>( pNode ) )
+		else if (const NodeMeshInterpolate* pTypedMI = dynamic_cast<const NodeMeshInterpolate*>(pNode))
         {
             if ( pTypedMI->GetTargetCount()>0 )
             {
                 pResult = FindSourceMesh( pTypedMI->GetTarget(0).get() );
             }
         }
-        else if ( auto pTypedMM = dynamic_cast<const NodeMeshMorph*>( pNode ) )
+		else if (const NodeMeshMorph* pTypedMM = dynamic_cast<const NodeMeshMorph*>(pNode))
         {
             pResult = FindSourceMesh( pTypedMM->GetBase().get() );
         }
-		else if (auto pTypedGO = dynamic_cast<const NodeMeshGeometryOperation*>(pNode))
+		else if (const NodeMeshGeometryOperation* pTypedGO = dynamic_cast<const NodeMeshGeometryOperation*>(pNode))
 		{
 			pResult = FindSourceMesh(pTypedGO->GetMeshA().get());
 		}
-		else if (auto pTypedR = dynamic_cast<const NodeMeshReshape*>(pNode))
+		else if (const NodeMeshReshape* pTypedR = dynamic_cast<const NodeMeshReshape*>(pNode))
 		{
 			pResult = FindSourceMesh(pTypedR->GetBaseMesh().get());
 		}
-		else if ( auto pTypedMF = dynamic_cast<const NodeMeshFormat*>( pNode ) )
+		else if (const NodeMeshFormat* pTypedMF = dynamic_cast<const NodeMeshFormat*>(pNode))
         {
             if ( pTypedMF->GetSource() )
             {
                 pResult = FindSourceMesh( pTypedMF->GetSource().get() );
             }
         }
-        else if ( auto pTypedMFrag = dynamic_cast<const NodeMeshFragment*>( pNode ) )
+ 		else if (const NodeMeshFragment* pTypedMFrag = dynamic_cast<const NodeMeshFragment*>(pNode))
         {
             if ( pTypedMFrag->GetMesh() )
             {
                 pResult = FindSourceMesh( pTypedMFrag->GetMesh().get() );
             }
         }
-        else if ( auto pCasted = dynamic_cast<const NodeMeshConstant*>( pNode ) )
+		else if (const NodeMeshConstant* pCasted = dynamic_cast<const NodeMeshConstant*>(pNode))
         {
             pResult = pCasted;
         }
-        else if ( auto pTypedCMP = dynamic_cast<const NodeMeshClipMorphPlane*>(pNode) )
+ 		else if (const NodeMeshClipMorphPlane* pTypedCMP = dynamic_cast<const NodeMeshClipMorphPlane*>(pNode))
 		{
-            pResult = FindSourceMesh(pTypedCMP->GetSource().get());
+			pResult = FindSourceMesh(pTypedCMP->GetSource().get());
 		}
-        else if ( auto pTypedCWM = dynamic_cast<const NodeMeshClipWithMesh*>(pNode) )
+		else if (const NodeMeshClipWithMesh* pTypedCWM = dynamic_cast<const NodeMeshClipWithMesh*>(pNode))
 		{
-            pResult = FindSourceMesh(pTypedCWM->GetSource().get());
+			pResult = FindSourceMesh(pTypedCWM->GetSource().get());
 		}
 		else
-        {
-            check( false );
-        }
+		{
+			check(false);
+		}
 
-        return pResult;
-    }
+		return pResult;
+	}
 
 
-    //---------------------------------------------------------------------------------------------
-    Ptr<ASTOp> CodeGenerator::Visit( const NodePatchImage::Private& node )
-    {
-        // Get the parent component layout
-        const NodeObjectNew::Private* pParent = nullptr;
-        if (m_currentParents.Num()>2)
-        {
-            pParent = m_currentParents[m_currentParents.Num()-2].m_pObject;
-        }
+	//---------------------------------------------------------------------------------------------
+	Ptr<ASTOp> CodeGenerator::Visit(const NodePatchImage::Private& node)
+	{
+		// Get the parent component layout
+		const NodeObjectNew::Private* pParent = nullptr;
+		if (m_currentParents.Num() > 2)
+		{
+			pParent = m_currentParents[m_currentParents.Num() - 2].m_pObject;
+		}
 
-        const NodeLayout* pNodeLayout = nullptr;
-        if (pParent)
-        {
-            pNodeLayout = pParent->GetLayout
-                (
-                    m_currentParents.Last().m_lod,
-                    m_currentParents.Last().m_component,
-                    m_currentParents.Last().m_surface,
-                    m_currentParents.Last().m_texture
-                ).get();
-        }
+		const NodeLayout* pNodeLayout = nullptr;
+		if (pParent)
+		{
+			pNodeLayout = pParent->GetLayout
+			(
+				m_currentParents.Last().m_lod,
+				m_currentParents.Last().m_component,
+				m_currentParents.Last().m_surface,
+				m_currentParents.Last().m_texture
+			).get();
+		}
 
-        if (!pNodeLayout)
-        {
-            char buf[256];
-            mutable_snprintf
-                (
-                    buf, 256,
-                    "In object [%s] NodePatchImage couldn't find the layout in parent.",
-                    m_currentParents.Last().m_pObject->m_name.c_str()
-                );
-
-            m_pErrorLog->GetPrivate()->Add( buf, ELMT_ERROR, node.m_errorContext );
-        }
+		if (!pNodeLayout)
+		{
+			char buf[256];
+			mutable_snprintf
+			(
+				buf, 256,
+				"In object [%s] NodePatchImage couldn't find the layout in parent.",
+				m_currentParents.Last().m_pObject->m_name.c_str()
+			);
+
+			m_pErrorLog->GetPrivate()->Add(buf, ELMT_ERROR, node.m_errorContext);
+		}
 
-        return 0;
-    }
+		return 0;
+	}
 
 
-    //---------------------------------------------------------------------------------------------
-    Ptr<ASTOp> CodeGenerator::Visit( const NodeLOD::Private& node )
-    {
-        // Build a series of operations to assemble all the LOD components
-        Ptr<ASTOp> lastCompOp;
+	//---------------------------------------------------------------------------------------------
+	Ptr<ASTOp> CodeGenerator::Visit(const NodeLOD::Private& node)
+	{
+		// Build a series of operations to assemble all the LOD components
+		Ptr<ASTOp> lastCompOp;
 
-        // Create the expression for each component in this object
-        // TODO: More components per operation
-        for ( int32 t=0; t<node.m_components.Num(); ++t )
-        {
-            if ( const NodeComponent* pComponentNode = node.m_components[t].get() )
-            {
-                m_currentParents.Last().m_component = (int)t;
+		// Create the expression for each component in this object
+		// TODO: More components per operation
+		for (int32 t = 0; t < node.m_components.Num(); ++t)
+		{
+			if (const NodeComponent* pComponentNode = node.m_components[t].get())
+			{
+				m_currentParents.Last().m_component = (int)t;
 
-                Ptr<ASTOp> componentOp = Generate( pComponentNode );
+				Ptr<ASTOp> componentOp = Generate(pComponentNode);
 
-                if (componentOp)
-                {
-                    check( componentOp->GetOpType() == OP_TYPE::IN_ADDCOMPONENT );
+				if (componentOp)
+				{
+					check(componentOp->GetOpType() == OP_TYPE::IN_ADDCOMPONENT);
 
-                    auto typedOp = dynamic_cast<ASTOpInstanceAdd*>( componentOp.get() );
+					ASTOpInstanceAdd* typedOp = dynamic_cast<ASTOpInstanceAdd*>(componentOp.get());
 
-                    // Complete the instruction adding the base
-                    typedOp->instance = lastCompOp;
+					// Complete the instruction adding the base
+					typedOp->instance = lastCompOp;
 
-                    lastCompOp = componentOp;
-                }
-            }
-        }
+					lastCompOp = componentOp;
+				}
+			}
+		}
 
-        // Add components from child objects
-        ADDITIONAL_COMPONENT_KEY thisKey;
-        thisKey.m_lod = m_currentParents.Last().m_lod;
-        thisKey.m_pObject = m_currentParents.Last().m_pObject;
-        auto addIt = m_additionalComponents.find( thisKey );
-        if (addIt!=m_additionalComponents.end())
-        {
-            for( const auto& cop: addIt->second )
-            {
-                // Add the additional components after the main ones, this means higher up in the
-                // op tree.
+		// Add components from child objects
+		ADDITIONAL_COMPONENT_KEY thisKey;
+		thisKey.m_lod = m_currentParents.Last().m_lod;
+		thisKey.m_pObject = m_currentParents.Last().m_pObject;
+		auto addIt = m_additionalComponents.find(thisKey);
+		if (addIt != m_additionalComponents.end())
+		{
+			for (const auto& cop : addIt->second)
+			{
+				// Add the additional components after the main ones, this means higher up in the
+				// op tree.
 
                 // First find the last op in the chain of IN_ADDCOMPONENT operations
                 check( cop->GetOpType() == OP_TYPE::IN_ADDCOMPONENT );
-                auto typedOp = dynamic_cast<ASTOpInstanceAdd*>( cop.get() );
-                auto bottomOp = typedOp;
-                while( bottomOp->instance )
-                {
-                    // Step down
-                    check( bottomOp->instance->GetOpType() == OP_TYPE::IN_ADDCOMPONENT );
-                    bottomOp = dynamic_cast<ASTOpInstanceAdd*>( bottomOp->instance.child().get() );
-                }
+				ASTOpInstanceAdd* typedOp = dynamic_cast<ASTOpInstanceAdd*>(cop.get());
+				ASTOpInstanceAdd* bottomOp = typedOp;
+				while (bottomOp->instance)
+				{
+					// Step down
+					check(bottomOp->instance->GetOpType() == OP_TYPE::IN_ADDCOMPONENT);
+					bottomOp = dynamic_cast<ASTOpInstanceAdd*>(bottomOp->instance.child().get());
+				}
 
-                // Chain
-                bottomOp->instance = lastCompOp;
-                lastCompOp = typedOp;
-            }
-        }
+				// Chain
+				bottomOp->instance = lastCompOp;
+				lastCompOp = typedOp;
+			}
+		}
 
-        // Store for possible parent objects if necessary
-        // 2 is because there must be a parent and there is always a null element as well.
-        if (lastCompOp && m_currentParents.Num()>2)
-        {
-            const auto& parentObjectKey = m_currentParents[ m_currentParents.Num()-2 ];
-            ADDITIONAL_COMPONENT_KEY parentKey;
-            parentKey.m_lod = m_currentParents.Last().m_lod;
-            parentKey.m_pObject = parentObjectKey.m_pObject;
-            m_additionalComponents[parentKey].Add( lastCompOp );
-        }
+		// Store for possible parent objects if necessary
+		// 2 is because there must be a parent and there is always a null element as well.
+		if (lastCompOp && m_currentParents.Num() > 2)
+		{
+			const auto& parentObjectKey = m_currentParents[m_currentParents.Num() - 2];
+			ADDITIONAL_COMPONENT_KEY parentKey;
+			parentKey.m_lod = m_currentParents.Last().m_lod;
+			parentKey.m_pObject = parentObjectKey.m_pObject;
+			m_additionalComponents[parentKey].Add(lastCompOp);
+		}
 
-        return lastCompOp;
-    }
+		return lastCompOp;
+	}
+
+
+	//---------------------------------------------------------------------------------------------
+	Ptr<ASTOp> CodeGenerator::ApplyTiling(Ptr<ASTOp> Source, FImageSize Size, EImageFormat Format)
+	{
+		// For now always apply tiling
+		if (m_compilerOptions->m_imageTiling==0)
+		{
+			return Source;
+		}
+
+		int32 TileSize = m_compilerOptions->m_imageTiling;
 
+		int32 TilesX = FMath::DivideAndRoundUp<int32>(Size[0], TileSize);
+		int32 TilesY = FMath::DivideAndRoundUp<int32>(Size[1], TileSize);
+		if (TilesX * TilesY <= 2)
+		{
+			return Source;
+		}
+
+		Ptr<ASTOpFixed> BaseImage = new ASTOpFixed;
+		BaseImage->op.type = OP_TYPE::IM_PLAINCOLOUR;
+		BaseImage->op.args.ImagePlainColour.size[0] = Size[0];
+		BaseImage->op.args.ImagePlainColour.size[1] = Size[1];
+		BaseImage->op.args.ImagePlainColour.format = Format;
+		BaseImage->op.args.ImagePlainColour.LODs = 1;
+
+		Ptr<ASTOp> CurrentImage = BaseImage;
+
+		for (int32 Y = 0; Y < TilesY; ++Y)
+		{
+			for (int32 X = 0; X < TilesX; ++X)
+			{
+				int32 MinX = X * TileSize;
+				int32 MinY = Y * TileSize;
+				int32 TileSizeX = FMath::Min(TileSize, Size[0] - MinX);
+				int32 TileSizeY = FMath::Min(TileSize, Size[1] - MinY);
+
+				Ptr<ASTOpFixed> TileImage = new ASTOpFixed();
+				TileImage->op.type = OP_TYPE::IM_CROP;
+				TileImage->SetChild(TileImage->op.args.ImageCrop.source, Source);
+				TileImage->op.args.ImageCrop.minX = MinX;
+				TileImage->op.args.ImageCrop.minY = MinY;
+				TileImage->op.args.ImageCrop.sizeX = TileSizeX;
+				TileImage->op.args.ImageCrop.sizeY = TileSizeY;
+
+				Ptr<ASTOpImagePatch> PatchedImage = new ASTOpImagePatch();
+				PatchedImage->base = CurrentImage;
+				PatchedImage->patch = TileImage;
+				PatchedImage->location[0] = MinX;
+				PatchedImage->location[1] = MinY;
+
+				CurrentImage = PatchedImage;
+			}
+		}
+
+		return CurrentImage;
+	}
 
 
     //---------------------------------------------------------------------------------------------
@@ -702,7 +760,7 @@ namespace mu
     {
         //MUTABLE_CPUPROFILER_SCOPE(GenerateSurface);
 
-        const auto& node = *surfaceNode->GetPrivate();
+        const NodeSurfaceNew::Private& node = *surfaceNode->GetPrivate();
 
         // Clear the surface generation state
 //        m_generatedLayouts.Empty();
@@ -799,7 +857,7 @@ namespace mu
             // Apply mesh merges from child objects "edit surface" nodes
             for ( int32 editIndex=0; editIndex<edits.Num(); ++editIndex )
             {
-                const auto& e = edits[editIndex];
+                const FirstPassGenerator::SURFACE::EDIT& e = edits[editIndex];
 
                 if ( e.node->m_pMesh )
                 {
@@ -861,7 +919,7 @@ namespace mu
             // "Removes" need to come after "Adds" because some removes may refer to added meshes,
             // and not the base.
             Ptr<ASTOpMeshRemoveMask> rop;
-            for ( const auto& e: edits )
+            for ( const FirstPassGenerator::SURFACE::EDIT& e: edits )
             {
                 if ( e.node->m_pMesh )
                 {
@@ -883,7 +941,7 @@ namespace mu
                         maskOp->op.type = OP_TYPE::ME_MASKDIFF;
 
                         // By default, remove from the base
-                        auto removeFrom = meshResults.baseMeshOp;
+                        Ptr<ASTOp> removeFrom = meshResults.baseMeshOp;
                         // See if we want to remove from an added mesh instead.
                         if ( e.node->m_pParent )
                         {
@@ -914,7 +972,7 @@ namespace mu
             }
 
             // Apply mesh morphs from child objects "edit surface" nodes
-            for ( const auto& e: edits )
+            for ( const FirstPassGenerator::SURFACE::EDIT& e: edits )
             {
                 if ( NodeMeshPtr pMorph = e.node->m_pMorph )
                 {
@@ -1032,7 +1090,7 @@ namespace mu
                     }
 
                     // Add children merged meshes layouts
-                    for ( const auto& data : meshResults.extraMeshLayouts )
+                    for ( const FMeshGenerationResult::FExtraLayouts& data : meshResults.extraMeshLayouts )
                     {
                         if ( !data.meshFragment )
                         {
@@ -1220,7 +1278,7 @@ namespace mu
                     // Look for patches to this block
                     for( int32 editIndex=0; editIndex<edits.Num(); ++editIndex )
                     {
-                        const auto& e = edits[editIndex];
+                        const FirstPassGenerator::SURFACE::EDIT& e = edits[editIndex];
                         if ( t<e.node->m_textures.Num() )
                         {
                             if ( const NodePatchImage* pPatch = e.node->m_textures[t].m_pPatch.get() )
@@ -1350,7 +1408,7 @@ namespace mu
                         for( int b=0; b<pLayout->GetBlockCount(); ++b )
                         {
                             // Block in layout grid units
-                            box< UE::Math::TIntVector2<int32> > rectInCells;
+                            box< UE::Math::TIntVector2<uint16> > rectInCells;
                             pLayout->GetBlock
                                 (
                                     b,
@@ -1359,7 +1417,7 @@ namespace mu
                                 );
 
                             // Transform to pixels
-                            box< UE::Math::TIntVector2<int32> > rect = rectInCells;
+                            box< UE::Math::TIntVector2<uint16> > rect = rectInCells;
                             rect.min[0] *= blockSizeX;
                             rect.min[1] *= blockSizeY;
                             rect.size[0] *= blockSizeX;
@@ -1368,8 +1426,9 @@ namespace mu
                             // Generate the image
                             IMAGE_STATE newState;
                             newState.m_imageSize = FIntVector2(desc.m_size);
-                            newState.m_imageRect = rect;
-                            newState.m_layoutBlockId = pLayout->m_blocks[b].m_id;
+							newState.m_imageRect.min = FIntVector2(rect.min);
+							newState.m_imageRect.size = FIntVector2(rect.size);
+							newState.m_layoutBlockId = pLayout->m_blocks[b].m_id;
                             newState.m_pLayout = pLayout;
                             m_imageState.Add( newState );
 
@@ -1378,7 +1437,7 @@ namespace mu
                             // Look for patches to this block
                             for( int32 editIndex=0; editIndex<edits.Num(); ++editIndex )
                             {
-                                const auto& e = edits[editIndex];
+                                const FirstPassGenerator::SURFACE::EDIT& e = edits[editIndex];
                                 if ( t<e.node->m_textures.Num() )
                                 {
                                     if ( const NodePatchImage* pPatch = e.node->m_textures[t].m_pPatch.get() )
@@ -1392,12 +1451,17 @@ namespace mu
                                 }
                             }
 
+							EImageFormat baseFormat = imageAd->GetImageDesc().m_format;
+							Ptr<ASTOp> FormattedBlock = GenerateImageFormat(blockAd, baseFormat);
+
+							// Apply tiling to avoid generating chunks of image that are too big.
+							FormattedBlock = ApplyTiling(FormattedBlock, FImageSize(rect.size), desc.m_format);
+
                             // Compose layout operation
                             Ptr<ASTOpImageCompose> composeOp = new ASTOpImageCompose();
                             composeOp->Layout = meshResults.layoutOps[ layout ];
                             composeOp->Base = imageAd;
-							EImageFormat baseFormat = imageAd->GetImageDesc().m_format;
-                            composeOp->BlockImage = GenerateImageFormat( blockAd, baseFormat );
+							composeOp->BlockImage = FormattedBlock;
 
                             // Set the absolute block index.
                             check( pLayout->m_blocks[b].m_id>=0 );
@@ -1412,10 +1476,10 @@ namespace mu
                         // Apply composition of blocks coming from child objects
                         for( int32 editIndex=0; editIndex<edits.Num(); ++editIndex )
                         {
-                            const auto& e = edits[editIndex];
+                            const FirstPassGenerator::SURFACE::EDIT& e = edits[editIndex];
                             if ( t<e.node->m_textures.Num() )
                             {
-                                auto pExtend = e.node->m_textures[t].m_pExtend;
+                                Ptr<NodeImage> pExtend = e.node->m_textures[t].m_pExtend;
                                 if (pExtend)
                                 {
                                     if ( layout>=meshResults.extraMeshLayouts[editIndex].GeneratedLayouts.Num()
@@ -1449,7 +1513,7 @@ namespace mu
                                         for (int b=0; b<pExtendLayout->GetBlockCount(); ++b)
                                         {
                                             // Block in layout grid units
-                                            box< UE::Math::TIntVector2<int32> > blockRect;
+                                            box< UE::Math::TIntVector2<uint16> > blockRect;
                                             pExtendLayout->GetBlock
                                                 (
                                                     b,
@@ -1458,11 +1522,11 @@ namespace mu
                                                 );
 
                                             // Transform to pixels
-                                            box< UE::Math::TIntVector2<int32> > rect = blockRect;
-                                            rect.min[0] = (rect.min[0] * extendDesc.m_size[0]) / extlayout[0];
-                                            rect.min[1] = (rect.min[1] * extendDesc.m_size[1]) / extlayout[1];
-                                            rect.size[0] = (rect.size[0] * extendDesc.m_size[0]) / extlayout[0];
-                                            rect.size[1] = (rect.size[1] * extendDesc.m_size[1]) / extlayout[1];
+                                            box< UE::Math::TIntVector2<int32> > rect;
+                                            rect.min[0] = (blockRect.min[0] * extendDesc.m_size[0]) / extlayout[0];
+                                            rect.min[1] = (blockRect.min[1] * extendDesc.m_size[1]) / extlayout[1];
+                                            rect.size[0] = (blockRect.size[0] * extendDesc.m_size[0]) / extlayout[0];
+                                            rect.size[1] = (blockRect.size[1] * extendDesc.m_size[1]) / extlayout[1];
 
                                             // Generate the image block
                                             IMAGE_STATE newState;
@@ -1487,6 +1551,8 @@ namespace mu
                                             //formatted = formatted;
                                             fragmentAd = formatted;
 
+											// Apply tiling to avoid generating chunks of image that are too big.
+											fragmentAd = ApplyTiling(fragmentAd, expectedSize, desc.m_format);
 
                                             // Compose operation
                                             Ptr<ASTOpImageCompose> composeOp = new ASTOpImageCompose();
@@ -1697,7 +1763,7 @@ namespace mu
         // Look for all surfaces that belong to this component
 		for (int32 i = 0; i<m_firstPass.surfaces.Num(); ++i, ++surfaceID)
 		{
-			const auto& its = m_firstPass.surfaces[i];
+			const FirstPassGenerator::SURFACE& its = m_firstPass.surfaces[i];
 			if (its.component==&node)
 			{
                 // Apply state conditions: only generate it if it enabled in this state
@@ -1910,7 +1976,7 @@ namespace mu
 				bool found = false;
                 for( int32 i = 0; !found && i != m_firstPass.objects.Num(); i++ )
 				{
-					auto& it = m_firstPass.objects[i];
+					FirstPassGenerator::OBJECT& it = m_firstPass.objects[i];
 					if (it.node == pChildNode->GetBasePrivate())
 					{
 						found = true;
@@ -1993,7 +2059,7 @@ namespace mu
 
 		if (tags.Num())
 		{
-			for (const auto& m: m_firstPass.modifiers)
+			for (const FirstPassGenerator::MODIFIER& m: m_firstPass.modifiers)
 			{
 				// Correct LOD?
 				if (m.lod != LOD)
@@ -2056,7 +2122,7 @@ namespace mu
 		Ptr<ASTOpMeshRemoveMask> removeOp;
 		for (const FirstPassGenerator::MODIFIER& m : modifiers)
 		{
-			if (auto TypedClipNode = dynamic_cast<const NodeModifierMeshClipWithMesh::Private*>(m.node))
+			if (const NodeModifierMeshClipWithMesh::Private* TypedClipNode = dynamic_cast<const NodeModifierMeshClipWithMesh::Private*>(m.node))
 			{
 				Ptr<ASTOpMeshMaskClipMesh> op = new ASTOpMeshMaskClipMesh();
 				op->source = preModifiersMesh;
@@ -2103,7 +2169,7 @@ namespace mu
 		{
 			Ptr<ASTOp> modifiedMeshOp;
 
-			if (auto TypedNode = dynamic_cast<const NodeModifierMeshClipMorphPlane::Private*>(m.node))
+			if (const NodeModifierMeshClipMorphPlane::Private* TypedNode = dynamic_cast<const NodeModifierMeshClipMorphPlane::Private*>(m.node))
 			{
 				Ptr<ASTOpMeshClipMorphPlane> op = new ASTOpMeshClipMorphPlane();
 				op->source = lastMeshOp;
diff --git a/Engine/Plugins/Experimental/Mutable/Source/MutableTools/Private/MuT/CodeGenerator.h b/Engine/Plugins/Experimental/Mutable/Source/MutableTools/Private/MuT/CodeGenerator.h
--- a/Engine/Plugins/Experimental/Mutable/Source/MutableTools/Private/MuT/CodeGenerator.h
+++ b/Engine/Plugins/Experimental/Mutable/Source/MutableTools/Private/MuT/CodeGenerator.h
@@ -188,8 +188,8 @@ namespace mu
             // This reference has to be the smart pointer to avoid memory aliasing, keeping
             // processed nodes alive.
             NodePtrConst pNode;
-			UE::Math::TIntVector2<int> imageSize;
-            box< UE::Math::TIntVector2<int> > imageRect;
+			UE::Math::TIntVector2<int32> imageSize;
+            box< UE::Math::TIntVector2<int32> > imageRect;
             int state = -1;
 			TArray<mu::string> activeTags;
 			TArray<LayoutPtrConst> overrideLayouts;
@@ -424,6 +424,8 @@ namespace mu
 		//!
 		FImageDesc CalculateImageDesc(const Node::Private&);
 
+		/** Evaluate if the image to generate is big enough to be split in separate operations and tiled afterwards. */
+		Ptr<ASTOp> ApplyTiling(Ptr<ASTOp> Source, FImageSize Size, EImageFormat Format);
 
         //-----------------------------------------------------------------------------------------
         // Meshes
diff --git a/Engine/Plugins/Experimental/Mutable/Source/MutableTools/Private/MuT/CodeGenerator_Image.cpp b/Engine/Plugins/Experimental/Mutable/Source/MutableTools/Private/MuT/CodeGenerator_Image.cpp
--- a/Engine/Plugins/Experimental/Mutable/Source/MutableTools/Private/MuT/CodeGenerator_Image.cpp
+++ b/Engine/Plugins/Experimental/Mutable/Source/MutableTools/Private/MuT/CodeGenerator_Image.cpp
@@ -136,7 +136,7 @@ namespace mu
 			IMAGE_STATE newState;
 			newState.m_imageSize[0] = desc.m_size[0] ? desc.m_size[0] : 256;
 			newState.m_imageSize[1] = desc.m_size[1] ? desc.m_size[1] : 256;
-			newState.m_imageRect.size = FIntVector2(desc.m_size);
+			newState.m_imageRect.size = UE::Math::TIntVector2<int32>(desc.m_size);
 			newState.m_imageRect.min[0] = 0;
 			newState.m_imageRect.min[1] = 0;
 			newState.m_layoutBlockId = -1;
@@ -1276,7 +1276,8 @@ namespace mu
         op->SetChild( op->op.args.ImagePlainColour.colour, base);
 		op->op.args.ImagePlainColour.format = node.Format;
         op->op.args.ImagePlainColour.size[0] = uint16(node.m_sizeX);
-        op->op.args.ImagePlainColour.size[1] = uint16(node.m_sizeY);
+		op->op.args.ImagePlainColour.size[1] = uint16(node.m_sizeY);
+		op->op.args.ImagePlainColour.LODs = 1;
 
         Ptr<ASTOpFixed> opSize = new ASTOpFixed();
         opSize->op.type = OP_TYPE::IM_RESIZE;
@@ -1572,10 +1573,10 @@ namespace mu
         }
 
         // Image size, from the current block being generated
-        op->sizeX = (uint16)m_imageState.Last().m_imageRect.size[0];
-        op->sizeY = (uint16)m_imageState.Last().m_imageRect.size[1];
+        op->SizeX = (uint16)m_imageState.Last().m_imageRect.size[0];
+        op->SizeY = (uint16)m_imageState.Last().m_imageRect.size[1];
 		// \TODO: Review naming of arg
-        op->blockIndex = GeneratedLayoutBlockId;
+        op->BlockIndex = GeneratedLayoutBlockId;
 
 		op->bIsRGBFadingEnabled = node.bIsRGBFadingEnabled;
 		op->bIsAlphaFadingEnabled = node.bIsAlphaFadingEnabled;
@@ -1602,11 +1603,9 @@ namespace mu
         // Target mask
         if ( node.m_pMask )
         {
-            auto mask = Generate( node.m_pMask.get() );
+            Ptr<ASTOp> mask = Generate( node.m_pMask.get() );
             mask = GenerateImageFormat( mask, EImageFormat::IF_L_UBYTE );
-            op->mask = GenerateImageSize( mask,
-                      FImageSize( (uint16)m_imageState.Last().m_imageRect.size[0],
-                                  (uint16)m_imageState.Last().m_imageRect.size[1]) );
+            op->mask = GenerateImageSize( mask, FImageSize( m_imageState.Last().m_imageRect.size ) );
         }
 
         // Seam correction operations
@@ -1615,11 +1614,13 @@ namespace mu
         rasterop->mesh = op->mesh.child();
         rasterop->image = 0;
         rasterop->mask = 0;
-        rasterop->blockIndex = op->blockIndex;
-        rasterop->sizeX = op->sizeX;
-        rasterop->sizeY = op->sizeY;
-		rasterop->SourceSizeX = op->sizeX;
-		rasterop->SourceSizeY = op->sizeX;
+        rasterop->BlockIndex = op->BlockIndex;
+		rasterop->SizeX = op->SizeX;
+		rasterop->SizeY = op->SizeY;
+		rasterop->UncroppedSizeX = op->UncroppedSizeX;
+		rasterop->UncroppedSizeY = op->UncroppedSizeY;
+		rasterop->CropMinX = op->CropMinX;
+		rasterop->CropMinY = op->CropMinY;
 		rasterop->SamplingMethod = ESamplingMethod::Point;
 		rasterop->MinFilterMethod = EMinFilterMethod::None;
 
diff --git a/Engine/Plugins/Experimental/Mutable/Source/MutableTools/Private/MuT/CodeGenerator_Mesh.cpp b/Engine/Plugins/Experimental/Mutable/Source/MutableTools/Private/MuT/CodeGenerator_Mesh.cpp
--- a/Engine/Plugins/Experimental/Mutable/Source/MutableTools/Private/MuT/CodeGenerator_Mesh.cpp
+++ b/Engine/Plugins/Experimental/Mutable/Source/MutableTools/Private/MuT/CodeGenerator_Mesh.cpp
@@ -164,9 +164,8 @@ class Node;
         {
             FIntPoint grid = pLayout->GetGridSize();
 
-            box< vec2<int> > block;
-            pLayout->GetBlock
-                ( b, &block.min[0], &block.min[1], &block.size[0], &block.size[1] );
+            box< UE::Math::TIntVector2<uint16> > block;
+            pLayout->GetBlock( b, &block.min[0], &block.min[1], &block.size[0], &block.size[1] );
 
             box< vec2<float> > rect;
             rect.min[0] = ( (float)block.min[0] ) / (float) grid[0];
diff --git a/Engine/Plugins/Experimental/Mutable/Source/MutableTools/Private/MuT/CodeOptimiser_Parameters.cpp b/Engine/Plugins/Experimental/Mutable/Source/MutableTools/Private/MuT/CodeOptimiser_Parameters.cpp
--- a/Engine/Plugins/Experimental/Mutable/Source/MutableTools/Private/MuT/CodeOptimiser_Parameters.cpp
+++ b/Engine/Plugins/Experimental/Mutable/Source/MutableTools/Private/MuT/CodeOptimiser_Parameters.cpp
@@ -250,7 +250,8 @@ namespace mu
             wplainOp->SetChild( wplainOp->op.args.ImagePlainColour.colour, whiteOp );
             wplainOp->op.args.ImagePlainColour.format = EImageFormat::IF_L_UBYTE;
             wplainOp->op.args.ImagePlainColour.size[0] = 4;
-            wplainOp->op.args.ImagePlainColour.size[1] = 4;
+			wplainOp->op.args.ImagePlainColour.size[1] = 4;
+			wplainOp->op.args.ImagePlainColour.LODs = 1;
 
             Ptr<ASTOpFixed> wresizeOp = new ASTOpFixed;
             wresizeOp->op.type = OP_TYPE::IM_RESIZELIKE;
@@ -475,7 +476,8 @@ namespace mu
                             plainOp->SetChild( plainOp->op.args.ImagePlainColour.colour, blackOp );
                             plainOp->op.args.ImagePlainColour.format = EImageFormat::IF_L_UBYTE;
                             plainOp->op.args.ImagePlainColour.size[0] = 4;
-                            plainOp->op.args.ImagePlainColour.size[1] = 4;
+							plainOp->op.args.ImagePlainColour.size[1] = 4;
+							plainOp->op.args.ImagePlainColour.LODs = 1;
 
                             Ptr<ASTOpFixed> resizeOp = new ASTOpFixed;
                             resizeOp->op.type = OP_TYPE::IM_RESIZELIKE;
@@ -519,6 +521,8 @@ namespace mu
 //							plainOp.type = OP_TYPE::IM_PLAINCOLOUR;
 //							plainOp.args.ImagePlainColour.colour = program.AddOp( blackOp );
 //							plainOp.args.ImagePlainColour.format = IF_L_UBYTE;
+//							plainOp.args.ImagePlainColour.size = ;
+//							plainOp.args.ImagePlainColour.LODs = ;
 
 //							OP resizeOp;
 //							resizeOp.type = OP_TYPE::IM_RESIZELIKE;
@@ -579,7 +583,8 @@ namespace mu
                             plainOp->SetChild( plainOp->op.args.ImagePlainColour.colour, blackOp );
                             plainOp->op.args.ImagePlainColour.format = EImageFormat::IF_L_UBYTE;
                             plainOp->op.args.ImagePlainColour.size[0] = 4;
-                            plainOp->op.args.ImagePlainColour.size[1] = 4;
+							plainOp->op.args.ImagePlainColour.size[1] = 4;
+							plainOp->op.args.ImagePlainColour.LODs = 1;
 
                             Ptr<ASTOpFixed> resizeOp = new ASTOpFixed;
                             resizeOp->op.type = OP_TYPE::IM_RESIZELIKE;
@@ -1046,7 +1051,8 @@ namespace mu
                         plainOp->SetChild( plainOp->op.args.ImagePlainColour.colour, blackOp );
                         plainOp->op.args.ImagePlainColour.format = EImageFormat::IF_L_UBYTE;
                         plainOp->op.args.ImagePlainColour.size[0] = 4;
-                        plainOp->op.args.ImagePlainColour.size[1] = 4;
+						plainOp->op.args.ImagePlainColour.size[1] = 4;
+						plainOp->op.args.ImagePlainColour.LODs = 1;
 
                         Ptr<ASTOpFixed> baseResizeOp = new ASTOpFixed;
                         baseResizeOp->op.type = OP_TYPE::IM_RESIZELIKE;
@@ -1102,7 +1108,8 @@ namespace mu
                         plainOp->SetChild( plainOp->op.args.ImagePlainColour.colour, blackOp );
                         plainOp->op.args.ImagePlainColour.format = EImageFormat::IF_L_UBYTE;
                         plainOp->op.args.ImagePlainColour.size[0] = 4;
-                        plainOp->op.args.ImagePlainColour.size[1] = 4;
+						plainOp->op.args.ImagePlainColour.size[1] = 4;
+						plainOp->op.args.ImagePlainColour.LODs = 1;
 
                         Ptr<ASTOpFixed> baseResizeOp = new ASTOpFixed;
                         baseResizeOp->op.type = OP_TYPE::IM_RESIZELIKE;
@@ -1130,7 +1137,8 @@ namespace mu
                         FImageDesc blendedDesc = baseAt->GetImageDesc();
                         plainOp->op.args.ImagePlainColour.format = blendedDesc.m_format;
                         plainOp->op.args.ImagePlainColour.size[0] = 4;
-                        plainOp->op.args.ImagePlainColour.size[1] = 4;
+						plainOp->op.args.ImagePlainColour.size[1] = 4;
+						plainOp->op.args.ImagePlainColour.LODs = 1;
 
                         Ptr<ASTOpFixed> resizeOp = new ASTOpFixed;
                         resizeOp->op.type = OP_TYPE::IM_RESIZELIKE;
@@ -1213,7 +1221,8 @@ namespace mu
                         plainOp->SetChild( plainOp->op.args.ImagePlainColour.colour, blackOp );
                         plainOp->op.args.ImagePlainColour.format = EImageFormat::IF_L_UBYTE; //TODO: FORMAT_LIKE
                         plainOp->op.args.ImagePlainColour.size[0] = 4;
-                        plainOp->op.args.ImagePlainColour.size[1] = 4;
+						plainOp->op.args.ImagePlainColour.size[1] = 4;
+						plainOp->op.args.ImagePlainColour.LODs = 1;
 
                         Ptr<ASTOpFixed> blockResizeOp = new ASTOpFixed;
                         blockResizeOp->op.type = OP_TYPE::IM_RESIZELIKE;
@@ -1259,7 +1268,8 @@ namespace mu
                         plainOp->SetChild( plainOp->op.args.ImagePlainColour.colour, blackOp );
                         plainOp->op.args.ImagePlainColour.format = EImageFormat::IF_L_UBYTE; //TODO: FORMAT_LIKE
                         plainOp->op.args.ImagePlainColour.size[0] = 4;
-                        plainOp->op.args.ImagePlainColour.size[1] = 4;
+						plainOp->op.args.ImagePlainColour.size[1] = 4;
+						plainOp->op.args.ImagePlainColour.LODs = 1;
 
                         Ptr<ASTOpFixed> blockResizeOp = new ASTOpFixed;
                         blockResizeOp->op.type = OP_TYPE::IM_RESIZELIKE;
diff --git a/Engine/Plugins/Experimental/Mutable/Source/MutableTools/Private/MuT/CodeOptimiser_Semantic.cpp b/Engine/Plugins/Experimental/Mutable/Source/MutableTools/Private/MuT/CodeOptimiser_Semantic.cpp
--- a/Engine/Plugins/Experimental/Mutable/Source/MutableTools/Private/MuT/CodeOptimiser_Semantic.cpp
+++ b/Engine/Plugins/Experimental/Mutable/Source/MutableTools/Private/MuT/CodeOptimiser_Semantic.cpp
@@ -106,7 +106,7 @@ namespace mu
 
         case OP_TYPE::BO_NOT:
         {
-            auto sourceAt = children[op.args.BoolNot.source].child();
+			Ptr<ASTOp> sourceAt = children[op.args.BoolNot.source].child();
             if  ( sourceAt->GetOpType() == OP_TYPE::BO_CONSTANT )
             {
                 Ptr<ASTOpConstantBool> newAt = new ASTOpConstantBool();
@@ -119,8 +119,8 @@ namespace mu
         case OP_TYPE::BO_AND:
         {
             bool changed = false;
-            auto aAt = children[op.args.BoolBinary.a].child();
-            auto bAt = children[op.args.BoolBinary.b].child();
+			Ptr<ASTOp> aAt = children[op.args.BoolBinary.a].child();
+			Ptr<ASTOp> bAt = children[op.args.BoolBinary.b].child();
             if ( !aAt )
             {
                 at = bAt;
@@ -201,8 +201,8 @@ namespace mu
         case OP_TYPE::BO_OR:
         {
             bool changed = false;
-            auto aAt = children[op.args.BoolBinary.a].child();
-            auto bAt = children[op.args.BoolBinary.b].child();
+			Ptr<ASTOp> aAt = children[op.args.BoolBinary.a].child();
+			Ptr<ASTOp> bAt = children[op.args.BoolBinary.b].child();
             if ( !aAt )
             {
                 at = bAt;
@@ -284,26 +284,25 @@ namespace mu
 		//-------------------------------------------------------------------------------------
 		case OP_TYPE::IM_RESIZEREL:
 		{
-			auto sourceAt = children[op.args.ImageResizeRel.source].child();
+			Ptr<ASTOp> sourceAt = children[op.args.ImageResizeRel.source].child();
 
 			// The instruction can be sunk
 			OP_TYPE sourceType = sourceAt->GetOpType();
 			switch (sourceType)
 			{
-
 			// This is done here instead of in the OptimizeSize step to prevent removing 
 			// resize_rel too early.
-            case OP_TYPE::IM_RESIZE:
-            {
+			case OP_TYPE::IM_RESIZE:
+			{
 				Ptr<ASTOpFixed> NewOp = mu::Clone<ASTOpFixed>(sourceAt.get());
 				NewOp->op.args.ImageResize.size[0] =
-                        int16_t(NewOp->op.args.ImageResize.size[0] * op.args.ImageResizeRel.factor[0] );
+					int16_t(NewOp->op.args.ImageResize.size[0] * op.args.ImageResizeRel.factor[0]);
 				NewOp->op.args.ImageResize.size[1] =
-                        int16_t(NewOp->op.args.ImageResize.size[1] * op.args.ImageResizeRel.factor[1] );
+					int16_t(NewOp->op.args.ImageResize.size[1] * op.args.ImageResizeRel.factor[1]);
 
 				at = NewOp;
-                break;
-            }
+				break;
+			}
 
 
 			default:
@@ -314,6 +313,32 @@ namespace mu
 			break;
 		}
 
+		//-------------------------------------------------------------------------------------
+		case OP_TYPE::IM_CROP:
+		{
+			Ptr<ASTOp> sourceAt = children[op.args.ImageCrop.source].child();
+
+			// The instruction can be sunk
+			OP_TYPE sourceType = sourceAt->GetOpType();
+			switch (sourceType)
+			{
+			case OP_TYPE::IM_PLAINCOLOUR:
+			{
+				Ptr<ASTOpFixed> NewOp = mu::Clone<ASTOpFixed>(sourceAt.get());
+				NewOp->op.args.ImagePlainColour.size[0] = op.args.ImageCrop.sizeX;
+				NewOp->op.args.ImagePlainColour.size[1] = op.args.ImageCrop.sizeY;
+				NewOp->op.args.ImagePlainColour.LODs = 1; // TODO
+				at = NewOp;
+				break;
+			}
+
+			default:
+				break;
+
+			}
+
+			break;
+		}
 
         //-------------------------------------------------------------------------------------
         case OP_TYPE::IM_PIXELFORMAT:
@@ -337,43 +362,43 @@ namespace mu
     //---------------------------------------------------------------------------------------------
     bool SinkOptimiserAST
     (
-            ASTOpList& roots,
-            const FModelOptimizationOptions& optimisationOptions
+            ASTOpList& InRoots,
+            const FModelOptimizationOptions& InOptimisationOptions
     )
     {
         MUTABLE_CPUPROFILER_SCOPE(SinkOptimiserAST);
 
-        bool modified = false;
+        bool bModified = false;
 
-		FOptimizeSinkContext context;
+		FOptimizeSinkContext Context;
 
-        ASTOp::Traverse_TopDown_Unique_Imprecise( roots, [&](Ptr<ASTOp>& n)
+        ASTOp::Traverse_TopDown_Unique_Imprecise(InRoots, [&](Ptr<ASTOp>& n)
         {
-            auto o = n->OptimiseSink(optimisationOptions, context);
+            auto o = n->OptimiseSink(InOptimisationOptions, Context);
             if (o && n!=o)
             {
-                modified = true;
+				bModified = true;
                 ASTOp::Replace(n,o);
             }
 
             return true;
         });
 
-        return modified;
+        return bModified;
     }
 
 
 	//---------------------------------------------------------------------------------------------
 	//---------------------------------------------------------------------------------------------
 	//---------------------------------------------------------------------------------------------
-    Ptr<ASTOp> Sink_ImageCropAST::Apply( const ASTOp* root )
+    Ptr<ASTOp> Sink_ImageCropAST::Apply( const ASTOp* InRoot )
     {
-        m_root = root;
+		check(InRoot->GetOpType() == OP_TYPE::IM_CROP);
+		
+		m_root = InRoot;
         OldToNew.Empty();
 
-        check(root->GetOpType()==OP_TYPE::IM_CROP);
-
-        auto typedRoot = dynamic_cast<const ASTOpFixed*>(root);
+		const ASTOpFixed* typedRoot = dynamic_cast<const ASTOpFixed*>(InRoot);
         m_initialSource = typedRoot->children[typedRoot->op.args.ImageCrop.source].child();
         Ptr<ASTOp> newSource = Visit( m_initialSource, typedRoot );
 
@@ -437,9 +462,9 @@ namespace mu
 
         case OP_TYPE::IM_PATCH:
         {
-            const ASTOpImagePatch* typedPatch = dynamic_cast<const ASTOpImagePatch*>(at.get());
+			const ASTOpImagePatch* typedPatch = dynamic_cast<const ASTOpImagePatch*>(at.get());
 
-            auto rectOp = typedPatch->patch.child();
+			Ptr<ASTOp> rectOp = typedPatch->patch.child();
             ASTOp::FGetImageDescContext context;
             FImageDesc patchDesc = rectOp->GetImageDesc( false, &context );
             box<vec2<int16>> patchBox;
@@ -463,7 +488,7 @@ namespace mu
             {
                 // Crop the base with the full crop, and the patch with the intersected part,
                 // adapting the patch origin
-                auto newOp = mu::Clone<ASTOpImagePatch>(at);
+                Ptr<ASTOpImagePatch> newOp = mu::Clone<ASTOpImagePatch>(at);
                 newOp->base = Visit(newOp->base.child(), currentCropOp);
 
                 box<vec2<int16>> ibox = patchBox.Intersect(cropBox);
@@ -486,16 +511,16 @@ namespace mu
 
         case OP_TYPE::IM_CROP:
         {
-            // We can combine the two crops into a possibliy smaller crop
-            auto childCrop = dynamic_cast<const ASTOpFixed*>(at.get());
+            // We can combine the two crops into a possibly smaller crop
+			const ASTOpFixed* childCrop = dynamic_cast<const ASTOpFixed*>(at.get());
 
-            box<vec2<int16_t>> childCropBox;
+            box<vec2<int16>> childCropBox;
             childCropBox.min[0] = childCrop->op.args.ImageCrop.minX;
             childCropBox.min[1] = childCrop->op.args.ImageCrop.minY;
             childCropBox.size[0] = childCrop->op.args.ImageCrop.sizeX;
             childCropBox.size[1] = childCrop->op.args.ImageCrop.sizeY;
 
-            box<vec2<int16_t>> cropBox;
+            box<vec2<int16>> cropBox;
             cropBox.min[0] = currentCropOp->op.args.ImageCrop.minX;
             cropBox.min[1] = currentCropOp->op.args.ImageCrop.minY;
             cropBox.size[0] = currentCropOp->op.args.ImageCrop.sizeX;
@@ -505,7 +530,7 @@ namespace mu
             // current ctop is applied to the result. So the final crop box would be:
             box<vec2<int16_t>> ibox;
             ibox.min = childCropBox.min + cropBox.min;
-            ibox.size = vec2<int16_t>::min(cropBox.size, childCropBox.size);
+            ibox.size = vec2<int16>::min(cropBox.size, childCropBox.size);
             check(cropBox.min.AllSmallerOrEqualThan(childCropBox.size));
             check((cropBox.min + cropBox.size)
                                 .AllSmallerOrEqualThan(childCropBox.size));
@@ -521,8 +546,7 @@ namespace mu
             else if (ibox == cropBox)
             {
                 // The child crop is not necessary
-                auto childSource =
-                    childCrop->children[childCrop->op.args.ImageCrop.source].child();
+                Ptr<ASTOp> childSource = childCrop->children[childCrop->op.args.ImageCrop.source].child();
                 newAt = Visit(childSource, currentCropOp);
             }
             else
@@ -534,8 +558,7 @@ namespace mu
                 newCropOp->op.args.ImageCrop.sizeX = ibox.size[0];
                 newCropOp->op.args.ImageCrop.sizeY = ibox.size[1];
 
-                auto childSource =
-                    childCrop->children[childCrop->op.args.ImageCrop.source].child();
+				Ptr<ASTOp> childSource = childCrop->children[childCrop->op.args.ImageCrop.source].child();
                 newAt = Visit(childSource, newCropOp.get());
             }
             break;
@@ -544,15 +567,15 @@ namespace mu
         case OP_TYPE::IM_LAYER:
         {
             // We move the op down the arguments
-            auto nop = mu::Clone<ASTOpImageLayer>(at);
+			Ptr<ASTOpImageLayer> nop = mu::Clone<ASTOpImageLayer>(at);
 
-            auto aOp = nop->base.child();
+			Ptr<ASTOp> aOp = nop->base.child();
             nop->base = Visit(aOp, currentCropOp);
 
-            auto bOp = nop->blend.child();
+			Ptr<ASTOp> bOp = nop->blend.child();
             nop->blend = Visit(bOp, currentCropOp);
 
-            auto mOp = nop->mask.child();
+			Ptr<ASTOp> mOp = nop->mask.child();
             nop->mask = Visit(mOp, currentCropOp);
 
             newAt = nop;
@@ -562,18 +585,60 @@ namespace mu
         case OP_TYPE::IM_LAYERCOLOUR:
         {
             // We move the op down the arguments
-            auto nop = mu::Clone<ASTOpImageLayerColor>(at);
+			Ptr<ASTOpImageLayerColor> nop = mu::Clone<ASTOpImageLayerColor>(at);
 
-            auto aOp = nop->base.child();
+			Ptr<ASTOp> aOp = nop->base.child();
             nop->base = Visit(aOp, currentCropOp);
 
-            auto mOp = nop->mask.child();
+			Ptr<ASTOp> mOp = nop->mask.child();
             nop->mask = Visit(mOp, currentCropOp);
 
             newAt = nop;
             break;
         }
 
+		case OP_TYPE::IM_DISPLACE:
+		{
+			// We move the op down the arguments
+			Ptr<ASTOpFixed> nop = mu::Clone<ASTOpFixed>(at);
+
+			Ptr<ASTOp> aOp = nop->children[nop->op.args.ImageDisplace.source].child();
+			nop->SetChild(nop->op.args.ImageDisplace.source, Visit(aOp, currentCropOp));
+
+			Ptr<ASTOp> bOp = nop->children[nop->op.args.ImageDisplace.displacementMap].child();
+			nop->SetChild(nop->op.args.ImageDisplace.displacementMap, Visit(bOp, currentCropOp));
+
+			newAt = nop;
+			break;
+		}
+
+		case OP_TYPE::IM_RASTERMESH:
+		{
+			// We add cropping data to the raster mesh if it doesn't have any
+			// \TODO: Is is possible to hit 2 crops on a raster mesh? Combine the crop.
+			Ptr<ASTOpImageRasterMesh> nop = mu::Clone<ASTOpImageRasterMesh>(at);
+
+			bool bRasterHasCrop = nop->UncroppedSizeX != 0;
+			if (!bRasterHasCrop)
+			{
+				box<vec2<int16>> cropBox;
+				cropBox.min[0] = currentCropOp->op.args.ImageCrop.minX;
+				cropBox.min[1] = currentCropOp->op.args.ImageCrop.minY;
+				cropBox.size[0] = currentCropOp->op.args.ImageCrop.sizeX;
+				cropBox.size[1] = currentCropOp->op.args.ImageCrop.sizeY;
+
+				nop->UncroppedSizeX = nop->SizeX;
+				nop->UncroppedSizeY = nop->SizeY;
+				nop->CropMinX = cropBox.min[0];
+				nop->CropMinY = cropBox.min[1];
+				nop->SizeX = cropBox.size[0];
+				nop->SizeY = cropBox.size[1];
+
+				newAt = nop;
+			}
+			break;
+		}
+
         case OP_TYPE::IM_INTERPOLATE:
         {
             // Move the op  down all the paths
@@ -581,9 +646,8 @@ namespace mu
 
             for (int v = 0; v < MUTABLE_OP_MAX_INTERPOLATE_COUNT; ++v)
             {
-                auto child =
-                    newOp->children[newOp->op.args.ImageInterpolate.targets[v]].child();
-                auto bOp = Visit(child, currentCropOp);
+				Ptr<ASTOp> child = newOp->children[newOp->op.args.ImageInterpolate.targets[v]].child();
+				Ptr<ASTOp> bOp = Visit(child, currentCropOp);
                 newOp->SetChild(newOp->op.args.ImageInterpolate.targets[v], bOp);
             }
 
@@ -1210,8 +1274,8 @@ namespace mu
 				//auto imageOp = nop->image.child();
 				//nop->image = Visit(imageOp, currentSinkingOp);
 
-				nop->sizeX = uint16(nop->sizeX * scaleX + 0.5f);
-				nop->sizeY = uint16(nop->sizeY * scaleY + 0.5f);
+				nop->SizeX = uint16(nop->SizeX * scaleX + 0.5f);
+				nop->SizeY = uint16(nop->SizeY * scaleY + 0.5f);
 				newAt = nop;
 				break;
 			}
@@ -1273,8 +1337,9 @@ namespace mu
                 // Set the size in the children and remove resize
                 auto sourceOp = mu::Clone<ASTOpFixed>(sourceAt.get());
                 sourceOp->op.args.ImagePlainColour.size[0] = op.args.ImageResize.size[0];
-                sourceOp->op.args.ImagePlainColour.size[1] = op.args.ImageResize.size[1];
-                at = sourceOp;
+				sourceOp->op.args.ImagePlainColour.size[1] = op.args.ImageResize.size[1];
+				sourceOp->op.args.ImagePlainColour.LODs = 1; // TODO
+				at = sourceOp;
                 break;
             }
 
@@ -1369,8 +1434,8 @@ namespace mu
                 //     ||
                 //     newOp->op.args.ImageRasterMesh.sizeY != op.args.ImageResize.size[1] )
                 {
-                    newOp->sizeX = op.args.ImageResize.size[0];
-                    newOp->sizeY = op.args.ImageResize.size[1];
+                    newOp->SizeX = op.args.ImageResize.size[0];
+                    newOp->SizeY = op.args.ImageResize.size[1];
 
 					if (newOp->mask)
 					{
@@ -1562,6 +1627,7 @@ namespace mu
 					uint16(newOp->op.args.ImagePlainColour.size[1]
 						* op.args.ImageResizeRel.factor[1]
 						+ 0.5f);
+				newOp->op.args.ImagePlainColour.LODs = 1; // TODO
 				at = newOp;
 				break;
 			}
diff --git a/Engine/Plugins/Experimental/Mutable/Source/MutableTools/Private/MuT/Compiler.cpp b/Engine/Plugins/Experimental/Mutable/Source/MutableTools/Private/MuT/Compiler.cpp
--- a/Engine/Plugins/Experimental/Mutable/Source/MutableTools/Private/MuT/Compiler.cpp
+++ b/Engine/Plugins/Experimental/Mutable/Source/MutableTools/Private/MuT/Compiler.cpp
@@ -132,11 +132,18 @@ namespace mu
     }
 
 
-    //---------------------------------------------------------------------------------------------
-    void CompilerOptions::SetImageCompressionQuality( int quality )
-    {
-        m_pD->m_imageCompressionQuality = quality;
-    }
+	//---------------------------------------------------------------------------------------------
+	void CompilerOptions::SetImageCompressionQuality(int quality)
+	{
+		m_pD->m_imageCompressionQuality = quality;
+	}
+
+
+	//---------------------------------------------------------------------------------------------
+	void CompilerOptions::SetImageTiling(int32 Tiling)
+	{
+		m_pD->m_imageTiling = Tiling;
+	}
 
 
 	//---------------------------------------------------------------------------------------------
diff --git a/Engine/Plugins/Experimental/Mutable/Source/MutableTools/Private/MuT/CompilerPrivate.h b/Engine/Plugins/Experimental/Mutable/Source/MutableTools/Private/MuT/CompilerPrivate.h
--- a/Engine/Plugins/Experimental/Mutable/Source/MutableTools/Private/MuT/CompilerPrivate.h
+++ b/Engine/Plugins/Experimental/Mutable/Source/MutableTools/Private/MuT/CompilerPrivate.h
@@ -97,6 +97,7 @@ namespace mu
 		int MinTextureResidentMipCount = 3;
 
         int m_imageCompressionQuality = 0;
+		int32 m_imageTiling=0 ;
 
         bool m_log = false;
 
diff --git a/Engine/Plugins/Experimental/Mutable/Source/MutableTools/Private/MuT/DataPacker.cpp b/Engine/Plugins/Experimental/Mutable/Source/MutableTools/Private/MuT/DataPacker.cpp
--- a/Engine/Plugins/Experimental/Mutable/Source/MutableTools/Private/MuT/DataPacker.cpp
+++ b/Engine/Plugins/Experimental/Mutable/Source/MutableTools/Private/MuT/DataPacker.cpp
@@ -766,6 +766,7 @@ namespace mu
 					NewPlain->op.args.ImagePlainColour.format = pOld->GetFormat();
 					NewPlain->op.args.ImagePlainColour.size[0] = pOld->GetSizeX();
 					NewPlain->op.args.ImagePlainColour.size[1] = pOld->GetSizeY();
+					NewPlain->op.args.ImagePlainColour.LODs = 1;
 
 					ASTOp::Replace(n, NewPlain);
 
diff --git a/Engine/Plugins/Experimental/Mutable/Source/MutableTools/Private/MuT/NodeLayout.cpp b/Engine/Plugins/Experimental/Mutable/Source/MutableTools/Private/MuT/NodeLayout.cpp
--- a/Engine/Plugins/Experimental/Mutable/Source/MutableTools/Private/MuT/NodeLayout.cpp
+++ b/Engine/Plugins/Experimental/Mutable/Source/MutableTools/Private/MuT/NodeLayout.cpp
@@ -169,19 +169,16 @@ namespace mu
 
 
 	//---------------------------------------------------------------------------------------------
-    void NodeLayoutBlocks::SetBlock( int index, int minx, int miny, int sizex, int sizey )
+	Ptr<const Layout> NodeLayoutBlocks::GetLayout() const
 	{
-        m_pD->m_pLayout->SetBlock( index, minx, miny, sizex, sizey );
+		return m_pD->m_pLayout;
 	}
 
 
 	//---------------------------------------------------------------------------------------------
-	void NodeLayoutBlocks::GetBlock(int index, int* minx, int* miny, int* sizex, int* sizey)
+    void NodeLayoutBlocks::SetBlock( int index, int minx, int miny, int sizex, int sizey )
 	{
-		if (minx && miny && sizex && sizey)
-		{
-			m_pD->m_pLayout->GetBlock(index, minx, miny, sizex, sizey);
-		}
+        m_pD->m_pLayout->SetBlock( index, minx, miny, sizex, sizey );
 	}
 
 
diff --git a/Engine/Plugins/Experimental/Mutable/Source/MutableTools/Public/MuT/Compiler.h b/Engine/Plugins/Experimental/Mutable/Source/MutableTools/Public/MuT/Compiler.h
--- a/Engine/Plugins/Experimental/Mutable/Source/MutableTools/Public/MuT/Compiler.h
+++ b/Engine/Plugins/Experimental/Mutable/Source/MutableTools/Public/MuT/Compiler.h
@@ -64,9 +64,14 @@ namespace mu
         //! process very slow, but will be able to compile very large models.
         void SetUseDiskCache( bool enabled );
 
-        //! Set the quality for the image compression algorithms. The level value is used internally
-        //! with System::SetImagecompressionQuality
-        void SetImageCompressionQuality( int quality );
+		//! Set the quality for the image compression algorithms. The level value is used internally
+		//! with System::SetImagecompressionQuality
+		void SetImageCompressionQuality(int quality);
+
+		/** Set the image tiling strategy :
+		 * If 0 (default) there is no tiling. Otherwise, images will be generated in tiles of the given size or less, and assembled afterwards as a final step.
+		 */
+		void SetImageTiling(int32 Tiling);
 
         //! 
         void SetDataPackingStrategy( int32 minRomSize, int32 MinTextureResidentMipCount );
diff --git a/Engine/Plugins/Experimental/Mutable/Source/MutableTools/Public/MuT/NodeLayout.h b/Engine/Plugins/Experimental/Mutable/Source/MutableTools/Public/MuT/NodeLayout.h
--- a/Engine/Plugins/Experimental/Mutable/Source/MutableTools/Public/MuT/NodeLayout.h
+++ b/Engine/Plugins/Experimental/Mutable/Source/MutableTools/Public/MuT/NodeLayout.h
@@ -132,14 +132,13 @@ namespace mu
 		//! Get the number of blocks in the layout
 		int GetBlockCount();
 
+		/** */
+		Ptr<const Layout> GetLayout() const;
+
 		//! Set a block of the layout.
 		//! minx and miny refer to the lowest left corner of the block.
         void SetBlock( int index, int minx, int miny, int sizex, int sizey );
 
-		//! Get a block of the layout.
-		//! minx and miny refer to the lowest left corner of the block.
-		void GetBlock(int index, int* minx, int* miny, int* sizex, int* sizey);
-
 		//! Set reduction block options like priority or if the block has to be reduced symmetrically.
 		void SetBlockOptions(int index, int priority, bool bUseSymmetry);
 
