[mutable] Fixed states view not showing up to date data.
- Now data gets updated not only after the instancing of the slate but also after host CO compilation.
- Cleared some commented sections of code.
- Removed unused includes

diff --git a/Engine/Plugins/Experimental/Mutable/Source/CustomizableObjectEditor/Private/MuCOE/CustomizableObjectDetails.cpp b/Engine/Plugins/Experimental/Mutable/Source/CustomizableObjectEditor/Private/MuCOE/CustomizableObjectDetails.cpp
--- a/Engine/Plugins/Experimental/Mutable/Source/CustomizableObjectEditor/Private/MuCOE/CustomizableObjectDetails.cpp
+++ b/Engine/Plugins/Experimental/Mutable/Source/CustomizableObjectEditor/Private/MuCOE/CustomizableObjectDetails.cpp
@@ -4,11 +4,9 @@
 
 #include "DetailLayoutBuilder.h"
 #include "DetailWidgetRow.h"
-#include "Framework/Views/TableViewMetadata.h"
 #include "IDetailsView.h"
 #include "MuCO/CustomizableObject.h"
 #include "MuR/Model.h"
-#include "MuR/MutableMemory.h"
 #include "MuR/Parameters.h"
 #include "MuR/Ptr.h"
 
@@ -62,16 +60,16 @@ public:
 		// Is it a state node?
 		if (ParameterIndex<0)
 		{
-			int numParams = Model->GetStateParameterCount( StateIndex );
-			for ( int i=0; i<numParams; ++i )
+			int ParameterCount = Model->GetStateParameterCount( StateIndex );
+			for ( int i=0; i<ParameterCount; ++i )
 			{
-				TSharedPtr<FStateDetailsNode> node = MakeShareable( new FStateDetailsNode );
+				TSharedPtr<FStateDetailsNode> SlateDetailsNode = MakeShareable( new FStateDetailsNode );
 
-				node->Model = Model;
-				node->StateIndex = StateIndex;
-				node->ParameterIndex = Model->GetStateParameterIndex( StateIndex, i );
+				SlateDetailsNode->Model = Model;
+				SlateDetailsNode->StateIndex = StateIndex;
+				SlateDetailsNode->ParameterIndex = Model->GetStateParameterIndex( StateIndex, i );
 
-				ChildrenList.Add( node );
+				ChildrenList.Add( SlateDetailsNode );
 			}
 		}
 
@@ -84,12 +82,10 @@ public:
 
 void FCustomizableObjectDetails::CustomizeDetails( IDetailLayoutBuilder& DetailBuilder )
 {
-	const UCustomizableObject* CustomObject = 0;
-
 	const IDetailsView* DetailsView = DetailBuilder.GetDetailsView();
 	if (DetailsView->GetSelectedObjects().Num())
 	{
-		CustomObject = Cast<const UCustomizableObject>(DetailsView->GetSelectedObjects()[0].Get());
+		CustomizableObject = Cast<UCustomizableObject>(DetailsView->GetSelectedObjects()[0].Get());
 	}
 
 	IDetailCategoryBuilder& StatesCategory = DetailBuilder.EditCategory( "States" );
@@ -99,80 +95,48 @@ void FCustomizableObjectDetails::CustomizeDetails( IDetailLayoutBuilder& DetailB
 	DetailBuilder.HideProperty("CustomizableObjectClassTags");
 	DetailBuilder.HideProperty("PopulationClassTags");
 
+	// Make the tree get automatically updated each time we compile the host CO
+	if (CustomizableObject && !CustomizableObject->PostCompileDelegate.IsBoundToObject(this))
+	{
+		CustomizableObject->PostCompileDelegate.AddSP(this, &FCustomizableObjectDetails::UpdateTree);
+	}
+	
+	// Cache the states defined in the CO
+	UpdateTree();
+	
 	StatesCategory.AddCustomRow( LOCTEXT("FCustomizableObjectDetails States", "States") )
 	[
 		SNew( SVerticalBox )
 		+ SVerticalBox::Slot()
 		.Padding( 2.0f )
 		[
-			//SNew( SFilterableDetail, LOCTEXT( "States", "States" ), &StatesCategory )
-			//[
-				SAssignNew(StatesTree, STreeView<TSharedPtr< FStateDetailsNode > >)
-					.SelectionMode(ESelectionMode::Single)
-					.TreeItemsSource( &RootTreeItems )
-					// Called to child items for any given parent item
-					.OnGetChildren( this, &FCustomizableObjectDetails::OnGetChildrenForStateTree )
-					// Generates the actual widget for a tree item
-					.OnGenerateRow( this, &FCustomizableObjectDetails::OnGenerateRowForStateTree ) 
-
-					// Generates the right click menu.
-					//.OnContextMenuOpening(this, &SClassViewer::BuildMenuWidget)
-
-					// Find out when the user selects something in the tree
-					//.OnSelectionChanged( this, &SClassViewer::OnClassViewerSelectionChanged )
-
-					// Allow for some spacing between items with a larger item height.
-					.ItemHeight(20.0f)
-
-					.HeaderRow
-					(
-						SNew(SHeaderRow)
-						.Visibility(EVisibility::Collapsed)
-						+ SHeaderRow::Column(TEXT("State"))
-						.DefaultLabel(NSLOCTEXT("CustomizableObjectDetails","State","State"))
-					)
-			//]
+			SAssignNew(StatesTree, STreeView<TSharedPtr< FStateDetailsNode > >)
+			.SelectionMode(ESelectionMode::Single)
+			.TreeItemsSource( &RootTreeItems )
+			// Called to child items for any given parent item
+			.OnGetChildren( this, &FCustomizableObjectDetails::OnGetChildrenForStateTree )
+			// Generates the actual widget for a tree item
+			.OnGenerateRow( this, &FCustomizableObjectDetails::OnGenerateRowForStateTree ) 
+			// Allow for some spacing between items with a larger item height.
+			.ItemHeight(20.0f)
+
+			.HeaderRow
+			(
+				SNew(SHeaderRow)
+				.Visibility(EVisibility::Collapsed)
+				+ SHeaderRow::Column(TEXT("State"))
+				.DefaultLabel(NSLOCTEXT("CustomizableObjectDetails","State","State"))
+			)
 		]
-//		+ SVerticalBox::Slot()
-//		.Padding( 2.0f )
-//		[
-//			SNew( SFilterableDetail, LOCTEXT( "SaveCollisionFromBuilderBrush", "Save Collision from Brush" ), &StatesCategory )
-//			[
-//				SNew( SButton )
-////				.ToolTipText( Commands.SaveBrushAsCollision->GetDescription() )
-//				//.OnClicked( this, &FCustomizableObjectDetails::OnSetCollisionFromBuilder )
-//				.VAlign( VAlign_Center )
-//				[
-//					SNew( STextBlock )
-//					.Text( LOCTEXT( "SaveCollisionFromBuilderBrush", "Save Collision from Brush" ) )
-//					.Font( IDetailLayoutBuilder::GetDetailFont() )
-//				]
-//			]
-//		]
 	];
 
 	StatesTree->SetIsRightClickScrollingEnabled(false);
-
-	if ( CustomObject &&  CustomObject->GetModel() )
-	{
-		uint32 numElements = CustomObject->GetModel()->GetStateCount();
-		RootTreeItems.SetNumUninitialized(0);
-		for ( uint32 i=0; i<numElements; ++i )
-		{
-			TSharedPtr<FStateDetailsNode> node = MakeShareable( new FStateDetailsNode );
-
-			node->Model = CustomObject->GetModel();
-			node->StateIndex = i;
-
-			RootTreeItems.Add( node );
-		}
-	}
-
+	
 	TSharedRef<IPropertyHandle> Property = DetailBuilder.GetProperty("ReferenceSkeletalMeshes");
 
-	if (Property->IsValidHandle() && CustomObject)
+	if (Property->IsValidHandle() && CustomizableObject)
 	{
-		if (CustomObject->bIsChildObject)
+		if (CustomizableObject->bIsChildObject)
 		{
 			Property->MarkHiddenByCustomization();
 		}
@@ -181,11 +145,30 @@ void FCustomizableObjectDetails::CustomizeDetails( IDetailLayoutBuilder& DetailB
 			Property->MarkResetToDefaultCustomized();
 		}
 	}
-
-	StatesTree->RequestTreeRefresh();
 }
 
+void FCustomizableObjectDetails::UpdateTree()
+{
+	if ( CustomizableObject &&  CustomizableObject->GetModel() )
+	{
+		RootTreeItems.SetNumUninitialized(0);
+		const uint32 NumElements = CustomizableObject->GetModel()->GetStateCount();
+		for ( uint32 i=0; i<NumElements; ++i )
+		{
+			TSharedPtr<FStateDetailsNode> SlateDetailsNode = MakeShareable( new FStateDetailsNode );
+
+			SlateDetailsNode->Model = CustomizableObject->GetModel();
+			SlateDetailsNode->StateIndex = i;
 
+			RootTreeItems.Add( SlateDetailsNode );
+		}
+	}
+
+	if (StatesTree)
+	{
+		StatesTree->RequestTreeRefresh();
+	}
+}
 /** The item used for visualizing the class in the tree. */
 class SStateItem : public STableRow< TSharedPtr<FString> >
 {
diff --git a/Engine/Plugins/Experimental/Mutable/Source/CustomizableObjectEditor/Private/MuCOE/CustomizableObjectDetails.h b/Engine/Plugins/Experimental/Mutable/Source/CustomizableObjectEditor/Private/MuCOE/CustomizableObjectDetails.h
--- a/Engine/Plugins/Experimental/Mutable/Source/CustomizableObjectEditor/Private/MuCOE/CustomizableObjectDetails.h
+++ b/Engine/Plugins/Experimental/Mutable/Source/CustomizableObjectEditor/Private/MuCOE/CustomizableObjectDetails.h
@@ -21,10 +21,15 @@ public:
 
 private:
 
-	TSharedPtr<STreeView<TSharedPtr< FStateDetailsNode > >> StatesTree;
+	class UCustomizableObject* CustomizableObject = nullptr;
+	
+	TSharedPtr<STreeView<TSharedPtr<FStateDetailsNode>>> StatesTree;
 
 	/** The root items to be displayed in the state tree. */
 	TArray< TSharedPtr< FStateDetailsNode > > RootTreeItems;
+
+	/** Re-caches the states found in the CO and calls for the update of the STreeView exposing the states and parameters*/
+	void UpdateTree();
 	
 	TSharedRef< ITableRow > OnGenerateRowForStateTree( TSharedPtr<FStateDetailsNode> Item, const TSharedRef< STableViewBase >& OwnerTable );
 	void OnGetChildrenForStateTree( TSharedPtr<FStateDetailsNode> InParent, TArray< TSharedPtr< FStateDetailsNode > >& OutChildren );
