[Mutable] Fix MutableTest SetProjectorValue crash.

#rnx

diff --git a/Engine/Plugins/Experimental/Mutable/Source/MutableRuntime/Private/MuR/Parameters.cpp b/Engine/Plugins/Experimental/Mutable/Source/MutableRuntime/Private/MuR/Parameters.cpp
--- a/Engine/Plugins/Experimental/Mutable/Source/MutableRuntime/Private/MuR/Parameters.cpp
+++ b/Engine/Plugins/Experimental/Mutable/Source/MutableRuntime/Private/MuR/Parameters.cpp
@@ -1017,8 +1017,22 @@ namespace mu
             return;
         }
 
-        FProjector* result = nullptr;
+	    // Parameters cannot change the projector type anymore
+		ParamProjectorType Value;
+        Value.type = PROJECTOR_TYPE::COUNT;
+        if (m_pD->m_pModel)
+        {
+            const FProgram& program = m_pD->m_pModel->GetPrivate()->m_program;
+            Value.type = program.m_parameters[ParameterIndex].m_defaultValue.Get<ParamProjectorType>().type;
+        }
 
+        Value.position = pos;
+        Value.direction = dir;
+        Value.up = up;        
+        Value.scale = scale;
+
+        Value.projectionAngle = projectionAngle;
+		
         // Single value case
         if (!RangePosition)
         {
@@ -1028,7 +1042,7 @@ namespace mu
                 m_pD->m_multiValues[ParameterIndex].Empty();
             }
 
-            result = &m_pD->m_values[ParameterIndex].Get<ParamProjectorType>();
+            m_pD->m_values[ParameterIndex].Set<ParamProjectorType>(Value);
         }
 
         // Multivalue case
@@ -1043,25 +1057,8 @@ namespace mu
 
 			TMap< TArray<int32_t>, PARAMETER_VALUE >& m = m_pD->m_multiValues[ParameterIndex];
 			PARAMETER_VALUE& it = m.FindOrAdd(RangePosition->m_pD->m_values);
-            result = &it.Get<ParamProjectorType>();
+            it.Set<ParamProjectorType>(Value);
         }
-
-        check( result );
-
-        // Parameters cannot change the projector type anymore
-        result->type = PROJECTOR_TYPE::COUNT;
-        if (m_pD->m_pModel)
-        {
-            const FProgram& program = m_pD->m_pModel->GetPrivate()->m_program;
-            result->type = program.m_parameters[ParameterIndex].m_defaultValue.Get<ParamProjectorType>().type;
-        }
-
-        result->position = pos;
-        result->direction = dir;
-        result->up = up;        
-        result->scale = scale;
-
-        result->projectionAngle = projectionAngle;
     }
 
 
