[mutable] Fix for static analysis errors.
#rnx

diff --git a/Engine/Plugins/Experimental/Mutable/Source/CustomizableObjectEditor/Private/MuCOE/CustomizableObjectEditor.cpp b/Engine/Plugins/Experimental/Mutable/Source/CustomizableObjectEditor/Private/MuCOE/CustomizableObjectEditor.cpp
--- a/Engine/Plugins/Experimental/Mutable/Source/CustomizableObjectEditor/Private/MuCOE/CustomizableObjectEditor.cpp
+++ b/Engine/Plugins/Experimental/Mutable/Source/CustomizableObjectEditor/Private/MuCOE/CustomizableObjectEditor.cpp
@@ -1141,15 +1141,19 @@ TSharedRef<SWidget> FCustomizableObjectEditor::GenerateCompileOptionsMenuContent
 		CompileOptimizationStrings.Add(MakeShareable(new FString(NSLOCTEXT("UnrealEd", "OptimizationMin", "Minimal").ToString())));
 		CompileOptimizationStrings.Add(MakeShareable(new FString(NSLOCTEXT("UnrealEd", "OptimizationMed", "Medium").ToString())));
 		CompileOptimizationStrings.Add(MakeShareable(new FString(NSLOCTEXT("UnrealEd", "OptimizationMax", "Maximum").ToString())));
+		int32 SelectedOptimization = FMath::Clamp(CustomizableObject->CompileOptions.OptimizationLevel,0, CompileOptimizationStrings.Num()-1);
 
-		CompileOptimizationCombo =
-			SNew(STextComboBox)
-			.OptionsSource(&CompileOptimizationStrings)
-			.InitiallySelectedItem(CompileOptimizationStrings[CustomizableObject->CompileOptions.OptimizationLevel])
-			.OnSelectionChanged(this, &FCustomizableObjectEditor::OnChangeCompileOptimizationLevel)
-			;
+		if (CustomizableObject)
+		{
+			CompileOptimizationCombo =
+				SNew(STextComboBox)
+				.OptionsSource(&CompileOptimizationStrings)
+				.InitiallySelectedItem(CompileOptimizationStrings[SelectedOptimization])
+				.OnSelectionChanged(this, &FCustomizableObjectEditor::OnChangeCompileOptimizationLevel)
+				;
 
-		MenuBuilder.AddWidget(CompileOptimizationCombo.ToSharedRef(), LOCTEXT("MutableCompileOptimizationLevel", "Optimization Level"));
+			MenuBuilder.AddWidget(CompileOptimizationCombo.ToSharedRef(), LOCTEXT("MutableCompileOptimizationLevel", "Optimization Level"));
+		}
 
 		// Image tiling
 		// Unfortunately SNumericDropDown doesn't work with integers at the time of writing.
@@ -1164,7 +1168,7 @@ TSharedRef<SWidget> FCustomizableObjectEditor::GenerateCompileOptionsMenuContent
 			.DropDownValues(TilingOptions)
 			.Value_Lambda([this]() 
 				{ 
-					return float(CustomizableObject->CompileOptions.ImageTiling); 
+					return CustomizableObject ? float(CustomizableObject->CompileOptions.ImageTiling) : 0.0f;
 				})
 			.OnValueChanged_Lambda([this](float Value) 
 				{ 
diff --git a/Engine/Plugins/Experimental/Mutable/Source/MutableTools/MutableTools.Build.cs b/Engine/Plugins/Experimental/Mutable/Source/MutableTools/MutableTools.Build.cs
--- a/Engine/Plugins/Experimental/Mutable/Source/MutableTools/MutableTools.Build.cs
+++ b/Engine/Plugins/Experimental/Mutable/Source/MutableTools/MutableTools.Build.cs
@@ -17,7 +17,6 @@ namespace UnrealBuildTool.Rules
 			// 
 			bUseRTTI = true;
 			bEnableExceptions = true;
-			//bUseUnity = false;
 
 			PrivateIncludePaths.AddRange(new string[] {
 				System.IO.Path.Combine(GetModuleDirectory("MutableRuntime"), "Private"),
