[mutable] Force initialization of box members.
#rnx

diff --git a/Engine/Plugins/Experimental/Mutable/Source/MutableRuntime/Public/MuR/MutableMath.h b/Engine/Plugins/Experimental/Mutable/Source/MutableRuntime/Public/MuR/MutableMath.h
--- a/Engine/Plugins/Experimental/Mutable/Source/MutableRuntime/Public/MuR/MutableMath.h
+++ b/Engine/Plugins/Experimental/Mutable/Source/MutableRuntime/Public/MuR/MutableMath.h
@@ -462,12 +462,19 @@ namespace mu
 			// Elements are cleared in the parent constructor
 		}
 
-        //-----------------------------------------------------------------------------------------
-        inline vec2( SCALAR x, SCALAR y )
-        {
-            this->m[0] = x;
-            this->m[1] = y;
-        }
+		//-----------------------------------------------------------------------------------------
+		inline explicit vec2(SCALAR v)
+		{
+			this->m[0] = v;
+			this->m[1] = v;
+		}
+
+		//-----------------------------------------------------------------------------------------
+		inline vec2(SCALAR x, SCALAR y)
+		{
+			this->m[0] = x;
+			this->m[1] = y;
+		}
 
         //-----------------------------------------------------------------------------------------
         inline vec2( const vec2& ) = default;
@@ -521,7 +528,15 @@ namespace mu
 		}
 
 		//-----------------------------------------------------------------------------------------
-		inline vec3( SCALAR x, SCALAR y, SCALAR z )
+		inline explicit vec3(SCALAR v)
+		{
+			this->m[0] = v;
+			this->m[1] = v;
+			this->m[2] = v;
+		}
+
+		//-----------------------------------------------------------------------------------------
+		inline vec3(SCALAR x, SCALAR y, SCALAR z)
 		{
 			this->m[0] = x;
 			this->m[1] = y;
@@ -902,6 +917,8 @@ namespace mu
 
 	public:
 
+		box() : min(0), size(0) {}
+
         //-----------------------------------------------------------------------------------------
         static box<VECTOR> FromMinSize(VECTOR m, VECTOR s)
         {
