[Mutable] Improved System tasks CPU scopes.

#rnx

diff --git a/Engine/Plugins/Experimental/Mutable/Source/CustomizableObject/Private/MuCO/CustomizableObjectSystem.cpp b/Engine/Plugins/Experimental/Mutable/Source/CustomizableObject/Private/MuCO/CustomizableObjectSystem.cpp
--- a/Engine/Plugins/Experimental/Mutable/Source/CustomizableObject/Private/MuCO/CustomizableObjectSystem.cpp
+++ b/Engine/Plugins/Experimental/Mutable/Source/CustomizableObject/Private/MuCO/CustomizableObjectSystem.cpp
@@ -1156,7 +1156,7 @@ namespace impl
 
 	void Subtask_Mutable_UpdateParameterDecorations(const TSharedPtr<FMutableOperationData>& OperationData, const TSharedPtr<mu::Model, ESPMode::ThreadSafe>& MutableModel, const mu::Parameters* MutableParameters)
 	{
-		MUTABLE_CPUPROFILER_SCOPE(UCustomizableInstancePrivateData::UpdateParameterDecorations)
+		MUTABLE_CPUPROFILER_SCOPE(Subtask_Mutable_UpdateParameterDecorations)
 
 		check(OperationData);
 		OperationData->ParametersUpdateData.Clear();
@@ -1215,11 +1215,11 @@ namespace impl
 	}
 
 
+	// This runs in the mutable thread.
 	void Subtask_Mutable_BeginUpdate_GetMesh(const TSharedPtr<FMutableOperationData>& OperationData, TSharedPtr<mu::Model, ESPMode::ThreadSafe> Model, const mu::Parameters* MutableParameters, int32 State)
 	{
-		MUTABLE_CPUPROFILER_SCOPE(Subtask_Mutable_BeginUpdate);
+		MUTABLE_CPUPROFILER_SCOPE(Subtask_Mutable_BeginUpdate_GetMesh)
 
-		// This runs in the mutable thread.
 		check(MutableParameters);
 		check(OperationData);
 		OperationData->InstanceUpdateData.Clear();
@@ -1432,15 +1432,13 @@ namespace impl
 				}
 			}
 		}
-
 	}
 
-	
+
+	// This runs in the mutable thread.
 	void Subtask_Mutable_GetImages(const TSharedPtr<FMutableOperationData>& OperationData, TSharedPtr<mu::Model, ESPMode::ThreadSafe> Model, const mu::Parameters* MutableParameters, int32 State)
 	{
-		MUTABLE_CPUPROFILER_SCOPE(Subtask_Mutable_GetImages);
-
-		// This runs in the mutable thread.
+		MUTABLE_CPUPROFILER_SCOPE(Subtask_Mutable_GetImages)
 
 		check(OperationData);
 		check(MutableParameters);
@@ -1544,9 +1542,10 @@ namespace impl
 	}
 	
 
+	// This runs in a worker thread
 	void Subtask_Mutable_PrepareTextures(const TSharedPtr<FMutableOperationData>& OperationData)
 	{
-		// This runs in a worker thread
+		MUTABLE_CPUPROFILER_SCOPE(Subtask_Mutable_PrepareTextures)
 
 		check(OperationData);
 		for (const FInstanceUpdateData::FSurface& Surface : OperationData->InstanceUpdateData.Surfaces)
@@ -1570,11 +1569,10 @@ namespace impl
 	}
 	
 
+	// This runs in a worker thread
 	void Subtask_Mutable_PrepareSkeletonData(const TSharedPtr<FMutableOperationData>& OperationData)
 	{
-		// This runs in a worker thread
-
-		MUTABLE_CPUPROFILER_SCOPE(PrepareSkeletonData);
+		MUTABLE_CPUPROFILER_SCOPE(Subtask_Mutable_PrepareSkeletonData)
 
 		check(OperationData);
 		const int32 LODCount = OperationData->InstanceUpdateData.LODs.Num();
@@ -1660,10 +1658,10 @@ namespace impl
 	}
 
 
+	// This runs in a worker thread.
 	void Task_Mutable_Update_GetMesh(TSharedPtr<FMutableOperationData> OperationData, TSharedPtr<mu::Model, ESPMode::ThreadSafe> Model, mu::ParametersPtrConst Parameters, bool bBuildParameterDecorations, int32 State)
 	{
-		// This runs in a worker thread.
-		MUTABLE_CPUPROFILER_SCOPE(Task_Mutable_GetMesh)
+		MUTABLE_CPUPROFILER_SCOPE(Task_Mutable_Update_GetMesh)
 
 #if WITH_EDITOR
 		uint32 StartCycles = FPlatformTime::Cycles();
@@ -1719,11 +1717,11 @@ namespace impl
 	}
 
 
+	// This runs in a worker thread.
 	void Task_Mutable_ReleaseInstance(TSharedPtr<FMutableOperationData> OperationData, mu::SystemPtr MutableSystem)
 	{
 		MUTABLE_CPUPROFILER_SCOPE(Task_Mutable_ReleaseInstance)
 
-		// This runs in a worker thread.
 		check(OperationData.IsValid());
 
 		if (OperationData->InstanceID > 0)
@@ -1982,6 +1980,8 @@ namespace impl
 	/** "Lock Cached Resources" */
 	void Task_Game_LockCache(TSharedPtr<FMutableOperationData> OperationData, const TWeakObjectPtr<UCustomizableObjectInstance>& CustomizableObjectInstancePtr, mu::Ptr<const mu::Parameters> Parameters, bool bBuildParameterDecorations)
 	{
+		MUTABLE_CPUPROFILER_SCOPE(Task_Game_ConvertResources)
+
 		check(IsInGameThread());
 
 		UCustomizableObjectSystem* System = UCustomizableObjectSystem::GetInstance();
@@ -2114,6 +2114,8 @@ namespace impl
 	/** Enqueue the release ID operation in the Mutable queue */
 	void Task_Game_ReleaseInstanceID(const mu::Instance::ID IDToRelease)
 	{
+		MUTABLE_CPUPROFILER_SCOPE(Task_Game_ReleaseInstanceID)
+
 		UCustomizableObjectSystem* System = UCustomizableObjectSystem::GetInstance();
 		check(System != nullptr);
 
@@ -2152,6 +2154,8 @@ namespace impl
 	/** Enqueue the release ID operation in the Mutable queue */
 	void Task_Game_ReleaseInstanceID(TSharedPtr<FMutableOperation> Operation)
 	{
+		MUTABLE_CPUPROFILER_SCOPE(Task_Game_ReleaseInstanceID)
+
 		check(Operation);
 		check(Operation->Type == FMutableOperation::EOperationType::IDRelease);
 
@@ -2162,6 +2166,8 @@ namespace impl
 	/** "Start Update" */
 	void Task_Game_StartUpdate(TSharedPtr<FMutableOperation> Operation)
 	{
+		MUTABLE_CPUPROFILER_SCOPE(Task_Game_StartUpdate)
+
 		check(Operation);
 		check(Operation->Type == FMutableOperation::EOperationType::Update);
 
@@ -2465,6 +2471,8 @@ void UCustomizableObjectSystem::AdvanceCurrentOperation()
 
 		case FMutableOperation::EOperationType::IDRelease:
 		{
+			MUTABLE_CPUPROFILER_SCOPE(OperationIDRelease);
+
 			impl::Task_Game_ReleaseInstanceID(Private->CurrentMutableOperation);
 
 			ClearCurrentMutableOperation();
