[mutable] Additional static analysis fixes.
#rnx

diff --git a/Engine/Plugins/Experimental/Mutable/Source/CustomizableObject/Private/MuCO/CustomizableInstanceLODManagement.cpp b/Engine/Plugins/Experimental/Mutable/Source/CustomizableObject/Private/MuCO/CustomizableInstanceLODManagement.cpp
--- a/Engine/Plugins/Experimental/Mutable/Source/CustomizableObject/Private/MuCO/CustomizableInstanceLODManagement.cpp
+++ b/Engine/Plugins/Experimental/Mutable/Source/CustomizableObject/Private/MuCO/CustomizableInstanceLODManagement.cpp
@@ -293,7 +293,11 @@ void UCustomizableInstanceLODManagement::UpdateInstanceDistsAndLODs()
 				}
 
 				// Blueprint instances and CO Editors
-				if (WorldType == EWorldType::EditorPreview || (!World && !CustomizableSkeletalComponent->GetAttachParentActor()))
+				
+				const USceneComponent* const AttachParentComponent = CustomizableSkeletalComponent->GetAttachParent();
+				bool bAttachParentActor = AttachParentComponent ? AttachParentComponent->GetOwner()!=nullptr : false;
+
+				if (WorldType == EWorldType::EditorPreview || (!World && !bAttachParentActor))
 				{
 					CustomizableSkeletalComponent->EditorUpdateComponent();
 					continue;
