[mutable] Fix for image displace checks that fail due to "image crop" optimizations.
#rnx

diff --git a/Engine/Plugins/Experimental/Mutable/Source/MutableRuntime/Private/MuR/OpImageDisplace.h b/Engine/Plugins/Experimental/Mutable/Source/MutableRuntime/Private/MuR/OpImageDisplace.h
--- a/Engine/Plugins/Experimental/Mutable/Source/MutableRuntime/Private/MuR/OpImageDisplace.h
+++ b/Engine/Plugins/Experimental/Mutable/Source/MutableRuntime/Private/MuR/OpImageDisplace.h
@@ -180,11 +180,15 @@ namespace mu
 					{
 						int dx, dy;
 						MutableDecodeOffset(*pMapData, dx, dy);
-						check(x + dx >= 0 && x + dx < sizeX);
-						check(y + dy >= 0 && y + dy < sizeY);
 
-						int offset = ((y + dy) * sizeX + (x + dx));
-						FMemory::Memcpy(pResultData, &pSourceData[offset], 1);
+						// This could actually happen since we enable the crop+displace optimization
+						//check(x + dx >= 0 && x + dx < sizeX);
+						//check(y + dy >= 0 && y + dy < sizeY);
+						if ((x + dx >= 0 && x + dx < sizeX) && (y + dy >= 0 && y + dy < sizeY))
+						{
+							int offset = ((y + dy) * sizeX + (x + dx));
+							FMemory::Memcpy(pResultData, &pSourceData[offset], 1);
+						}
 
 						++pResultData;
 						++pMapData;
@@ -210,11 +214,14 @@ namespace mu
 					{
 						int dx, dy;
 						MutableDecodeOffset(*pMapData, dx, dy);
-						check(x + dx >= 0 && x + dx < sizeX);
-						check(y + dy >= 0 && y + dy < sizeY);
-
-						int offset = ((y + dy) * sizeX + (x + dx)) * 3;
-						FMemory::Memcpy(pResultData, &pSourceData[offset], 3);
+						// This could actually happen since we enable the crop+displace optimization
+						//check(x + dx >= 0 && x + dx < sizeX);
+						//check(y + dy >= 0 && y + dy < sizeY);
+						if ((x + dx >= 0 && x + dx < sizeX) && (y + dy >= 0 && y + dy < sizeY))
+						{
+							int offset = ((y + dy) * sizeX + (x + dx)) * 3;
+							FMemory::Memcpy(pResultData, &pSourceData[offset], 3);
+						}
 
 						pResultData += 3;
 						++pMapData;
@@ -242,11 +249,14 @@ namespace mu
 					{
 						int dx, dy;
 						MutableDecodeOffset(*pMapData, dx, dy);
-						check(x + dx >= 0 && x + dx < sizeX);
-						check(y + dy >= 0 && y + dy < sizeY);
-
-						int offset = ((y + dy) * sizeX + (x + dx)) * pixelSize;
-						FMemory::Memcpy(pResultData, &pSourceData[offset], pixelSize);
+						// This could actually happen since we enable the crop+displace optimization
+						//check(x + dx >= 0 && x + dx < sizeX);
+						//check(y + dy >= 0 && y + dy < sizeY);
+						if ((x + dx >= 0 && x + dx < sizeX) && (y + dy >= 0 && y + dy < sizeY))
+						{
+							int offset = ((y + dy) * sizeX + (x + dx)) * pixelSize;
+							FMemory::Memcpy(pResultData, &pSourceData[offset], pixelSize);
+						}
 
 						pResultData += pixelSize;
 						++pMapData;
