[mutable] Added widget to remove bones from the final generated skeleton.

#rnx

diff --git a/Engine/Plugins/Experimental/Mutable/Source/CustomizableObjectEditor/Classes/MuCOE/Nodes/CustomizableObjectNodeObject.h b/Engine/Plugins/Experimental/Mutable/Source/CustomizableObjectEditor/Classes/MuCOE/Nodes/CustomizableObjectNodeObject.h
--- a/Engine/Plugins/Experimental/Mutable/Source/CustomizableObjectEditor/Classes/MuCOE/Nodes/CustomizableObjectNodeObject.h
+++ b/Engine/Plugins/Experimental/Mutable/Source/CustomizableObjectEditor/Classes/MuCOE/Nodes/CustomizableObjectNodeObject.h
@@ -16,6 +16,42 @@ class UObject;
 struct FPropertyChangedEvent;
 struct FSoftObjectPath;
 
+USTRUCT()
+struct FBoneToRemove
+{
+	GENERATED_USTRUCT_BODY()
+
+	UPROPERTY(EditAnywhere, Category = CustomizableObject)
+	bool bIncludeBone = false;
+
+	UPROPERTY(EditAnywhere, Category = CustomizableObject)
+	FName BoneName;
+};
+
+
+USTRUCT()
+struct FLODReductionSettings
+{
+	GENERATED_USTRUCT_BODY()
+
+	/** Selects which bones will be removed from the final skeleton
+	* BoneName: Name of the bone that will be removed. Its children will be removed too.
+	* Remove Only Children: If true, only the children of the selected bone will be removed. The selected bone will remain.
+	*/
+	UPROPERTY(EditAnywhere, Category = CustomizableObject)
+	TArray<FBoneToRemove> BonesToRemove;
+};
+
+
+USTRUCT()
+struct FComponentSettings
+{
+	GENERATED_USTRUCT_BODY()
+	
+	UPROPERTY(EditAnywhere, Category = CustomizableObject)
+	TArray<FLODReductionSettings> LODReductionSettings;
+};
+
 
 USTRUCT()
 struct CUSTOMIZABLEOBJECTEDITOR_API FCustomizableObjectState
@@ -154,7 +190,12 @@ public:
     // so it is need to keep them syncronized. 
     // This overrides the per skeletal mesh node selection 
     UPROPERTY()
-    TArray<FRealTimeMorphSelectionOverride> RealTimeMorphSelectionOverrides; 
+    TArray<FRealTimeMorphSelectionOverride> RealTimeMorphSelectionOverrides;
+
+	// Array of bones to remove from the mesh.All influences assigned to these bones will be transferred to the closest valid bone.
+	// Selected per component and LOD. Bones will be accumulated down the line.
+	UPROPERTY(VisibleAnywhere, Category = CustomizableObject)
+	TArray<FComponentSettings> ComponentSettings;
 	
     // To avoid any no properly saved GUIDs
 	FGuid IdentifierVerification;
@@ -231,5 +272,9 @@ public:
 	{
 		NumMeshComponents = MeshComponentNum;
 	}
+
+	// Node Details Support
+	int32 CurrentComponent = 0;
+	int32 CurrentLOD = 0;
 };
 
diff --git a/Engine/Plugins/Experimental/Mutable/Source/CustomizableObjectEditor/Private/MuCOE/CustomizableObjectEditorModule.cpp b/Engine/Plugins/Experimental/Mutable/Source/CustomizableObjectEditor/Private/MuCOE/CustomizableObjectEditorModule.cpp
--- a/Engine/Plugins/Experimental/Mutable/Source/CustomizableObjectEditor/Private/MuCOE/CustomizableObjectEditorModule.cpp
+++ b/Engine/Plugins/Experimental/Mutable/Source/CustomizableObjectEditor/Private/MuCOE/CustomizableObjectEditorModule.cpp
@@ -57,11 +57,13 @@
 #include "MuCOE/Nodes/CustomizableObjectNodeStaticMesh.h"
 #include "MuCOE/Nodes/CustomizableObjectNodeTable.h"
 #include "MuCOE/Nodes/CustomizableObjectNodeTableDetails.h"
+#include "MuCOE/Widgets/CustomizableObjectLODReductionSettings.h"
 #include "PropertyEditorModule.h"
 #include "MuCO/CustomizableObjectInstance.h"
 #include "MuCOE/Nodes/CustomizableObjectNodeGroupProjectorParameter.h"
 #include "UObject/UObjectIterator.h"
 
+
 class AActor;
 class FString;
 class ICustomizableObjectDebugger;
@@ -188,6 +190,7 @@ void FCustomizableObjectEditorModule::StartupModule()
 	//Custom properties
 	PropertyModule.RegisterCustomPropertyTypeLayout("CustomizableObjectIdentifier", FOnGetPropertyTypeCustomizationInstance::CreateStatic(&FCustomizableObjectIdentifierCustomization::MakeInstance));
 	PropertyModule.RegisterCustomPropertyTypeLayout(FMeshReshapeBoneReference::StaticStruct()->GetFName(), FOnGetPropertyTypeCustomizationInstance::CreateStatic(&FMeshReshapeBonesReferenceCustomization::MakeInstance));
+	PropertyModule.RegisterCustomPropertyTypeLayout(FBoneToRemove::StaticStruct()->GetFName(), FOnGetPropertyTypeCustomizationInstance::CreateStatic(&FCustomizableObjectLODReductionSettings::MakeInstance));
 
 	PropertyModule.NotifyCustomizationModuleChanged();
 
diff --git a/Engine/Plugins/Experimental/Mutable/Source/CustomizableObjectEditor/Private/MuCOE/Nodes/CustomizableObjectNodeObject.cpp b/Engine/Plugins/Experimental/Mutable/Source/CustomizableObjectEditor/Private/MuCOE/Nodes/CustomizableObjectNodeObject.cpp
--- a/Engine/Plugins/Experimental/Mutable/Source/CustomizableObjectEditor/Private/MuCOE/Nodes/CustomizableObjectNodeObject.cpp
+++ b/Engine/Plugins/Experimental/Mutable/Source/CustomizableObjectEditor/Private/MuCOE/Nodes/CustomizableObjectNodeObject.cpp
@@ -31,6 +31,12 @@ UCustomizableObjectNodeObject::UCustomizableObjectNodeObject()
 		Identifier = FGuid::NewGuid();
 		IdentifierVerification = Identifier;
 	}
+
+	if (ComponentSettings.IsEmpty())
+	{
+		ComponentSettings.Add(FComponentSettings());
+		ComponentSettings.Last().LODReductionSettings.Add(FLODReductionSettings());
+	}
 }
 
 
@@ -75,8 +81,23 @@ void UCustomizableObjectNodeObject::PostEditChangeProperty(FPropertyChangedEvent
 	{
 		NumLODs = FMath::Clamp(NumLODs, 1, 64);
 
+		for (int32 CompSetIndex = 0; CompSetIndex < ComponentSettings.Num(); ++CompSetIndex)
+		{
+			ComponentSettings[CompSetIndex].LODReductionSettings.SetNum(NumLODs);
+		}
+
 		ReconstructNode();
 	}
+
+	if (PropertyThatChanged && PropertyThatChanged->GetName() == TEXT("NumMeshComponents"))
+	{
+		ComponentSettings.SetNum(NumMeshComponents);
+
+		for (int32 CompSetIndex = 0; CompSetIndex < ComponentSettings.Num(); ++CompSetIndex)
+		{
+			ComponentSettings[CompSetIndex].LODReductionSettings.SetNum(NumLODs);
+		}
+	}
 	
 	Super::PostEditChangeProperty(PropertyChangedEvent);
 }
diff --git a/Engine/Plugins/Experimental/Mutable/Source/CustomizableObjectEditor/Private/MuCOE/Nodes/CustomizableObjectNodeObjectDetails.cpp b/Engine/Plugins/Experimental/Mutable/Source/CustomizableObjectEditor/Private/MuCOE/Nodes/CustomizableObjectNodeObjectDetails.cpp
--- a/Engine/Plugins/Experimental/Mutable/Source/CustomizableObjectEditor/Private/MuCOE/Nodes/CustomizableObjectNodeObjectDetails.cpp
+++ b/Engine/Plugins/Experimental/Mutable/Source/CustomizableObjectEditor/Private/MuCOE/Nodes/CustomizableObjectNodeObjectDetails.cpp
@@ -3,14 +3,20 @@
 #include "MuCOE/Nodes/CustomizableObjectNodeObjectDetails.h"
 
 #include "DetailLayoutBuilder.h"
+#include "Framework/MultiBox/MultiBoxBuilder.h"
 #include "IDetailsView.h"
 #include "MuCOE/CustomizableObjectEditorUtilities.h"
+#include "MuCOE/CustomizableObjectGraph.h"
+#include "MuCOE/Nodes/CustomizableObjectNode.h"
 #include "MuCOE/Nodes/CustomizableObjectNodeObject.h"
 #include "MuCOE/Nodes/CustomizableObjectNodeObjectGroup.h"
 #include "MuCOE/SCustomizableObjectNodeObjectRTMorphTargetOverride.h"
+#include "MuCOE/Widgets/CustomizableObjectLODReductionSettings.h"
 #include "PropertyCustomizationHelpers.h"
 #include "ScopedTransaction.h"
+#include "Widgets/Input/SComboButton.h"
 #include "Widgets/Input/STextComboBox.h"
+#include "Widgets/Text/STextBlock.h"
 
 
 #define LOCTEXT_NAMESPACE "CustomizableObjectDetails"
@@ -37,17 +43,32 @@ void FCustomizableObjectNodeObjectDetails::CustomizeDetails( IDetailLayoutBuilde
 	IDetailCategoryBuilder& RuntimeParameters = DetailBuilder.EditCategory("States Runtime Parameters");
 	IDetailCategoryBuilder& ExternalCategory = DetailBuilder.EditCategory("AttachedToExternalObject");
 	IDetailCategoryBuilder& RealTimeMorphTargets = DetailBuilder.EditCategory("RealTime Morph Targets");
+	IDetailCategoryBuilder& LODCustomSettings = DetailBuilder.EditCategory("LODCustomSettings");
+
 	//StatesCategory.CategoryIcon( "ActorClassIcon.CustomizableObject" );
 
 	if (Node)
 	{
+		// Properties
 		TSharedRef<IPropertyHandle> StatesProperty = DetailBuilder.GetProperty("States");
-		StatesProperty->SetOnPropertyValueChanged(FSimpleDelegate::CreateSP(this, &FCustomizableObjectNodeObjectDetails::OnStatesPropertyChanged));
-
 		TSharedRef<IPropertyHandle> ParentObjectProperty = DetailBuilder.GetProperty("ParentObject");
+		TSharedRef<IPropertyHandle> ComponentsProperty = DetailBuilder.GetProperty("NumMeshComponents");
+		TSharedRef<IPropertyHandle> LODsProperty = DetailBuilder.GetProperty("NumLODs");
+		TSharedRef<IPropertyHandle> ComponentSettingsProperty= DetailBuilder.GetProperty("ComponentSettings");
+
+		FName BonesToRemovePropertyPath = FName("ComponentSettings[" + FString::FromInt(Node->CurrentComponent) + "].LODReductionSettings[" + FString::FromInt(Node->CurrentLOD) + "].BonesToRemove");
+		TSharedRef<IPropertyHandle> BonesToRemoveProperty = DetailBuilder.GetProperty(BonesToRemovePropertyPath);
 
+		// Callbacks
+		StatesProperty->SetOnPropertyValueChanged(FSimpleDelegate::CreateSP(this, &FCustomizableObjectNodeObjectDetails::OnStatesPropertyChanged));
+		ComponentsProperty->SetOnPropertyValueChangedWithData(TDelegate<void(const FPropertyChangedEvent&)>::CreateSP(this, &FCustomizableObjectNodeObjectDetails::OnNumComponentsOrLODsChanged));
+		LODsProperty->SetOnPropertyValueChangedWithData(TDelegate<void(const FPropertyChangedEvent&)>::CreateSP(this, &FCustomizableObjectNodeObjectDetails::OnNumComponentsOrLODsChanged));
+
+		// Hidden Properties
 		DetailBuilder.HideProperty("ParentObjectGroupId");
 		DetailBuilder.HideProperty("ParentObject");
+		DetailBuilder.HideProperty("ComponentSettings");
+
 		GroupNodeComboOptions.Empty();
 
 		if (Node->bIsBase)
@@ -116,8 +137,42 @@ void FCustomizableObjectNodeObjectDetails::CustomizeDetails( IDetailLayoutBuilde
                     SNew(SCustomizableObjectNodeSkeletalMeshRTMorphTargetOverride)
                     .Node(Node)
                 ];
-			}
 
+				// Component Settings Category ----------
+
+				// Component Picker
+				LODCustomSettings.AddCustomRow((LOCTEXT("ComponentCustomModeSelect", "Select Component")))
+				.NameContent()
+				[
+					SNew(STextBlock)
+					.Text(LOCTEXT("ComponentCustomSettingsSelectTitle", "Component"))
+					.Font(IDetailLayoutBuilder::GetDetailFont())
+					.ToolTipText(LOCTEXT("ComponentCustomSettingsSelectTooltip", "Select the component to edit."))
+				]
+				.ValueContent()
+				.VAlign(VAlign_Center)
+				[
+					OnGenerateComponentComboBoxForPicker()
+				];
+				
+				// LOD Picker
+				LODCustomSettings.AddCustomRow((LOCTEXT("LODCustomModeSelect", "Select LOD")))
+				.NameContent()
+				[
+					SNew(STextBlock)
+					.Text(LOCTEXT("LODCustomSettingsSelectTitle", "LOD"))
+					.Font(IDetailLayoutBuilder::GetDetailFont())
+					.ToolTipText(LOCTEXT("LODCustomSettingsSelectTitle", "Select the component's LOD to edit."))
+				]
+				.ValueContent()
+				.VAlign(VAlign_Center)
+				[
+					OnGenerateLODComboBoxForPicker()
+				];
+
+				// Bones to remove widget
+				LODCustomSettings.AddProperty(BonesToRemoveProperty);
+			}
 		}
 		else
 		{
@@ -224,5 +279,147 @@ void FCustomizableObjectNodeObjectDetails::OnStatesPropertyChanged()
 	}
 }
 
-#undef LOCTEXT_NAMESPACE
 
+TSharedRef<SWidget> FCustomizableObjectNodeObjectDetails::OnGenerateComponentComboBoxForPicker()
+{
+	return SNew(SComboButton)
+		.OnGetMenuContent(this, &FCustomizableObjectNodeObjectDetails::OnGenerateComponentMenuForPicker)
+		.VAlign(VAlign_Center)
+		.ContentPadding(0)
+		.ButtonContent()
+		[
+			SNew(STextBlock)
+			.Font(IDetailLayoutBuilder::GetDetailFont())
+			.Text(this, &FCustomizableObjectNodeObjectDetails::GetCurrentComponentName)
+		];
+}
+
+
+TSharedRef<SWidget> FCustomizableObjectNodeObjectDetails::OnGenerateLODComboBoxForPicker()
+{
+	return SNew(SComboButton)
+		.OnGetMenuContent(this, &FCustomizableObjectNodeObjectDetails::OnGenerateLODMenuForPicker)
+		.VAlign(VAlign_Center)
+		.ButtonContent()
+		[
+			SNew(STextBlock)
+			.Font(IDetailLayoutBuilder::GetDetailFont())
+			.Text(this, &FCustomizableObjectNodeObjectDetails::GetCurrentLODName)
+		];
+}
+
+
+TSharedRef<SWidget> FCustomizableObjectNodeObjectDetails::OnGenerateComponentMenuForPicker()
+{
+	if (Node)
+	{
+		int32 NumComponents = Node->NumMeshComponents;
+		FMenuBuilder MenuBuilder(true, NULL);
+
+		for (int32 ComponentIndex = 0; ComponentIndex < NumComponents; ++ComponentIndex)
+		{
+			FText ComponentString = FText::FromString((TEXT("Component ") + FString::FromInt(ComponentIndex)));
+			FUIAction Action(FExecuteAction::CreateSP(this, &FCustomizableObjectNodeObjectDetails::OnSelectedComponentChanged, ComponentIndex));
+			MenuBuilder.AddMenuEntry(ComponentString, FText::GetEmpty(), FSlateIcon(), Action);
+		}
+
+		return MenuBuilder.MakeWidget();
+	}
+
+	return SNullWidget::NullWidget;
+}
+
+
+TSharedRef<SWidget> FCustomizableObjectNodeObjectDetails::OnGenerateLODMenuForPicker()
+{
+	if (Node)
+	{
+		int32 NumLODs = Node->NumLODs;
+		FMenuBuilder MenuBuilder(true, NULL);
+
+		for (int32 LODIndex = 0; LODIndex < NumLODs; ++LODIndex)
+		{
+			FText LODString = FText::FromString((TEXT("LOD ") + FString::FromInt(LODIndex)));
+			FUIAction Action(FExecuteAction::CreateSP(this, &FCustomizableObjectNodeObjectDetails::OnSelectedLODChanged, LODIndex));
+			MenuBuilder.AddMenuEntry(LODString, FText::GetEmpty(), FSlateIcon(), Action);
+		}
+
+		return MenuBuilder.MakeWidget();
+	}
+
+	return SNullWidget::NullWidget;
+}
+
+
+void FCustomizableObjectNodeObjectDetails::OnSelectedComponentChanged(int32 NewComponentIndex)
+{
+	if (Node)
+	{
+		Node->CurrentComponent = NewComponentIndex;
+		Node->CurrentLOD = 0;
+	}
+	
+	DetailBuilderPtr->ForceRefreshDetails();
+}
+
+
+void FCustomizableObjectNodeObjectDetails::OnSelectedLODChanged(int32 NewLODIndex)
+{
+	if (Node)
+	{
+		Node->CurrentLOD = NewLODIndex;
+	}
+
+	DetailBuilderPtr->ForceRefreshDetails();
+}
+
+
+FText FCustomizableObjectNodeObjectDetails::GetCurrentComponentName() const
+{
+	FText ComponentText;
+
+	if (Node)
+	{
+		ComponentText = FText::FromString(FString(TEXT("Component ")) + FString::FromInt(Node->CurrentComponent));
+	}
+
+	return ComponentText;
+}
+
+
+FText FCustomizableObjectNodeObjectDetails::GetCurrentLODName() const
+{
+	FText LODText;
+
+	if (Node)
+	{
+		LODText = FText::FromString(FString(TEXT("LOD ")) + FString::FromInt(Node->CurrentLOD));
+	}
+
+	return LODText;
+}
+
+
+void FCustomizableObjectNodeObjectDetails::OnNumComponentsOrLODsChanged(const FPropertyChangedEvent& PropertyChangedEvent)
+{
+	if (DetailBuilderPtr && PropertyChangedEvent.ChangeType == EPropertyChangeType::ValueSet)
+	{
+		int32 ComponentToSelect = FMath::Min(Node->NumMeshComponents - 1, Node->CurrentComponent);
+
+		if (ComponentToSelect != Node->CurrentComponent)
+		{
+			Node->CurrentComponent = ComponentToSelect;
+
+			// Reset the LOD selection
+			Node->CurrentLOD = 0;
+		}
+		else
+		{
+			Node->CurrentLOD = FMath::Min(Node->NumLODs - 1, Node->CurrentLOD);
+		}
+
+		DetailBuilderPtr->ForceRefreshDetails();
+	}
+}
+
+#undef LOCTEXT_NAMESPACE
diff --git a/Engine/Plugins/Experimental/Mutable/Source/CustomizableObjectEditor/Private/MuCOE/Nodes/CustomizableObjectNodeObjectDetails.h b/Engine/Plugins/Experimental/Mutable/Source/CustomizableObjectEditor/Private/MuCOE/Nodes/CustomizableObjectNodeObjectDetails.h
--- a/Engine/Plugins/Experimental/Mutable/Source/CustomizableObjectEditor/Private/MuCOE/Nodes/CustomizableObjectNodeObjectDetails.h
+++ b/Engine/Plugins/Experimental/Mutable/Source/CustomizableObjectEditor/Private/MuCOE/Nodes/CustomizableObjectNodeObjectDetails.h
@@ -9,7 +9,9 @@ namespace ESelectInfo { enum Type : int; }
 class FString;
 class IDetailLayoutBuilder;
 class IPropertyHandle;
+class SWidget;
 struct FAssetData;
+struct FPropertyChangedEvent;
 
 
 class FCustomizableObjectNodeObjectDetails : public IDetailCustomization
@@ -34,6 +36,25 @@ private:
 
 	void OnGroupNodeComboBoxSelectionChanged(TSharedPtr<FString> Selection, ESelectInfo::Type SelectInfo, TSharedRef<IPropertyHandle> ParentProperty);
 
+	// Component and LOD combobox generation callbacks.
+	TSharedRef<SWidget> OnGenerateComponentComboBoxForPicker();
+	TSharedRef<SWidget> OnGenerateLODComboBoxForPicker();
+
+	// Component and LOD menu generation callbacks
+	TSharedRef<SWidget> OnGenerateComponentMenuForPicker();
+	TSharedRef<SWidget> OnGenerateLODMenuForPicker();
+
+	// Component and LOD OnSelect callbacks
+	void OnSelectedComponentChanged(int32 NewComponentIndex);
+	void OnSelectedLODChanged(int32 NewLODIndex);
+
+	// Component and LOD name generation callbacks
+	FText GetCurrentComponentName() const;
+	FText GetCurrentLODName() const;
+
+	// Calback to force the refresh of the detail if the number of components or LODs changes
+	void OnNumComponentsOrLODsChanged(const FPropertyChangedEvent& PropertyChangedEvent);
+
 	// Refreshes the Details View when the states variable has been updated
 	void OnStatesPropertyChanged();
 
diff --git a/Engine/Plugins/Experimental/Mutable/Source/CustomizableObjectEditor/Private/MuCOE/Widgets/CustomizableObjectLODReductionSettings.cpp b/Engine/Plugins/Experimental/Mutable/Source/CustomizableObjectEditor/Private/MuCOE/Widgets/CustomizableObjectLODReductionSettings.cpp
new file mode 100644
--- /dev/null
+++ b/Engine/Plugins/Experimental/Mutable/Source/CustomizableObjectEditor/Private/MuCOE/Widgets/CustomizableObjectLODReductionSettings.cpp
@@ -0,0 +1,229 @@
+// Copyright Epic Games, Inc. All Rights Reserved.
+
+#include "CustomizableObjectLODReductionSettings.h"
+
+#include "BoneSelectionWidget.h"
+#include "DetailLayoutBuilder.h"
+#include "DetailWidgetRow.h"
+#include "Engine/SkeletalMesh.h"
+#include "MuCOE/CustomizableObjectGraph.h"
+#include "MuCOE/Nodes/CustomizableObjectNodeObject.h"
+#include "Widgets/Input/SCheckBox.h"
+#include "Widgets/Text/STextBlock.h"
+
+
+#define LOCTEXT_NAMESPACE "CustomizableObjectDetails"
+
+
+TSharedRef<IPropertyTypeCustomization> FCustomizableObjectLODReductionSettings::MakeInstance()
+{
+	return MakeShareable(new FCustomizableObjectLODReductionSettings());
+}
+
+
+void FCustomizableObjectLODReductionSettings::CustomizeHeader(TSharedRef<IPropertyHandle> StructPropertyHandle, class FDetailWidgetRow& HeaderRow, IPropertyTypeCustomizationUtils& StructCustomizationUtils)
+{
+	constexpr bool bRecurse = false;
+
+	BoneNameProperty = StructPropertyHandle->GetChildHandle(GET_MEMBER_NAME_CHECKED(FBoneToRemove, BoneName), bRecurse);
+	IncludeBoneProperty = StructPropertyHandle->GetChildHandle(GET_MEMBER_NAME_CHECKED(FBoneToRemove, bIncludeBone), bRecurse);
+
+	TSharedPtr<IPropertyHandle> Parent = StructPropertyHandle->GetParentHandle();
+	if (Parent)
+	{
+		ParentArrayProperty = Parent->AsArray();
+	}
+
+	if (BoneNameProperty.IsValid() && BoneNameProperty->IsValidHandle() 
+		&& IncludeBoneProperty.IsValid() && IncludeBoneProperty->IsValidHandle())
+	{
+		ObjectNode = GetObjectNode();
+
+		bool bIncludeBoneProperty = false;
+		IncludeBoneProperty->GetValue(bIncludeBoneProperty);
+
+		HeaderRow
+		.NameContent()
+		[
+			StructPropertyHandle->CreatePropertyNameWidget()
+		]
+		.ValueContent()
+		.MaxDesiredWidth(0.0f)
+		[
+			SNew(SVerticalBox)
+			+SVerticalBox::Slot()
+			.AutoHeight()
+			.Padding(0.0f, 1.0f, 0.0, 0.0f)
+			[
+				SNew(SHorizontalBox)
+				+ SHorizontalBox::Slot()
+				.Padding(0.0f,5.0f,5.0f,0.0f)
+				.AutoWidth()
+				[
+					SNew(STextBlock)
+					.Text(LOCTEXT("BoneSelector_Text", "Bone:"))
+					.ToolTipText(LOCTEXT("BoneSelector_TooltipText", "Select the bone to remove. Bone children will be removed too."))
+					.Font(IDetailLayoutBuilder::GetDetailFont())
+				]
+
+				+SHorizontalBox::Slot()
+				.AutoWidth()
+				[
+					SAssignNew(SelectionWidget, SBoneSelectionWidget)
+					.ToolTipText(StructPropertyHandle->GetToolTipText())
+					.OnBoneSelectionChanged(this, &FCustomizableObjectLODReductionSettings::OnBoneSelectionChanged)
+					.OnGetSelectedBone(this, &FCustomizableObjectLODReductionSettings::GetSelectedBone)
+					.OnGetReferenceSkeleton(this, &FCustomizableObjectLODReductionSettings::GetReferenceSkeleton)
+				]
+			]
+			
+			+ SVerticalBox::Slot()
+			.AutoHeight()
+			.Padding(0.0f,5.0f,0.0,3.0f)
+			[
+				SNew(SHorizontalBox)
+				+ SHorizontalBox::Slot()
+				.AutoWidth()
+				.Padding(0.0f,2.5f,5.0f,0.0f)
+				[
+					SNew(STextBlock)
+					.Text(LOCTEXT("BoneCheckBox_Text", "Remove Only Children:"))
+					.ToolTipText(LOCTEXT("BoneCheckBox_TooltipText","If true, only the bone children of the selected bone will be removed."))
+					.Font(IDetailLayoutBuilder::GetDetailFont())
+				]
+
+				+ SHorizontalBox::Slot()
+				.AutoWidth()
+				[
+					SNew(SCheckBox)
+					.OnCheckStateChanged(this, &FCustomizableObjectLODReductionSettings::OnBoneCheckBoxChanged)
+					.IsChecked(bIncludeBoneProperty ? ECheckBoxState::Checked : ECheckBoxState::Unchecked)
+				]
+			]
+		];
+
+		DefaultColor = SelectionWidget->GetColorAndOpacity();
+	}
+}
+
+
+UCustomizableObjectNodeObject* FCustomizableObjectLODReductionSettings::GetObjectNode()
+{
+	if (!BoneNameProperty->IsValidHandle())
+	{
+		return nullptr;
+	}
+
+	TArray<UObject*> Objects;
+	BoneNameProperty->GetOuterObjects(Objects);
+
+	if (Objects.Num())
+	{
+		return Cast<UCustomizableObjectNodeObject>(Objects[0]);
+	}
+
+	return nullptr;
+}
+
+
+FName FCustomizableObjectLODReductionSettings::GetSelectedBone(bool& bMultipleValues) const
+{
+	if (BoneNameProperty->IsValidHandle())
+	{
+		FString OutText;
+	
+		FPropertyAccess::Result Result = BoneNameProperty->GetValueAsFormattedString(OutText);
+	
+		return FName(*OutText);
+	}
+
+	return FName();
+}
+
+
+const struct FReferenceSkeleton& FCustomizableObjectLODReductionSettings::GetReferenceSkeleton() const
+{
+	static FReferenceSkeleton DummySkeleton;
+	USkeletalMesh* SkeletalMesh = nullptr;
+
+	if (ObjectNode)
+	{
+		if (const UCustomizableObject* CustomizableObject = Cast<UCustomizableObject>(ObjectNode->GetCustomizableObjectGraph()->GetOuter()))
+		{
+			if (CustomizableObject->ReferenceSkeletalMeshes.IsValidIndex(ObjectNode->CurrentComponent))
+			{
+				SkeletalMesh = CustomizableObject->ReferenceSkeletalMeshes[ObjectNode->CurrentComponent];
+			}
+		}
+	}
+
+	return SkeletalMesh ? SkeletalMesh->GetRefSkeleton() : DummySkeleton;
+}
+
+
+void FCustomizableObjectLODReductionSettings::OnBoneSelectionChanged(FName Name)
+{
+	if (!ParentArrayProperty)
+	{
+		BoneNameProperty->SetValue(Name);
+		IncludeBoneProperty->SetValue(false);
+		return;
+	}
+	
+	uint32 NumParentElems = 0;
+	FPropertyAccess::Result Result = ParentArrayProperty->GetNumElements(NumParentElems);
+	check(Result == FPropertyAccess::Result::Success);
+	
+	// Find if already set.
+	bool BoneNameFound = false;
+	for (uint32 ElementIndex = 0; ElementIndex < NumParentElems; ++ElementIndex)
+	{
+		TSharedPtr<IPropertyHandle> ElemProperty = ParentArrayProperty->GetElement(ElementIndex);
+		check(ElemProperty);
+	
+		constexpr bool bRecurse = false;
+		TSharedPtr<IPropertyHandle> ElemBoneNameProperty =
+			ElemProperty->GetChildHandle(GET_MEMBER_NAME_CHECKED(FBoneToRemove, BoneName), bRecurse);
+	
+		check(ElemBoneNameProperty);
+	
+		FName BoneNameValue;
+		Result = ElemBoneNameProperty->GetValue(BoneNameValue);
+	
+		check(Result == FPropertyAccess::Result::Success);
+	
+		if (BoneNameValue == Name)
+		{
+			BoneNameFound = true;
+			break;
+		}
+	}
+	
+	// TODO: This behavior does not seem to be correct, if the option is repeated the value is not set
+	// only logging a warning to the console, which might not be open and could make dificult to see what 
+	// is the problem. 
+	BoneNameProperty->SetValue(Name);
+	IncludeBoneProperty->SetValue(false);
+
+	if (BoneNameFound)
+	{
+		const FLinearColor WarningColor = FLinearColor::Yellow;
+		SelectionWidget->SetColorAndOpacity(WarningColor);
+	}
+	else
+	{
+		SelectionWidget->SetColorAndOpacity(DefaultColor);
+	}
+}
+
+void FCustomizableObjectLODReductionSettings::OnBoneCheckBoxChanged(ECheckBoxState NewState)
+{
+	bool bValue = NewState == ECheckBoxState::Checked;
+
+	if (IncludeBoneProperty->IsValidHandle())
+	{
+		IncludeBoneProperty->SetValue(bValue);
+	}
+}
+
+#undef LOCTEXT_NAMESPACE
\ No newline at end of file
diff --git a/Engine/Plugins/Experimental/Mutable/Source/CustomizableObjectEditor/Private/MuCOE/Widgets/CustomizableObjectLODReductionSettings.h b/Engine/Plugins/Experimental/Mutable/Source/CustomizableObjectEditor/Private/MuCOE/Widgets/CustomizableObjectLODReductionSettings.h
new file mode 100644
--- /dev/null
+++ b/Engine/Plugins/Experimental/Mutable/Source/CustomizableObjectEditor/Private/MuCOE/Widgets/CustomizableObjectLODReductionSettings.h
@@ -0,0 +1,59 @@
+// Copyright Epic Games, Inc. All Rights Reserved.
+
+#pragma once
+
+#include "IPropertyTypeCustomization.h"
+
+class IPropertyHandle;
+class IPropertyHandleArray;
+class SBoneSelectionWidget;
+class UCustomizableObjectNodeObject;
+
+struct FReferenceSkeleton;
+
+enum class ECheckBoxState : uint8;
+
+class FCustomizableObjectLODReductionSettings : public IPropertyTypeCustomization
+{
+public:
+	static TSharedRef<IPropertyTypeCustomization> MakeInstance();
+
+	// IPropertyTypeCustomization interface
+	void CustomizeHeader(TSharedRef<IPropertyHandle> StructPropertyHandle, class FDetailWidgetRow& HeaderRow, IPropertyTypeCustomizationUtils& StructCustomizationUtils) override;
+	void CustomizeChildren(TSharedRef<IPropertyHandle> StructPropertyHandle, class IDetailChildrenBuilder& ChildBuilder, IPropertyTypeCustomizationUtils& StructCustomizationUtils) override {};
+
+private:
+
+	// Gets the Reference Skeleton from the target skeleton
+	const FReferenceSkeleton& GetReferenceSkeleton() const;
+	
+	// Gets the current selected bone in the combo button
+	FName GetSelectedBone(bool& bMultipleValues) const;
+	
+	// OnSelectionChanged callback to return the selected bone
+	void OnBoneSelectionChanged(FName Name);
+	
+	// OnSelectionChanged callback 
+	void OnBoneCheckBoxChanged(ECheckBoxState NewState);
+
+	// Returns the object node of the property represented in the widget
+	UCustomizableObjectNodeObject* GetObjectNode();
+
+private:
+
+	// Pointer to the object node
+	UCustomizableObjectNodeObject* ObjectNode;
+
+	// Property to change after bone has been picked
+	TSharedPtr<IPropertyHandle> BoneNameProperty;
+	TSharedPtr<IPropertyHandle> IncludeBoneProperty;
+
+	// Parent propery, as array to identify repeated values.
+	TSharedPtr<IPropertyHandleArray> ParentArrayProperty;
+
+	// Widget to select the bone to remove
+	TSharedPtr<SBoneSelectionWidget> SelectionWidget;
+
+	// Color of the bone name
+	FLinearColor DefaultColor;
+};
