[mutable] Now, table nodes support empty meshes.

#rnx

diff --git a/Engine/Plugins/Experimental/Mutable/Source/MutableTools/Private/MuT/CodeGenerator.h b/Engine/Plugins/Experimental/Mutable/Source/MutableTools/Private/MuT/CodeGenerator.h
--- a/Engine/Plugins/Experimental/Mutable/Source/MutableTools/Private/MuT/CodeGenerator.h
+++ b/Engine/Plugins/Experimental/Mutable/Source/MutableTools/Private/MuT/CodeGenerator.h
@@ -739,7 +739,11 @@ namespace mu
             check( pTable->GetPrivate()->m_rows[i].m_id <= 0xFFFF);
             auto condition = (uint16)pTable->GetPrivate()->m_rows[i].m_id;
             Ptr<ASTOp> Branch = GenerateOption( node, colIndex, (int)i, m_pErrorLog.get() );
-			SwitchOp->cases.Add(ASTOpSwitch::FCase(condition, SwitchOp, Branch ));
+
+			if (Branch || TYPE != TCT_MESH)
+			{
+				SwitchOp->cases.Add(ASTOpSwitch::FCase(condition, SwitchOp, Branch));
+			}
         }
 
         return SwitchOp;
diff --git a/Engine/Plugins/Experimental/Mutable/Source/MutableTools/Private/MuT/CodeGenerator_Mesh.cpp b/Engine/Plugins/Experimental/Mutable/Source/MutableTools/Private/MuT/CodeGenerator_Mesh.cpp
--- a/Engine/Plugins/Experimental/Mutable/Source/MutableTools/Private/MuT/CodeGenerator_Mesh.cpp
+++ b/Engine/Plugins/Experimental/Mutable/Source/MutableTools/Private/MuT/CodeGenerator_Mesh.cpp
@@ -717,43 +717,39 @@ class Node;
 		Ptr<ASTOp> Op = GenerateTableSwitch<NodeMeshTable::Private, TCT_MESH, OP_TYPE::ME_SWITCH>(*TableNode->GetPrivate(), 
 			[this, &NewResult, &t, &InOptions] (const NodeMeshTable::Private& node, int colIndex, int row, ErrorLog* pErrorLog)
 			{
-				NodeMeshConstantPtr pCell = new NodeMeshConstant();
 				MeshPtr pMesh = node.m_pTable->GetPrivate()->m_rows[row].m_values[colIndex].m_pMesh;
+				FMeshGenerationResult BranchResults;
 
-				if (!pMesh)
+				if (pMesh)
 				{
-					char temp[256];
-					mutable_snprintf(temp, 256,
-						"Table has a missing mesh in column %d, row %d.", colIndex, row);
-					pErrorLog->GetPrivate()->Add(temp, ELMT_ERROR, node.m_errorContext);
-				}
+					NodeMeshConstantPtr pCell = new NodeMeshConstant();
+					pCell->SetValue(pMesh);
 
-				pCell->SetValue(pMesh);
+					// TODO Take into account layout strategy
+					int numLayouts = node.m_layouts.Num();
+					pCell->SetLayoutCount(numLayouts);
+					for (int i = 0; i < numLayouts; ++i)
+					{
+						pCell->SetLayout(i, node.m_layouts[i]);
+					}
 
-				// TODO Take into account layout strategy
-				int numLayouts = node.m_layouts.Num();
-				pCell->SetLayoutCount(numLayouts);
-				for (int i = 0; i < numLayouts; ++i)
-				{
-					pCell->SetLayout(i, node.m_layouts[i]);
-				}
+					FMeshGenerationOptions TargetOptions = InOptions;
 
-				FMeshGenerationOptions TargetOptions = InOptions;
+					if (t != 0)
+					{
+						TargetOptions.OverrideLayouts = NewResult.GeneratedLayouts;
+					}
 
-				if (t != 0)
-				{
-					TargetOptions.OverrideLayouts = NewResult.GeneratedLayouts;
-				}
+					GenerateMesh(TargetOptions, BranchResults, pCell);
 
-				FMeshGenerationResult BranchResults;
-				GenerateMesh(TargetOptions, BranchResults, pCell);
+					if (t == 0)
+					{
+						NewResult = BranchResults;
+					}
 
-				if (t == 0)
-				{
-					NewResult = BranchResults;
-				}				
+					++t;
+				}
 
-				++t;
 				return BranchResults.meshOp;
 			});
 
