[mutable] Fix COInstance drag and drop to level

#rnx

diff --git a/Engine/Plugins/Experimental/Mutable/Source/CustomizableObjectEditor/Private/MuCOE/CustomizableObjectInstanceFactory.cpp b/Engine/Plugins/Experimental/Mutable/Source/CustomizableObjectEditor/Private/MuCOE/CustomizableObjectInstanceFactory.cpp
--- a/Engine/Plugins/Experimental/Mutable/Source/CustomizableObjectEditor/Private/MuCOE/CustomizableObjectInstanceFactory.cpp
+++ b/Engine/Plugins/Experimental/Mutable/Source/CustomizableObjectEditor/Private/MuCOE/CustomizableObjectInstanceFactory.cpp
@@ -20,7 +20,7 @@ UCustomizableObjectInstanceFactory::UCustomizableObjectInstanceFactory(const FOb
     : Super(ObjectInitializer)
 {
     DisplayName = LOCTEXT("CustomizableObjectInstanceDisplayName", "Customizable Object Instance");
-    NewActorClass = ASkeletalMeshActor::StaticClass();
+    NewActorClass = ACustomizableSkeletalMeshActor::StaticClass();
     bUseSurfaceOrientation = true;
 }
 
@@ -144,7 +144,7 @@ void UCustomizableObjectInstanceFactory::PostSpawnActor(UObject* Asset, AActor*
 			}			
 		}
 
-		UActorFactory::PostSpawnActor(COInstance, NewActor);
+		Super::PostSpawnActor(COInstance, NewActor);
 	}
 }
 
