[mutable] Fixed LOCTEXT conflict in node object details

#rnx

diff --git a/Engine/Plugins/Experimental/Mutable/Source/CustomizableObjectEditor/Private/MuCOE/Nodes/CustomizableObjectNodeObjectDetails.cpp b/Engine/Plugins/Experimental/Mutable/Source/CustomizableObjectEditor/Private/MuCOE/Nodes/CustomizableObjectNodeObjectDetails.cpp
--- a/Engine/Plugins/Experimental/Mutable/Source/CustomizableObjectEditor/Private/MuCOE/Nodes/CustomizableObjectNodeObjectDetails.cpp
+++ b/Engine/Plugins/Experimental/Mutable/Source/CustomizableObjectEditor/Private/MuCOE/Nodes/CustomizableObjectNodeObjectDetails.cpp
@@ -162,7 +162,7 @@ void FCustomizableObjectNodeObjectDetails::CustomizeDetails( IDetailLayoutBuilde
 					SNew(STextBlock)
 					.Text(LOCTEXT("LODCustomSettingsSelectTitle", "LOD"))
 					.Font(IDetailLayoutBuilder::GetDetailFont())
-					.ToolTipText(LOCTEXT("LODCustomSettingsSelectTitle", "Select the component's LOD to edit."))
+					.ToolTipText(LOCTEXT("LODCustomSettingsSelectTooltip", "Select the component's LOD to edit."))
 				]
 				.ValueContent()
 				.VAlign(VAlign_Center)
