[mutable] Fixes for triangle raster memory stomp with non-multiple of 8 images.
Ensure initialization of attributes moved to unreal vectors.
#rnx

diff --git a/Engine/Plugins/Experimental/Mutable/Source/CustomizableObjectEditor/Private/MuCOE/SMutableLayoutViewer.cpp b/Engine/Plugins/Experimental/Mutable/Source/CustomizableObjectEditor/Private/MuCOE/SMutableLayoutViewer.cpp
--- a/Engine/Plugins/Experimental/Mutable/Source/CustomizableObjectEditor/Private/MuCOE/SMutableLayoutViewer.cpp
+++ b/Engine/Plugins/Experimental/Mutable/Source/CustomizableObjectEditor/Private/MuCOE/SMutableLayoutViewer.cpp
@@ -104,8 +104,8 @@ void SMutableLayoutViewer::Tick(const FGeometry& AllottedGeometry, const double
 			Blocks.SetNum( BlockCount );
 			for (int32 BlockIndex = 0; BlockIndex < BlockCount; ++BlockIndex)
 			{
-				UE::Math::TIntVector2<uint16> Min;
-				UE::Math::TIntVector2<uint16> Size;
+				UE::Math::TIntVector2<uint16> Min(ForceInit);
+				UE::Math::TIntVector2<uint16> Size(ForceInit);
 				MutableLayout->GetBlock(BlockIndex, &Min[0], &Min[1], &Size[0], &Size[1]);
 				Blocks[BlockIndex].Min[0] = Min[0];
 				Blocks[BlockIndex].Min[1] = Min[1];
diff --git a/Engine/Plugins/Experimental/Mutable/Source/MutableRuntime/Private/MuR/OpImageProject.cpp b/Engine/Plugins/Experimental/Mutable/Source/MutableRuntime/Private/MuR/OpImageProject.cpp
--- a/Engine/Plugins/Experimental/Mutable/Source/MutableRuntime/Private/MuR/OpImageProject.cpp
+++ b/Engine/Plugins/Experimental/Mutable/Source/MutableRuntime/Private/MuR/OpImageProject.cpp
@@ -587,7 +587,7 @@ void ImageRasterProjected_Optimised( const Mesh* pMesh, Image* pImage,
 				!scratch->culledVertex[i1] ||
 				!scratch->culledVertex[i2])
 			{
-				Triangle(pImage->GetData(),
+				Triangle(pImage->GetData(), pImage->GetDataSize(),
 					sizeX, sizeY,
 					pixelSize,
 					scratch->vertices[i0],
@@ -712,7 +712,7 @@ void ImageRasterProjected_OptimisedWrapping( const Mesh* pMesh, Image* pImage, P
              !scratch->culledVertex[i1] ||
              !scratch->culledVertex[i2] )
         {
-            Triangle( pImage->GetData(),
+            Triangle( pImage->GetData(), pImage->GetDataSize(),
 				sizeX, sizeY,
 				pixelSize,
 				scratch->vertices[i0], scratch->vertices[i1], scratch->vertices[i2],
diff --git a/Engine/Plugins/Experimental/Mutable/Source/MutableRuntime/Private/MuR/OpImageRasterMesh.h b/Engine/Plugins/Experimental/Mutable/Source/MutableRuntime/Private/MuR/OpImageRasterMesh.h
--- a/Engine/Plugins/Experimental/Mutable/Source/MutableRuntime/Private/MuR/OpImageRasterMesh.h
+++ b/Engine/Plugins/Experimental/Mutable/Source/MutableRuntime/Private/MuR/OpImageRasterMesh.h
@@ -85,7 +85,7 @@ namespace mu
 				vertices, indices, pImage, sizeX, sizeY, pixelProc
 			] (int32 f)
 			{
-				Triangle(pImage->GetData(),
+				Triangle(pImage->GetData(), pImage->GetDataSize(),
 					sizeX, sizeY,
 					1,
 					vertices[indices[f * 3 + 0]],
@@ -123,7 +123,7 @@ namespace mu
 				// TODO: Select faces outside for loop?
 				if (blocks[indices[f * 3 + 0]] == block)
 				{
-					Triangle(pImage->GetData(),
+					Triangle(pImage->GetData(), pImage->GetDataSize(),
 						sizeX, sizeY,
 						1,
 						vertices[indices[f * 3 + 0]],
diff --git a/Engine/Plugins/Experimental/Mutable/Source/MutableRuntime/Private/MuR/OpLayoutPack.h b/Engine/Plugins/Experimental/Mutable/Source/MutableRuntime/Private/MuR/OpLayoutPack.h
--- a/Engine/Plugins/Experimental/Mutable/Source/MutableRuntime/Private/MuR/OpLayoutPack.h
+++ b/Engine/Plugins/Experimental/Mutable/Source/MutableRuntime/Private/MuR/OpLayoutPack.h
@@ -149,7 +149,7 @@ namespace mu
 		}
 
 		int32 index;
-		UE::Math::TIntVector2<uint16> size;
+		UE::Math::TIntVector2<uint16> size = UE::Math::TIntVector2<uint16>(0,0);
 		int32 priority;
 		bool useSymmetry;
 	};
diff --git a/Engine/Plugins/Experimental/Mutable/Source/MutableRuntime/Private/MuR/Raster.h b/Engine/Plugins/Experimental/Mutable/Source/MutableRuntime/Private/MuR/Raster.h
--- a/Engine/Plugins/Experimental/Mutable/Source/MutableRuntime/Private/MuR/Raster.h
+++ b/Engine/Plugins/Experimental/Mutable/Source/MutableRuntime/Private/MuR/Raster.h
@@ -41,16 +41,18 @@ namespace mu
 	//! TODO Optimize anyway.
 	//! TODO Option for two sided? right now it culls back facing.
     template<class PIXEL_PROCESSOR, int NUM_INTERPOLATORS, typename INT=int64_t>
-	void Triangle( unsigned char* aBuffer, int aWidth, int aHeight, int aPixelSize
+	void Triangle( uint8* Buffer, int32 BufferSize, int Width, int Height, int PixelSize
 				 , const RasterVertex<NUM_INTERPOLATORS> &av1
 				 , const RasterVertex<NUM_INTERPOLATORS> &av2
 				 , const RasterVertex<NUM_INTERPOLATORS> &av3
 				 , PIXEL_PROCESSOR &processor
 				 , bool cullBackface = true )
 	{
+		check(Width * Height * PixelSize <= BufferSize);
+
 		float interpolatedValues[NUM_INTERPOLATORS];
 
-		int stride = aWidth*aPixelSize;
+		int Stride = Width*PixelSize;
 
 		// Check backface culling
 		// TODO better option than resorting vertices?
@@ -86,20 +88,24 @@ namespace mu
         const INT FDY23 = DY23 * 16;
         const INT FDY31 = DY31 * 16;
 
-	    // Bounding rectangle
-        INT minx = FMath::Max(INT(0),        (FMath::Min3(X1, X2, X3) + 0xF) / 16 );
-        INT maxx = FMath::Min(INT(aWidth-1), (FMath::Max3(X1, X2, X3) + 0xF) / 16 );
-        INT miny = FMath::Max(INT(0),		 (FMath::Min3(Y1, Y2, Y3) + 0xF) / 16 );
-        INT maxy = FMath::Min(INT(aHeight-1),(FMath::Max3(Y1, Y2, Y3) + 0xF) / 16 );
+		// Block size, standard 8x8 (must be power of two)
+		constexpr INT q = 8;
 
-	    // Block size, standard 8x8 (must be power of two)
-        const INT q = 8;
+		// Skip partial blocks
+		INT ClampWidth = Width & ~(q - 1);
+		INT ClampHeight = Height & ~(q - 1);
+
+	    // Bounding rectangle
+        INT minx = FMath::Max(INT(0),				(FMath::Min3(X1, X2, X3) + 0xF) / 16 );
+        INT maxx = FMath::Min(INT(ClampWidth -1),	(FMath::Max3(X1, X2, X3) + 0xF) / 16 );
+        INT miny = FMath::Max(INT(0),				(FMath::Min3(Y1, Y2, Y3) + 0xF) / 16 );
+        INT maxy = FMath::Min(INT(ClampHeight -1),	(FMath::Max3(Y1, Y2, Y3) + 0xF) / 16 );
 
-	    // Start in corner of 8x8 block
-	    minx &= ~(q - 1);
-	    miny &= ~(q - 1);
+		// Start in corner of 8x8 block
+		minx &= ~(q - 1);
+		miny &= ~(q - 1);
 
-		aBuffer += miny * stride;
+		Buffer += miny * Stride;
 
 	    // Half-edge constants
         INT C1 = DY12 * X1 - DX12 * Y1;
@@ -153,9 +159,12 @@ namespace mu
                 INT c = (INT(c00) * 1) | (INT(c10) * 2) | (INT(c01) * 4) | (INT(c11) * 8);
 
 	            // Skip block when outside an edge
-	            if(a == 0x0 || b == 0x0 || c == 0x0) continue;
+				if (a == 0x0 || b == 0x0 || c == 0x0)
+				{
+					continue;
+				}
 
-				unsigned char *buffer = aBuffer;
+				uint8 *CurrentBuffer = Buffer;
 
 	            // Accept whole block when totally covered
 	            if(a == 0xF && b == 0xF && c == 0xF)
@@ -181,10 +190,10 @@ namespace mu
 							}
 
 	                        // Process pixel
-							processor.ProcessPixel( &buffer[ix*aPixelSize], interpolatedValues );
+							processor.ProcessPixel( &CurrentBuffer[ix*PixelSize], interpolatedValues );
 	                    }
 
-	                    buffer += stride;
+						CurrentBuffer += Stride;
 	                }
 	            }
 	            else // Partially covered block
@@ -219,7 +228,7 @@ namespace mu
 															+ beta * v3.interpolators[ii];
 								}
 
-								processor.ProcessPixel( &buffer[ix*aPixelSize], interpolatedValues );
+								processor.ProcessPixel( &CurrentBuffer[ix*PixelSize], interpolatedValues );
 	                        }
 
 	                        CX1 -= FDY12;
@@ -231,12 +240,12 @@ namespace mu
 	                    CY2 += FDX23;
 	                    CY3 += FDX31;
 
-	                    buffer += stride;
+						CurrentBuffer += Stride;
 	                }
 	            }
 	        }
 
-			aBuffer += q * stride;
+			Buffer += q * Stride;
 	    }
 	}
 
diff --git a/Engine/Plugins/Experimental/Mutable/Source/MutableRuntime/Public/MuR/Image.h b/Engine/Plugins/Experimental/Mutable/Source/MutableRuntime/Public/MuR/Image.h
--- a/Engine/Plugins/Experimental/Mutable/Source/MutableRuntime/Public/MuR/Image.h
+++ b/Engine/Plugins/Experimental/Mutable/Source/MutableRuntime/Public/MuR/Image.h
@@ -91,7 +91,7 @@ namespace mu
 		}
 
 		//!
-		FImageSize m_size;
+		FImageSize m_size = FImageSize(0, 0);
 
 		//!
 		EImageFormat m_format;
@@ -251,7 +251,7 @@ namespace mu
 		//-----------------------------------------------------------------------------------------
 
 		//!
-		FImageSize m_size;
+		FImageSize m_size = FImageSize(0, 0);
 
 		//! Non-persistent internal id
 		mutable uint32 m_internalId = 0;
diff --git a/Engine/Plugins/Experimental/Mutable/Source/MutableRuntime/Public/MuR/Layout.h b/Engine/Plugins/Experimental/Mutable/Source/MutableRuntime/Public/MuR/Layout.h
--- a/Engine/Plugins/Experimental/Mutable/Source/MutableRuntime/Public/MuR/Layout.h
+++ b/Engine/Plugins/Experimental/Mutable/Source/MutableRuntime/Public/MuR/Layout.h
@@ -164,8 +164,8 @@ namespace mu
 				bUseSymmetry = false;
 			}
 
-			UE::Math::TIntVector2<uint16> m_min;
-			UE::Math::TIntVector2<uint16> m_size;
+			UE::Math::TIntVector2<uint16> m_min = UE::Math::TIntVector2<uint16>(0,0);
+			UE::Math::TIntVector2<uint16> m_size = UE::Math::TIntVector2<uint16>(0, 0);
 
 			//! Absolute id used to control merging of various layouts
 			int32 m_id;
@@ -229,17 +229,17 @@ namespace mu
 
 
 		//!
-		UE::Math::TIntVector2<uint16> m_size;
+		UE::Math::TIntVector2<uint16> m_size = UE::Math::TIntVector2<uint16>(0, 0);
 
-		UE::Math::TIntVector2<uint16> m_maxsize;
+		UE::Math::TIntVector2<uint16> m_maxsize = UE::Math::TIntVector2<uint16>(0, 0);
 
 		//!
 		TArray<FBlock> m_blocks;
 
 		//! Packing strategy
 		EPackStrategy m_strategy = EPackStrategy::RESIZABLE_LAYOUT;
-
-		int32 FirstLODToIgnoreWarnings;
+		 
+		int32 FirstLODToIgnoreWarnings = 0;
 
 		EReductionMethod ReductionMethod = EReductionMethod::HALVE_REDUCTION;
 
diff --git a/Engine/Plugins/Experimental/Mutable/Source/MutableTools/Private/MuT/AST.h b/Engine/Plugins/Experimental/Mutable/Source/MutableTools/Private/MuT/AST.h
--- a/Engine/Plugins/Experimental/Mutable/Source/MutableTools/Private/MuT/AST.h
+++ b/Engine/Plugins/Experimental/Mutable/Source/MutableTools/Private/MuT/AST.h
@@ -79,7 +79,7 @@ namespace mu
         } type;
 
         // For constant sizes
-        FImageSize size;
+        FImageSize size = FImageSize(0, 0);
 
         // For layout factor sizes
         Ptr<class ASTOp> layout;
diff --git a/Engine/Plugins/Experimental/Mutable/Source/MutableTools/Private/MuT/ASTOpImagePatch.h b/Engine/Plugins/Experimental/Mutable/Source/MutableTools/Private/MuT/ASTOpImagePatch.h
--- a/Engine/Plugins/Experimental/Mutable/Source/MutableTools/Private/MuT/ASTOpImagePatch.h
+++ b/Engine/Plugins/Experimental/Mutable/Source/MutableTools/Private/MuT/ASTOpImagePatch.h
@@ -20,7 +20,7 @@ struct FProgram;
 
 		ASTChild base;
 		ASTChild patch;
-		UE::Math::TIntVector2<uint16> location;
+		UE::Math::TIntVector2<uint16> location = UE::Math::TIntVector2<uint16>(0, 0);
 
 	public:
 
diff --git a/Engine/Plugins/Experimental/Mutable/Source/MutableTools/Private/MuT/CodeGenerator.h b/Engine/Plugins/Experimental/Mutable/Source/MutableTools/Private/MuT/CodeGenerator.h
--- a/Engine/Plugins/Experimental/Mutable/Source/MutableTools/Private/MuT/CodeGenerator.h
+++ b/Engine/Plugins/Experimental/Mutable/Source/MutableTools/Private/MuT/CodeGenerator.h
@@ -188,7 +188,7 @@ namespace mu
             // This reference has to be the smart pointer to avoid memory aliasing, keeping
             // processed nodes alive.
             NodePtrConst pNode;
-			UE::Math::TIntVector2<int32> imageSize;
+			UE::Math::TIntVector2<int32> imageSize = UE::Math::TIntVector2<int32>(0, 0);
             box< UE::Math::TIntVector2<int32> > imageRect;
             int state = -1;
 			TArray<mu::string> activeTags;
@@ -237,7 +237,7 @@ namespace mu
         //! image that we are generating.
         struct IMAGE_STATE
         {
-            UE::Math::TIntVector2<int32> m_imageSize;
+            UE::Math::TIntVector2<int32> m_imageSize = UE::Math::TIntVector2<int32>(0, 0);
             box< UE::Math::TIntVector2<int32> > m_imageRect;
             int32 m_layoutBlockId;
             LayoutPtrConst m_pLayout;
