[mutable] Add support for reshaping of animBP RigidBody Override Physics assets.

#rnx

diff --git a/Engine/Plugins/Experimental/Mutable/Source/CustomizableObject/Classes/MuCO/CustomizableObject.h b/Engine/Plugins/Experimental/Mutable/Source/CustomizableObject/Classes/MuCO/CustomizableObject.h
--- a/Engine/Plugins/Experimental/Mutable/Source/CustomizableObject/Classes/MuCO/CustomizableObject.h
+++ b/Engine/Plugins/Experimental/Mutable/Source/CustomizableObject/Classes/MuCO/CustomizableObject.h
@@ -119,7 +119,6 @@ struct FProfileParameterDat
 	TArray<FCustomizableObjectProjectorParameterValue> ProjectorParameters;
 };
 
-
 USTRUCT()
 struct FCompilationOptions
 {
@@ -177,6 +176,9 @@ struct FCompilationOptions
 	// Used to enable physics asset merge.
 	bool bPhysicsAssetMergeEnabled = false;
 
+	// Used to enable AnimBp override physics mainipualtion.  
+	bool bAnimBpPhysicsManipulationEnabled = false;
+
 	// Used to reduce the number of notifications when compiling objects
 	bool bSilentCompilation = true;
 
@@ -425,6 +427,57 @@ private:
 };
 
 
+
+USTRUCT()
+struct CUSTOMIZABLEOBJECT_API FAnimBpOverridePhysicsAssetsInfo
+{
+	GENERATED_BODY()
+
+	UPROPERTY()
+	TSoftClassPtr<UAnimInstance> AnimInstanceClass;
+	
+	UPROPERTY()
+	TSoftObjectPtr<UPhysicsAsset> SourceAsset;
+	
+	UPROPERTY()
+	int32 PropertyIndex = -1;
+
+	friend FArchive& operator<<(FArchive& Ar, FAnimBpOverridePhysicsAssetsInfo& Info)
+	{
+		FString AnimInstanceClassPathString;
+		FString PhysicsAssetPathString;
+
+		if (Ar.IsLoading())
+		{	
+			Ar << AnimInstanceClassPathString;
+			Ar << PhysicsAssetPathString;
+			Ar << Info.PropertyIndex;
+
+			Info.AnimInstanceClass = TSoftClassPtr<UAnimInstance>(AnimInstanceClassPathString);
+			Info.SourceAsset = TSoftObjectPtr<UPhysicsAsset>(PhysicsAssetPathString);
+		}
+
+		if (Ar.IsSaving())
+		{
+			AnimInstanceClassPathString = Info.AnimInstanceClass.ToString();
+			PhysicsAssetPathString = Info.SourceAsset.ToString();
+
+			Ar << AnimInstanceClassPathString;
+			Ar << PhysicsAssetPathString;
+			Ar << Info.PropertyIndex;
+		}
+
+		return Ar;
+	}
+
+	friend bool operator==(const FAnimBpOverridePhysicsAssetsInfo& Lhs, const FAnimBpOverridePhysicsAssetsInfo& Rhs)
+	{
+		return Lhs.AnimInstanceClass == Rhs.AnimInstanceClass && 
+			   Lhs.SourceAsset	     == Rhs.SourceAsset       && 
+		 	   Lhs.PropertyIndex     == Rhs.PropertyIndex;
+	}
+};
+
 USTRUCT()
 struct FMorphTargetInfo
 {
@@ -620,6 +673,7 @@ struct FMutableSkinWeightProfileInfo
 };
 
 
+
 UCLASS()
 class CUSTOMIZABLEOBJECT_API UMutableMaskOutCache : public UObject
 {
@@ -1165,6 +1219,9 @@ public:
 	UPROPERTY(EditAnywhere, Category = CompileOptions)
 	bool bEnablePhysicsAssetMerge = false;
 
+	UPROPERTY(EditAnywhere, Category = CompileOptions)
+	bool bEnableAnimBpPhysicsAssetsManipualtion = false;
+
 	// Options when compiling this customizable object (see EMutableCompileMeshType declaration for info)
 	UPROPERTY(EditAnywhere, Category = CompileOptions)
 	EMutableCompileMeshType MeshCompileType = EMutableCompileMeshType::LocalAndChildren;
@@ -1368,7 +1425,7 @@ private:
 	// This is a manual version number for the binary blobs in this asset.
 	// Increasing it invalidates all the previously compiled models.
 	// Warning: If while merging code both versions have changed, take the highest+1.
-	static const int32 CurrentSupportedVersion = 377;
+	static const int32 CurrentSupportedVersion = 379;
 
 public:
 
@@ -1511,6 +1568,9 @@ public:
 	UPROPERTY()
 	TMap<FString, TSoftClassPtr<UAnimInstance>> AnimBPAssetsMap;
 
+	UPROPERTY()
+	TArray<FAnimBpOverridePhysicsAssetsInfo> AnimBpOverridePhysiscAssetsInfo;
+
 	UPROPERTY()
 	/** Stores the sockets provided by the part skeletal meshes, to be merged in the generated meshes */
 	TArray<FMutableRefSocket> SocketArray;
diff --git a/Engine/Plugins/Experimental/Mutable/Source/CustomizableObject/Classes/MuCO/CustomizableObjectInstance.h b/Engine/Plugins/Experimental/Mutable/Source/CustomizableObject/Classes/MuCO/CustomizableObjectInstance.h
--- a/Engine/Plugins/Experimental/Mutable/Source/CustomizableObject/Classes/MuCO/CustomizableObjectInstance.h
+++ b/Engine/Plugins/Experimental/Mutable/Source/CustomizableObject/Classes/MuCO/CustomizableObjectInstance.h
@@ -130,6 +130,7 @@ public:
 	UPROPERTY(Transient, VisibleAnywhere, Category = CustomizableSkeletalMesh)
 	TArray< TObjectPtr<USkeletalMesh> > SkeletalMeshes;
 
+
 	// Will store status description of current skeletal mesh generation (for instance, "EmptyLOD0" or "EmptyMesh"
 	UPROPERTY()
 	FString SkeletalMeshStatus;
@@ -528,6 +529,9 @@ public:
 
 	void ForEachAnimInstance(int32 ComponentIndex, FEachComponentAnimInstanceClassNativeDelegate Delegate) const;
 
+	bool AnimInstanceNeedsFixup(TSubclassOf<UAnimInstance> AnimInstance) const;
+	void AnimInstanceFixup(UAnimInstance* AnimInstance) const;
+
 	/** Serializes/Deserializes all the customization parameters to/from a descriptor, ready to be sent/read, works like a typical UE4 Two-Way Save System Function  */
 	void SaveDescriptor(FArchive &CustomizableObjectDescriptor);
 	void LoadDescriptor(FArchive &CustomizableObjectDescriptor);
diff --git a/Engine/Plugins/Experimental/Mutable/Source/CustomizableObject/CustomizableObject.Build.cs b/Engine/Plugins/Experimental/Mutable/Source/CustomizableObject/CustomizableObject.Build.cs
--- a/Engine/Plugins/Experimental/Mutable/Source/CustomizableObject/CustomizableObject.Build.cs
+++ b/Engine/Plugins/Experimental/Mutable/Source/CustomizableObject/CustomizableObject.Build.cs
@@ -28,6 +28,7 @@ public class CustomizableObject : ModuleRules
             "ApplicationCore",
 			"ClothingSystemRuntimeCommon",
 			"ClothingSystemRuntimeInterface",
+			"AnimGraphRuntime",
 			//"ClothingSystemEditor",
 			"UMG",
 			"MutableRuntime",
@@ -41,6 +42,7 @@ public class CustomizableObject : ModuleRules
 			"ClothingSystemRuntimeCommon",
             "GameplayTags",
 			"MutableRuntime",
+			"AnimGraphRuntime",
 		});
 
         PrivateIncludePathModuleNames.AddRange(
diff --git a/Engine/Plugins/Experimental/Mutable/Source/CustomizableObject/Private/MuCO/CustomizableInstancePrivateData.h b/Engine/Plugins/Experimental/Mutable/Source/CustomizableObject/Private/MuCO/CustomizableInstancePrivateData.h
--- a/Engine/Plugins/Experimental/Mutable/Source/CustomizableObject/Private/MuCO/CustomizableInstancePrivateData.h
+++ b/Engine/Plugins/Experimental/Mutable/Source/CustomizableObject/Private/MuCO/CustomizableInstancePrivateData.h
@@ -93,6 +93,11 @@ struct FReferencedPhysicsAssets
 	
 	UPROPERTY(Transient)
 	TArray< TObjectPtr<UPhysicsAsset> > PhysicsAssetsToMerge;
+
+	TArray<int32> AdditionalPhysicsAssetsToLoad;
+
+	UPROPERTY(Transient)
+	TArray<TObjectPtr<UPhysicsAsset>> AdditionalPhysicsAssets;
 };
 
 
@@ -144,6 +149,26 @@ struct FCustomizableInstanceComponentData
 	bool operator==(const FCustomizableInstanceComponentData& Other) const { return ComponentIndex == Other.ComponentIndex; }
 };
 
+USTRUCT()
+struct FAnimInstanceOverridePhysicsAsset
+{
+	GENERATED_BODY()
+
+	UPROPERTY(Transient)
+	int32 PropertyIndex;
+
+	UPROPERTY(Transient)
+	TObjectPtr<UPhysicsAsset> PhysicsAsset;
+};
+
+USTRUCT()
+struct FAnimBpGeneratedPhysicsAssets
+{
+	GENERATED_BODY()
+
+	UPROPERTY(Transient)
+	TArray<FAnimInstanceOverridePhysicsAsset> AnimInstancePropertyIndexAndPhysicsAssets;
+};
 
 UCLASS()
 class UCustomizableInstancePrivateData : public UObject
@@ -238,6 +263,8 @@ public:
 
 	void DiscardResourcesAndSetReferenceSkeletalMesh(UCustomizableObjectInstance* Public);
 
+	const TArray<FAnimInstanceOverridePhysicsAsset>* GetGeneratedPhysicsAssetsForAnimInstance(TSubclassOf<UAnimInstance> AnimInstance) const;
+
 	/** 
 	* \param OnlyLOD: If not 0, extract and convert only one single LOD from the source image.
 	* \param ExtractChannel: If different than -1, extract a single-channel image with the specified source channel data.
@@ -267,12 +294,11 @@ private:
 	USkeleton* MergeSkeletons(UCustomizableObjectInstance* Public, const FMutableRefSkeletalMeshData* RefSkeletalMeshData, int32 ComponentIndex);
 
 	//
-	UPhysicsAsset* GetOrBuildPhysicsAsset(TObjectPtr<class UPhysicsAsset> TamplateAsset, const mu::PhysicsBody* PhysicsBody, int32 ComponentIndex, bool bDisableCollisionBetweenAssets);
+	UPhysicsAsset* GetOrBuildMainPhysicsAsset(TObjectPtr<class UPhysicsAsset> TamplateAsset, const mu::PhysicsBody* PhysicsBody, int32 ComponentIndex, bool bDisableCollisionBetweenAssets);
 	
 	// Create a transient texture and add it to the TextureTrackerArray
 	UTexture2D* CreateTexture();
 
-	
 	void InvalidateGeneratedData();
 
 	bool DoComponentsNeedUpdate(UCustomizableObjectInstance* CustomizableObjectInstance, const TSharedPtr<FMutableOperationData>& OperationData, TArray<bool>& OutComponentNeedsUpdate, bool& bOutEmptyMesh);
@@ -318,6 +344,9 @@ public:
 	UPROPERTY(Transient, Category = CustomizableObjectInstance, editfixedsize, VisibleAnywhere)
 	FGameplayTagContainer AnimBPGameplayTags;
 
+	UPROPERTY(Transient, Category = CustomizableObjectInstance, editfixedsize, VisibleAnywhere)
+	TMap<TSubclassOf<UAnimInstance>, FAnimBpGeneratedPhysicsAssets> AnimBpPhysicsAssets;
+
 	// Struct used during an update to avoid generating resources that can be reused.
 	FInstanceGeneratedData LastUpdateData;
 
diff --git a/Engine/Plugins/Experimental/Mutable/Source/CustomizableObject/Private/MuCO/CustomizableObject.cpp b/Engine/Plugins/Experimental/Mutable/Source/CustomizableObject/Private/MuCO/CustomizableObject.cpp
--- a/Engine/Plugins/Experimental/Mutable/Source/CustomizableObject/Private/MuCO/CustomizableObject.cpp
+++ b/Engine/Plugins/Experimental/Mutable/Source/CustomizableObject/Private/MuCO/CustomizableObject.cpp
@@ -330,6 +330,7 @@ void UCustomizableObject::ClearCompiledData()
 	ContributingClothingAssetsData.Empty();
 	ClothSharedConfigsData.Empty();
 	SkinWeightProfilesInfo.Empty();
+	AnimBpOverridePhysiscAssetsInfo.Empty();
 
 #if WITH_EDITORONLY_DATA
 	CustomizableObjectPathMap.Empty();
@@ -427,6 +428,8 @@ void UCustomizableObject::SaveCompiledData(FArchive& MemoryWriter, bool bSkipEdi
 	
 	MemoryWriter << SkinWeightProfilesInfo;
 
+	MemoryWriter << AnimBpOverridePhysiscAssetsInfo;
+
 	MemoryWriter << HashToStreamableBlock;
 
 	// All Editor Only data must be serialized here
@@ -524,6 +527,8 @@ void UCustomizableObject::LoadCompiledData(FArchive& MemoryReader, bool bSkipEdi
 
 		MemoryReader << SkinWeightProfilesInfo;
 
+		MemoryReader << AnimBpOverridePhysiscAssetsInfo;
+
 		MemoryReader << HashToStreamableBlock;
 
 		// All Editor Only data must be loaded here
diff --git a/Engine/Plugins/Experimental/Mutable/Source/CustomizableObject/Private/MuCO/CustomizableObjectInstance.cpp b/Engine/Plugins/Experimental/Mutable/Source/CustomizableObject/Private/MuCO/CustomizableObjectInstance.cpp
--- a/Engine/Plugins/Experimental/Mutable/Source/CustomizableObject/Private/MuCO/CustomizableObjectInstance.cpp
+++ b/Engine/Plugins/Experimental/Mutable/Source/CustomizableObject/Private/MuCO/CustomizableObjectInstance.cpp
@@ -3,8 +3,11 @@
 #include "MuCO/CustomizableObjectInstance.h"
 
 #include "Algo/MaxElement.h"
+#include "Animation/AnimClassInterface.h"
+#include "Animation/AnimBlueprintGeneratedClass.h"
 #include "Animation/AnimInstance.h"
 #include "Animation/Skeleton.h"
+#include "BoneControllers/AnimNode_RigidBody.h"
 #include "ClothConfig.h"
 #include "ClothingAsset.h"
 #include "Engine/SkeletalMeshSocket.h"
@@ -26,7 +29,6 @@
 #include "SkeletalMergingLibrary.h"
 #include "UObject/UObjectIterator.h"
 #include "Widgets/Notifications/SNotificationList.h"
-
 // Required for engine branch preprocessor defines.
 #include "MuCO/UnrealPortabilityHelpers.h"
 #include "MuR/Model.h"
@@ -1059,121 +1061,67 @@ USkeleton* UCustomizableInstancePrivateData::MergeSkeletons(UCustomizableObjectI
 	return FinalSkeleton;
 }
 
-UPhysicsAsset* UCustomizableInstancePrivateData::GetOrBuildPhysicsAsset(
-	TObjectPtr<UPhysicsAsset> TemplateAsset,
-	const mu::PhysicsBody* MutablePhysics, int32 ComponentId,
-	bool bDisableCollisionsBetweenDifferentAssets)
+namespace
 {
-
-	MUTABLE_CPUPROFILER_SCOPE(MergePhysicsAssets);
-
-	check(MutablePhysics);
-
-	UPhysicsAsset* Result = nullptr;
-
-	FCustomizableInstanceComponentData* ComponentData = GetComponentData(ComponentId);
-	check(ComponentData);
-
-	TArray<TObjectPtr<UPhysicsAsset>>& PhysicsAssets = ComponentData->PhysicsAssets.PhysicsAssetsToMerge;
-
-	TArray<TObjectPtr<UPhysicsAsset>> ValidAssets;
-
-	const int32 NumPhysicsAssets = ComponentData->PhysicsAssets.PhysicsAssetsToMerge.Num();
-	for (int32 I = 0; I < NumPhysicsAssets; ++I)
-	{
-		const TObjectPtr<UPhysicsAsset>& PhysicsAsset = ComponentData->PhysicsAssets.PhysicsAssetsToMerge[I];
-
-		if (PhysicsAsset)
-		{
-			ValidAssets.AddUnique(PhysicsAsset);
-		}
-	}
-
-	if (!ValidAssets.Num())
-	{
-		return Result;
-	}
-
-	// Just get the referenced asset if no recontrution or merge is needed.
-	if (ValidAssets.Num() == 1 && !MutablePhysics->bBodiesModified)
-	{
-		return ValidAssets[0];
-	}
-
-	TemplateAsset = TemplateAsset ? TemplateAsset : ValidAssets[0];
-	check(TemplateAsset);
-
-	Result = NewObject<UPhysicsAsset>();
-
-	if (!Result)
-	{
-		return nullptr;
-	}
-
-	Result->SolverSettings = TemplateAsset->SolverSettings;
-	Result->SolverType = TemplateAsset->SolverType;
-
-	Result->bNotForDedicatedServer = TemplateAsset->bNotForDedicatedServer;
-
-
-	auto MakeAggGeomFromMutablePhysics = [](int32 BodyIndex, const mu::PhysicsBody* MutablePhysicsBody) -> FKAggregateGeom
+	FKAggregateGeom MakeAggGeomFromMutablePhysics(int32 BodyIndex, const mu::PhysicsBody* MutablePhysicsBody)
 	{
 		FKAggregateGeom BodyAggGeom;
 
 		auto GetCollisionEnabledFormFlags = [](uint32 Flags) -> ECollisionEnabled::Type
 		{
-			return ECollisionEnabled::Type( Flags & 0xFF ); 
+			return ECollisionEnabled::Type(Flags & 0xFF);
 		};
 
 		auto GetContributeToMassFromFlags = [](uint32 Flags) -> bool
 		{
-			return static_cast<bool>( (Flags >> 8 ) & 1 );
+			return static_cast<bool>((Flags >> 8) & 1);
 		};
 
-		const int32 NumSpheres = MutablePhysicsBody->GetSphereCount( BodyIndex );
+		const int32 NumSpheres = MutablePhysicsBody->GetSphereCount(BodyIndex);
 		TArray<FKSphereElem>& AggSpheres = BodyAggGeom.SphereElems;
 		AggSpheres.Empty(NumSpheres);
 		for (int32 I = 0; I < NumSpheres; ++I)
 		{
-			uint32 Flags = MutablePhysicsBody->GetSphereFlags( BodyIndex, I );
-			FString Name = MutablePhysicsBody->GetSphereName( BodyIndex, I );
+			uint32 Flags = MutablePhysicsBody->GetSphereFlags(BodyIndex, I);
+			FString Name = MutablePhysicsBody->GetSphereName(BodyIndex, I);
 
 			FVector3f Position;
 			float Radius;
 
-			MutablePhysicsBody->GetSphere( BodyIndex, I, Position, Radius );
+			MutablePhysicsBody->GetSphere(BodyIndex, I, Position, Radius);
 			FKSphereElem& NewElem = AggSpheres.AddDefaulted_GetRef();
 			
 			NewElem.Center = FVector(Position);
 			NewElem.Radius = Radius;
-			NewElem.SetContributeToMass( GetContributeToMassFromFlags( Flags ) );
-			NewElem.SetCollisionEnabled( GetCollisionEnabledFormFlags( Flags ) );
+			NewElem.SetContributeToMass(GetContributeToMassFromFlags(Flags));
+			NewElem.SetCollisionEnabled(GetCollisionEnabledFormFlags(Flags));
 			NewElem.SetName(FName(*Name));
 		}
 		
-		const int32 NumBoxes = MutablePhysicsBody->GetBoxCount( BodyIndex );
+		const int32 NumBoxes = MutablePhysicsBody->GetBoxCount(BodyIndex);
 		TArray<FKBoxElem>& AggBoxes = BodyAggGeom.BoxElems;
 		AggBoxes.Empty(NumBoxes);
 		for (int32 I = 0; I < NumBoxes; ++I)
 		{
-			uint32 Flags = MutablePhysicsBody->GetBoxFlags( BodyIndex, I );
-			FString Name = MutablePhysicsBody->GetBoxName( BodyIndex, I );
+			uint32 Flags = MutablePhysicsBody->GetBoxFlags(BodyIndex, I);
+			FString Name = MutablePhysicsBody->GetBoxName(BodyIndex, I);
 
 			FVector3f Position;
-			FQuat4f Orientation;
+			FQuat4f Orientation
+				;
 			FVector3f Size;
-			MutablePhysicsBody->GetBox( BodyIndex, I, Position, Orientation, Size );
+			MutablePhysicsBody->GetBox(BodyIndex, I, Position, Orientation, Size);
 
 			FKBoxElem& NewElem = AggBoxes.AddDefaulted_GetRef();
 			
-			NewElem.Center = FVector( Position );
-			NewElem.Rotation = FRotator( Orientation.Rotator() );
+			NewElem.Center = FVector(Position);
+			NewElem.Rotation = FRotator(Orientation.Rotator());
 			NewElem.X = Size.X;
 			NewElem.Y = Size.Y;
 			NewElem.Z = Size.Z;
-			NewElem.SetContributeToMass( GetContributeToMassFromFlags( Flags ) );
-			NewElem.SetCollisionEnabled( GetCollisionEnabledFormFlags( Flags ) );
-			NewElem.SetName( FName(*Name) );
+			NewElem.SetContributeToMass(GetContributeToMassFromFlags(Flags));
+			NewElem.SetCollisionEnabled(GetCollisionEnabledFormFlags(Flags));
+			NewElem.SetName(FName(*Name));
 		}
 
 		//const int32 NumConvexes = MutablePhysicsBody->GetConvexCount( BodyIndex );
@@ -1197,41 +1145,41 @@ UPhysicsAsset* UCustomizableInstancePrivateData::GetOrBuildPhysicsAsset(
 		//}
 
 		TArray<FKSphylElem>& AggSphyls = BodyAggGeom.SphylElems;
-		const int32 NumSphyls = MutablePhysicsBody->GetSphylCount( BodyIndex );
+		const int32 NumSphyls = MutablePhysicsBody->GetSphylCount(BodyIndex);
 		AggSphyls.Empty(NumSphyls);
 
 		for (int32 I = 0; I < NumSphyls; ++I)
 		{
-			uint32 Flags = MutablePhysicsBody->GetSphylFlags( BodyIndex, I );
-			FString Name = MutablePhysicsBody->GetSphylName( BodyIndex, I );
+			uint32 Flags = MutablePhysicsBody->GetSphylFlags(BodyIndex, I);
+			FString Name = MutablePhysicsBody->GetSphylName(BodyIndex, I);
 
 			FVector3f Position;
 			FQuat4f Orientation;
 			float Radius;
 			float Length;
 
-			MutablePhysicsBody->GetSphyl( BodyIndex, I, Position, Orientation, Radius, Length );
+			MutablePhysicsBody->GetSphyl(BodyIndex, I, Position, Orientation, Radius, Length);
 
 			FKSphylElem& NewElem = AggSphyls.AddDefaulted_GetRef();
 			
-			NewElem.Center = FVector( Position );
-			NewElem.Rotation = FRotator( Orientation.Rotator() );
+			NewElem.Center = FVector(Position);
+			NewElem.Rotation = FRotator(Orientation.Rotator());
 			NewElem.Radius = Radius;
 			NewElem.Length = Length;
 			
-			NewElem.SetContributeToMass( GetContributeToMassFromFlags( Flags ) );
-			NewElem.SetCollisionEnabled( GetCollisionEnabledFormFlags( Flags ) );
-			NewElem.SetName( FName(*Name) );
+			NewElem.SetContributeToMass(GetContributeToMassFromFlags(Flags));
+			NewElem.SetCollisionEnabled(GetCollisionEnabledFormFlags(Flags));
+			NewElem.SetName(FName(*Name));
 		}	
 
 		TArray<FKTaperedCapsuleElem>& AggTaperedCapsules = BodyAggGeom.TaperedCapsuleElems;
-		const int32 NumTaperedCapsules = MutablePhysicsBody->GetTaperedCapsuleCount( BodyIndex );
+		const int32 NumTaperedCapsules = MutablePhysicsBody->GetTaperedCapsuleCount(BodyIndex);
 		AggTaperedCapsules.Empty(NumTaperedCapsules);
 
 		for (int32 I = 0; I < NumTaperedCapsules; ++I)
 		{
-			uint32 Flags = MutablePhysicsBody->GetTaperedCapsuleFlags( BodyIndex, I );
-			FString Name = MutablePhysicsBody->GetTaperedCapsuleName( BodyIndex, I );
+			uint32 Flags = MutablePhysicsBody->GetTaperedCapsuleFlags(BodyIndex, I);
+			FString Name = MutablePhysicsBody->GetTaperedCapsuleName(BodyIndex, I);
 
 			FVector3f Position;
 			FQuat4f Orientation;
@@ -1239,24 +1187,341 @@ UPhysicsAsset* UCustomizableInstancePrivateData::GetOrBuildPhysicsAsset(
 			float Radius1;
 			float Length;
 
-			MutablePhysicsBody->GetTaperedCapsule( BodyIndex, I, Position, Orientation, Radius0, Radius1, Length );
+			MutablePhysicsBody->GetTaperedCapsule(BodyIndex, I, Position, Orientation, Radius0, Radius1, Length);
 
 			FKTaperedCapsuleElem& NewElem = AggTaperedCapsules.AddDefaulted_GetRef();
 			
-			NewElem.Center = FVector( Position );
-			NewElem.Rotation = FRotator( Orientation.Rotator() );
+			NewElem.Center = FVector(Position);
+			NewElem.Rotation = FRotator(Orientation.Rotator());
 			NewElem.Radius0 = Radius0;
 			NewElem.Radius1 = Radius1;
 			NewElem.Length = Length;
 			
-			NewElem.SetContributeToMass( GetContributeToMassFromFlags( Flags ) );
-			NewElem.SetCollisionEnabled( GetCollisionEnabledFormFlags( Flags ) );
-			NewElem.SetName( FName(*Name) );	
+			NewElem.SetContributeToMass(GetContributeToMassFromFlags(Flags));
+			NewElem.SetCollisionEnabled(GetCollisionEnabledFormFlags(Flags));
+			NewElem.SetName(FName(*Name));
 		}	
 
 		return BodyAggGeom;
 	};
 
+	TObjectPtr<UPhysicsAsset> MakePhysicsAssetFromTemplateAndMutableBody(
+		TObjectPtr<UPhysicsAsset> TemplateAsset, const mu::PhysicsBody* MutablePhysics) 
+	{
+		check(TemplateAsset);
+		TObjectPtr<UPhysicsAsset> Result = NewObject<UPhysicsAsset>();
+
+		if (!Result)
+		{
+			return nullptr;
+		}
+
+		Result->SolverSettings = TemplateAsset->SolverSettings;
+		Result->SolverType = TemplateAsset->SolverType;
+
+		Result->bNotForDedicatedServer = TemplateAsset->bNotForDedicatedServer;
+
+		TMap<FName, int32> BonesInUse;
+
+		const int32 MutablePhysicsBodyCount = MutablePhysics->GetBodyCount();
+		BonesInUse.Reserve(MutablePhysicsBodyCount);
+		for ( int32 I = 0; I < MutablePhysicsBodyCount; ++I )
+		{
+			BonesInUse.Add(FName(MutablePhysics->GetBodyBoneName(I)), I);
+		}
+
+		const int32 PhysicsAssetBodySetupNum = TemplateAsset->SkeletalBodySetups.Num();
+		bool bTemplateBodyNotUsedFound = false;
+
+		TArray<uint8> UsageMap;
+		UsageMap.Init(1, PhysicsAssetBodySetupNum);
+
+		for (int32 BodySetupIndex = 0; BodySetupIndex < PhysicsAssetBodySetupNum; ++BodySetupIndex)
+		{
+			const TObjectPtr<USkeletalBodySetup>& BodySetup = TemplateAsset->SkeletalBodySetups[BodySetupIndex];
+
+			int32* MutableBodyIndex = BonesInUse.Find(BodySetup->BoneName);
+			if (!MutableBodyIndex)
+			{
+				bTemplateBodyNotUsedFound = true;
+				UsageMap[BodySetupIndex] = 0;
+				continue;
+			}
+
+			TObjectPtr<USkeletalBodySetup> NewBodySetup = NewObject<USkeletalBodySetup>(Result);
+			NewBodySetup->BodySetupGuid = FGuid::NewGuid();
+
+			// Copy Body properties 	
+			NewBodySetup->BoneName = BodySetup->BoneName;
+			NewBodySetup->PhysicsType = BodySetup->PhysicsType;
+			NewBodySetup->bConsiderForBounds = BodySetup->bConsiderForBounds;
+			NewBodySetup->bMeshCollideAll = BodySetup->bMeshCollideAll;
+			NewBodySetup->bDoubleSidedGeometry = BodySetup->bDoubleSidedGeometry;
+			NewBodySetup->bGenerateNonMirroredCollision = BodySetup->bGenerateNonMirroredCollision;
+			NewBodySetup->bSharedCookedData = BodySetup->bSharedCookedData;
+			NewBodySetup->bGenerateMirroredCollision = BodySetup->bGenerateMirroredCollision;
+			NewBodySetup->PhysMaterial = BodySetup->PhysMaterial;
+			NewBodySetup->CollisionReponse = BodySetup->CollisionReponse;
+			NewBodySetup->CollisionTraceFlag = BodySetup->CollisionTraceFlag;
+			NewBodySetup->DefaultInstance = BodySetup->DefaultInstance;
+			NewBodySetup->WalkableSlopeOverride = BodySetup->WalkableSlopeOverride;
+			NewBodySetup->BuildScale3D = BodySetup->BuildScale3D;
+			NewBodySetup->bSkipScaleFromAnimation = BodySetup->bSkipScaleFromAnimation;
+
+			// PhysicalAnimationProfiles can't be added with the current UPhysicsAsset API outside the editor.
+			// Don't pouplate them for now.	
+			//NewBodySetup->PhysicalAnimationData = BodySetup->PhysicalAnimationData;
+
+			NewBodySetup->AggGeom = MakeAggGeomFromMutablePhysics(*MutableBodyIndex, MutablePhysics);
+
+			Result->SkeletalBodySetups.Add(NewBodySetup);
+		}
+
+		if (!bTemplateBodyNotUsedFound)
+		{
+			Result->CollisionDisableTable = TemplateAsset->CollisionDisableTable;
+			Result->ConstraintSetup = TemplateAsset->ConstraintSetup;
+		}
+		else
+		{
+			// Recreate the collision disable entry
+			Result->CollisionDisableTable.Reserve(TemplateAsset->CollisionDisableTable.Num());
+			for (const TPair<FRigidBodyIndexPair, bool>& CollisionDisableEntry : TemplateAsset->CollisionDisableTable)
+			{
+				const bool bIndex0Used = UsageMap[CollisionDisableEntry.Key.Indices[0]] > 0;
+				const bool bIndex1Used = UsageMap[CollisionDisableEntry.Key.Indices[1]] > 0;
+
+				if (bIndex0Used && bIndex1Used)
+				{
+					Result->CollisionDisableTable.Add(CollisionDisableEntry);
+				}
+			}
+
+			// Only add constraints that are part of the bones used for the mutable physics volumes description.
+			for (const TObjectPtr<UPhysicsConstraintTemplate>& Constrain : TemplateAsset->ConstraintSetup)
+			{
+				const FName BoneA = Constrain->DefaultInstance.ConstraintBone1;
+				const FName BoneB = Constrain->DefaultInstance.ConstraintBone2;
+
+				if (BonesInUse.Find(BoneA) && BonesInUse.Find(BoneB))
+				{
+					Result->ConstraintSetup.Add(Constrain);
+				}	
+			}
+		}
+		
+		Result->ConstraintSetup = TemplateAsset->ConstraintSetup;
+		
+		Result->UpdateBodySetupIndexMap();
+		Result->UpdateBoundsBodiesArray();
+
+
+
+	#if WITH_EDITORONLY_DATA
+		Result->ConstraintProfiles = TemplateAsset->ConstraintProfiles;
+	#endif
+
+		return Result;
+	}
+}
+
+UPhysicsAsset* UCustomizableInstancePrivateData::GetOrBuildMainPhysicsAsset(
+	TObjectPtr<UPhysicsAsset> TemplateAsset,
+	const mu::PhysicsBody* MutablePhysics, int32 ComponentId,
+	bool bDisableCollisionsBetweenDifferentAssets)
+{
+
+	MUTABLE_CPUPROFILER_SCOPE(MergePhysicsAssets);
+
+	check(MutablePhysics);
+
+	UPhysicsAsset* Result = nullptr;
+
+	FCustomizableInstanceComponentData* ComponentData = GetComponentData(ComponentId);
+	check(ComponentData);
+
+	TArray<TObjectPtr<UPhysicsAsset>>& PhysicsAssets = ComponentData->PhysicsAssets.PhysicsAssetsToMerge;
+
+	TArray<TObjectPtr<UPhysicsAsset>> ValidAssets;
+
+	const int32 NumPhysicsAssets = ComponentData->PhysicsAssets.PhysicsAssetsToMerge.Num();
+	for (int32 I = 0; I < NumPhysicsAssets; ++I)
+	{
+		const TObjectPtr<UPhysicsAsset>& PhysicsAsset = ComponentData->PhysicsAssets.PhysicsAssetsToMerge[I];
+
+		if (PhysicsAsset)
+		{
+			ValidAssets.AddUnique(PhysicsAsset);
+		}
+	}
+
+	if (!ValidAssets.Num())
+	{
+		return Result;
+	}
+
+	// Just get the referenced asset if no recontrution or merge is needed.
+	if (ValidAssets.Num() == 1 && !MutablePhysics->bBodiesModified)
+	{
+		return ValidAssets[0];
+	}
+
+	TemplateAsset = TemplateAsset ? TemplateAsset : ValidAssets[0];
+	check(TemplateAsset);
+
+	Result = NewObject<UPhysicsAsset>();
+
+	if (!Result)
+	{
+		return nullptr;
+	}
+
+	Result->SolverSettings = TemplateAsset->SolverSettings;
+	Result->SolverType = TemplateAsset->SolverType;
+
+	Result->bNotForDedicatedServer = TemplateAsset->bNotForDedicatedServer;
+
+	//auto MakeAggGeomFromMutablePhysics = [](int32 BodyIndex, const mu::PhysicsBody* MutablePhysicsBody) -> FKAggregateGeom
+	//{
+	//	FKAggregateGeom BodyAggGeom;
+
+	//	auto GetCollisionEnabledFormFlags = [](uint32 Flags) -> ECollisionEnabled::Type
+	//	{
+	//		return ECollisionEnabled::Type( Flags & 0xFF ); 
+	//	};
+
+	//	auto GetContributeToMassFromFlags = [](uint32 Flags) -> bool
+	//	{
+	//		return static_cast<bool>( (Flags >> 8 ) & 1 );
+	//	};
+
+	//	const int32 NumSpheres = MutablePhysicsBody->GetSphereCount( BodyIndex );
+	//	TArray<FKSphereElem>& AggSpheres = BodyAggGeom.SphereElems;
+	//	AggSpheres.Empty(NumSpheres);
+	//	for (int32 I = 0; I < NumSpheres; ++I)
+	//	{
+	//		uint32 Flags = MutablePhysicsBody->GetSphereFlags( BodyIndex, I );
+	//		FString Name = MutablePhysicsBody->GetSphereName( BodyIndex, I );
+
+	//		FVector3f Position;
+	//		float Radius;
+
+	//		MutablePhysicsBody->GetSphere( BodyIndex, I, Position, Radius );
+	//		FKSphereElem& NewElem = AggSpheres.AddDefaulted_GetRef();
+	//		
+	//		NewElem.Center = FVector(Position);
+	//		NewElem.Radius = Radius;
+	//		NewElem.SetContributeToMass( GetContributeToMassFromFlags( Flags ) );
+	//		NewElem.SetCollisionEnabled( GetCollisionEnabledFormFlags( Flags ) );
+	//		NewElem.SetName(FName(*Name));
+	//	}
+	//	
+	//	const int32 NumBoxes = MutablePhysicsBody->GetBoxCount( BodyIndex );
+	//	TArray<FKBoxElem>& AggBoxes = BodyAggGeom.BoxElems;
+	//	AggBoxes.Empty(NumBoxes);
+	//	for (int32 I = 0; I < NumBoxes; ++I)
+	//	{
+	//		uint32 Flags = MutablePhysicsBody->GetBoxFlags( BodyIndex, I );
+	//		FString Name = MutablePhysicsBody->GetBoxName( BodyIndex, I );
+
+	//		FVector3f Position;
+	//		FQuat4f Orientation;
+	//		FVector3f Size;
+	//		MutablePhysicsBody->GetBox( BodyIndex, I, Position, Orientation, Size );
+
+	//		FKBoxElem& NewElem = AggBoxes.AddDefaulted_GetRef();
+	//		
+	//		NewElem.Center = FVector( Position );
+	//		NewElem.Rotation = FRotator( Orientation.Rotator() );
+	//		NewElem.X = Size.X;
+	//		NewElem.Y = Size.Y;
+	//		NewElem.Z = Size.Z;
+	//		NewElem.SetContributeToMass( GetContributeToMassFromFlags( Flags ) );
+	//		NewElem.SetCollisionEnabled( GetCollisionEnabledFormFlags( Flags ) );
+	//		NewElem.SetName( FName(*Name) );
+	//	}
+
+	//	//const int32 NumConvexes = MutablePhysicsBody->GetConvexCount( BodyIndex );
+	//	//TArray<FKConvexElem>& AggConvexes = BodyAggGeom.ConvexElems;
+	//	//AggConvexes.Empty();
+	//	//for (int32 I = 0; I < NumConvexes; ++I)
+	//	//{
+	//	//	uint32 Flags = MutablePhysicsBody->GetConvexFlags( BodyIndex, I );
+	//	//	FString Name = MutablePhysicsBody->GetConvexName( BodyIndex, I );
+
+	//	//	const FVector3f* Vertices;
+	//	//	const int32* Indices;
+	//	//	int32 NumVertices;
+	//	//	int32 NumIndices;
+	//	//	FTransform3f Transform;
+
+	//	//	MutablePhysicsBody->GetConvex( BodyIndex, I, Vertices, NumVertices, Indices, NumIndices, Transform );
+	//	//	
+	//	//	TArrayView<const FVector3f> VerticesView( Vertices, NumVertices );
+	//	//	TArrayView<const int32> IndicesView( Indices, NumIndices );
+	//	//}
+
+	//	TArray<FKSphylElem>& AggSphyls = BodyAggGeom.SphylElems;
+	//	const int32 NumSphyls = MutablePhysicsBody->GetSphylCount( BodyIndex );
+	//	AggSphyls.Empty(NumSphyls);
+
+	//	for (int32 I = 0; I < NumSphyls; ++I)
+	//	{
+	//		uint32 Flags = MutablePhysicsBody->GetSphylFlags( BodyIndex, I );
+	//		FString Name = MutablePhysicsBody->GetSphylName( BodyIndex, I );
+
+	//		FVector3f Position;
+	//		FQuat4f Orientation;
+	//		float Radius;
+	//		float Length;
+
+	//		MutablePhysicsBody->GetSphyl( BodyIndex, I, Position, Orientation, Radius, Length );
+
+	//		FKSphylElem& NewElem = AggSphyls.AddDefaulted_GetRef();
+	//		
+	//		NewElem.Center = FVector( Position );
+	//		NewElem.Rotation = FRotator( Orientation.Rotator() );
+	//		NewElem.Radius = Radius;
+	//		NewElem.Length = Length;
+	//		
+	//		NewElem.SetContributeToMass( GetContributeToMassFromFlags( Flags ) );
+	//		NewElem.SetCollisionEnabled( GetCollisionEnabledFormFlags( Flags ) );
+	//		NewElem.SetName( FName(*Name) );
+	//	}	
+
+	//	TArray<FKTaperedCapsuleElem>& AggTaperedCapsules = BodyAggGeom.TaperedCapsuleElems;
+	//	const int32 NumTaperedCapsules = MutablePhysicsBody->GetTaperedCapsuleCount( BodyIndex );
+	//	AggTaperedCapsules.Empty(NumTaperedCapsules);
+
+	//	for (int32 I = 0; I < NumTaperedCapsules; ++I)
+	//	{
+	//		uint32 Flags = MutablePhysicsBody->GetTaperedCapsuleFlags( BodyIndex, I );
+	//		FString Name = MutablePhysicsBody->GetTaperedCapsuleName( BodyIndex, I );
+
+	//		FVector3f Position;
+	//		FQuat4f Orientation;
+	//		float Radius0;
+	//		float Radius1;
+	//		float Length;
+
+	//		MutablePhysicsBody->GetTaperedCapsule( BodyIndex, I, Position, Orientation, Radius0, Radius1, Length );
+
+	//		FKTaperedCapsuleElem& NewElem = AggTaperedCapsules.AddDefaulted_GetRef();
+	//		
+	//		NewElem.Center = FVector( Position );
+	//		NewElem.Rotation = FRotator( Orientation.Rotator() );
+	//		NewElem.Radius0 = Radius0;
+	//		NewElem.Radius1 = Radius1;
+	//		NewElem.Length = Length;
+	//		
+	//		NewElem.SetContributeToMass( GetContributeToMassFromFlags( Flags ) );
+	//		NewElem.SetCollisionEnabled( GetCollisionEnabledFormFlags( Flags ) );
+	//		NewElem.SetName( FName(*Name) );	
+	//	}	
+
+	//	return BodyAggGeom;
+	//};
+
 	TMap<FName, int32> BonesInUse;
 
 	const int32 MutablePhysicsBodyCount = MutablePhysics->GetBodyCount();
@@ -1715,20 +1980,61 @@ bool UCustomizableInstancePrivateData::UpdateSkeletalMesh_PostBeginUpdate0(UCust
 				break;
 			}
 
-			// Build PhysicsAsset Merge Physics Assets coming from SubMeshes of the newly generated Mesh
+			// Build PhysicsAsset merging physics assets coming from SubMeshes of the newly generated Mesh
 			if (mu::Ptr<const mu::PhysicsBody> MutablePhysics = Component.Mesh->GetPhysicsBody())
 			{
 				constexpr bool bDisallowCollisionBetweenAssets = true;
-				UPhysicsAsset* PhysicsAssetResult = GetOrBuildPhysicsAsset(
+				UPhysicsAsset* PhysicsAssetResult = GetOrBuildMainPhysicsAsset(
 					RefSkeletalMeshData->PhysicsAsset.Get(), MutablePhysics.get(), Component.Id, bDisallowCollisionBetweenAssets);
 
 				SkeletalMesh->SetPhysicsAsset(PhysicsAssetResult);
 
 				if (PhysicsAssetResult)
 				{
+					// We are setting the physics asset mesh preview to the generated skeletal mesh.
+					// this is fine if the phyiscs asset is also generated, but not sure about referenced assets.
+					// In any case, don't mark the asset as modified.
+
 					// This is only called in editor, no need to add editor guards.	
-					PhysicsAssetResult->SetPreviewMesh(SkeletalMesh);
+					constexpr bool bMarkAsDirty = false;
+					PhysicsAssetResult->SetPreviewMesh(SkeletalMesh, bMarkAsDirty);
+				}
 				}
+
+			const int32 NumAdditionalPhysicsNum = Component.Mesh->AdditionalPhysicsBodies.Num();
+			for (int32 I = 0; I < NumAdditionalPhysicsNum; ++I)
+			{
+				mu::Ptr<const mu::PhysicsBody> AdditionalPhysiscBody = Component.Mesh->AdditionalPhysicsBodies[I];
+				
+				check(AdditionalPhysiscBody);
+				if (!AdditionalPhysiscBody->bBodiesModified)
+				{
+					continue;
+				}
+
+				const int32 PhysicsBodyExternalId = Component.Mesh->AdditionalPhysicsBodies[I]->CustomId;
+				
+				const FAnimBpOverridePhysicsAssetsInfo& Info = CustomizableObject->AnimBpOverridePhysiscAssetsInfo[PhysicsBodyExternalId];
+
+				// Make sure the AnimInstance class is loaded. It is expected to be already loaded at this point though. 
+				UClass* AnimInstanceClassLoaded = Info.AnimInstanceClass.LoadSynchronous();
+				TSubclassOf<UAnimInstance> AnimInstanceClass = TSubclassOf<UAnimInstance>(AnimInstanceClassLoaded);
+				if (!ensureAlways(AnimInstanceClass))
+				{
+					continue;
+				}
+
+				FAnimBpGeneratedPhysicsAssets& PhysicsAssetsUsedByAnimBp = AnimBpPhysicsAssets.FindOrAdd(AnimInstanceClass);
+
+				TObjectPtr<UPhysicsAsset> PhysicsAssetTemplate = TObjectPtr<UPhysicsAsset>(Info.SourceAsset.Get());
+
+				check(PhysicsAssetTemplate);
+
+				FAnimInstanceOverridePhysicsAsset& Entry =
+					PhysicsAssetsUsedByAnimBp.AnimInstancePropertyIndexAndPhysicsAssets.Emplace_GetRef();
+
+				Entry.PropertyIndex = Info.PropertyIndex;
+				Entry.PhysicsAsset = MakePhysicsAssetFromTemplateAndMutableBody(PhysicsAssetTemplate, AdditionalPhysiscBody.get());
 			}
 
 			// Add sockets from the SkeletalMesh of reference and from the MutableMesh
@@ -4672,6 +4978,7 @@ FGraphEventRef UCustomizableInstancePrivateData::LoadAdditionalAssetsAsync(const
 
 	GatheredAnimBPs.Empty();
 	AnimBPGameplayTags.Reset();
+	AnimBpPhysicsAssets.Reset();
 
 	for (const FInstanceUpdateData::FSurface& Surface : OperationData->InstanceUpdateData.Surfaces)
 	{
@@ -4867,6 +5174,15 @@ FGraphEventRef UCustomizableInstancePrivateData::LoadAdditionalAssetsAsync(const
 				}
 #endif
 			}
+
+			const int32 AdditionalPhysicsNum = MutableMesh->AdditionalPhysicsBodies.Num();
+			for (int32 I = 0; I < AdditionalPhysicsNum; ++I)
+			{
+				const int32 ExternalId = MutableMesh->AdditionalPhysicsBodies[I]->CustomId;
+				
+				ComponentData->PhysicsAssets.AdditionalPhysicsAssetsToLoad.Add(ExternalId);
+				AssetsToStream.Add(CustomizableObject->AnimBpOverridePhysiscAssetsInfo[ExternalId].SourceAsset.ToSoftObjectPath());
+			}
 		}
 	}
 
@@ -5009,6 +5325,17 @@ void UCustomizableInstancePrivateData::AdditionalAssetsAsyncLoaded(UCustomizable
 			}
 #endif
 		}
+
+		const int32 AdditionalPhysicsNum = ComponentData.PhysicsAssets.AdditionalPhysicsAssetsToLoad.Num();
+		ComponentData.PhysicsAssets.AdditionalPhysicsAssets.Reserve(AdditionalPhysicsNum);
+		for (int32 I = 0; I < AdditionalPhysicsNum; ++I)
+		{
+			// Make the loaded assets references strong.
+			const int32 AnimBpPhysicsOverrideIndex = ComponentData.PhysicsAssets.AdditionalPhysicsAssetsToLoad[I];
+			ComponentData.PhysicsAssets.AdditionalPhysicsAssets.Add( 
+					CustomizableObject->AnimBpOverridePhysiscAssetsInfo[AnimBpPhysicsOverrideIndex].SourceAsset.Get());
+		}
+		ComponentData.PhysicsAssets.AdditionalPhysicsAssetsToLoad.Empty();
 	}
 }
 
@@ -6147,6 +6474,94 @@ void UCustomizableObjectInstance::ForEachAnimInstance(int32 ComponentIndex, FEac
 	}
 }
 
+bool UCustomizableObjectInstance::AnimInstanceNeedsFixup(TSubclassOf<UAnimInstance> AnimInstanceClass) const
+{
+	return PrivateData->AnimBpPhysicsAssets.Contains(AnimInstanceClass);
+}
+
+void UCustomizableObjectInstance::AnimInstanceFixup(UAnimInstance* InAnimInstance) const
+{
+	if (!InAnimInstance)
+	{
+		return;
+	}
+
+	TSubclassOf<UAnimInstance> AnimInstanceClass = InAnimInstance->GetClass();
+
+	const TArray<FAnimInstanceOverridePhysicsAsset>* AnimInstanceOverridePhysicsAssets = 
+			PrivateData->GetGeneratedPhysicsAssetsForAnimInstance(AnimInstanceClass);
+	
+	if (!AnimInstanceOverridePhysicsAssets)
+	{
+		return;
+	}
+
+	// Swap RigidBody anim nodes override physics assets with mutable generated ones.
+	if (UAnimBlueprintGeneratedClass* AnimClass = Cast<UAnimBlueprintGeneratedClass>(AnimInstanceClass))
+	{
+		bool bPropertyMismatchFound = false;
+		const int32 AnimNodePropertiesNum = AnimClass->AnimNodeProperties.Num();
+
+		for (const FAnimInstanceOverridePhysicsAsset& PropIndexAndAsset : *AnimInstanceOverridePhysicsAssets)
+		{
+			check(PropIndexAndAsset.PropertyIndex >= 0);
+			if (PropIndexAndAsset.PropertyIndex >= AnimNodePropertiesNum)
+			{
+				bPropertyMismatchFound = true;
+				continue;
+			}
+
+			const int32 AnimNodePropIndex = PropIndexAndAsset.PropertyIndex;
+
+			FStructProperty* StructProperty = AnimClass->AnimNodeProperties[AnimNodePropIndex];
+
+			if (!ensure(StructProperty))
+			{
+				bPropertyMismatchFound = true;
+				continue;
+			}
+
+			const bool bIsRigidBodyNode = StructProperty->Struct->IsChildOf(FAnimNode_RigidBody::StaticStruct());
+
+			if (!bIsRigidBodyNode)
+			{
+				bPropertyMismatchFound = true;
+				continue;
+			}
+
+			FAnimNode_RigidBody* RbanNode = StructProperty->ContainerPtrToValuePtr<FAnimNode_RigidBody>(InAnimInstance);
+
+			if (!ensure(RbanNode))
+			{
+				bPropertyMismatchFound = true;
+				continue;
+			}
+
+			RbanNode->OverridePhysicsAsset = PropIndexAndAsset.PhysicsAsset;
+		}
+#if WITH_EDITOR
+		if (bPropertyMismatchFound)
+		{
+			UE_LOG(LogMutable, Warning, TEXT("AnimBp %s is not in sync with the data stored in the CO %s. A CO recompilation may be needed."),
+				*AnimInstanceClass.Get()->GetName(), 
+				*GetCustomizableObject()->GetName());
+		}
+#endif
+	}
+}
+
+const TArray<FAnimInstanceOverridePhysicsAsset>* UCustomizableInstancePrivateData::GetGeneratedPhysicsAssetsForAnimInstance(TSubclassOf<UAnimInstance> AnimInstanceClass) const
+{
+	const FAnimBpGeneratedPhysicsAssets* Found = AnimBpPhysicsAssets.Find(AnimInstanceClass);
+
+	if (!Found)
+	{
+		return nullptr;
+	}
+
+	return &Found->AnimInstancePropertyIndexAndPhysicsAssets;
+}
+
 void UCustomizableObjectInstance::ForEachAnimInstance(int32 ComponentIndex, FEachComponentAnimInstanceClassNativeDelegate Delegate) const
 {
 	// allow us to log out both bad states with one pass
diff --git a/Engine/Plugins/Experimental/Mutable/Source/CustomizableObjectEditor/Classes/MuCOE/Nodes/CustomizableObjectNodeMeshReshape.h b/Engine/Plugins/Experimental/Mutable/Source/CustomizableObjectEditor/Classes/MuCOE/Nodes/CustomizableObjectNodeMeshReshape.h
--- a/Engine/Plugins/Experimental/Mutable/Source/CustomizableObjectEditor/Classes/MuCOE/Nodes/CustomizableObjectNodeMeshReshape.h
+++ b/Engine/Plugins/Experimental/Mutable/Source/CustomizableObjectEditor/Classes/MuCOE/Nodes/CustomizableObjectNodeMeshReshape.h
@@ -46,6 +46,10 @@ public:
 
 	void Serialize(FArchive& Ar) override;
 
+	/** Enable the deformation of the skeleton of the base mesh. */
+	UPROPERTY(EditAnywhere, Category = CustomizableObject)
+	bool bReshapeVertices = true;	
+
 	/** Enable the deformation of the skeleton of the base mesh. */
 	UPROPERTY(EditAnywhere, Category = CustomizableObject)
 	bool bReshapeSkeleton = false;	
diff --git a/Engine/Plugins/Experimental/Mutable/Source/CustomizableObjectEditor/Classes/MuCOE/Nodes/CustomizableObjectNodeTable.h b/Engine/Plugins/Experimental/Mutable/Source/CustomizableObjectEditor/Classes/MuCOE/Nodes/CustomizableObjectNodeTable.h
--- a/Engine/Plugins/Experimental/Mutable/Source/CustomizableObjectEditor/Classes/MuCOE/Nodes/CustomizableObjectNodeTable.h
+++ b/Engine/Plugins/Experimental/Mutable/Source/CustomizableObjectEditor/Classes/MuCOE/Nodes/CustomizableObjectNodeTable.h
@@ -217,8 +217,11 @@ public:
 		return nullptr;
 	}
 
+
 	// We should do this in a template!
 	USkeletalMesh* GetSkeletalMeshAt(const UEdGraphPin* Pin, const FName& RowName) const;
+	TSoftClassPtr<UAnimInstance> GetAnimInstanceAt(const UEdGraphPin* Pin, const FName& RowName) const;
+
 
 	// Return the name of the enabled rows in the data table.
 	// Returns the name if the row has a bool column set as false (true == disabled)
diff --git a/Engine/Plugins/Experimental/Mutable/Source/CustomizableObjectEditor/CustomizableObjectEditor.Build.cs b/Engine/Plugins/Experimental/Mutable/Source/CustomizableObjectEditor/CustomizableObjectEditor.Build.cs
--- a/Engine/Plugins/Experimental/Mutable/Source/CustomizableObjectEditor/CustomizableObjectEditor.Build.cs
+++ b/Engine/Plugins/Experimental/Mutable/Source/CustomizableObjectEditor/CustomizableObjectEditor.Build.cs
@@ -33,7 +33,9 @@ public class CustomizableObjectEditor : ModuleRules
 				"WorkspaceMenuStructure",
 				"AdvancedPreviewScene",
 				"MessageLog",
-				"KismetWidgets"
+				"KismetWidgets",
+				"AnimGraphRuntime",
+				"AnimGraph",
 			}
 		);
 
@@ -70,6 +72,8 @@ public class CustomizableObjectEditor : ModuleRules
 				"UMG",
 				"Persona",
 				"CurveEditor",
+				"AnimGraphRuntime",
+				"AnimGraph",
 
 				"MutableRuntime",
 				"MutableTools",
diff --git a/Engine/Plugins/Experimental/Mutable/Source/CustomizableObjectEditor/Private/MuCOE/CustomizableObjectCompiler.cpp b/Engine/Plugins/Experimental/Mutable/Source/CustomizableObjectEditor/Private/MuCOE/CustomizableObjectCompiler.cpp
--- a/Engine/Plugins/Experimental/Mutable/Source/CustomizableObjectEditor/Private/MuCOE/CustomizableObjectCompiler.cpp
+++ b/Engine/Plugins/Experimental/Mutable/Source/CustomizableObjectEditor/Private/MuCOE/CustomizableObjectCompiler.cpp
@@ -910,6 +910,7 @@ void FCustomizableObjectCompiler::CompileInternal(UCustomizableObject* Object, c
 	Options.b16BitBoneWeightsEnabled = Object->bEnable16BitBoneWeights;
 	Options.bSkinWeightProfilesEnabled = Object->bEnableAltSkinWeightProfiles;
 	Options.bPhysicsAssetMergeEnabled = Object->bEnablePhysicsAssetMerge;
+	Options.bAnimBpPhysicsManipulationEnabled = Object->bEnableAnimBpPhysicsAssetsManipualtion;
 
 	if (Object->IsLocked() || !UCustomizableObjectSystem::GetInstance()->LockObject(Object))
 	{
@@ -993,6 +994,7 @@ void FCustomizableObjectCompiler::CompileInternal(UCustomizableObject* Object, c
 		Object->MorphTargetReconstructionData = MoveTemp(GenerationContext.MorphTargetReconstructionData);
 		
 		Object->SkinWeightProfilesInfo = MoveTemp(GenerationContext.SkinWeightProfilesInfo);
+		Object->AnimBpOverridePhysiscAssetsInfo = MoveTemp(GenerationContext.AnimBpOverridePhysicsAssetsInfo);
 		
 		// Clothing	
 		Object->ClothMeshToMeshVertData = MoveTemp(GenerationContext.ClothMeshToMeshVertData);
diff --git a/Engine/Plugins/Experimental/Mutable/Source/CustomizableObjectEditor/Private/MuCOE/CustomizableObjectLayout.cpp b/Engine/Plugins/Experimental/Mutable/Source/CustomizableObjectEditor/Private/MuCOE/CustomizableObjectLayout.cpp
--- a/Engine/Plugins/Experimental/Mutable/Source/CustomizableObjectEditor/Private/MuCOE/CustomizableObjectLayout.cpp
+++ b/Engine/Plugins/Experimental/Mutable/Source/CustomizableObjectEditor/Private/MuCOE/CustomizableObjectLayout.cpp
@@ -93,7 +93,7 @@ void UCustomizableObjectLayout::GenerateBlocksFromUVs()
 			GenerationContext.MeshGenerationFlags.Push(ShapeFlags);
 
 			GenerationContext.ComponentInfos.Add(SkeletalMesh);
-			MutableMesh = ConvertSkeletalMeshToMutable(SkeletalMesh, LOD, Material, GenerationContext, Node);
+			MutableMesh = ConvertSkeletalMeshToMutable(SkeletalMesh, TSoftClassPtr<UAnimInstance>(), LOD, Material, GenerationContext, Node);
 		}
 		else if (UStaticMesh* StaticMesh = Cast<UStaticMesh>(Mesh))
 		{
diff --git a/Engine/Plugins/Experimental/Mutable/Source/CustomizableObjectEditor/Private/MuCOE/GenerateMutableSource/GenerateMutableSource.h b/Engine/Plugins/Experimental/Mutable/Source/CustomizableObjectEditor/Private/MuCOE/GenerateMutableSource/GenerateMutableSource.h
--- a/Engine/Plugins/Experimental/Mutable/Source/CustomizableObjectEditor/Private/MuCOE/GenerateMutableSource/GenerateMutableSource.h
+++ b/Engine/Plugins/Experimental/Mutable/Source/CustomizableObjectEditor/Private/MuCOE/GenerateMutableSource/GenerateMutableSource.h
@@ -650,6 +650,8 @@ struct FMutableGraphGenerationContext
 	// Data used for SkinWeightProfiles reconstruction
 	TArray<FMutableSkinWeightProfileInfo> SkinWeightProfilesInfo;
 
+	TArray<FAnimBpOverridePhysicsAssetsInfo> AnimBpOverridePhysicsAssetsInfo;
+
 	// Hierarchy of current ComponentNew nodes, each stored for every LOD
 	struct ObjectParent
 	{
diff --git a/Engine/Plugins/Experimental/Mutable/Source/CustomizableObjectEditor/Private/MuCOE/GenerateMutableSource/GenerateMutableSourceMesh.cpp b/Engine/Plugins/Experimental/Mutable/Source/CustomizableObjectEditor/Private/MuCOE/GenerateMutableSource/GenerateMutableSourceMesh.cpp
--- a/Engine/Plugins/Experimental/Mutable/Source/CustomizableObjectEditor/Private/MuCOE/GenerateMutableSource/GenerateMutableSourceMesh.cpp
+++ b/Engine/Plugins/Experimental/Mutable/Source/CustomizableObjectEditor/Private/MuCOE/GenerateMutableSource/GenerateMutableSourceMesh.cpp
@@ -4,6 +4,9 @@
 
 #include "Algo/Count.h"
 #include "Animation/PoseAsset.h"
+#include "Animation/AnimBlueprintGeneratedClass.h"
+#include "Animation/AnimInstance.h"
+#include "AnimGraphNode_RigidBody.h"
 #include "ClothConfigBase.h"
 #include "ClothingAsset.h"
 #include "MeshUtilities.h"
@@ -372,13 +375,195 @@ void SetAndPropagatePoseBoneUsage(
 
 }
 
-mu::MeshPtr ConvertSkeletalMeshToMutable(USkeletalMesh* InSkeletalMesh, int LOD, int MaterialIndex, FMutableGraphGenerationContext& GenerationContext, const UCustomizableObjectNode* CurrentNode)
+TArray<TTuple<UPhysicsAsset*, int32>> GetPhysicsAssetsFromAnimInstance(const TSoftClassPtr<UAnimInstance>& AnimInstance)
+{
+	// TODO: Consider caching the result in the GenerationContext.
+	TArray<TTuple<UPhysicsAsset*, int32>> Result;
+
+	if (AnimInstance.IsNull())
+	{
+		return Result;
+	}
+
+	UClass* AnimInstanceClass = AnimInstance.LoadSynchronous();
+	UAnimBlueprintGeneratedClass* AnimClass = Cast<UAnimBlueprintGeneratedClass>(AnimInstanceClass);
+
+	if (AnimClass)
+	{
+		const int32 AnimNodePropertiesNum = AnimClass->AnimNodeProperties.Num();
+		for ( int32 PropertyIndex = 0; PropertyIndex < AnimNodePropertiesNum; ++PropertyIndex)
+		{
+			FStructProperty* StructProperty = AnimClass->AnimNodeProperties[PropertyIndex];
+
+			if (StructProperty->Struct->IsChildOf(FAnimNode_RigidBody::StaticStruct()))
+			{
+				FAnimNode_RigidBody* Rban = StructProperty->ContainerPtrToValuePtr<FAnimNode_RigidBody>(AnimInstanceClass->GetDefaultObject());
+			
+				if (Rban && Rban->OverridePhysicsAsset)
+				{
+					Result.Emplace(Rban->OverridePhysicsAsset, PropertyIndex);
+				}
+			}
+		}
+	}		
+
+	return Result;
+}
+
+TArray<uint8> MakePhysicsAssetBodySetupRelevancyMap(UPhysicsAsset* Asset, const mu::Ptr<mu::Mesh>& Mesh)
+{
+	const int32 BodySetupsNum = Asset->SkeletalBodySetups.Num();
+
+	TArray<uint8> RelevancyMap;
+	RelevancyMap.Init(0, BodySetupsNum);
+
+	if (!Mesh->GetSkeleton())
+	{
+		return RelevancyMap;
+	}
+
+	for (int32 BodyIndex = 0; BodyIndex < BodySetupsNum; ++BodyIndex)
+	{
+		FString BodyBoneName = Asset->SkeletalBodySetups[BodyIndex]->BoneName.ToString();
+
+		RelevancyMap[BodyIndex] = Mesh->GetSkeleton()->FindBone(TCHAR_TO_ANSI(*BodyBoneName)) >= 0;
+	}
+
+	return RelevancyMap;
+}
+
+mu::Ptr<mu::PhysicsBody> MakePhysicsBodyFromAsset(UPhysicsAsset* Asset, const TArray<uint8>& BodySetupRelevancyMap)
+{
+	check(Asset);
+	check(Asset->SkeletalBodySetups.Num() == BodySetupRelevancyMap.Num());
+
+	// Find BodySetups with relevant bones.
+	TArray<TObjectPtr<USkeletalBodySetup>>& SkeletalBodySetups = Asset->SkeletalBodySetups;
+	
+	const int32 NumRelevantSetups = Algo::CountIf(BodySetupRelevancyMap, [](uint8 V) { return V; });
+
+	mu::Ptr<mu::PhysicsBody> PhysicsBody = new mu::PhysicsBody;
+	
+	PhysicsBody->SetBodyCount(NumRelevantSetups);
+
+	auto GetKBodyElemFlags = [](const FKShapeElem& KElem) -> uint32
+	{
+		uint8 ElemCollisionEnabled = static_cast<uint8>( KElem.GetCollisionEnabled() );
+		
+		uint32 Flags = static_cast<uint32>( ElemCollisionEnabled );
+		Flags = Flags | (static_cast<uint32>(KElem.GetContributeToMass()) << 8);
+
+		return Flags; 
+	};
+
+	for (int32 B = 0, SourceBodyIndex = 0; B < NumRelevantSetups; ++B)
+	{
+		if (!BodySetupRelevancyMap[SourceBodyIndex])
+		{
+			continue;
+		}
+
+		TObjectPtr<USkeletalBodySetup>& BodySetup = SkeletalBodySetups[SourceBodyIndex++];
+
+		FString BodyBoneName = BodySetup->BoneName.ToString();
+		PhysicsBody->SetBodyBoneName(B, TCHAR_TO_ANSI(*BodyBoneName));
+		
+		const int32 NumSpheres = BodySetup->AggGeom.SphereElems.Num();
+		PhysicsBody->SetSphereCount(B, NumSpheres);
+
+		for (int32 I = 0; I < NumSpheres; ++I)
+		{
+			const FKSphereElem& SphereElem = BodySetup->AggGeom.SphereElems[I];
+			PhysicsBody->SetSphere(B, I, FVector3f(SphereElem.Center), SphereElem.Radius);
+
+			const FString ElemName = SphereElem.GetName().ToString();
+			PhysicsBody->SetSphereName(B, I, TCHAR_TO_ANSI(*ElemName));	
+			PhysicsBody->SetSphereFlags(B, I, GetKBodyElemFlags(SphereElem));
+		}
+
+		const int32 NumBoxes = BodySetup->AggGeom.BoxElems.Num();
+		PhysicsBody->SetBoxCount(B, NumBoxes);
+
+		for (int32 I = 0; I < NumBoxes; ++I)
+		{
+			const FKBoxElem& BoxElem = BodySetup->AggGeom.BoxElems[I];
+			PhysicsBody->SetBox(B, I, 
+					FVector3f(BoxElem.Center), 
+					FQuat4f(BoxElem.Rotation.Quaternion()), 
+					FVector3f(BoxElem.X, BoxElem.Y, BoxElem.Z));
+
+			const FString KElemName = BoxElem.GetName().ToString();
+			PhysicsBody->SetBoxName(B, I, TCHAR_TO_ANSI(*KElemName));
+			PhysicsBody->SetBoxFlags(B, I, GetKBodyElemFlags(BoxElem));
+		}
+
+		const int32 NumConvex = BodySetup->AggGeom.ConvexElems.Num();
+		PhysicsBody->SetConvexCount(B, NumConvex);
+		for (int32 I = 0; I < NumConvex; ++I)
+		{
+			const FKConvexElem& ConvexElem = BodySetup->AggGeom.ConvexElems[I];
+
+			// Convert to FVector3f
+			TArray<FVector3f> VertexData;
+			VertexData.SetNumUninitialized(ConvexElem.VertexData.Num());
+			for (int32 Elem = VertexData.Num() - 1; Elem >= 0; --Elem)
+			{
+				VertexData[Elem] = FVector3f(ConvexElem.VertexData[Elem]);
+			}
+			
+			PhysicsBody->SetConvexMesh(B, I,
+					TArrayView<const FVector3f>(VertexData.GetData(), ConvexElem.VertexData.Num()),
+					TArrayView<const int32>(ConvexElem.IndexData.GetData(), ConvexElem.IndexData.Num()));
+
+			PhysicsBody->SetConvexTransform(B, I, FTransform3f(ConvexElem.GetTransform()));
+			
+			const FString KElemName = ConvexElem.GetName().ToString();
+			PhysicsBody->SetConvexName(B, I, TCHAR_TO_ANSI(*KElemName));
+			PhysicsBody->SetConvexFlags(B, I, GetKBodyElemFlags(ConvexElem));
+		}
+
+		const int32 NumSphyls = BodySetup->AggGeom.SphylElems.Num();
+		PhysicsBody->SetSphylCount( B, NumSphyls );
+
+		for (int32 I = 0; I < NumSphyls; ++I)
+		{
+			const FKSphylElem& SphylElem = BodySetup->AggGeom.SphylElems[I];
+			PhysicsBody->SetSphyl( B, I, 
+					FVector3f(SphylElem.Center), 
+					FQuat4f(SphylElem.Rotation.Quaternion()), 
+					SphylElem.Radius, SphylElem.Length );
+
+			const FString KElemName = SphylElem.GetName().ToString();
+			PhysicsBody->SetSphylName(B, I, TCHAR_TO_ANSI(*KElemName));
+			PhysicsBody->SetSphylFlags(B, I, GetKBodyElemFlags(SphylElem));
+		}
+
+		const int32 NumTaperedCapsules = BodySetup->AggGeom.TaperedCapsuleElems.Num();
+		PhysicsBody->SetTaperedCapsuleCount( B, NumTaperedCapsules );
+
+		for (int32 I = 0; I < NumTaperedCapsules; ++I)
+		{
+			const FKTaperedCapsuleElem& TaperedCapsuleElem = BodySetup->AggGeom.TaperedCapsuleElems[I];
+			PhysicsBody->SetTaperedCapsule(B, I, 
+					FVector3f(TaperedCapsuleElem.Center), 
+					FQuat4f(TaperedCapsuleElem.Rotation.Quaternion()), 
+					TaperedCapsuleElem.Radius0, TaperedCapsuleElem.Radius1, TaperedCapsuleElem.Length);
+			
+			const FString KElemName = TaperedCapsuleElem.GetName().ToString();
+			PhysicsBody->SetTaperedCapsuleName(B, I, TCHAR_TO_ANSI(*KElemName));
+			PhysicsBody->SetTaperedCapsuleFlags(B, I, GetKBodyElemFlags(TaperedCapsuleElem));
+		}
+	}
+
+	return PhysicsBody;
+}
+
+mu::MeshPtr ConvertSkeletalMeshToMutable(USkeletalMesh* InSkeletalMesh, const TSoftClassPtr<UAnimInstance>& AnimBp, int LOD, int MaterialIndex, FMutableGraphGenerationContext& GenerationContext, const UCustomizableObjectNode* CurrentNode)
 {
 	// Get the mesh generation flags to use
 	const EMutableMeshConversionFlags CurrentFlags = GenerationContext.MeshGenerationFlags.Last();
 	const bool bIgnoreSkeleton = EnumHasAnyFlags(CurrentFlags, EMutableMeshConversionFlags::IgnoreSkinning);
-	const bool bIgnorePhysics = EnumHasAnyFlags(CurrentFlags, EMutableMeshConversionFlags::IgnorePhysics) || 
-								!GenerationContext.Options.bPhysicsAssetMergeEnabled;
+	const bool bIgnorePhysics = EnumHasAnyFlags(CurrentFlags, EMutableMeshConversionFlags::IgnorePhysics);
 
 	mu::MeshPtr MutableMesh = new mu::Mesh();
 		
@@ -1339,7 +1524,8 @@ mu::MeshPtr ConvertSkeletalMeshToMutable(USkeletalMesh* InSkeletalMesh, int LOD,
 		}
 	}
 
-	if (!bIgnorePhysics && SkeletalMesh->GetPhysicsAsset() && MutableMesh->GetSkeleton())
+	const bool bPhysicsMergeEnabled = GenerationContext.Options.bPhysicsAssetMergeEnabled;
+	if (!bIgnorePhysics && SkeletalMesh->GetPhysicsAsset() && MutableMesh->GetSkeleton() && bPhysicsMergeEnabled)
 	{
 		// Find BodySetups with relevant bones.
 		TArray<TObjectPtr<USkeletalBodySetup>>& SkeletalBodySetups = SkeletalMesh->GetPhysicsAsset()->SkeletalBodySetups;
@@ -1555,6 +1741,36 @@ mu::MeshPtr ConvertSkeletalMeshToMutable(USkeletalMesh* InSkeletalMesh, int LOD,
 		}
 	}
 
+	const bool bAnimPhysicsManipulationEnabled = GenerationContext.Options.bAnimBpPhysicsManipulationEnabled;
+
+	if (!bIgnorePhysics && !AnimBp.IsNull() && MutableMesh->GetSkeleton() && bAnimPhysicsManipulationEnabled)
+	{
+		using AnimPhysicsInfoType = TTuple<UPhysicsAsset*, int32>;
+		const TArray<AnimPhysicsInfoType> AnimPhysicsInfo = GetPhysicsAssetsFromAnimInstance(AnimBp);
+
+		for (const AnimPhysicsInfoType PropertyInfo : AnimPhysicsInfo)
+		{
+			UPhysicsAsset* const PropertyAsset = PropertyInfo.Get<UPhysicsAsset*>();
+			const int32 PropertyIndex = PropertyInfo.Get<int32>();
+
+			FAnimBpOverridePhysicsAssetsInfo Info;
+			{
+				Info.AnimInstanceClass = AnimBp;
+				Info.PropertyIndex = PropertyIndex;
+				Info.SourceAsset = TSoftObjectPtr<UPhysicsAsset>(PropertyAsset);
+			}
+
+			const int32 PhysicsAssetId = GenerationContext.AnimBpOverridePhysicsAssetsInfo.AddUnique(Info);
+
+			mu::Ptr<mu::PhysicsBody> MutableBody = MakePhysicsBodyFromAsset(
+					PropertyAsset,
+					MakePhysicsAssetBodySetupRelevancyMap(PropertyAsset, MutableMesh));
+			MutableBody->CustomId = PhysicsAssetId;
+
+			MutableMesh->AddAdditionalPhysicsBody(MutableBody);
+		}
+	}
+
 	return MutableMesh;
 }
 
@@ -1741,14 +1957,13 @@ mu::MeshPtr ConvertStaticMeshToMutable(UStaticMesh* StaticMesh, int LOD, int Mat
 		}
 	}
 
-
 	return MutableMesh;
 }
 
 
 // Convert a Mesh constant to a mutable format. UniqueTags are the tags that make this Mesh unique that cannot be merged in the cache 
 //  with the exact same Mesh with other tags
-mu::MeshPtr GenerateMutableMesh(UObject * Mesh, int32 LOD, int32 MaterialIndex, const FString& UniqueTags, FMutableGraphGenerationContext & GenerationContext, const UCustomizableObjectNode* CurrentNode)
+mu::MeshPtr GenerateMutableMesh(UObject * Mesh, const TSoftClassPtr<UAnimInstance>& AnimInstance, int32 LOD, int32 MaterialIndex, const FString& UniqueTags, FMutableGraphGenerationContext & GenerationContext, const UCustomizableObjectNode* CurrentNode)
 {
 	// Get the mesh generation flags to use
 	EMutableMeshConversionFlags CurrentFlags = GenerationContext.MeshGenerationFlags.Last();
@@ -1760,7 +1975,7 @@ mu::MeshPtr GenerateMutableMesh(UObject * Mesh, int32 LOD, int32 MaterialIndex,
 	{
 		if (USkeletalMesh* SkeletalMesh = Cast<USkeletalMesh>(Mesh))
 		{
-			MutableMesh = ConvertSkeletalMeshToMutable(SkeletalMesh, LOD, MaterialIndex, GenerationContext, CurrentNode);
+			MutableMesh = ConvertSkeletalMeshToMutable(SkeletalMesh, AnimInstance, LOD, MaterialIndex, GenerationContext, CurrentNode);
 		}
 		else if (UStaticMesh* StaticMesh = Cast<UStaticMesh>(Mesh))
 		{
@@ -1858,7 +2073,7 @@ mu::MeshPtr BuildMorphedMutableMesh(const UEdGraphPin* BaseSourcePin, const FStr
 
 
 		// Get the base mesh
-		mu::MeshPtr BaseSourceMesh = GenerateMutableMesh(SkeletalMesh, LODIndex, SectionIndex, FString(), GenerationContext, Node);
+		mu::MeshPtr BaseSourceMesh = GenerateMutableMesh(SkeletalMesh, TSoftClassPtr<UAnimInstance>(), LODIndex, SectionIndex, FString(), GenerationContext, Node);
 		if (BaseSourceMesh)
 		{
 			// Clone it (it will probably be shared)
@@ -1947,22 +2162,21 @@ void GenerateMorphFactor(const UCustomizableObjectNode* Node, const UEdGraphPin&
 	}
 }
 
-
-TArray<USkeletalMesh*> GetSkeletalMeshesForReshapeSelection(
+TArray<TTuple<USkeletalMesh*, TSoftClassPtr<UAnimInstance>>> GetSkeletalMeshesInfoForReshapeSelection(
 		const UEdGraphNode* SkeletalMeshOrTableNode, const UEdGraphPin* SourceMeshPin)
 {
-	TArray<USkeletalMesh*> SkeletalMeshes;
+	TArray<TTuple<USkeletalMesh*, TSoftClassPtr<UAnimInstance>>> SkeletalMeshesInfo;
 
 	if (!(SkeletalMeshOrTableNode && SourceMeshPin))
 	{
-		return SkeletalMeshes;
+		return SkeletalMeshesInfo;
 	}
 
 	if (const UCustomizableObjectNodeSkeletalMesh* SkeletalMeshNode = Cast<UCustomizableObjectNodeSkeletalMesh>(SkeletalMeshOrTableNode))
 	{
 		if (SkeletalMeshNode->SkeletalMesh)
 		{
-			SkeletalMeshes.Add(SkeletalMeshNode->SkeletalMesh);
+			SkeletalMeshesInfo.Emplace(SkeletalMeshNode->SkeletalMesh, SkeletalMeshNode->AnimInstance);
 		}
 	}
 	else if (const UCustomizableObjectNodeTable* TableNode = Cast<UCustomizableObjectNodeTable>(SkeletalMeshOrTableNode))
@@ -1972,10 +2186,11 @@ TArray<USkeletalMesh*> GetSkeletalMeshesForReshapeSelection(
 			for (const FName& RowName : TableNode->GetRowNames())
 			{
 				USkeletalMesh* SkeletalMesh = TableNode->GetSkeletalMeshAt(SourceMeshPin, RowName);
+				TSoftClassPtr<UAnimInstance> MeshAnimInstance = TableNode->GetAnimInstanceAt(SourceMeshPin, RowName);
 
 				if (SkeletalMesh)
 				{
-					SkeletalMeshes.Add(SkeletalMesh);
+					SkeletalMeshesInfo.Emplace(SkeletalMesh, MeshAnimInstance);
 				}
 			}
 		}	
@@ -1985,18 +2200,20 @@ TArray<USkeletalMesh*> GetSkeletalMeshesForReshapeSelection(
 		checkf(false, TEXT("Node not expected."));
 	}
 
-	return SkeletalMeshes;
+	return SkeletalMeshesInfo;
 }
 
 
 bool GetAndValidateReshapeBonesToDeform(
 	TArray<FString>& OutBonesToDeform,
 	const TArray<FMeshReshapeBoneReference>& InBonesToDeform,
-	const TArray<USkeletalMesh*>& SkeletalMeshes,
+	const TArray<TTuple<USkeletalMesh*, TSoftClassPtr<UAnimInstance>>>& SkeletalMeshesInfo,
 	const UCustomizableObjectNode* Node,
 	const EBoneDeformSelectionMethod SelectionMethod,
 	FMutableGraphGenerationContext& GenerationContext)
 {
+	using MeshInfoType = TTuple<USkeletalMesh*, TSoftClassPtr<UAnimInstance>>;
+
 	bool bSetRefreshWarning = false;
 
 	TArray<uint8> MissingBones;
@@ -2011,8 +2228,10 @@ bool GetAndValidateReshapeBonesToDeform(
 
 			const FName BoneName = InBonesToDeform[InBoneIndex].BoneName;
 
-			for (const USkeletalMesh* SkeletalMesh : SkeletalMeshes)
+			for (const MeshInfoType& Mesh : SkeletalMeshesInfo)
 			{
+				const USkeletalMesh* SkeletalMesh = Mesh.Get<USkeletalMesh*>();
+
 				int32 BoneIndex = SkeletalMesh->GetRefSkeleton().FindBoneIndex(BoneName);
 				if (BoneIndex != INDEX_NONE)
 				{
@@ -2074,13 +2293,13 @@ bool GetAndValidateReshapeBonesToDeform(
 
 	else if (SelectionMethod == EBoneDeformSelectionMethod::ALL_BUT_SELECTED)
 	{
-		for (const USkeletalMesh* SkeletalMesh : SkeletalMeshes)
+		for (const MeshInfoType& Mesh : SkeletalMeshesInfo)
 		{
-			int32 NumBonesToDeform = SkeletalMesh->GetRefSkeleton().GetRawBoneNum();
+			int32 NumBonesToDeform = Mesh.Get<USkeletalMesh*>()->GetRefSkeleton().GetRawBoneNum();
 
 			for (int32 BoneIndex = 0; BoneIndex < NumBonesToDeform; ++BoneIndex)
 			{
-				FName BoneName = SkeletalMesh->GetRefSkeleton().GetBoneName(BoneIndex);
+				FName BoneName = Mesh.Get<USkeletalMesh*>()->GetRefSkeleton().GetBoneName(BoneIndex);
 				bool bFound = false;
 				int32 InNumBonesToDeform = InBonesToDeform.Num();
 
@@ -2093,7 +2312,7 @@ bool GetAndValidateReshapeBonesToDeform(
 					}
 				}
 
-				if (!bFound && SkeletalMesh->GetRefSkeleton().GetParentIndex(BoneIndex) != INDEX_NONE)
+				if (!bFound && Mesh.Get<USkeletalMesh*>()->GetRefSkeleton().GetParentIndex(BoneIndex) != INDEX_NONE)
 				{
 					OutBonesToDeform.AddUnique(BoneName.ToString());
 				}
@@ -2121,8 +2340,10 @@ bool GetAndValidateReshapeBonesToDeform(
 		// Getting reference skeleton from the reference skeletal mesh of the current component
 		const FReferenceSkeleton RefSkeleton = GenerationContext.ComponentInfos[GenerationContext.CurrentMeshComponent].RefSkeletalMesh->GetRefSkeleton();
 
-		for (const USkeletalMesh* SkeletalMesh : SkeletalMeshes)
+		for (const MeshInfoType& Mesh : SkeletalMeshesInfo)
 		{
+			const USkeletalMesh* SkeletalMesh = Mesh.Get<USkeletalMesh*>();
+
 			int32 NumBones = SkeletalMesh->GetRefSkeleton().GetRawBoneNum();
 
 			for (int32 BoneIndex = 0; BoneIndex < NumBones; ++BoneIndex)
@@ -2145,7 +2366,7 @@ bool GetAndValidateReshapeBonesToDeform(
 bool GetAndValidateReshapePhysicsToDeform(
 	TArray<FString>& OutPhysiscsToDeform,
 	const TArray<FMeshReshapeBoneReference>& InPhysicsToDeform,
-	const TArray<USkeletalMesh*>& SkeletalMeshes,
+	const TArray<TTuple<USkeletalMesh*, TSoftClassPtr<UAnimInstance>>>& SkeletalMeshesInfo,
 	EBoneDeformSelectionMethod SelectionMethod,
 	const UCustomizableObjectNode* Node,
 	FMutableGraphGenerationContext& GenerationContext)
@@ -2155,26 +2376,84 @@ bool GetAndValidateReshapePhysicsToDeform(
 		SelectionMethod == EBoneDeformSelectionMethod::DEFORM_REF_SKELETON ||
 		SelectionMethod == EBoneDeformSelectionMethod::DEFORM_NONE_REF_SKELETON;
 
+	// Find all used Bone names;
+	TSet<FName> UsedBoneNames;
+
+	using MeshInfoType = TTuple<USkeletalMesh*, TSoftClassPtr<UAnimInstance>>;
+	using PhysicsInfoType = TTuple<UPhysicsAsset*, const FReferenceSkeleton&>;
+
+	const TArray<PhysicsInfoType> ContributingPhysicsAssetsInfo = Invoke([&]() -> TArray<PhysicsInfoType>
+	{
+		TArray<PhysicsInfoType> PhysicsAssetsInfo;
+
+		const bool bAnimBpOverridePhysicsManipulationEnabled = GenerationContext.Options.bAnimBpPhysicsManipulationEnabled;
+		for (const MeshInfoType& Mesh : SkeletalMeshesInfo)
+		{
+			const USkeletalMesh* SkeletalMesh = Mesh.Get<USkeletalMesh*>();
+
+			if (!SkeletalMesh)
+			{
+				continue;
+			}
+			
+			{
+				UPhysicsAsset* PhysicsAsset = SkeletalMesh->GetPhysicsAsset();
+
+				if (PhysicsAsset)
+				{
+					PhysicsAssetsInfo.Emplace(PhysicsAsset, SkeletalMesh->GetRefSkeleton());
+				}
+			}
+
+			if (bAnimBpOverridePhysicsManipulationEnabled)
+			{
+				TSoftClassPtr<UAnimInstance> AnimInstance = Mesh.Get<TSoftClassPtr<UAnimInstance>>();
+
+				TArray<TTuple<UPhysicsAsset*, int32>> AnimInstanceOverridePhysicsAssets = GetPhysicsAssetsFromAnimInstance(AnimInstance);
+
+				for (const TTuple<UPhysicsAsset*, int32>& AnimPhysicsAssetInfo : AnimInstanceOverridePhysicsAssets)
+				{
+					int32 PropertyIndex = AnimPhysicsAssetInfo.Get<int32>();
+					UPhysicsAsset* AnimPhysicsAsset = AnimPhysicsAssetInfo.Get<UPhysicsAsset*>();
+
+					const bool bIsAnimPhysicsValid = PropertyIndex >= 0 && AnimPhysicsAsset;
+					if (bIsAnimPhysicsValid)
+					{
+						PhysicsAssetsInfo.Emplace(AnimPhysicsAsset, SkeletalMesh->GetRefSkeleton());
+					}
+				}
+			}
+		}
+
+		return PhysicsAssetsInfo;
+	});	
+
 	// Get the participant bone names.
-	TArray<FName> BoneNamesInUserSelection;
+	const TArray<FName> BoneNamesInUserSelection = Invoke([&]() -> TArray<FName>
+	{
+		TArray<FName> BoneNames;
+
 	if (bIsReferenceSkeletalMeshMethod)
 	{
 		const FReferenceSkeleton& RefSkeleton =
 			GenerationContext.ComponentInfos[GenerationContext.CurrentMeshComponent].RefSkeletalMesh->GetRefSkeleton();
 
 		const int32 RefSkeletonNumBones = RefSkeleton.GetRawBoneNum();
-		BoneNamesInUserSelection.SetNum(RefSkeletonNumBones);
+			BoneNames.SetNum(RefSkeletonNumBones);
 		for (int32 I = 0; I < RefSkeletonNumBones; ++I)
 		{
-			BoneNamesInUserSelection[I] = RefSkeleton.GetBoneName(I);
+				BoneNames[I] = RefSkeleton.GetBoneName(I);
 		}
 	}
 	else
 	{
-		BoneNamesInUserSelection.Reserve(InPhysicsToDeform.Num());
-		Algo::Transform(InPhysicsToDeform, BoneNamesInUserSelection, [](const FMeshReshapeBoneReference& B) { return B.BoneName; });
+			BoneNames.Reserve(InPhysicsToDeform.Num());
+			Algo::Transform(InPhysicsToDeform, BoneNames, [](const FMeshReshapeBoneReference& B) { return B.BoneName; });
 	}
 
+		return BoneNames;
+	});
+
 	int32 NumUserSelectedBones = BoneNamesInUserSelection.Num();
 
 	struct FMissingBoneStatus
@@ -2186,21 +2465,16 @@ bool GetAndValidateReshapePhysicsToDeform(
 	TArray<FMissingBoneStatus> MissingBones;
 	MissingBones.Init(FMissingBoneStatus{ false, true }, NumUserSelectedBones);
 
-	for (const USkeletalMesh* SkeletalMesh : SkeletalMeshes)
+	for (const PhysicsInfoType& PhysicsInfo : ContributingPhysicsAssetsInfo)
 	{
-		check(SkeletalMesh)
 		check(GenerationContext.ComponentInfos[GenerationContext.CurrentMeshComponent].RefSkeletalMesh);
 
-		UPhysicsAsset* PhysicsAsset = SkeletalMesh->GetPhysicsAsset();
-
-		if (!PhysicsAsset)
-		{
-			continue;
-		}
-
 		const FReferenceSkeleton& RefSkeleton = bIsReferenceSkeletalMeshMethod
 			? GenerationContext.ComponentInfos[GenerationContext.CurrentMeshComponent].RefSkeletalMesh->GetRefSkeleton()
-			: SkeletalMesh->GetRefSkeleton();
+			: PhysicsInfo.Get<const FReferenceSkeleton&>();
+
+		UPhysicsAsset* PhysicsAsset = PhysicsInfo.Get<UPhysicsAsset*>();
+		check(PhysicsAsset);
 
 		TArray<uint8> BoneInclusionSet;
 		BoneInclusionSet.Init(0, PhysicsAsset->SkeletalBodySetups.Num());
@@ -2451,7 +2725,8 @@ mu::NodeMeshPtr GenerateMorphMesh(const UEdGraphPin* Pin,
 					const UEdGraphPin* SourceMeshPin = ConnectedPin ? FindMeshBaseSource(*ConnectedPin, false) : nullptr;
 					const UEdGraphNode* SkeletalMeshNode = SourceMeshPin ? SourceMeshPin->GetOwningNode() : nullptr;
 
-					TArray<USkeletalMesh*> SkeletalMeshesToDeform = GetSkeletalMeshesForReshapeSelection(SkeletalMeshNode, SourceMeshPin);
+					TArray<TTuple<USkeletalMesh*, TSoftClassPtr<UAnimInstance>>> SkeletalMeshesToDeform = 
+							GetSkeletalMeshesInfoForReshapeSelection(SkeletalMeshNode, SourceMeshPin);
 
 					bool bWarningFound = false;
 					if (TypedMorphNode->bReshapeSkeleton)
@@ -2531,7 +2806,6 @@ void GenerateMorphTarget(const UCustomizableObjectNode* Node, const UEdGraphPin*
 	}
 }
 
-
 /** Convert a CustomizableObject Source Graph into a mutable source graph  */
 mu::NodeMeshPtr GenerateMutableSourceMesh(const UEdGraphPin * Pin,
 	FMutableGraphGenerationContext & GenerationContext,
@@ -2631,7 +2905,7 @@ mu::NodeMeshPtr GenerateMutableSourceMesh(const UEdGraphPin * Pin,
 				MeshUniqueTags += AnimBPTag;
 			}
 
-			mu::MeshPtr MutableMesh = GenerateMutableMesh(TypedNodeSkel->SkeletalMesh, LOD, SectionIndex, MeshUniqueTags, GenerationContext, TypedNodeSkel);
+			mu::MeshPtr MutableMesh = GenerateMutableMesh(TypedNodeSkel->SkeletalMesh, TypedNodeSkel->AnimInstance, LOD, SectionIndex, MeshUniqueTags, GenerationContext, TypedNodeSkel);
 			if (MutableMesh)
 			{
 				MeshNode->SetValue(MutableMesh);
@@ -2822,7 +3096,7 @@ mu::NodeMeshPtr GenerateMutableSourceMesh(const UEdGraphPin * Pin,
 			}
 			check(MaterialIndex < TypedNodeStatic->LODs[LOD].Materials.Num());
 
-			mu::MeshPtr MutableMesh = GenerateMutableMesh(TypedNodeStatic->StaticMesh, LOD, MaterialIndex, FString(), GenerationContext, TypedNodeStatic);
+			mu::MeshPtr MutableMesh = GenerateMutableMesh(TypedNodeStatic->StaticMesh, TSoftClassPtr<UAnimInstance>(), LOD, MaterialIndex, FString(), GenerationContext, TypedNodeStatic);
 			if (MutableMesh)
 			{
 				MeshNode->SetValue(MutableMesh);
@@ -3149,6 +3423,8 @@ mu::NodeMeshPtr GenerateMutableSourceMesh(const UEdGraphPin * Pin,
 		}
 	
 		{
+
+			MeshNode->SetReshapeVertices(TypedNodeReshape->bReshapeVertices);
 			MeshNode->SetReshapeSkeleton(TypedNodeReshape->bReshapeSkeleton);
 			MeshNode->SetReshapePhysicsVolumes(TypedNodeReshape->bReshapePhysicsVolumes);
 			MeshNode->SetEnableRigidParts(TypedNodeReshape->bEnableRigidParts);
@@ -3157,7 +3433,8 @@ mu::NodeMeshPtr GenerateMutableSourceMesh(const UEdGraphPin * Pin,
 			const UEdGraphPin* SourceMeshPin = ConnectedPin ? FindMeshBaseSource(*ConnectedPin, false) : nullptr;
 			const UEdGraphNode* SkeletalMeshNode = SourceMeshPin ? SourceMeshPin->GetOwningNode() : nullptr;
 
-			TArray<USkeletalMesh*> SkeletalMeshesToDeform = GetSkeletalMeshesForReshapeSelection(SkeletalMeshNode, SourceMeshPin);
+			TArray<TTuple<USkeletalMesh*, TSoftClassPtr<UAnimInstance>>> SkeletalMeshesToDeform = 
+					GetSkeletalMeshesInfoForReshapeSelection(SkeletalMeshNode, SourceMeshPin);
 
 			bool bWarningFound = false;
 			if (TypedNodeReshape->bReshapeSkeleton)
diff --git a/Engine/Plugins/Experimental/Mutable/Source/CustomizableObjectEditor/Private/MuCOE/GenerateMutableSource/GenerateMutableSourceMesh.h b/Engine/Plugins/Experimental/Mutable/Source/CustomizableObjectEditor/Private/MuCOE/GenerateMutableSource/GenerateMutableSourceMesh.h
--- a/Engine/Plugins/Experimental/Mutable/Source/CustomizableObjectEditor/Private/MuCOE/GenerateMutableSource/GenerateMutableSourceMesh.h
+++ b/Engine/Plugins/Experimental/Mutable/Source/CustomizableObjectEditor/Private/MuCOE/GenerateMutableSource/GenerateMutableSourceMesh.h
@@ -5,6 +5,8 @@
 #include "MuT/NodeModifierMeshClipWithMesh.h"
 #include "MuT/Table.h"
 
+#include "UObject/SoftObjectPtr.h"
+
 class FName;
 class FString;
 
@@ -16,13 +18,13 @@ class UStaticMesh;
 struct FMutableGraphGenerationContext;
 struct FMutableGraphMeshGenerationData;
 
-mu::MeshPtr ConvertSkeletalMeshToMutable(USkeletalMesh* InSkeletalMesh, int LOD, int MaterialIndex, FMutableGraphGenerationContext& GenerationContext, const UCustomizableObjectNode* CurrentNode);
+mu::MeshPtr ConvertSkeletalMeshToMutable(USkeletalMesh* InSkeletalMesh, const TSoftClassPtr<class UAnimInstance>& AnimBp, int LOD, int MaterialIndex, FMutableGraphGenerationContext& GenerationContext, const UCustomizableObjectNode* CurrentNode);
 
 
 mu::MeshPtr ConvertStaticMeshToMutable(UStaticMesh* StaticMesh, int LOD, int MaterialIndex, FMutableGraphGenerationContext& GenerationContext, const UCustomizableObjectNode* CurrentNode);
 
 
-mu::MeshPtr GenerateMutableMesh(UObject * Mesh, int32 LOD, int32 MaterialIndex, const FString& MeshUniqueTags, FMutableGraphGenerationContext & GenerationContext, const UCustomizableObjectNode* CurrentNode);
+mu::MeshPtr GenerateMutableMesh(UObject * Mesh, const TSoftClassPtr<class UAnimInstance>& AnimBp, int32 LOD, int32 MaterialIndex, const FString& MeshUniqueTags, FMutableGraphGenerationContext & GenerationContext, const UCustomizableObjectNode* CurrentNode);
 
 
 mu::MeshPtr BuildMorphedMutableMesh(const UEdGraphPin * BaseSourcePin, const FString& MorphTargetName, FMutableGraphGenerationContext & GenerationContext, const FName& RowName = "");
diff --git a/Engine/Plugins/Experimental/Mutable/Source/CustomizableObjectEditor/Private/MuCOE/GenerateMutableSource/GenerateMutableSourceTable.cpp b/Engine/Plugins/Experimental/Mutable/Source/CustomizableObjectEditor/Private/MuCOE/GenerateMutableSource/GenerateMutableSourceTable.cpp
--- a/Engine/Plugins/Experimental/Mutable/Source/CustomizableObjectEditor/Private/MuCOE/GenerateMutableSource/GenerateMutableSourceTable.cpp
+++ b/Engine/Plugins/Experimental/Mutable/Source/CustomizableObjectEditor/Private/MuCOE/GenerateMutableSource/GenerateMutableSourceTable.cpp
@@ -202,7 +202,8 @@ void FillTableColumn(const UCustomizableObjectNodeTable* TableNode,	mu::TablePtr
 					MeshUniqueTags += GenerateGameplayTag(Tag.ToString());
 				}
 
-				mu::MeshPtr MutableMesh = GenerateMutableMesh(SkeletalMesh, CurrentLOD, MatIndex, MeshUniqueTags, GenerationContext, TableNode);
+				//TODO: Add AnimBp physics to Tables.
+				mu::MeshPtr MutableMesh = GenerateMutableMesh(SkeletalMesh, TSoftClassPtr<UAnimInstance>(), CurrentLOD, MatIndex, MeshUniqueTags, GenerationContext, TableNode);
 
 				if (MutableMesh)
 				{
@@ -298,7 +299,7 @@ void FillTableColumn(const UCustomizableObjectNodeTable* TableNode,	mu::TablePtr
 					CurrentColumn = MutableTable->AddColumn(StringCast<ANSICHAR>(*MutableColumnName).Get(), mu::TABLE_COLUMN_TYPE::TCT_MESH);
 				}
 
-				mu::MeshPtr MutableMesh = GenerateMutableMesh(StaticMesh, CurrentLOD, MatIndex, FString(), GenerationContext, TableNode);
+				mu::MeshPtr MutableMesh = GenerateMutableMesh(StaticMesh, TSoftClassPtr<UAnimInstance>(), CurrentLOD, MatIndex, FString(), GenerationContext, TableNode);
 
 				if (MutableMesh)
 				{
diff --git a/Engine/Plugins/Experimental/Mutable/Source/CustomizableObjectEditor/Private/MuCOE/Nodes/CustomizableObjectNodeTable.cpp b/Engine/Plugins/Experimental/Mutable/Source/CustomizableObjectEditor/Private/MuCOE/Nodes/CustomizableObjectNodeTable.cpp
--- a/Engine/Plugins/Experimental/Mutable/Source/CustomizableObjectEditor/Private/MuCOE/Nodes/CustomizableObjectNodeTable.cpp
+++ b/Engine/Plugins/Experimental/Mutable/Source/CustomizableObjectEditor/Private/MuCOE/Nodes/CustomizableObjectNodeTable.cpp
@@ -894,6 +894,40 @@ USkeletalMesh* UCustomizableObjectNodeTable::GetSkeletalMeshAt(const UEdGraphPin
 	return nullptr;
 }
 
+TSoftClassPtr<UAnimInstance> UCustomizableObjectNodeTable::GetAnimInstanceAt(const UEdGraphPin* Pin, const FName& RowName) const
+{
+	if (!Table || !Table->GetRowStruct() || !Pin || !Table->GetRowNames().Contains(RowName))
+	{
+		return TSoftClassPtr<UAnimInstance>();
+	}
+
+	// Getting Struct Pointer
+	const UScriptStruct* TableStruct = Table->GetRowStruct();
+
+	FString ColumnName = GetColumnNameByPin(Pin);
+
+	for (TFieldIterator<FProperty> It(TableStruct); It; ++It)
+	{
+		FProperty* ColumnProperty = *It;
+
+		if (!ColumnProperty || ColumnName != DataTableUtils::GetPropertyExportName(ColumnProperty))
+		{
+			continue;
+		}
+
+		if (const FSoftClassProperty* SoftClassProperty = CastField<FSoftClassProperty>(ColumnProperty))
+		{
+			TSoftClassPtr<UAnimInstance> AnimInstance(SoftClassProperty->GetPropertyValue(SoftClassProperty).ToSoftObjectPath());
+
+			if (!AnimInstance.IsNull())
+			{
+				return AnimInstance;
+			}
+		}
+	}
+
+	return TSoftClassPtr<UAnimInstance>();
+}
 
 TArray<FName> UCustomizableObjectNodeTable::GetRowNames() const
 {
diff --git a/Engine/Plugins/Experimental/Mutable/Source/MutableRuntime/Private/MuR/CodeRunner.cpp b/Engine/Plugins/Experimental/Mutable/Source/MutableRuntime/Private/MuR/CodeRunner.cpp
--- a/Engine/Plugins/Experimental/Mutable/Source/MutableRuntime/Private/MuR/CodeRunner.cpp
+++ b/Engine/Plugins/Experimental/Mutable/Source/MutableRuntime/Private/MuR/CodeRunner.cpp
@@ -2013,17 +2013,17 @@ namespace mu
 
 		case OP_TYPE::ME_BINDSHAPE:
 		{
-			OP::MeshBindShapeArgs args = pModel->GetPrivate()->m_program.GetOpArgs<OP::MeshBindShapeArgs>(item.At);
-			const uint8* data = pModel->GetPrivate()->m_program.GetOpArgsPointer(item.At);
+			OP::MeshBindShapeArgs Args = pModel->GetPrivate()->m_program.GetOpArgs<OP::MeshBindShapeArgs>(item.At);
+			const uint8* Data = pModel->GetPrivate()->m_program.GetOpArgsPointer(item.At);
 
 			switch (item.Stage)
 			{
 			case 0:
-				if (args.mesh)
+				if (Args.mesh)
 				{
 					AddOp(FScheduledOp(item.At, item, 1),
-							FScheduledOp(args.mesh, item),
-							FScheduledOp(args.shape, item));
+							FScheduledOp(Args.mesh, item),
+							FScheduledOp(Args.shape, item));
 				}
 				else
 				{
@@ -2034,47 +2034,70 @@ namespace mu
 			case 1:
 			{
 				MUTABLE_CPUPROFILER_SCOPE(ME_BINDSHAPE_1)
+				Ptr<const Mesh> BaseMesh = GetMemory().GetMesh(FCacheAddress(Args.mesh, item));
+				Ptr<const Mesh> Shape = GetMemory().GetMesh(FCacheAddress(Args.shape, item));
 				
-				Ptr<const Mesh> BaseMesh = GetMemory().GetMesh(FCacheAddress(args.mesh, item));
-				Ptr<const Mesh> Shape = GetMemory().GetMesh(FCacheAddress(args.shape, item));
 				
-				EShapeBindingMethod BindingMethod = static_cast<EShapeBindingMethod>(args.bindingMethod); 
+				EShapeBindingMethod BindingMethod = static_cast<EShapeBindingMethod>(Args.bindingMethod); 
 
 				if (BindingMethod == EShapeBindingMethod::ReshapeClosestProject)
 				{ 
 					// Bones are stored after the args
-					data += sizeof(args);
+					Data += sizeof(Args);
 
 					// Rebuilding array of bone names ----
 					int32 NumBones;
-					FMemory::Memcpy(&NumBones, data, sizeof(int32)); 
-					data += sizeof(int32);
+					FMemory::Memcpy(&NumBones, Data, sizeof(int32)); 
+					Data += sizeof(int32);
 					
 					TArray<string> BonesToDeform;
 					BonesToDeform.SetNum(NumBones);
 					for (int32 b = 0; b < NumBones; ++b)
 					{
-						BonesToDeform[b] = pModel->GetPrivate()->m_program.m_constantStrings[*data].c_str();
-						data += sizeof(int32);
+						int32 ConstStringIndex;
+						FMemory::Memcpy(&ConstStringIndex, Data, sizeof(int32));
+						BonesToDeform[b] = pModel->GetPrivate()->m_program.m_constantStrings[ConstStringIndex].c_str();
+						Data += sizeof(int32);
 					}
 
 					int32 NumPhysicsBodies;
-					FMemory::Memcpy(&NumPhysicsBodies, data, sizeof(int32)); 
-					data += sizeof(int32);
+					FMemory::Memcpy(&NumPhysicsBodies, Data, sizeof(int32)); 
+					Data += sizeof(int32);
 
 					TArray<string> PhysicsToDeform;
 					PhysicsToDeform.SetNum(NumPhysicsBodies);
 					for (int32 b = 0; b < NumPhysicsBodies; ++b)
 					{
-						PhysicsToDeform[b] = pModel->GetPrivate()->m_program.m_constantStrings[*data].c_str();
-						data += sizeof(int32);
+						int32 ConstStringIndex;
+						FMemory::Memcpy(&ConstStringIndex, Data, sizeof(int32));
+						PhysicsToDeform[b] = pModel->GetPrivate()->m_program.m_constantStrings[ConstStringIndex].c_str();
+						Data += sizeof(int32);
 					}
 
+					const EMeshBindShapeFlags BindFlags = static_cast<EMeshBindShapeFlags>(Args.flags);
+
 					// -----------
-					Ptr<Mesh> pResult = MeshBindShapeReshape( BaseMesh.get(), Shape.get(), BonesToDeform, PhysicsToDeform, 
-						    static_cast<EMeshBindShapeFlags>(args.flags));
+					Ptr<Mesh> BindMesh = 
+						MeshBindShapeReshape( BaseMesh.get(), Shape.get(), BonesToDeform, PhysicsToDeform, BindFlags);
 					
-					GetMemory().SetMesh(item, pResult);
+					// null bind mesh indicates nothing has bond so the base mesh can be reused.
+					if (BindMesh)
+					{
+						if (!EnumHasAnyFlags(BindFlags, EMeshBindShapeFlags::ReshapeVertices))
+						{
+							Ptr<Mesh> BindMeshNoVerts = BaseMesh->Clone();
+							BindMeshNoVerts->m_AdditionalBuffers = MoveTemp(BindMesh->m_AdditionalBuffers);
+							GetMemory().SetMesh(item, BindMeshNoVerts);
+						}
+						else
+						{
+							GetMemory().SetMesh(item, BindMesh);
+						}
+					}
+					else
+					{
+						GetMemory().SetMesh(item, BaseMesh);
+					}	
 				}	
 				else
 				{
@@ -2118,10 +2141,37 @@ namespace mu
 				Ptr<const Mesh> BaseMesh = GetMemory().GetMesh(FCacheAddress(args.mesh, item));
 				Ptr<const Mesh> Shape = GetMemory().GetMesh(FCacheAddress(args.shape, item));
 
+				const EMeshBindShapeFlags ReshapeFlags = static_cast<EMeshBindShapeFlags>(args.flags);
+				Ptr<Mesh> ReshapedMesh = MeshApplyShape(BaseMesh.get(), Shape.get(), ReshapeFlags);
 
-				Ptr<Mesh> pResult = MeshApplyShape(BaseMesh.get(), Shape.get(), static_cast<EMeshBindShapeFlags>(args.flags));
+				const bool bReshapeVertices = EnumHasAnyFlags(ReshapeFlags, EMeshBindShapeFlags::ReshapeVertices);
+
+				if (!ReshapedMesh)
+				{
+					GetMemory().SetMesh(item, BaseMesh);
+				}
+				else if (!bReshapeVertices)
+				{
+					// Clone without Skeleton, Physics or Poses 
+					EMeshCloneFlags CloneFlags = ~(
+							EMeshCloneFlags::WithSkeleton          | 
+							EMeshCloneFlags::WithPhysicsBody       |
+							EMeshCloneFlags::WithAdditionalPhysics |
+							EMeshCloneFlags::WithPoses);
+ 
+					Ptr<Mesh> Result = BaseMesh->Clone(CloneFlags);
+					Result->SetSkeleton(ReshapedMesh->GetSkeleton().get());
+					Result->SetPhysicsBody(ReshapedMesh->GetPhysicsBody().get());
+					Result->AdditionalPhysicsBodies = ReshapedMesh->AdditionalPhysicsBodies;
+					Result->BonePoses = ReshapedMesh->BonePoses;
+
+					GetMemory().SetMesh(item, Result);
+				}
+				else
+				{
+					GetMemory().SetMesh(item, ReshapedMesh);
+				}
 
-				GetMemory().SetMesh(item, pResult);
 				break;
 			}
 
diff --git a/Engine/Plugins/Experimental/Mutable/Source/MutableRuntime/Private/MuR/Mesh.cpp b/Engine/Plugins/Experimental/Mutable/Source/MutableRuntime/Private/MuR/Mesh.cpp
--- a/Engine/Plugins/Experimental/Mutable/Source/MutableRuntime/Private/MuR/Mesh.cpp
+++ b/Engine/Plugins/Experimental/Mutable/Source/MutableRuntime/Private/MuR/Mesh.cpp
@@ -75,6 +75,7 @@ MeshPtr Mesh::Clone() const
 
 	// Clone SkeletonIDs
 	pResult->SkeletonIDs = SkeletonIDs;
+	pResult->AdditionalPhysicsBodies = AdditionalPhysicsBodies;
 
     return pResult;
 }
@@ -89,43 +90,43 @@ MeshPtr Mesh::Clone(EMeshCloneFlags Flags) const
     pResult->m_internalId = m_internalId;
 	pResult->m_staticFormatFlags = m_staticFormatFlags;
 
-	if (EnumHasAnyFlags(Flags, EMeshCloneFlags::WithSurfaces ))
+	if (EnumHasAnyFlags(Flags, EMeshCloneFlags::WithSurfaces))
 	{
 		pResult->m_surfaces = m_surfaces;
 	}
 
-	if (EnumHasAnyFlags(Flags, EMeshCloneFlags::WithSkeleton ))
+	if (EnumHasAnyFlags(Flags, EMeshCloneFlags::WithSkeleton))
 	{
 		pResult->m_pSkeleton = m_pSkeleton;
 	}
 
-	if (EnumHasAnyFlags(Flags, EMeshCloneFlags::WithPhysicsBody ))
+	if (EnumHasAnyFlags(Flags, EMeshCloneFlags::WithPhysicsBody))
 	{
 		pResult->m_pPhysicsBody = m_pPhysicsBody;
 	}
 
-	if (EnumHasAnyFlags(Flags, EMeshCloneFlags::WithFaceGroups ))
+	if (EnumHasAnyFlags(Flags, EMeshCloneFlags::WithFaceGroups))
 	{
 		pResult->m_faceGroups = m_faceGroups;
 	}
 
-	if (EnumHasAnyFlags(Flags, EMeshCloneFlags::WithTags ))
+	if (EnumHasAnyFlags(Flags, EMeshCloneFlags::WithTags))
 	{
 		pResult->m_tags = m_tags;
 	}
 
     // Clone the main buffers
-	if (EnumHasAnyFlags(Flags, EMeshCloneFlags::WithVertexBuffers ))
+	if (EnumHasAnyFlags(Flags, EMeshCloneFlags::WithVertexBuffers))
     {
 		pResult->m_VertexBuffers = m_VertexBuffers;
 	}
 
-	if (EnumHasAnyFlags(Flags, EMeshCloneFlags::WithIndexBuffers ))
+	if (EnumHasAnyFlags(Flags, EMeshCloneFlags::WithIndexBuffers))
     {
 		pResult->m_IndexBuffers = m_IndexBuffers;
 	}
 
-	if (EnumHasAnyFlags(Flags, EMeshCloneFlags::WithFaceBuffers ))
+	if (EnumHasAnyFlags(Flags, EMeshCloneFlags::WithFaceBuffers))
     {
 		pResult->m_FaceBuffers = m_FaceBuffers;
 	}
@@ -158,6 +159,11 @@ MeshPtr Mesh::Clone(EMeshCloneFlags Flags) const
 		pResult->SkeletonIDs = SkeletonIDs;
 	}
 
+	if (EnumHasAnyFlags(Flags, EMeshCloneFlags::WithAdditionalPhysics))
+	{
+		pResult->AdditionalPhysicsBodies = AdditionalPhysicsBodies;
+	}
+
     return pResult;
 }
 
@@ -216,6 +222,20 @@ void Mesh::SetPhysicsBody( Ptr<const PhysicsBody> PhysicsBody )
     m_pPhysicsBody = PhysicsBody;
 }
 
+//---------------------------------------------------------------------------------------------
+int32 Mesh::AddAdditionalPhysicsBody(Ptr<const PhysicsBody> Body)
+{
+	return AdditionalPhysicsBodies.Add(Body);
+}
+
+Ptr<const PhysicsBody> Mesh::GetAdditionalPhysicsBody(int32 I) const
+{
+	check(I >= 0 && I < AdditionalPhysicsBodies.Num());
+
+	return AdditionalPhysicsBodies[I];
+}
+
+
 //---------------------------------------------------------------------------------------------
 int Mesh::GetFaceCount() const
 {
diff --git a/Engine/Plugins/Experimental/Mutable/Source/MutableRuntime/Private/MuR/OpMeshBind.h b/Engine/Plugins/Experimental/Mutable/Source/MutableRuntime/Private/MuR/OpMeshBind.h
--- a/Engine/Plugins/Experimental/Mutable/Source/MutableRuntime/Private/MuR/OpMeshBind.h
+++ b/Engine/Plugins/Experimental/Mutable/Source/MutableRuntime/Private/MuR/OpMeshBind.h
@@ -135,6 +135,7 @@ namespace mu
 
 	struct FReshapeVertexBindingDataBufferDescriptor
 	{
+		constexpr static int ElementSize = sizeof(FReshapeVertexBindingData);
 		constexpr static int Channels = 4;
 		constexpr static MESH_BUFFER_SEMANTIC Semantics[Channels] = { MBS_BARYCENTRICCOORDS, MBS_DISTANCE, MBS_TRIANGLEINDEX, MBS_OTHER };
 		constexpr static MESH_BUFFER_FORMAT Formats[Channels] = { MBF_FLOAT32, MBF_FLOAT32, MBF_INT32, MBF_FLOAT32 };
@@ -161,6 +162,7 @@ namespace mu
 
 	struct FReshapePointBindingDataBufferDescriptor
 	{
+		constexpr static int ElementSize = sizeof(FReshapeVertexBindingData);
 		constexpr static int Channels = 4;
 		constexpr static MESH_BUFFER_SEMANTIC Semantics[Channels] = { MBS_BARYCENTRICCOORDS, MBS_DISTANCE, MBS_TRIANGLEINDEX, MBS_OTHER };
 		constexpr static MESH_BUFFER_FORMAT Formats[Channels] = { MBF_FLOAT32, MBF_FLOAT32, MBF_INT32, MBF_FLOAT32 };
@@ -175,8 +177,9 @@ namespace mu
 		}
 	};
 	
-	struct FUsedIndicesBufferDescriptor
+	struct FIntBufferDescriptor
 	{
+		constexpr static int ElementSize = sizeof(int32);
 		constexpr static int Channels = 1;
 		constexpr static MESH_BUFFER_SEMANTIC Semantics[Channels] = { MBS_OTHER };
 		constexpr static MESH_BUFFER_FORMAT Formats[Channels] = { MBF_INT32 };
@@ -184,7 +187,7 @@ namespace mu
 		constexpr static int Offsets[Channels] = { 0 };
 		constexpr static int SemanticIndices[Channels] = { 0 };
 
-		FUsedIndicesBufferDescriptor()
+		FIntBufferDescriptor()
 		{
 		}
 	};
@@ -486,59 +489,95 @@ namespace mu
 		BindReshapePoint(ShapeMeshTree, BoundPoint, OutBindingData, BindTolerance);
 	}
 
-	inline TTuple<TArray<FReshapePointBindingData>, TArray<int32>> BindPhysicsBodies( 
-			const PhysicsBody* PBody, FShapeMeshTree& ShapeMeshTree, const Mesh* pMesh, 
+	inline TTuple<TArray<FReshapePointBindingData>, TArray<int32>, TArray<int32>> BindPhysicsBodies( 
+			TArray<const PhysicsBody*> PhysicsBodies, FShapeMeshTree& ShapeMeshTree, const Mesh* pMesh, 
 			const TArray<string>& PhysicsToDeform )
 	{
-		TTuple<TArray<FReshapePointBindingData>, TArray<int32>> ReturnValue;
+		TTuple<TArray<FReshapePointBindingData>, TArray<int32>, TArray<int32>> ReturnValue;
 
 		TArray<int32>& BodiesToDeformIndices = ReturnValue.Get<1>();
+		TArray<int32>& BodiesToDeformOffsets = ReturnValue.Get<2>();
+
+		const int32 NumPhysicsBodies = PhysicsBodies.Num();
+
+		BodiesToDeformIndices.Reserve(NumPhysicsBodies * PhysicsToDeform.Num());
+		BodiesToDeformOffsets.SetNum(NumPhysicsBodies + 1);
 	
-		BodiesToDeformIndices.Reserve(PhysicsToDeform.Num());
-		const int32 NumBodies = PBody->GetBodyCount();
-		for (int32 I = 0; I < NumBodies; ++I)
+		BodiesToDeformOffsets[0] = 0;
+		for (int32 PhysicsBodyIndex = 0; PhysicsBodyIndex < NumPhysicsBodies; ++PhysicsBodyIndex)
 		{
-			if (PhysicsToDeform.Contains(PBody->GetBodyBoneName(I)))
+			if (const PhysicsBody* Body = PhysicsBodies[PhysicsBodyIndex])
 			{
-				BodiesToDeformIndices.Add(I);
+				const int32 NumBodies = PhysicsBodies[PhysicsBodyIndex]->GetBodyCount();
+				for (int32 BodyIndex = 0; BodyIndex < NumBodies; ++BodyIndex)
+				{
+					if (PhysicsToDeform.Contains(Body->GetBodyBoneName(BodyIndex)))
+					{
+						BodiesToDeformIndices.Add(BodyIndex);
+					}
+				}
 			}
+
+			BodiesToDeformOffsets[PhysicsBodyIndex + 1] = BodiesToDeformIndices.Num();
 		}
+
+		int32 TotalNumPoints = 0;
+		for (int32 PhysicsBodyIndex = 0; PhysicsBodyIndex < NumPhysicsBodies; ++PhysicsBodyIndex)
+		{
+			const int32 IndicesBegin = BodiesToDeformOffsets[PhysicsBodyIndex];
+			const int32 IndicesEnd = BodiesToDeformOffsets[PhysicsBodyIndex + 1];
+
+			TArrayView<int32> BodyIndices = TArrayView<int32>(
+				BodiesToDeformIndices.GetData() + IndicesBegin, IndicesEnd - IndicesBegin);
+
+			const PhysicsBody* Body = PhysicsBodies[PhysicsBodyIndex];
 		
-		// Count how many poinst will be needed.
-		const int32 NumBodiesToDeformIndices = BodiesToDeformIndices.Num();
-		int32 NumPoints = 0;
-		for (const int32 I : BodiesToDeformIndices)
+			int32 BodyNumPoints = 0;
+			for (const int32 I : BodyIndices)
 		{
-			NumPoints += PBody->GetSphereCount(I) * 6;
-			NumPoints += PBody->GetBoxCount(I) * 14;
-			NumPoints += PBody->GetSphylCount(I) * 14;
-			NumPoints += PBody->GetTaperedCapsuleCount(I) * 14;
+				BodyNumPoints += Body->GetSphereCount(I) * 6;
+				BodyNumPoints += Body->GetBoxCount(I) * 14;
+				BodyNumPoints += Body->GetSphylCount(I) * 14;
+				BodyNumPoints += Body->GetTaperedCapsuleCount(I) * 14;
 
-			const int32 ConvexCount = PBody->GetConvexCount(I);
-			for ( int C = 0; C < ConvexCount; ++C )
+				const int32 ConvexCount = Body->GetConvexCount(I);
+				for (int C = 0; C < ConvexCount; ++C)
 			{
 				TArrayView<const FVector3f> Vertices;
 				TArrayView<const int32> Indices;
 
 				FTransform3f Transform;
-				PBody->GetConvex(I, C, Vertices, Indices, Transform);
+					Body->GetConvex(I, C, Vertices, Indices, Transform);
 
-				NumPoints += Vertices.Num();
+					BodyNumPoints += Vertices.Num();
 			}
 		}
 		
+			TotalNumPoints += BodyNumPoints;
+		}
+		
 		TArray<FVector3f> Points;
-		Points.SetNumUninitialized(NumPoints);
+		Points.SetNumUninitialized(TotalNumPoints);
 	
 		// Bone transform needs to be applied to the body's sample points so they are in mesh space.
 
 		// Create a point soup to be deformed based on the shapes in the aggregate.
-		// Currently using oriented bounding boxes corners for all shapes except convex which uses the deformed mesh directly.
 
 		int32 AddedPoints = 0;
-		for (const int32 B : BodiesToDeformIndices)
+
+		for (int32 PhysicsBodyIndex = 0; PhysicsBodyIndex < NumPhysicsBodies; ++PhysicsBodyIndex)
 		{
-			const char* BoneName = PBody->GetBodyBoneName(B);
+			const PhysicsBody* Body = PhysicsBodies[PhysicsBodyIndex];
+
+			const int32 IndicesBegin = BodiesToDeformOffsets[PhysicsBodyIndex];
+			const int32 IndicesEnd = BodiesToDeformOffsets[PhysicsBodyIndex + 1];
+
+			TArrayView<int32> BodyIndices = TArrayView<int32>(
+					BodiesToDeformIndices.GetData() + IndicesBegin, IndicesEnd - IndicesBegin);
+
+			for (const int32 B : BodyIndices)
+		{
+				const char* BoneName = Body->GetBodyBoneName(B);
 			int32 BoneIdx = pMesh->FindBonePose(BoneName);
 
 			FTransform3f T = FTransform3f::Identity;
@@ -547,13 +586,13 @@ namespace mu
 				pMesh->GetBoneTransform(BoneIdx, T);
 			}
 			
-			const int32 SphereCount = PBody->GetSphereCount(B);
-			for ( int32 I = 0; I < SphereCount; ++I, AddedPoints += 6)
+				const int32 SphereCount = Body->GetSphereCount(B);
+				for (int32 I = 0; I < SphereCount; ++I, AddedPoints += 6)
 			{
 				FVector3f P;
 				float R;
 
-				PBody->GetSphere(B, I, P, R);	
+					Body->GetSphere(B, I, P, R);
 
 				Points[AddedPoints + 0] = T.TransformPosition(P + FVector3f(R, 0.0f, 0.0f));
 				Points[AddedPoints + 1] = T.TransformPosition(P - FVector3f(R, 0.0f, 0.0f));
@@ -565,47 +604,47 @@ namespace mu
 				Points[AddedPoints + 5] = T.TransformPosition(P - FVector3f(0.0f, 0.0f, R));
 			}
 
-			const int32 BoxCount = PBody->GetBoxCount(B);
-			for ( int32 I = 0; I < BoxCount; ++I, AddedPoints += 14)
+				const int32 BoxCount = Body->GetBoxCount(B);
+				for (int32 I = 0; I < BoxCount; ++I, AddedPoints += 14)
 			{
 				FVector3f P;
 				FQuat4f Q;
 				FVector3f S;
 			
-				PBody->GetBox(B, I, P, Q, S);
+					Body->GetBox(B, I, P, Q, S);
 
 				const FVector3f BasisX = Q.RotateVector(FVector3f::UnitX());
 				const FVector3f BasisY = Q.RotateVector(FVector3f::UnitY());
 				const FVector3f BasisZ = Q.RotateVector(FVector3f::UnitZ());
 
-				Points[AddedPoints + 0] = T.TransformPosition(P + BasisX*S.X + BasisY*S.Y + BasisZ*S.Z);
-				Points[AddedPoints + 1] = T.TransformPosition(P + BasisX*S.X - BasisY*S.Y + BasisZ*S.Z);
-				Points[AddedPoints + 2] = T.TransformPosition(P - BasisX*S.X + BasisY*S.Y + BasisZ*S.Z);
-				Points[AddedPoints + 3] = T.TransformPosition(P - BasisX*S.X - BasisY*S.Y + BasisZ*S.Z);
+					Points[AddedPoints + 0] = T.TransformPosition(P + BasisX * S.X + BasisY * S.Y + BasisZ * S.Z);
+					Points[AddedPoints + 1] = T.TransformPosition(P + BasisX * S.X - BasisY * S.Y + BasisZ * S.Z);
+					Points[AddedPoints + 2] = T.TransformPosition(P - BasisX * S.X + BasisY * S.Y + BasisZ * S.Z);
+					Points[AddedPoints + 3] = T.TransformPosition(P - BasisX * S.X - BasisY * S.Y + BasisZ * S.Z);
 
-				Points[AddedPoints + 4] = T.TransformPosition(P + BasisX*S.X + BasisY*S.Y - BasisZ*S.Z);
-				Points[AddedPoints + 5] = T.TransformPosition(P + BasisX*S.X - BasisY*S.Y - BasisZ*S.Z);
-				Points[AddedPoints + 6] = T.TransformPosition(P - BasisX*S.X + BasisY*S.Y - BasisZ*S.Z);
-				Points[AddedPoints + 7] = T.TransformPosition(P - BasisX*S.X - BasisY*S.Y - BasisZ*S.Z);
+					Points[AddedPoints + 4] = T.TransformPosition(P + BasisX * S.X + BasisY * S.Y - BasisZ * S.Z);
+					Points[AddedPoints + 5] = T.TransformPosition(P + BasisX * S.X - BasisY * S.Y - BasisZ * S.Z);
+					Points[AddedPoints + 6] = T.TransformPosition(P - BasisX * S.X + BasisY * S.Y - BasisZ * S.Z);
+					Points[AddedPoints + 7] = T.TransformPosition(P - BasisX * S.X - BasisY * S.Y - BasisZ * S.Z);
 
-				Points[AddedPoints + 8]  = T.TransformPosition(P + BasisX*S.X);
-				Points[AddedPoints + 9]  = T.TransformPosition(P + BasisY*S.Y);
-				Points[AddedPoints + 10] = T.TransformPosition(P + BasisZ*S.Z);
+					Points[AddedPoints + 8] = T.TransformPosition(P + BasisX * S.X);
+					Points[AddedPoints + 9] = T.TransformPosition(P + BasisY * S.Y);
+					Points[AddedPoints + 10] = T.TransformPosition(P + BasisZ * S.Z);
 
-				Points[AddedPoints + 11] = T.TransformPosition(P - BasisX*S.X);
-				Points[AddedPoints + 12] = T.TransformPosition(P - BasisY*S.Y);
-				Points[AddedPoints + 13] = T.TransformPosition(P - BasisZ*S.Z);
+					Points[AddedPoints + 11] = T.TransformPosition(P - BasisX * S.X);
+					Points[AddedPoints + 12] = T.TransformPosition(P - BasisY * S.Y);
+					Points[AddedPoints + 13] = T.TransformPosition(P - BasisZ * S.Z);
 			}
 
-			const int32 SphylCount = PBody->GetSphylCount(B);
-			for ( int32 I = 0; I < SphylCount; ++I, AddedPoints += 14)
+				const int32 SphylCount = Body->GetSphylCount(B);
+				for (int32 I = 0; I < SphylCount; ++I, AddedPoints += 14)
 			{
 				FVector3f P;
 				FQuat4f Q;
 				float R;
 				float L;
 			
-				PBody->GetSphyl(B, I, P, Q, R, L);
+					Body->GetSphyl(B, I, P, Q, R, L);
 				
 				const float H = L * 0.5f;
 
@@ -614,31 +653,31 @@ namespace mu
 				const FVector3f BasisZ = Q.RotateVector(FVector3f::UnitZ());
 
 				// Top and Bottom
-				Points[AddedPoints + 0] = T.TransformPosition(P + BasisZ*(H+R));
-				Points[AddedPoints + 1] = T.TransformPosition(P - BasisZ*(H+R));
+					Points[AddedPoints + 0] = T.TransformPosition(P + BasisZ * (H + R));
+					Points[AddedPoints + 1] = T.TransformPosition(P - BasisZ * (H + R));
 
 				// Top ring
-				Points[AddedPoints + 2] = T.TransformPosition(P + BasisX*R + BasisZ*H);
-				Points[AddedPoints + 3] = T.TransformPosition(P - BasisX*R + BasisZ*H);
-				Points[AddedPoints + 4] = T.TransformPosition(P + BasisY*R + BasisZ*H);
-				Points[AddedPoints + 5] = T.TransformPosition(P - BasisY*R + BasisZ*H);
+					Points[AddedPoints + 2] = T.TransformPosition(P + BasisX * R + BasisZ * H);
+					Points[AddedPoints + 3] = T.TransformPosition(P - BasisX * R + BasisZ * H);
+					Points[AddedPoints + 4] = T.TransformPosition(P + BasisY * R + BasisZ * H);
+					Points[AddedPoints + 5] = T.TransformPosition(P - BasisY * R + BasisZ * H);
 
 				// Center ring
-				Points[AddedPoints + 6] = T.TransformPosition(P + BasisX*R);
-				Points[AddedPoints + 7] = T.TransformPosition(P - BasisX*R);
-				Points[AddedPoints + 8] = T.TransformPosition(P + BasisY*R);
-				Points[AddedPoints + 9] = T.TransformPosition(P - BasisY*R);
+					Points[AddedPoints + 6] = T.TransformPosition(P + BasisX * R);
+					Points[AddedPoints + 7] = T.TransformPosition(P - BasisX * R);
+					Points[AddedPoints + 8] = T.TransformPosition(P + BasisY * R);
+					Points[AddedPoints + 9] = T.TransformPosition(P - BasisY * R);
 	
 				// Bottom ring
-				Points[AddedPoints + 10] = T.TransformPosition(P + BasisX*R - BasisZ*H);
-				Points[AddedPoints + 11] = T.TransformPosition(P - BasisX*R - BasisZ*H);
-				Points[AddedPoints + 12] = T.TransformPosition(P + BasisY*R - BasisZ*H);
-				Points[AddedPoints + 13] = T.TransformPosition(P - BasisY*R - BasisZ*H);
+					Points[AddedPoints + 10] = T.TransformPosition(P + BasisX * R - BasisZ * H);
+					Points[AddedPoints + 11] = T.TransformPosition(P - BasisX * R - BasisZ * H);
+					Points[AddedPoints + 12] = T.TransformPosition(P + BasisY * R - BasisZ * H);
+					Points[AddedPoints + 13] = T.TransformPosition(P - BasisY * R - BasisZ * H);
 
 			}
 	
-			const int32 TaperedCapsuleCount = PBody->GetTaperedCapsuleCount(B);
-			for ( int32 I = 0; I < TaperedCapsuleCount; ++I, AddedPoints += 14)
+				const int32 TaperedCapsuleCount = Body->GetTaperedCapsuleCount(B);
+				for (int32 I = 0; I < TaperedCapsuleCount; ++I, AddedPoints += 14)
 			{
 				FVector3f P;
 				FQuat4f Q;
@@ -646,9 +685,9 @@ namespace mu
 				float R1;
 				float L;
 			
-				PBody->GetTaperedCapsule(B, I, P, Q, R0, R1, L);
+					Body->GetTaperedCapsule(B, I, P, Q, R0, R1, L);
 		
-				const float H = L*0.5f;
+					const float H = L * 0.5f;
 				const float RCenter = (R0 + R1) * 0.5f;
 
 				const FVector3f BasisX = Q.RotateVector(FVector3f::UnitX());
@@ -656,62 +695,62 @@ namespace mu
 				const FVector3f BasisZ = Q.RotateVector(FVector3f::UnitZ());
 
 				// Top and Bottom
-				Points[AddedPoints + 0] = T.TransformPosition(P + BasisZ*(H+R0));
-				Points[AddedPoints + 1] = T.TransformPosition(P - BasisZ*(H+R1));
+					Points[AddedPoints + 0] = T.TransformPosition(P + BasisZ * (H + R0));
+					Points[AddedPoints + 1] = T.TransformPosition(P - BasisZ * (H + R1));
 
 				// Top ring
-				Points[AddedPoints + 2] = T.TransformPosition(P + BasisX*R0 + BasisZ*H);
-				Points[AddedPoints + 3] = T.TransformPosition(P - BasisX*R0 + BasisZ*H);
-				Points[AddedPoints + 4] = T.TransformPosition(P + BasisY*R0 + BasisZ*H);
-				Points[AddedPoints + 5] = T.TransformPosition(P - BasisY*R0 + BasisZ*H);
+					Points[AddedPoints + 2] = T.TransformPosition(P + BasisX * R0 + BasisZ * H);
+					Points[AddedPoints + 3] = T.TransformPosition(P - BasisX * R0 + BasisZ * H);
+					Points[AddedPoints + 4] = T.TransformPosition(P + BasisY * R0 + BasisZ * H);
+					Points[AddedPoints + 5] = T.TransformPosition(P - BasisY * R0 + BasisZ * H);
 
 				// Center ring
-				Points[AddedPoints + 6] = T.TransformPosition(P + BasisX*RCenter);
-				Points[AddedPoints + 7] = T.TransformPosition(P - BasisX*RCenter);
-				Points[AddedPoints + 8] = T.TransformPosition(P + BasisY*RCenter);
-				Points[AddedPoints + 9] = T.TransformPosition(P - BasisY*RCenter);
+					Points[AddedPoints + 6] = T.TransformPosition(P + BasisX * RCenter);
+					Points[AddedPoints + 7] = T.TransformPosition(P - BasisX * RCenter);
+					Points[AddedPoints + 8] = T.TransformPosition(P + BasisY * RCenter);
+					Points[AddedPoints + 9] = T.TransformPosition(P - BasisY * RCenter);
 	
 				// Bottom ring
-				Points[AddedPoints + 10] = T.TransformPosition(P + BasisX*R1 - BasisZ*H);
-				Points[AddedPoints + 11] = T.TransformPosition(P - BasisX*R1 - BasisZ*H);
-				Points[AddedPoints + 12] = T.TransformPosition(P + BasisY*R1 - BasisZ*H);
-				Points[AddedPoints + 13] = T.TransformPosition(P - BasisY*R1 - BasisZ*H);
+					Points[AddedPoints + 10] = T.TransformPosition(P + BasisX * R1 - BasisZ * H);
+					Points[AddedPoints + 11] = T.TransformPosition(P - BasisX * R1 - BasisZ * H);
+					Points[AddedPoints + 12] = T.TransformPosition(P + BasisY * R1 - BasisZ * H);
+					Points[AddedPoints + 13] = T.TransformPosition(P - BasisY * R1 - BasisZ * H);
 			}
 
-			const int32 ConvexCount = PBody->GetConvexCount(B);
-			for ( int32 I = 0; I < ConvexCount; ++I )
+				const int32 ConvexCount = Body->GetConvexCount(B);
+				for (int32 I = 0; I < ConvexCount; ++I)
 			{
 			
 				TArrayView<const FVector3f> VerticesView;
 				TArrayView<const int32> IndicesView;
 				FTransform3f ConvexT;
 
-				PBody->GetConvex(B, I, VerticesView, IndicesView, ConvexT);
+					Body->GetConvex(B, I, VerticesView, IndicesView, ConvexT);
 
 				ConvexT = T * ConvexT;
-				for ( const FVector3f& P : VerticesView )
+					for (const FVector3f& P : VerticesView)
 				{
 					Points[AddedPoints++] = ConvexT.TransformPosition(P);
 				}	
 			}
 		}
+		}
 		
 		TArray<FReshapePointBindingData>& PhysicsBodyBindData = ReturnValue.Get<0>();
-		PhysicsBodyBindData.SetNumUninitialized( NumPoints );
+		PhysicsBodyBindData.SetNumUninitialized(TotalNumPoints);
 
 		FReshapeVertexBindingData VertexBindData;
 
-		for (int32 I = 0; I < NumPoints; ++I)
+		for (int32 PointIndex = 0; PointIndex < TotalNumPoints; ++PointIndex)
 		{
-			BindReshapePoint(ShapeMeshTree, Points[I], VertexBindData, 0.1f);
-			PhysicsBodyBindData[I] = FReshapePointBindingData
+			BindReshapePoint(ShapeMeshTree, Points[PointIndex], VertexBindData, 0.1f);
+			PhysicsBodyBindData[PointIndex] = FReshapePointBindingData
 				{ VertexBindData.S, VertexBindData.T, VertexBindData.D, VertexBindData.Triangle, VertexBindData.Weight };
 		}
 
 		return ReturnValue;
 	}
 
-	
 	inline TArray<FReshapeVertexBindingData> BindVerticesReshape(
 			const Mesh* BaseMesh, FShapeMeshTree& ShapeMeshTree, bool bEnableRigidParts )	
 	{
@@ -862,7 +901,7 @@ namespace mu
 		// Early out if nothing will be modified and the vertices discarted. return null in this
 		// case indicating nothing has modified so the Base Mesh can be reused.
 		const bool bSkeletonModification = BaseMesh->GetSkeleton() && bReshapeSkeleton;
-		const bool bPhysicsModification = BaseMesh->GetPhysicsBody() && bReshapePhysics;
+		const bool bPhysicsModification = (BaseMesh->GetPhysicsBody() || BaseMesh->AdditionalPhysicsBodies.Num()) && bReshapePhysics;
 
 		if (!bReshapeVertices && !bSkeletonModification && !bPhysicsModification)
 		{
@@ -935,14 +974,15 @@ namespace mu
 	
 		Ptr<Mesh> Result;
 		
-		// If no vertices are needed, it is assumed we only want to reshape physics or skeleton
+		// If no vertices are needed, it is assumed we only want to reshape physics or skeleton.
 		// In that case, remove everything except physics bodies, the skeleton and pose.
 		if (!bReshapeVertices)
 		{
 			constexpr EMeshCloneFlags CloneFlags = 
 					EMeshCloneFlags::WithSkeleton    | 
 					EMeshCloneFlags::WithPhysicsBody | 
-					EMeshCloneFlags::WithPoses;
+				EMeshCloneFlags::WithPoses |
+				EMeshCloneFlags::WithAdditionalPhysics;
 
 			Result = BaseMesh->Clone(CloneFlags);
 		}
@@ -954,14 +994,14 @@ namespace mu
 		int32 BindingDataIndex = 0;
 		if (bReshapeVertices)
 		{
-			TArray<FReshapeVertexBindingData> VerticesBindData = BindVerticesReshape( BaseMesh, ShapeMeshTree, bEnableRigidParts );
+			TArray<FReshapeVertexBindingData> VerticesBindData = BindVerticesReshape(BaseMesh, ShapeMeshTree, bEnableRigidParts);
 			
 			// Add the binding information to the mesh
 			// \TODO: Check that there is no other binding data.
 			// \TODO: Support specifying the binding data channel for multiple binding support.
 			FMeshBufferSet& VB = Result->GetVertexBuffers();
 			int NewBufferIndex = VB.GetBufferCount();
-			VB.SetBufferCount(NewBufferIndex+1);
+			VB.SetBufferCount(NewBufferIndex + 1);
 
 			FReshapeVertexBindingDataBufferDescriptor BufDesc(BindingDataIndex);
 			VB.SetBuffer(NewBufferIndex, sizeof(FReshapeVertexBindingData), BufDesc.Channels, BufDesc.Semantics, BufDesc.SemanticIndices, BufDesc.Formats, BufDesc.Components);
@@ -1006,14 +1046,29 @@ namespace mu
 		}
 
 		const PhysicsBody* ResultPhysicsBody = Result->m_pPhysicsBody.get();
-		if (bReshapePhysics && ResultPhysicsBody && PhysicsToDeform.Num())
+		if (bReshapePhysics && (ResultPhysicsBody || Result->AdditionalPhysicsBodies.Num()) && PhysicsToDeform.Num())
 		{
 			MUTABLE_CPUPROFILER_SCOPE(BindPhysicsBody);
 
-			TTuple<TArray<FReshapePointBindingData>, TArray<int32>> PhysicsBindingData = 
-					BindPhysicsBodies(ResultPhysicsBody, ShapeMeshTree, Result.get(), PhysicsToDeform);
+			// Gather bodies respecting order, firts main physics body then additional bodies.
+			// Null entries are need to be able to mantin that order.  
+			TArray<const PhysicsBody*> PhysicsBodiesToBind;
+			{
+				PhysicsBodiesToBind.Reserve(Result->AdditionalPhysicsBodies.Num() + 1);
+
+				PhysicsBodiesToBind.Add(ResultPhysicsBody);
+				for (const Ptr<const PhysicsBody>& Body : Result->AdditionalPhysicsBodies)
+				{
+					PhysicsBodiesToBind.Add(Body.get());
+				}
+			}
+
+			TTuple<TArray<FReshapePointBindingData>, TArray<int32>, TArray<int32>> PhysicsBindingData = 
+					BindPhysicsBodies(PhysicsBodiesToBind, ShapeMeshTree, Result.get(), PhysicsToDeform);
+
 			const TArray<FReshapePointBindingData>& PhysicsBindDataArray = PhysicsBindingData.Get<0>();
 			const TArray<int32>& DeformedBodyIndices = PhysicsBindingData.Get<1>();
+			const TArray<int32>& DeformedBodyIndicesOffsets = PhysicsBindingData.Get<2>();
 
 			FMeshBufferSet PhysicsBodyBuffer;
 			PhysicsBodyBuffer.SetBufferCount(1);
@@ -1024,21 +1079,23 @@ namespace mu
 			PhysicsBodyBuffer.SetBuffer(0, sizeof(FReshapePointBindingData), BufDesc.Channels, BufDesc.Semantics, BufDesc.SemanticIndices, BufDesc.Formats, BufDesc.Components, BufDesc.Offsets);	
 			FMemory::Memcpy(PhysicsBodyBuffer.GetBufferData(0), PhysicsBindDataArray.GetData(), PhysicsBindDataArray.Num() * sizeof(FReshapePointBindingData));
 
-			// Bone indices buffer
-			MESH_BUFFER_SEMANTIC BoneSemantics[1] = { MBS_OTHER };
-			MESH_BUFFER_FORMAT BoneFormats[1] = { MBF_INT32 };
-			int BoneSemanticIndices[1] = { 0 };
-			int BoneComponents[1] = { 1 };
-			int BoneOffsets[1] = { 0 };
+			FIntBufferDescriptor IntBufDesc;
 
 			FMeshBufferSet PhysicsBodySelectionBuffer;
 			PhysicsBodySelectionBuffer.SetBufferCount(1);
 			PhysicsBodySelectionBuffer.SetElementCount(DeformedBodyIndices.Num());
-			PhysicsBodySelectionBuffer.SetBuffer(0, sizeof(int32), 1, BoneSemantics, BoneSemanticIndices, BoneFormats, BoneComponents, BoneOffsets);	
+			PhysicsBodySelectionBuffer.SetBuffer(0, sizeof(int32), IntBufDesc.Channels, IntBufDesc.Semantics, IntBufDesc.SemanticIndices, IntBufDesc.Formats, IntBufDesc.Components, IntBufDesc.Offsets);	
 			FMemory::Memcpy(PhysicsBodySelectionBuffer.GetBufferData(0), DeformedBodyIndices.GetData(), DeformedBodyIndices.Num() * sizeof(int32));
 
+			FMeshBufferSet PhysicsBodySelectionOffsetsBuffer;
+			PhysicsBodySelectionOffsetsBuffer.SetBufferCount(1);
+			PhysicsBodySelectionOffsetsBuffer.SetElementCount(DeformedBodyIndicesOffsets.Num());
+			PhysicsBodySelectionOffsetsBuffer.SetBuffer(0, sizeof(int32), IntBufDesc.Channels, IntBufDesc.Semantics, IntBufDesc.SemanticIndices, IntBufDesc.Formats, IntBufDesc.Components, IntBufDesc.Offsets);	
+			FMemory::Memcpy(PhysicsBodySelectionOffsetsBuffer.GetBufferData(0), DeformedBodyIndicesOffsets.GetData(), DeformedBodyIndicesOffsets.Num() * sizeof(int32));
+
 			Result->m_AdditionalBuffers.Emplace(EMeshBufferType::PhysicsBodyDeformBinding, MoveTemp(PhysicsBodyBuffer));
 			Result->m_AdditionalBuffers.Emplace(EMeshBufferType::PhysicsBodyDeformSelection, MoveTemp(PhysicsBodySelectionBuffer));
+			Result->m_AdditionalBuffers.Emplace(EMeshBufferType::PhysicsBodyDeformOffsets, MoveTemp(PhysicsBodySelectionOffsetsBuffer));
 		}
 		
         return Result;
diff --git a/Engine/Plugins/Experimental/Mutable/Source/MutableRuntime/Private/MuR/OpMeshFormat.cpp b/Engine/Plugins/Experimental/Mutable/Source/MutableRuntime/Private/MuR/OpMeshFormat.cpp
--- a/Engine/Plugins/Experimental/Mutable/Source/MutableRuntime/Private/MuR/OpMeshFormat.cpp
+++ b/Engine/Plugins/Experimental/Mutable/Source/MutableRuntime/Private/MuR/OpMeshFormat.cpp
@@ -674,6 +674,9 @@ namespace mu
 
 		pResult->SkeletonIDs = pSource->SkeletonIDs;
 
+		// A shallow copy is done here, it should not be a problem.
+		pResult->AdditionalPhysicsBodies = pSource->AdditionalPhysicsBodies;
+
 		return pResult;
 
 	}
diff --git a/Engine/Plugins/Experimental/Mutable/Source/MutableRuntime/Private/MuR/OpMeshMerge.h b/Engine/Plugins/Experimental/Mutable/Source/MutableRuntime/Private/MuR/OpMeshMerge.h
--- a/Engine/Plugins/Experimental/Mutable/Source/MutableRuntime/Private/MuR/OpMeshMerge.h
+++ b/Engine/Plugins/Experimental/Mutable/Source/MutableRuntime/Private/MuR/OpMeshMerge.h
@@ -571,11 +571,11 @@ namespace mu
 			auto AppendPhysicsBodiesUnique = [](PhysicsBody& OutPhysicsBody, const PhysicsBody& InPhysicsBody) -> bool
 			{
 				TArray<string>& OutBones = OutPhysicsBody.Bones;
-				TArray<int32>& OutCustomIds = OutPhysicsBody.CustomIds;
+				TArray<int32>& OutCustomIds = OutPhysicsBody.BodiesCustomIds;
 				TArray<FPhysicsBodyAggregate>& OutBodies = OutPhysicsBody.Bodies;
 
 				const TArray<string>& InBones = InPhysicsBody.Bones;
-				const TArray<int32>& InCustomIds = InPhysicsBody.CustomIds;
+				const TArray<int32>& InCustomIds = InPhysicsBody.BodiesCustomIds;
 				const TArray<FPhysicsBodyAggregate>& InBodies = InPhysicsBody.Bodies;
 
 				const int32 InBodyCount = InPhysicsBody.GetBodyCount();
@@ -679,7 +679,30 @@ namespace mu
 				pResult->SetPhysicsBody(MergedResultPhysicsBody);
 			}
 
+			// Additional physics bodies.
+			const int32 MaxAdditionalPhysicsResultNum = pFirst->AdditionalPhysicsBodies.Num() + pSecond->AdditionalPhysicsBodies.Num();
+
+			pResult->AdditionalPhysicsBodies.Reserve(MaxAdditionalPhysicsResultNum);
+			pResult->AdditionalPhysicsBodies.Append(pFirst->AdditionalPhysicsBodies);
+			
+			// Not very many additional bodies expected, do a quadratic search to have unique bodies based on external id.
+			const int32 NumSecondAdditionalBodies = pSecond->AdditionalPhysicsBodies.Num();
+			for (int32 Index = 0; Index < NumSecondAdditionalBodies; ++Index)
+			{
+				const int32 CustomIdToMerge = pSecond->AdditionalPhysicsBodies[Index]->CustomId;
+
+				const mu::Ptr<const PhysicsBody>* Found = pFirst->AdditionalPhysicsBodies.FindByPredicate(
+					[CustomIdToMerge](const Ptr<const mu::PhysicsBody>& Body) { return CustomIdToMerge == Body->CustomId; });
+
+				// TODO: current usages do not expect collisions, but same Id collision with bodies modified in differnet ways
+				// may need to be contemplated at some point.
+				if (!Found)
+				{
+					pResult->AdditionalPhysicsBodies.Add(pSecond->AdditionalPhysicsBodies[Index]);
+				}
+			}
 		}
+
 		// Vertices
 		//-----------------
 		{
diff --git a/Engine/Plugins/Experimental/Mutable/Source/MutableRuntime/Private/MuR/OpMeshReshape.h b/Engine/Plugins/Experimental/Mutable/Source/MutableRuntime/Private/MuR/OpMeshReshape.h
--- a/Engine/Plugins/Experimental/Mutable/Source/MutableRuntime/Private/MuR/OpMeshReshape.h
+++ b/Engine/Plugins/Experimental/Mutable/Source/MutableRuntime/Private/MuR/OpMeshReshape.h
@@ -515,8 +515,10 @@ namespace mu
 	}
 
 	inline void ApplyToPhysicsBodies(
-			PhysicsBody& PBody, const Mesh& BaseMesh, 
-			TArrayView<const FReshapePointBindingData> BindingData, TArrayView<const int32> UsedIndices, 
+			PhysicsBody& PBody, int32& InOutNumProcessedBindPoints, 
+			const Mesh& BaseMesh, 
+			TArrayView<const FReshapePointBindingData> BindingData, 
+			TArrayView<const int32> UsedIndices, 
 			const FShapeMeshDescriptorApply& Shape)
 	{
 #if DO_CHECK
@@ -530,8 +532,6 @@ namespace mu
 		}
 #endif
 		
-		int32 NumProcessedBindPoints = 0;
-
 		bool bAnyModified = false;
 
 		// Retrieve them in the same order the boxes where put in, so they can be linked to the physics body volumes.
@@ -553,7 +553,7 @@ namespace mu
 				float R;
 
 				TStaticArray<FVector3f, 6> Points; 
-				const bool bDeformed = GetDeformedPoints(Shape, &BindingData[NumProcessedBindPoints], Points);
+				const bool bDeformed = GetDeformedPoints(Shape, &BindingData[InOutNumProcessedBindPoints], Points);
 
 				if (bDeformed)
 				{
@@ -562,14 +562,14 @@ namespace mu
 					bAnyModified = true;
 				}
 
-				NumProcessedBindPoints += Points.Num();
+				InOutNumProcessedBindPoints += Points.Num();
 			}
 
 			const int32 NumBoxes = PBody.GetBoxCount(B);
 			for ( int32 I = 0; I < NumBoxes; ++I )
 			{
 				TStaticArray<FVector3f, 14> Points;
-				const bool bDeformed = GetDeformedPoints(Shape, &BindingData[NumProcessedBindPoints], Points);
+				const bool bDeformed = GetDeformedPoints(Shape, &BindingData[InOutNumProcessedBindPoints], Points);
 				
 				if (bDeformed)
 				{
@@ -582,14 +582,14 @@ namespace mu
 					bAnyModified = true;
 				}
 
-				NumProcessedBindPoints += Points.Num();
+				InOutNumProcessedBindPoints += Points.Num();
 			}
 			
 			const int32 NumSphyls = PBody.GetSphylCount(B);
 			for ( int32 I = 0; I < NumSphyls; ++I )
 			{
 				TStaticArray<FVector3f, 14> Points;
-				const bool bDeformed = GetDeformedPoints(Shape, &BindingData[NumProcessedBindPoints], Points);
+				const bool bDeformed = GetDeformedPoints(Shape, &BindingData[InOutNumProcessedBindPoints], Points);
 
 				if ( bDeformed )
 				{
@@ -603,14 +603,14 @@ namespace mu
 					bAnyModified = true;
 				}
 
-				NumProcessedBindPoints += Points.Num();
+				InOutNumProcessedBindPoints += Points.Num();
 			}
 
 			const int32 NumTaperedCapsules = PBody.GetTaperedCapsuleCount(B);
 			for ( int32 I = 0; I < NumTaperedCapsules; ++I )
 			{
 				TStaticArray<FVector3f, 14> Points;
-				const bool bDeformed = GetDeformedPoints(Shape, &BindingData[NumProcessedBindPoints], Points);
+				const bool bDeformed = GetDeformedPoints(Shape, &BindingData[InOutNumProcessedBindPoints], Points);
 
 				if ( bDeformed )
 				{
@@ -625,7 +625,7 @@ namespace mu
 					bAnyModified = true;
 				}
 				
-				NumProcessedBindPoints += Points.Num();
+				InOutNumProcessedBindPoints += Points.Num();
 			}
 
 			const int32 NumConvex = PBody.GetConvexCount(B);
@@ -638,7 +638,7 @@ namespace mu
 				FTransform3f ConvexTransform;
 				PBody.GetConvexTransform(B, I, ConvexTransform);
 
-				GetDeformedConvex(Shape, &BindingData[NumProcessedBindPoints], VerticesView);
+				GetDeformedConvex(Shape, &BindingData[InOutNumProcessedBindPoints], VerticesView);
 
 				FTransform3f InvConvexT = InvBoneTransform * ConvexTransform.Inverse();
 				for ( FVector3f& V : VerticesView )
@@ -646,12 +646,56 @@ namespace mu
 					V = InvConvexT.TransformPosition( V );				
 				}
 				
-				NumProcessedBindPoints += VerticesView.Num();
+				InOutNumProcessedBindPoints += VerticesView.Num();
 				bAnyModified = true;
 			}
 		}
 
 		PBody.bBodiesModified = bAnyModified;
+	}
+
+	inline void ApplyToAllPhysicsBodies(
+		Mesh& OutNewMesh, const Mesh& BaseMesh, 
+		TArrayView<const FReshapePointBindingData> BindingData,
+		TArrayView<const int32> UsedIndices,
+		TArrayView<const int32> BodyOffsets,
+		const FShapeMeshDescriptorApply& Shape)
+	{
+		if (!BodyOffsets.Num())
+		{
+			return; // Nothing to do.
+		}
+
+		int32 NumProcessedBindPoints = 0;
+
+		auto ApplyPhysicsBody = [&](PhysicsBody& OutBody, int32 IndicesBegin, int32 IndicesEnd) -> void
+		{
+			TArrayView<const int32> BodyUsedIndices(UsedIndices.GetData() + IndicesBegin, IndicesEnd - IndicesBegin);
+
+			ApplyToPhysicsBodies(OutBody, NumProcessedBindPoints, BaseMesh, BindingData, BodyUsedIndices, Shape);
+		};
+
+
+		check(BodyOffsets.Num() > 1);
+		check(OutNewMesh.AdditionalPhysicsBodies.Num() + 1 == BodyOffsets.Num() - 1);
+		const int32 PhysicsBodiesNum = BodyOffsets.Num() - 1;
+		
+		// Apply main physics body
+		if (BaseMesh.m_pPhysicsBody)
+		{
+			Ptr<PhysicsBody> NewBody = BaseMesh.m_pPhysicsBody->Clone();
+			ApplyPhysicsBody(*NewBody, 0, BodyOffsets[1]);
+			OutNewMesh.m_pPhysicsBody = NewBody;
+		}
+
+
+		// Apply additional physics bodies
+		for (int32 I = 1; I < PhysicsBodiesNum; ++I)
+		{	
+			Ptr<PhysicsBody> NewBody = BaseMesh.AdditionalPhysicsBodies[I - 1]->Clone();
+			ApplyPhysicsBody(*NewBody, BodyOffsets[I], BodyOffsets[I + 1]);
+			OutNewMesh.AdditionalPhysicsBodies[I - 1] = NewBody;
+		}
 
 		check(NumProcessedBindPoints == BindingData.Num());	
 	}
@@ -675,7 +719,7 @@ namespace mu
 
 		// Early out if nothing will be modified and the vertices discarted.
 		const bool bSkeletonModification = BaseMesh->GetSkeleton() && bReshapeSkeleton;
-		const bool bPhysicsModification = BaseMesh->GetPhysicsBody() && bReshapePhysicsVolumes;
+		const bool bPhysicsModification = (BaseMesh->GetPhysicsBody() || BaseMesh->AdditionalPhysicsBodies.Num()) && bReshapePhysicsVolumes;
 
 		if (!bReshapeVertices && !bSkeletonModification && !bPhysicsModification)
 		{
@@ -719,7 +763,8 @@ namespace mu
 			const bool bIsBindOpBuffer = 
 					A.Key == EMeshBufferType::SkeletonDeformBinding || 
 					A.Key == EMeshBufferType::PhysicsBodyDeformBinding ||
-					A.Key == EMeshBufferType::PhysicsBodyDeformSelection;
+					A.Key == EMeshBufferType::PhysicsBodyDeformSelection ||
+					A.Key == EMeshBufferType::PhysicsBodyDeformOffsets;
 	
 			if (!bIsBindOpBuffer)
 			{
@@ -844,8 +889,9 @@ namespace mu
 
 		// Transform physics volumes based on the deformed sampling points.
 		const PhysicsBody* OldPhysicsBody = Result->m_pPhysicsBody.get();
+		const int32 AdditionalPhysicsBodiesNum = Result->AdditionalPhysicsBodies.Num();
 
-		if (bReshapePhysicsVolumes && OldPhysicsBody)
+		if (bReshapePhysicsVolumes && (OldPhysicsBody || AdditionalPhysicsBodiesNum))
 		{	
 			MUTABLE_CPUPROFILER_SCOPE(ReshapePhysicsBodies);
 			
@@ -856,22 +902,29 @@ namespace mu
 			const BufferEntryType* FoundPhysicsBindSelectionBuffer = BaseMesh->m_AdditionalBuffers.FindByPredicate(
 					[](BufferEntryType& E){ return E.Key == EMeshBufferType::PhysicsBodyDeformSelection; });
 	
+			const BufferEntryType* FoundPhysicsBindOffsetsBuffer = BaseMesh->m_AdditionalBuffers.FindByPredicate(
+					[](BufferEntryType& E){ return E.Key == EMeshBufferType::PhysicsBodyDeformOffsets; });
+
 			BarycentricDataBuffer = -1;
 			if (FoundPhysicsBindBuffer)
 			{
 				FoundPhysicsBindBuffer->Value.FindChannel(MBS_BARYCENTRICCOORDS, BindingDataIndex, &BarycentricDataBuffer, &BarycentricDataChannel);
 			}
 			
-			if (FoundPhysicsBindBuffer && FoundPhysicsBindSelectionBuffer && BarycentricDataBuffer >= 0)
+			const bool bAllNeededBuffersFound =
+				FoundPhysicsBindBuffer && FoundPhysicsBindSelectionBuffer && FoundPhysicsBindOffsetsBuffer && BarycentricDataBuffer >= 0;
+
+			if (bAllNeededBuffersFound)
 			{
 				const FMeshBufferSet& PhysicsBindBuffer = FoundPhysicsBindBuffer->Value;
 				const FMeshBufferSet& PhyiscsBindSelectionBuffer = FoundPhysicsBindSelectionBuffer->Value;
+				const FMeshBufferSet& PhyiscsBindOffsetsBuffer = FoundPhysicsBindOffsetsBuffer->Value;
 
 				// \TODO: More checks
 				check(BarycentricDataChannel == 0);
 				check(PhysicsBindBuffer.GetElementSize(BarycentricDataBuffer) == (int)sizeof(FReshapePointBindingData));
 					
-				TArrayView<const FReshapePointBindingData> PhysicsBindingData(  
+				TArrayView<const FReshapePointBindingData> BindingData(  
 						(const FReshapePointBindingData*)PhysicsBindBuffer.GetBufferData(BarycentricDataBuffer),
 						PhysicsBindBuffer.GetElementCount() );
 
@@ -879,11 +932,11 @@ namespace mu
 						(const int32*)PhyiscsBindSelectionBuffer.GetBufferData(0), 
 					    PhyiscsBindSelectionBuffer.GetElementCount());
 
-				Ptr<PhysicsBody> NewPhysicsBody = OldPhysicsBody->Clone();
-				Result->SetPhysicsBody(NewPhysicsBody);
+				TArrayView<const int32> Offsets(
+						(const int32*)PhyiscsBindOffsetsBuffer.GetBufferData(0), 
+					    PhyiscsBindOffsetsBuffer.GetElementCount());
 
-				check(NewPhysicsBody);
-				ApplyToPhysicsBodies(*NewPhysicsBody, *Result, PhysicsBindingData, UsedIndices, ShapeDescriptor);
+				ApplyToAllPhysicsBodies(*Result, *BaseMesh, BindingData, UsedIndices, Offsets, ShapeDescriptor);
 			}
 		}
 		
diff --git a/Engine/Plugins/Experimental/Mutable/Source/MutableRuntime/Private/MuR/PhysicsBody.cpp b/Engine/Plugins/Experimental/Mutable/Source/MutableRuntime/Private/MuR/PhysicsBody.cpp
--- a/Engine/Plugins/Experimental/Mutable/Source/MutableRuntime/Private/MuR/PhysicsBody.cpp
+++ b/Engine/Plugins/Experimental/Mutable/Source/MutableRuntime/Private/MuR/PhysicsBody.cpp
@@ -27,19 +27,36 @@ namespace mu
 	{
 		mu::Ptr<PhysicsBody> pResult = new PhysicsBody();
 
+		pResult->CustomId = CustomId;
+
 		pResult->Bodies = Bodies;
 		pResult->Bones = Bones;
-		pResult->CustomIds = CustomIds;
+		pResult->BodiesCustomIds = BodiesCustomIds;
+
+		pResult->bBodiesModified = bBodiesModified;
+
 
 		return pResult;
 	}
 
+	//-------------------------------------------------------------------------------------------------
+	void PhysicsBody::SetCustomId(int32 InCustomId)
+	{
+		CustomId = InCustomId;
+	}
+
+	//-------------------------------------------------------------------------------------------------
+	int32 PhysicsBody::GetCustomId() const
+	{
+		return CustomId;
+	}
+
 	//-------------------------------------------------------------------------------------------------
 	void PhysicsBody::SetBodyCount(int32 Count)
 	{
 		Bodies.SetNum(Count);
 		Bones.SetNum(Count);
-		CustomIds.Init(-1, Count);
+		BodiesCustomIds.Init(-1, Count);
 	}
 
 	//-------------------------------------------------------------------------------------------------
@@ -67,17 +84,17 @@ namespace mu
 	//-------------------------------------------------------------------------------------------------
 	int32 PhysicsBody::GetBodyCustomId(int32 B) const
 	{
-		check(B >= 0 && B < CustomIds.Num());
+		check(B >= 0 && B < BodiesCustomIds.Num());
 
-		return CustomIds[B];
+		return BodiesCustomIds[B];
 	}
 
 	//-------------------------------------------------------------------------------------------------
-	void PhysicsBody::SetBodyCustomId(int32 B, int32 CustomId)
+	void PhysicsBody::SetBodyCustomId(int32 B, int32 BodyCustomId)
 	{
-		check(B >= 0 && B < CustomIds.Num());
+		check(B >= 0 && B < BodiesCustomIds.Num());
 
-		CustomIds[B] = CustomId;
+		BodiesCustomIds[B] = BodyCustomId;
 	}
 
 	//-------------------------------------------------------------------------------------------------
diff --git a/Engine/Plugins/Experimental/Mutable/Source/MutableRuntime/Public/MuR/Mesh.h b/Engine/Plugins/Experimental/Mutable/Source/MutableRuntime/Public/MuR/Mesh.h
--- a/Engine/Plugins/Experimental/Mutable/Source/MutableRuntime/Public/MuR/Mesh.h
+++ b/Engine/Plugins/Experimental/Mutable/Source/MutableRuntime/Public/MuR/Mesh.h
@@ -95,7 +95,8 @@ namespace mu
 		None,
 		SkeletonDeformBinding,
 		PhysicsBodyDeformBinding,
-		PhysicsBodyDeformSelection
+		PhysicsBodyDeformSelection,
+		PhysicsBodyDeformOffsets
 	};
 	MUTABLE_DEFINE_ENUM_SERIALISABLE(EMeshBufferType)
 
@@ -124,7 +125,8 @@ namespace mu
 		WithAdditionalBuffers = 1 << 10,
 		WithLayouts           = 1 << 11,
 		WithPoses			  = 1 << 12,
-		WithSkeletonIDs		  = 1 << 13
+		WithSkeletonIDs		  = 1 << 13,
+		WithAdditionalPhysics = 1 << 14
 	};
 	
 	ENUM_CLASS_FLAGS(EMeshCloneFlags);
@@ -226,6 +228,10 @@ namespace mu
         void SetPhysicsBody( Ptr<const PhysicsBody> );
         Ptr<const PhysicsBody> GetPhysicsBody() const;
 
+		int32 AddAdditionalPhysicsBody(Ptr<const PhysicsBody> Body);
+		Ptr<const PhysicsBody> GetAdditionalPhysicsBody(int32 I) const;
+		//int32 GetAdditionalPhysicsBodyExternalId(int32 I) const;
+
         //! \}
 
         //! \name Face groups
@@ -347,6 +353,9 @@ namespace mu
 		Ptr<const Skeleton> m_pSkeleton;
 		Ptr<const PhysicsBody> m_pPhysicsBody;
 
+		//! Additional physics bodies referenced by the mesh that don't merge.
+		TArray<Ptr<const PhysicsBody>> AdditionalPhysicsBodies;
+
 		//! Texture Layout blocks attached to this mesh. They are const because they could be shared with
 		//! other meshes, so they need to be cloned and replaced if a modification is needed.
 		TArray<Ptr<const Layout>> m_layouts;
@@ -441,7 +450,7 @@ namespace mu
 		//!
 		inline void Serialise(OutputArchive& arch) const
 		{
-			uint32 ver = 14;
+			uint32 ver = 15;
 			arch << ver;
 
 			arch << m_IndexBuffers;
@@ -462,6 +471,8 @@ namespace mu
 			arch << m_tags;
 
 			arch << BonePoses;
+
+			arch << AdditionalPhysicsBodies;
 		}
 
 		//!
@@ -469,7 +480,7 @@ namespace mu
 		{
 			uint32 ver;
 			arch >> ver;
-			check(ver <= 14);
+			check(ver <= 15);
 
 			arch >> m_IndexBuffers;
 			arch >> m_VertexBuffers;
@@ -515,6 +526,11 @@ namespace mu
 					BonePoses[BoneIndex].BoneTransform = m_pSkeleton->m_boneTransforms_DEPRECATED[BoneIndex];
 				}
 			}
+
+			if (ver >= 15)
+			{
+				arch >> AdditionalPhysicsBodies;
+			}
 		}
 
 
@@ -547,16 +563,24 @@ namespace mu
 				equal &= (*m_layouts[i]) == (*o.m_layouts[i]);
 			}
 
+			equal &= m_AdditionalBuffers.Num() == o.m_AdditionalBuffers.Num();
 			for (int32 i = 0; equal && i < m_AdditionalBuffers.Num(); ++i)
 			{
 				equal &= m_AdditionalBuffers[i] == o.m_AdditionalBuffers[i];
 			}
 
+			equal &= BonePoses.Num() == o.BonePoses.Num();
 			for (int32 i = 0; equal && i < BonePoses.Num(); ++i)
 			{
 				equal &= BonePoses[i] == o.BonePoses[i];
 			}
 
+			equal &= AdditionalPhysicsBodies.Num() == o.AdditionalPhysicsBodies.Num();
+			for (int32 i = 0; equal && i < AdditionalPhysicsBodies.Num(); ++i)
+			{
+				equal &= *AdditionalPhysicsBodies[i] == *o.AdditionalPhysicsBodies[i];
+			}
+
 			return equal;
 		}
 
diff --git a/Engine/Plugins/Experimental/Mutable/Source/MutableRuntime/Public/MuR/PhysicsBody.h b/Engine/Plugins/Experimental/Mutable/Source/MutableRuntime/Public/MuR/PhysicsBody.h
--- a/Engine/Plugins/Experimental/Mutable/Source/MutableRuntime/Public/MuR/PhysicsBody.h
+++ b/Engine/Plugins/Experimental/Mutable/Source/MutableRuntime/Public/MuR/PhysicsBody.h
@@ -322,34 +322,37 @@ namespace mu
 
         //bool operator==( const PhysicsBody& Other ) const;
 
-		void SetBodyCount( int32 B );
-		int32 GetBodyCount( ) const;
+		void SetCustomId(int32 Id);
+		int32 GetCustomId() const;
 
-		void SetBodyBoneName( int32 B, const char* BoneName );
-		const char* GetBodyBoneName( int32 B ) const;
+		void SetBodyCount(int32 B);
+		int32 GetBodyCount() const;
 	
-		void SetBodyCustomId( int32 B, int32 CustomId );
-		int32 GetBodyCustomId( int32 B ) const;
+		void SetBodyBoneName(int32 B, const char* BoneName);
+		const char* GetBodyBoneName(int32 B) const;
 		
-		int32 GetSphereCount( int32 B ) const;
-		int32 GetBoxCount( int32 B) const;
-		int32 GetConvexCount( int32 B) const;
-		int32 GetSphylCount( int32 B ) const;
-		int32 GetTaperedCapsuleCount(int32 B ) const;
+		void SetBodyCustomId(int32 B, int32 BodyCustomId);
+		int32 GetBodyCustomId(int32 B) const;
+		
+		int32 GetSphereCount(int32 B) const;
+		int32 GetBoxCount(int32 B) const;
+		int32 GetConvexCount(int32 B) const;
+		int32 GetSphylCount(int32 B) const;
+		int32 GetTaperedCapsuleCount(int32 B) const;
 		
 		void SetSphereCount(int32 B, int32 Count);
 		void SetBoxCount(int32 B, int32 Count);
 		void SetConvexCount(int32 B, int32 Count);
 		void SetSphylCount(int32 B, int32 Count);
-		void SetTaperedCapsuleCount( int32 B, int32 Count);
+		void SetTaperedCapsuleCount(int32 B, int32 Count);
 		
 		void SetSphere( 
 				int32 B, int32 I, 
-				FVector3f Position, float Radius );
+				FVector3f Position, float Radius);
 
 		void SetBox( 
 				int32 B, int32 I, 
-				FVector3f Position, FQuat4f Orientation, FVector3f Size );
+				FVector3f Position, FQuat4f Orientation, FVector3f Size);
 
 		void SetConvexMesh( 
 				int32 B, int32 I,
@@ -357,29 +360,29 @@ namespace mu
 
 		void SetConvexTransform( 
 				int32 B, int32 I, 
-				const FTransform3f& Transform );
+				const FTransform3f& Transform);
 
 		void SetSphyl( 
 				int32 B, int32 I, 
 				FVector3f Position, FQuat4f Orientation, 
-				float Radius, float Length );
+				float Radius, float Length);
 
 		void SetTaperedCapsule( 
 				int32 B, int32 I, 
 				FVector3f Position, FQuat4f Orientation, 
-				float Radius0, float Radius1, float Length );
-
-		void SetSphereFlags( int32 B, int32 I, uint32 Flags );
-		void SetBoxFlags( int32 B, int32 I, uint32 Flags );
-		void SetConvexFlags( int32 B, int32 I, uint32 Flags );
-		void SetSphylFlags( int32 B, int32 I, uint32 Flags );
-		void SetTaperedCapsuleFlags( int32 B, int32 I, uint32 Flags );
-
-		void SetSphereName( int32 B, int32 I, const char* Name );
-		void SetBoxName( int32 B, int32 I, const char* Name );
-		void SetConvexName( int32 B, int32 I, const char* Name );
-		void SetSphylName( int32 B, int32 I, const char* Name );
-		void SetTaperedCapsuleName( int32 B, int32 I, const char* Name );
+				float Radius0, float Radius1, float Length);
+
+		void SetSphereFlags(int32 B, int32 I, uint32 Flags);
+		void SetBoxFlags(int32 B, int32 I, uint32 Flags);
+		void SetConvexFlags(int32 B, int32 I, uint32 Flags);
+		void SetSphylFlags(int32 B, int32 I, uint32 Flags);
+		void SetTaperedCapsuleFlags(int32 B, int32 I, uint32 Flags);
+
+		void SetSphereName(int32 B, int32 I, const char* Name);
+		void SetBoxName(int32 B, int32 I, const char* Name);
+		void SetConvexName(int32 B, int32 I, const char* Name);
+		void SetSphylName(int32 B, int32 I, const char* Name);
+		void SetTaperedCapsuleName(int32 B, int32 I, const char* Name);
 	
 		void GetSphere( 
 				int32 B, int32 I, 
@@ -387,11 +390,11 @@ namespace mu
 
 		void GetBox( 
 				int32 B, int32 I, 
-				FVector3f& OutPosition, FQuat4f& OutOrientation, FVector3f& OutSize ) const;
+				FVector3f& OutPosition, FQuat4f& OutOrientation, FVector3f& OutSize) const;
 
 		void GetConvex( 
 				int32 B, int32 I, 
-				TArrayView<const FVector3f>& OutVertices, TArrayView<const int32>& OutIndices, FTransform3f& OutTransform ) const;
+				TArrayView<const FVector3f>& OutVertices, TArrayView<const int32>& OutIndices, FTransform3f& OutTransform) const;
 
 		void GetConvexMeshView(
 				int32 B, int32 I, 
@@ -411,17 +414,17 @@ namespace mu
 				FVector3f& OutPosition, FQuat4f& OutOrientation, 
 				float& OutRadius0, float& OutRadius1, float& OutLength) const;
 		
-		uint32 GetSphereFlags( int32 B, int32 I ) const;
-		uint32 GetBoxFlags( int32 B, int32 I ) const;
-		uint32 GetConvexFlags( int32 B, int32 I ) const;
-		uint32 GetSphylFlags( int32 B, int32 I ) const;
-		uint32 GetTaperedCapsuleFlags( int32 B, int32 I ) const;
-
-		const char* GetSphereName( int32 B, int32 I ) const;
-		const char* GetBoxName( int32 B, int32 I ) const;
-		const char* GetConvexName( int32 B, int32 I ) const;
-		const char* GetSphylName( int32 B, int32 I ) const;
-		const char* GetTaperedCapsuleName( int32 B, int32 I ) const;
+		uint32 GetSphereFlags(int32 B, int32 I) const;
+		uint32 GetBoxFlags(int32 B, int32 I) const;
+		uint32 GetConvexFlags(int32 B, int32 I) const;
+		uint32 GetSphylFlags(int32 B, int32 I) const;
+		uint32 GetTaperedCapsuleFlags(int32 B, int32 I) const;
+
+		const char* GetSphereName(int32 B, int32 I) const;
+		const char* GetBoxName(int32 B, int32 I) const;
+		const char* GetConvexName(int32 B, int32 I) const;
+		const char* GetSphylName(int32 B, int32 I) const;
+		const char* GetTaperedCapsuleName(int32 B, int32 I) const;
 
 	protected:
 		//! Forbidden. Manage with the Ptr<> template.
@@ -430,22 +433,26 @@ namespace mu
 		}
 
 	public:
+		int32 CustomId = -1;
+
 		// Bone name the physics volume aggregate is bound to. 
 		TArray<string> Bones;
 		TArray<FPhysicsBodyAggregate> Bodies;
-		TArray<int32> CustomIds;
+		TArray<int32> BodiesCustomIds;
+
 
 		bool bBodiesModified = false;
 
         //!
         inline void Serialise( OutputArchive& arch ) const
 		{
-            uint32 ver = 1;
+            uint32 ver = 2;
 			arch << ver;
         	
+			arch << CustomId;
         	arch << Bodies;
         	arch << Bones;
-        	arch << CustomIds;
+        	arch << BodiesCustomIds;
 			arch << bBodiesModified;
         }
 
@@ -454,11 +461,16 @@ namespace mu
 		{
             uint32 ver;
 			arch >> ver;
-            check(ver <= 1);
+            check(ver <= 2);
+        	
+			if (ver >= 2)
+			{
+				arch >> CustomId;
+			}
         	
         	arch >> Bodies;
         	arch >> Bones;
-        	arch >> CustomIds;
+        	arch >> BodiesCustomIds;
 
 			if (ver >= 1)
 			{
@@ -467,12 +479,13 @@ namespace mu
 		}
 		
         //!
-		inline bool operator==( const PhysicsBody& o ) const
+		inline bool operator==(const PhysicsBody& Other) const
         {
-        	return Bones == o.Bones &&
-        		   Bodies == o.Bodies &&
-        	       CustomIds == o.CustomIds &&
-				   bBodiesModified == o.bBodiesModified;
+        	return CustomId        == Other.CustomId        &&
+				   Bones		   == Other.Bones           &&
+        		   Bodies          == Other.Bodies          &&
+        	       BodiesCustomIds == Other.BodiesCustomIds &&
+				   bBodiesModified == Other.bBodiesModified;
         }
 	};
 }
diff --git a/Engine/Plugins/Experimental/Mutable/Source/MutableTools/Private/MuT/CodeGenerator_Mesh.cpp b/Engine/Plugins/Experimental/Mutable/Source/MutableTools/Private/MuT/CodeGenerator_Mesh.cpp
--- a/Engine/Plugins/Experimental/Mutable/Source/MutableTools/Private/MuT/CodeGenerator_Mesh.cpp
+++ b/Engine/Plugins/Experimental/Mutable/Source/MutableTools/Private/MuT/CodeGenerator_Mesh.cpp
@@ -1379,7 +1379,7 @@ class Node;
 		OpBind->BonesToDeform = node.m_bonesToDeform;
     	OpBind->bReshapePhysicsVolumes = node.m_reshapePhysicsVolumes;
 		OpBind->PhysicsToDeform = node.m_physicsToDeform;
-		OpBind->bReshapeVertices = true;
+		OpBind->bReshapeVertices = node.m_reshapeVertices;
 		OpBind->BindingMethod = static_cast<uint32>(EShapeBindingMethod::ReshapeClosestProject);
 
 		OpApply->bReshapeVertices = OpBind->bReshapeVertices;
diff --git a/Engine/Plugins/Experimental/Mutable/Source/MutableTools/Private/MuT/NodeMeshReshape.cpp b/Engine/Plugins/Experimental/Mutable/Source/MutableTools/Private/MuT/NodeMeshReshape.cpp
--- a/Engine/Plugins/Experimental/Mutable/Source/MutableTools/Private/MuT/NodeMeshReshape.cpp
+++ b/Engine/Plugins/Experimental/Mutable/Source/MutableTools/Private/MuT/NodeMeshReshape.cpp
@@ -116,6 +116,11 @@ namespace mu
 		m_pD->m_pTargetShape = pNode;
 	}
 
+	//---------------------------------------------------------------------------------------------
+	void NodeMeshReshape::SetReshapeVertices(bool bEnable)
+	{
+		m_pD->m_reshapeVertices = bEnable;
+	}
 
 	//---------------------------------------------------------------------------------------------
 	void NodeMeshReshape::SetReshapeSkeleton(bool bEnable)
diff --git a/Engine/Plugins/Experimental/Mutable/Source/MutableTools/Private/MuT/NodeMeshReshapePrivate.h b/Engine/Plugins/Experimental/Mutable/Source/MutableTools/Private/MuT/NodeMeshReshapePrivate.h
--- a/Engine/Plugins/Experimental/Mutable/Source/MutableTools/Private/MuT/NodeMeshReshapePrivate.h
+++ b/Engine/Plugins/Experimental/Mutable/Source/MutableTools/Private/MuT/NodeMeshReshapePrivate.h
@@ -24,6 +24,7 @@ namespace mu
 		Ptr<NodeMesh> m_pBaseMesh;
 		Ptr<NodeMesh> m_pBaseShape;
 		Ptr<NodeMesh> m_pTargetShape;
+		bool m_reshapeVertices = true;
 		bool m_reshapeSkeleton = false;
 		bool m_enableRigidParts = false;
 		bool m_reshapePhysicsVolumes = false;
@@ -34,12 +35,13 @@ namespace mu
         //!
 		void Serialise( OutputArchive& arch ) const
 		{
-            uint32_t ver = 7;
+            uint32_t ver = 8;
 			arch << ver;
 
 			arch << m_pBaseMesh;
 			arch << m_pBaseShape;
 			arch << m_pTargetShape;
+			arch << m_reshapeVertices;
 			arch << m_reshapeSkeleton;
 			arch << m_enableRigidParts;
 			arch << m_bonesToDeform;
@@ -52,11 +54,17 @@ namespace mu
 		{
             uint32_t ver;
 			arch >> ver;
-			check(ver <= 7);
+			check(ver <= 8);
 
 			arch >> m_pBaseMesh;
 			arch >> m_pBaseShape;
 			arch >> m_pTargetShape;
+			
+			if (ver >= 8)
+			{
+				arch >> m_reshapeVertices;
+			}
+
 			arch >> m_reshapeSkeleton;
 			arch >> m_enableRigidParts;
 
diff --git a/Engine/Plugins/Experimental/Mutable/Source/MutableTools/Public/MuT/NodeMeshReshape.h b/Engine/Plugins/Experimental/Mutable/Source/MutableTools/Public/MuT/NodeMeshReshape.h
--- a/Engine/Plugins/Experimental/Mutable/Source/MutableTools/Public/MuT/NodeMeshReshape.h
+++ b/Engine/Plugins/Experimental/Mutable/Source/MutableTools/Public/MuT/NodeMeshReshape.h
@@ -53,6 +53,8 @@ namespace mu
 		const NodeMeshPtr& GetTargetShape() const;
 		void SetTargetShape(const NodeMeshPtr&);
 
+		void SetReshapeVertices(bool);
+
 		/** Also deform the mesh skeleton. Disabled by default. */
 		void SetReshapeSkeleton(bool);
 		
