[Mutable] Invalidate generated data if the descriptor being set has a different CO than the current one.

#rnx

diff --git a/Engine/Plugins/Experimental/Mutable/Source/CustomizableObject/Private/MuCO/CustomizableObjectInstance.cpp b/Engine/Plugins/Experimental/Mutable/Source/CustomizableObject/Private/MuCO/CustomizableObjectInstance.cpp
--- a/Engine/Plugins/Experimental/Mutable/Source/CustomizableObject/Private/MuCO/CustomizableObjectInstance.cpp
+++ b/Engine/Plugins/Experimental/Mutable/Source/CustomizableObject/Private/MuCO/CustomizableObjectInstance.cpp
@@ -154,8 +154,9 @@ const FCustomizableObjectInstanceDescriptor& UCustomizableObjectInstance::GetDes
 
 void UCustomizableObjectInstance::SetDescriptor(const FCustomizableObjectInstanceDescriptor& InDescriptor)
 {
+	const bool bInvalidatePreviousData = GetCustomizableObject() != InDescriptor.CustomizableObject;
 	Descriptor = InDescriptor;
-	PrivateData->ReloadParameters(this);
+	PrivateData->ReloadParameters(this, bInvalidatePreviousData);
 }
 
 
