[mutable] Fix bug in plain colour optimization for RGB images.
Additional data in mutable debugger for some code tree operations.
#rnx

diff --git a/Engine/Plugins/Experimental/Mutable/Source/CustomizableObjectEditor/Private/MuCOE/SMutableCodeViewer.cpp b/Engine/Plugins/Experimental/Mutable/Source/CustomizableObjectEditor/Private/MuCOE/SMutableCodeViewer.cpp
--- a/Engine/Plugins/Experimental/Mutable/Source/CustomizableObjectEditor/Private/MuCOE/SMutableCodeViewer.cpp
+++ b/Engine/Plugins/Experimental/Mutable/Source/CustomizableObjectEditor/Private/MuCOE/SMutableCodeViewer.cpp
@@ -230,6 +230,27 @@ public:
 				break;
 			}
 
+			case mu::OP_TYPE::IM_PLAINCOLOUR:
+			{
+				mu::OP::ImagePlainColourArgs Args = Program.GetOpArgs<mu::OP::ImagePlainColourArgs>(RowItem->MutableOperation);
+				OpName += TEXT(" format: ");
+				OpName += StringCast<TCHAR>(mu::TypeInfo::s_imageFormatName[int32(Args.format)]).Get();
+				OpName += FString::Printf(TEXT(" size %d x %d"), Args.size[0], Args.size[1]);
+				OpName += FString::Printf(TEXT(" mips %d"), Args.LODs);
+				break;
+			}
+
+			case mu::OP_TYPE::IN_ADDIMAGE:
+			{
+				mu::OP::InstanceAddArgs Args = Program.GetOpArgs<mu::OP::InstanceAddArgs>(RowItem->MutableOperation);
+				if (Program.m_constantStrings.IsValidIndex(Args.name))
+				{
+					OpName += TEXT(" name: ");
+					OpName += FString(Program.m_constantStrings[Args.name].c_str());
+				}
+ 				break;
+			}
+
 			default:
 				break;
 			}
diff --git a/Engine/Plugins/Experimental/Mutable/Source/MutableRuntime/Private/MuR/Image.cpp b/Engine/Plugins/Experimental/Mutable/Source/MutableRuntime/Private/MuR/Image.cpp
--- a/Engine/Plugins/Experimental/Mutable/Source/MutableRuntime/Private/MuR/Image.cpp
+++ b/Engine/Plugins/Experimental/Mutable/Source/MutableRuntime/Private/MuR/Image.cpp
@@ -708,8 +708,8 @@ namespace mu
 				if (res)
 				{
 					colour[0] = r / 255.0f;
-					colour[0] = g / 255.0f;
-					colour[0] = b / 255.0f;
+					colour[1] = g / 255.0f;
+					colour[2] = b / 255.0f;
 					colour[3] = 1.0f;
 				}
 				break;
