[mutable] code cleanup to ease backports.
#rnx

diff --git a/Engine/Plugins/Experimental/Mutable/Source/CustomizableObject/Private/MuCO/CustomizableInstanceLODManagement.cpp b/Engine/Plugins/Experimental/Mutable/Source/CustomizableObject/Private/MuCO/CustomizableInstanceLODManagement.cpp
--- a/Engine/Plugins/Experimental/Mutable/Source/CustomizableObject/Private/MuCO/CustomizableInstanceLODManagement.cpp
+++ b/Engine/Plugins/Experimental/Mutable/Source/CustomizableObject/Private/MuCO/CustomizableInstanceLODManagement.cpp
@@ -12,6 +12,7 @@
 #include "MuCO/CustomizableObjectInstance.h"
 #include "MuCO/CustomizableObjectSystem.h"
 #include "MuCO/CustomizableSkeletalComponent.h"
+#include "MuCO/UnrealPortabilityHelpers.h"
 #include "UObject/UObjectIterator.h"
 #include "Components/SkeletalMeshComponent.h"
 
@@ -324,7 +325,10 @@ void UCustomizableInstanceLODManagement::UpdateInstanceDistsAndLODs()
 					COI->SetIsBeingUsedByComponentInPlay(true);
 
 					// If it's the local player set max priority
-					const APawn* Pawn = Cast<APawn>(Parent->GetAttachParentActor());
+					const USceneComponent* const AttachParentParentComponent = Parent->GetAttachParent();
+					AActor* ParentParentActor = AttachParentParentComponent ? AttachParentParentComponent->GetOwner() : nullptr;
+
+					const APawn* Pawn = Cast<APawn>(ParentParentActor);
 					if (Pawn && Pawn->IsLocallyControlled())
 					{
 						COI->SetMinSquareDistToPlayer(-1.f);
@@ -339,7 +343,7 @@ void UCustomizableInstanceLODManagement::UpdateInstanceDistsAndLODs()
 					LODTracker.MinLOD = FMath::Min(LODTracker.MinLOD, Parent->bOverrideMinLod ? Parent->MinLodModel : 0);
 
 					// If the parent component have a SkeletalMesh use the RequestedLODLevel of the component as reference to know which LODs mutable should generate.
-					if (Parent->GetSkeletalMeshAsset() && LODTracker.RequestedLODsPerComponent.IsValidIndex(CustomizableSkeletalComponent->ComponentIndex))
+					if (UE_MUTABLE_GETSKELETALMESHASSET(Parent) && LODTracker.RequestedLODsPerComponent.IsValidIndex(CustomizableSkeletalComponent->ComponentIndex))
 					{
 						LODTracker.RequestedLODsPerComponent[CustomizableSkeletalComponent->ComponentIndex] |= 1 << Parent->GetPredictedLODLevel();
 					}
diff --git a/Engine/Plugins/Experimental/Mutable/Source/CustomizableObject/Private/MuCO/CustomizableObject.cpp b/Engine/Plugins/Experimental/Mutable/Source/CustomizableObject/Private/MuCO/CustomizableObject.cpp
--- a/Engine/Plugins/Experimental/Mutable/Source/CustomizableObject/Private/MuCO/CustomizableObject.cpp
+++ b/Engine/Plugins/Experimental/Mutable/Source/CustomizableObject/Private/MuCO/CustomizableObject.cpp
@@ -22,6 +22,7 @@
 #include "MuCO/ICustomizableObjectModule.h"
 #include "MuCO/MutableProjectorTypeUtils.h"
 #include "MuCO/UnrealMutableModelDiskStreamer.h"
+#include "MuCO/UnrealPortabilityHelpers.h"
 #include "MuR/ModelPrivate.h"
 #include "MuR/ParametersPrivate.h"
 #include "PhysicsEngine/PhysicsAsset.h"
@@ -757,7 +758,7 @@ bool UCustomizableObject::ConditionalAutoCompile()
 	else
 	{
 		const FAssetRegistryModule& AssetRegistryModule = FModuleManager::LoadModuleChecked<FAssetRegistryModule>("AssetRegistry");
-		const FAssetData AssetData = AssetRegistryModule.Get().GetAssetByObjectPath(FSoftObjectPath(GetPathName()));
+		const FAssetData AssetData = AssetRegistryModule.Get().GetAssetByObjectPath(UE_MUTABLE_OBJECTPATH(GetPathName()));
 		System->RecompileCustomizableObjectAsync(AssetData, this);
 	}
 
diff --git a/Engine/Plugins/Experimental/Mutable/Source/CustomizableObject/Private/MuCO/CustomizableObjectInstance.cpp b/Engine/Plugins/Experimental/Mutable/Source/CustomizableObject/Private/MuCO/CustomizableObjectInstance.cpp
--- a/Engine/Plugins/Experimental/Mutable/Source/CustomizableObject/Private/MuCO/CustomizableObjectInstance.cpp
+++ b/Engine/Plugins/Experimental/Mutable/Source/CustomizableObject/Private/MuCO/CustomizableObjectInstance.cpp
@@ -22,18 +22,17 @@
 #include "MuCO/CustomizableSkeletalComponent.h"
 #include "MuCO/DefaultImageProvider.h"
 #include "MuCO/UnrealConversionUtils.h"
+#include "MuCO/UnrealPortabilityHelpers.h"
+#include "MuR/OpImagePixelFormat.h"
+#include "MuR/Model.h"
 #include "MuR/OpImagePixelFormat.h"
+#include "MuR/OpImageSwizzle.h"
 #include "PhysicsEngine/PhysicsAsset.h"
 #include "Rendering/Texture2DResource.h"
 #include "RenderingThread.h"
 #include "SkeletalMergingLibrary.h"
 #include "UObject/UObjectIterator.h"
 #include "Widgets/Notifications/SNotificationList.h"
-// Required for engine branch preprocessor defines.
-#include "MuCO/UnrealPortabilityHelpers.h"
-#include "MuR/Model.h"
-#include "MuR/OpImagePixelFormat.h"
-#include "MuR/OpImageSwizzle.h"
 #include "PhysicsEngine/PhysicsConstraintTemplate.h"
 
 #include UE_INLINE_GENERATED_CPP_BY_NAME(CustomizableObjectInstance)
@@ -734,7 +733,7 @@ void UCustomizableInstancePrivateData::PostEditChangePropertyWithoutEditor(USkel
 
 	if (!InSkeletalMesh) return;
 
-	FSkeletalMeshRenderData* Resource = Helper_GetResourceForRendering(InSkeletalMesh);
+	FSkeletalMeshRenderData* Resource = InSkeletalMesh->GetResourceForRendering();
 	if (!Resource) return;
 
 	{
@@ -1033,9 +1032,7 @@ USkeleton* UCustomizableInstancePrivateData::MergeSkeletons(UCustomizableObjectI
 	Exchange(Params.SkeletonsToMerge, ReferencedSkeletons.SkeletonsToMerge);
 
 #if !UE_BUILD_SHIPPING
-#if !MUTABLE_CLEAN_ENGINE_BRANCH
-	Params.bCheckSkeletonsCompatibility = true;
-#endif
+	//Params.bCheckSkeletonsCompatibility = true;
 #endif
 
 	USkeleton* FinalSkeleton = USkeletalMergingLibrary::MergeSkeletons(Params);
@@ -2103,8 +2100,8 @@ bool UCustomizableInstancePrivateData::UpdateSkeletalMesh_PostBeginUpdate0(UCust
 				BuildOrCopyMorphTargetsData(OperationData, SkeletalMesh, OldSkeletalMesh, Public, ComponentIndex);
 				BuildOrCopyClothingData(OperationData, SkeletalMesh, OldSkeletalMesh, Public, ComponentIndex);
 
-				ensure(Helper_GetLODData(SkeletalMesh).Num() > 0);
-				ensure(Helper_GetLODInfoArray(SkeletalMesh).Num() > 0);
+				ensure(SkeletalMesh->GetResourceForRendering()->LODRenderData.Num() > 0);
+				ensure(SkeletalMesh->GetLODInfoArray().Num() > 0);
 			}
 		}
 	}
@@ -3301,7 +3298,7 @@ void UCustomizableInstancePrivateData::InitSkeletalMeshData(const TSharedPtr<FMu
 		FSkeletalMeshRenderData* RenderData = SkeletalMesh->GetResourceForRendering();
 		for (int32 LODIndex = 0; LODIndex <= OperationData->CurrentMaxLOD; ++LODIndex)
 		{
-			RenderData->LODRenderData.Add(new Helper_LODDataType());
+			RenderData->LODRenderData.Add(new FSkeletalMeshLODRenderData());
 
 			const FMutableRefLODData& LODData = RefSkeletalMeshData->LODData[LODIndex];
 
@@ -3714,7 +3711,7 @@ void UCustomizableInstancePrivateData::BuildOrCopyElementData(const TSharedPtr<F
 
 		for (int32 SurfaceIndex = 0; SurfaceIndex < Component.SurfaceCount; ++SurfaceIndex)
 		{
-			new(Helper_GetLODRenderSections(SkeletalMesh, LODIndex)) Helper_SkelMeshRenderSection();
+			new(SkeletalMesh->GetResourceForRendering()->LODRenderData[LODIndex].RenderSections) FSkelMeshRenderSection();
 		}
 	}
 }
@@ -4865,7 +4862,7 @@ bool UCustomizableInstancePrivateData::BuildOrCopyRenderData(const TSharedPtr<FM
 			Component.Mesh,
 			LODIndex);
 
-		FSkeletalMeshLODRenderData& LODModel = Helper_GetLODData(SkeletalMesh)[LODIndex];
+		FSkeletalMeshLODRenderData& LODModel = SkeletalMesh->GetResourceForRendering()->LODRenderData[LODIndex];
 
 		if (LODModel.DoesVertexBufferUse16BitBoneIndex() && !UCustomizableObjectSystem::GetInstance()->IsSupport16BitBoneIndexEnabled())
 		{
@@ -5472,9 +5469,9 @@ void UCustomizableInstancePrivateData::BuildMaterials(const TSharedPtr<FMutableO
 		int32 LastValidLODIndex = OperationData->CurrentMaxLOD;
 		for (int32 LODIndex = OperationData->CurrentMaxLOD; LODIndex >= FirstLODAvailable; --LODIndex)
 		{
-			if (Helper_GetLODInfoArray(SkeletalMesh).Num() != 0)
+			if (SkeletalMesh->GetLODInfoArray().Num() != 0)
 			{
-				Helper_GetLODInfoArray(SkeletalMesh)[LODIndex].LODMaterialMap.Reset();
+				SkeletalMesh->GetLODInfoArray()[LODIndex].LODMaterialMap.Reset();
 			}
 
 			if (LODIndex >= OperationData->CurrentMinLOD && OperationData->InstanceUpdateData.Components[OperationData->InstanceUpdateData.LODs[LODIndex].FirstComponent + ComponentIndex].bGenerated)
@@ -5499,8 +5496,8 @@ void UCustomizableInstancePrivateData::BuildMaterials(const TSharedPtr<FMutableO
 				if (LastValidLODIndex != LODIndex)
 				{
 					// Copy information from the LastValidLODIndex
-					Helper_GetLODInfoArray(SkeletalMesh)[LODIndex].LODMaterialMap = Helper_GetLODInfoArray(SkeletalMesh)[LastValidLODIndex].LODMaterialMap;
-					Helper_GetLODRenderSections(SkeletalMesh, LODIndex)[SurfaceIndex].MaterialIndex = Helper_GetLODRenderSections(SkeletalMesh, LastValidLODIndex)[SurfaceIndex].MaterialIndex;
+					SkeletalMesh->GetLODInfoArray()[LODIndex].LODMaterialMap = SkeletalMesh->GetLODInfoArray()[LastValidLODIndex].LODMaterialMap;
+					SkeletalMesh->GetResourceForRendering()->LODRenderData[LODIndex].RenderSections[SurfaceIndex].MaterialIndex = SkeletalMesh->GetResourceForRendering()->LODRenderData[LastValidLODIndex].RenderSections[SurfaceIndex].MaterialIndex;
 					continue;
 				}
 
@@ -5811,8 +5808,8 @@ void UCustomizableInstancePrivateData::BuildMaterials(const TSharedPtr<FMutableO
 				{
 					int32 MatIndex = (*FoundMaterialPlaceholder)->MatIndex;
 					check(MatIndex >= 0);
-					int32 LODMaterialIndex = Helper_GetLODInfoArray(SkeletalMesh)[LODIndex].LODMaterialMap.Add(MatIndex);
-					Helper_GetLODRenderSections(SkeletalMesh, LODIndex)[SurfaceIndex].MaterialIndex = LODMaterialIndex;
+					int32 LODMaterialIndex = SkeletalMesh->GetLODInfoArray()[LODIndex].LODMaterialMap.Add(MatIndex);
+					SkeletalMesh->GetResourceForRendering()->LODRenderData[LODIndex].RenderSections[SurfaceIndex].MaterialIndex = LODMaterialIndex;
 				}
 				else
 				{
@@ -5847,8 +5844,8 @@ void UCustomizableInstancePrivateData::BuildMaterials(const TSharedPtr<FMutableO
 							SkeletalMaterial.MaterialSlotName = CustomizableObject->ReferencedMaterialSlotNames[Surface.MaterialIndex];
 							SetMeshUVChannelDensity(SkeletalMaterial.UVChannelData, RefSkeletalMeshData->Settings.DefaultUVChannelDensity);
 
-							const int32 LODMaterialIndex = Helper_GetLODInfoArray(SkeletalMesh)[LODIndex].LODMaterialMap.Add(MatIndex);
-							Helper_GetLODRenderSections(SkeletalMesh, LODIndex)[SurfaceIndex].MaterialIndex = LODMaterialIndex;
+							const int32 LODMaterialIndex = SkeletalMesh->GetLODInfoArray()[LODIndex].LODMaterialMap.Add(MatIndex);
+							SkeletalMesh->GetResourceForRendering()->LODRenderData[LODIndex].RenderSections[SurfaceIndex].MaterialIndex = LODMaterialIndex;
 						}
 
 						if (MaterialInstance)
diff --git a/Engine/Plugins/Experimental/Mutable/Source/CustomizableObject/Private/MuCO/CustomizableObjectSystem.cpp b/Engine/Plugins/Experimental/Mutable/Source/CustomizableObject/Private/MuCO/CustomizableObjectSystem.cpp
--- a/Engine/Plugins/Experimental/Mutable/Source/CustomizableObject/Private/MuCO/CustomizableObjectSystem.cpp
+++ b/Engine/Plugins/Experimental/Mutable/Source/CustomizableObject/Private/MuCO/CustomizableObjectSystem.cpp
@@ -23,6 +23,7 @@
 #include "MuCO/UnrealBakeHelpers.h"
 #include "MuCO/UnrealMutableImageProvider.h"
 #include "MuCO/UnrealMutableModelDiskStreamer.h"
+#include "MuCO/UnrealPortabilityHelpers.h"
 #include "MuR/Model.h"
 #include "MuR/Settings.h"
 #include "TextureResource.h"
@@ -2888,7 +2889,7 @@ void UCustomizableObjectSystem::OnPreBeginPIE(const bool bIsSimulatingInEditor)
 	
 	// Find root customizable objects
 	FARFilter AssetRegistryFilter;
-	AssetRegistryFilter.ClassPaths.Add(FTopLevelAssetPath(TEXT("/Script/CustomizableObject"), TEXT("CustomizableObject")));
+	UE_MUTABLE_GET_CLASSPATHS(AssetRegistryFilter).Add(UE_MUTABLE_TOPLEVELASSETPATH(TEXT("/Script/CustomizableObject"), TEXT("CustomizableObject")));
 	AssetRegistryFilter.TagsAndValues.Add(FName("IsRoot"), FString::FromInt(1));
 
 	TArray<FAssetData> OutAssets;
diff --git a/Engine/Plugins/Experimental/Mutable/Source/CustomizableObject/Private/MuCO/CustomizableSkeletalComponent.cpp b/Engine/Plugins/Experimental/Mutable/Source/CustomizableObject/Private/MuCO/CustomizableSkeletalComponent.cpp
--- a/Engine/Plugins/Experimental/Mutable/Source/CustomizableObject/Private/MuCO/CustomizableSkeletalComponent.cpp
+++ b/Engine/Plugins/Experimental/Mutable/Source/CustomizableObject/Private/MuCO/CustomizableSkeletalComponent.cpp
@@ -12,6 +12,7 @@
 #include "MuCO/CustomizableInstancePrivateData.h"
 #include "Rendering/SkeletalMeshRenderData.h"
 #include "UObject/ObjectSaveContext.h"
+#include "MuCO/UnrealPortabilityHelpers.h"
 
 #include UE_INLINE_GENERATED_CPP_BY_NAME(CustomizableSkeletalComponent)
 
@@ -55,7 +56,7 @@ void UCustomizableSkeletalComponent::SetSkeletalMesh(USkeletalMesh* SkeletalMesh
 		}
 #endif
 
-		if(Parent->GetSkinnedAsset() == SkeletalMesh)
+		if(UE_MUTABLE_GETSKINNEDASSET(Parent) == SkeletalMesh)
 		{
 			Parent->RecreateRenderState_Concurrent();
 		}
@@ -113,7 +114,7 @@ USkeletalMesh* UCustomizableSkeletalComponent::GetAttachedSkeletalMesh() const
 
 	if (Parent)
 	{
-		return Parent->GetSkeletalMeshAsset();
+		return UE_MUTABLE_GETSKELETALMESHASSET(Parent);
 	}
 
 	return nullptr;
diff --git a/Engine/Plugins/Experimental/Mutable/Source/CustomizableObject/Private/MuCO/LogBenchmarkUtil.cpp b/Engine/Plugins/Experimental/Mutable/Source/CustomizableObject/Private/MuCO/LogBenchmarkUtil.cpp
--- a/Engine/Plugins/Experimental/Mutable/Source/CustomizableObject/Private/MuCO/LogBenchmarkUtil.cpp
+++ b/Engine/Plugins/Experimental/Mutable/Source/CustomizableObject/Private/MuCO/LogBenchmarkUtil.cpp
@@ -19,7 +19,7 @@ void LogBenchmarkUtil::shutdownAndSaveResults()
 {
 	if (bLoggingActive)
 	{
-		FString SaveDirectory = Helper_GetSavedDir() + "/Logs";
+		FString SaveDirectory = FPaths::ConvertRelativePathToFull(FPaths::ProjectSavedDir()) + "/Logs";
 		FString FileName = FString("BenchmarkResult.txt");
 		FString AbsoluteFilePath = SaveDirectory + "/" + FileName;
 		IPlatformFile& PlatformFile = FPlatformFileManager::Get().GetPlatformFile();
diff --git a/Engine/Plugins/Experimental/Mutable/Source/CustomizableObject/Private/MuCO/UnrealConversionUtils.cpp b/Engine/Plugins/Experimental/Mutable/Source/CustomizableObject/Private/MuCO/UnrealConversionUtils.cpp
--- a/Engine/Plugins/Experimental/Mutable/Source/CustomizableObject/Private/MuCO/UnrealConversionUtils.cpp
+++ b/Engine/Plugins/Experimental/Mutable/Source/CustomizableObject/Private/MuCO/UnrealConversionUtils.cpp
@@ -4,6 +4,7 @@
 
 #include "MuCO/CustomizableObjectInstance.h"
 #include "MuCO/UnrealPortabilityHelpers.h"
+#include "Engine/SkeletalMesh.h"
 
 class USkeleton;
 
@@ -169,7 +170,7 @@ namespace UnrealConversionUtils
 			int32 FirstVertex;
 			int32 VertexCount;
 			InMutableMesh->GetSurface(SurfaceIndex, &FirstVertex, &VertexCount, &FirstIndex, &IndexCount);
-			FSkelMeshRenderSection& Section = Helper_GetLODRenderSections(OutSkeletalMesh, MeshLODIndex)[SurfaceIndex];
+			FSkelMeshRenderSection& Section = OutSkeletalMesh->GetResourceForRendering()->LODRenderData[MeshLODIndex].RenderSections[SurfaceIndex];
 
 			Section.DuplicatedVerticesBuffer.Init(1, TMap<int, TArray<int32>>());
 
@@ -197,7 +198,7 @@ namespace UnrealConversionUtils
 	{
 		MUTABLE_CPUPROFILER_SCOPE(CopyMutableVertexBuffers);
 
-		FSkeletalMeshLODRenderData& LODModel = Helper_GetLODData(SkeletalMesh)[MeshLODIndex];
+		FSkeletalMeshLODRenderData& LODModel = SkeletalMesh->GetResourceForRendering()->LODRenderData[MeshLODIndex];
 		const mu::FMeshBufferSet& MutableMeshVertexBuffers = MutableMesh->GetVertexBuffers();
 		const int32 NumVertices = MutableMeshVertexBuffers.GetElementCount();
 
@@ -208,7 +209,7 @@ namespace UnrealConversionUtils
 		const bool bHasVertexColors = EnumHasAllFlags(BuildFlags, ESkeletalMeshVertexFlags::HasVertexColors);
 		const int NumTexCoords = MutableMeshVertexBuffers.GetBufferChannelCount(MUTABLE_VERTEXBUFFER_TEXCOORDS);
 
-		const bool bNeedsCPUAccess = Helper_GetLODInfoArray(SkeletalMesh)[MeshLODIndex].bAllowCPUAccess;
+		const bool bNeedsCPUAccess = SkeletalMesh->GetLODInfoArray()[MeshLODIndex].bAllowCPUAccess;
 
 		FStaticMeshVertexBuffers_InitWithMutableData(
 			LODModel.StaticVertexBuffers,
@@ -448,21 +449,21 @@ namespace UnrealConversionUtils
 			 return;
 		 }
 
-		 const FSkeletalMeshLODRenderData& SrcLODModel = Helper_GetLODData(SrcSkeletalMesh)[SrcLODIndex];
-		 FSkeletalMeshLODRenderData& DestLODModel = Helper_GetLODData(DestSkeletalMesh)[DestLODIndex];
+		 const FSkeletalMeshLODRenderData& SrcLODModel = SrcSkeletalMesh->GetResourceForRendering()->LODRenderData[SrcLODIndex];
+		 FSkeletalMeshLODRenderData& DestLODModel = DestSkeletalMesh->GetResourceForRendering()->LODRenderData[DestLODIndex];
 
 		 // Copying render sections
 		 {
-			 const FSkeletalMeshLODInfo& SrcLODInfo = Helper_GetLODInfoArray(SrcSkeletalMesh)[SrcLODIndex];
-			 FSkeletalMeshLODInfo& DestLODInfo = Helper_GetLODInfoArray(DestSkeletalMesh)[DestLODIndex];
+			 const FSkeletalMeshLODInfo& SrcLODInfo = SrcSkeletalMesh->GetLODInfoArray()[SrcLODIndex];
+			 FSkeletalMeshLODInfo& DestLODInfo = DestSkeletalMesh->GetLODInfoArray()[DestLODIndex];
 
 			 DestLODInfo.LODMaterialMap = SrcLODInfo.LODMaterialMap;
 
 			 const int32 SurfaceCount = SrcLODModel.RenderSections.Num();
 			 for (int32 SurfaceIndex = 0; SurfaceIndex < SurfaceCount; ++SurfaceIndex)
 			 {
-				 const FSkelMeshRenderSection& SrcSection = Helper_GetLODRenderSections(SrcSkeletalMesh, SrcLODIndex)[SurfaceIndex];
-				 FSkelMeshRenderSection* DestSection = new(Helper_GetLODRenderSections(DestSkeletalMesh, DestLODIndex)) Helper_SkelMeshRenderSection();
+				 const FSkelMeshRenderSection& SrcSection = SrcSkeletalMesh->GetResourceForRendering()->LODRenderData[SrcLODIndex].RenderSections[SurfaceIndex];
+				 FSkelMeshRenderSection* DestSection = new(DestSkeletalMesh->GetResourceForRendering()->LODRenderData[DestLODIndex].RenderSections) FSkelMeshRenderSection();
 
 				 DestSection->DuplicatedVerticesBuffer.Init(1, TMap<int, TArray<int32>>());
 				 DestSection->bDisabled = SrcSection.bDisabled;
@@ -485,7 +486,7 @@ namespace UnrealConversionUtils
 		 const int32 NumVertices = SrcStaticVertexBuffer.PositionVertexBuffer.GetNumVertices();
 		 const int32 NumTexCoords = SrcStaticVertexBuffer.StaticMeshVertexBuffer.GetNumTexCoords();
 
-		 bool bAllowCPUAccess = Helper_GetLODInfoArray(SrcSkeletalMesh)[SrcLODIndex].bAllowCPUAccess;
+		 bool bAllowCPUAccess = SrcSkeletalMesh->GetLODInfoArray()[SrcLODIndex].bAllowCPUAccess;
 
 		 // Copying Static Vertex Buffers
 		 {
@@ -562,7 +563,7 @@ namespace UnrealConversionUtils
 				 int32 IndexCount = SrcLODModel.MultiSizeIndexContainer.GetIndexBuffer()->Num();
 				 int32 ElementSize = SrcLODModel.MultiSizeIndexContainer.GetDataTypeSize();
 
-				 const void* Data = Helper_GetLODData(SrcSkeletalMesh)[SrcLODIndex].MultiSizeIndexContainer.GetIndexBuffer()->GetPointerTo(0);
+				 const void* Data = SrcSkeletalMesh->GetResourceForRendering()->LODRenderData[SrcLODIndex].MultiSizeIndexContainer.GetIndexBuffer()->GetPointerTo(0);
 
 				 DestLODModel.MultiSizeIndexContainer.CreateIndexBuffer(ElementSize);
 				 DestLODModel.MultiSizeIndexContainer.GetIndexBuffer()->Insert(0, IndexCount);
diff --git a/Engine/Plugins/Experimental/Mutable/Source/CustomizableObject/Private/MuCO/UnrealMutableImageProvider.cpp b/Engine/Plugins/Experimental/Mutable/Source/CustomizableObject/Private/MuCO/UnrealMutableImageProvider.cpp
--- a/Engine/Plugins/Experimental/Mutable/Source/CustomizableObject/Private/MuCO/UnrealMutableImageProvider.cpp
+++ b/Engine/Plugins/Experimental/Mutable/Source/CustomizableObject/Private/MuCO/UnrealMutableImageProvider.cpp
@@ -358,7 +358,7 @@ mu::EImageFormat GetMutablePixelFormat(EPixelFormat InTextureFormat)
 			else
 			{
 				// Bulk data already loaded
-				const void* Data = BulkData.IsUnlocked() ? BulkData.LockReadOnly() : nullptr; // TODO: Retry if it fails?
+				const void* Data = (!BulkData.IsLocked()) ? BulkData.LockReadOnly() : nullptr; // TODO: Retry if it fails?
 				
 				if (Data)
 				{
diff --git a/Engine/Plugins/Experimental/Mutable/Source/CustomizableObject/Public/MuCO/UnrealPortabilityHelpers.h b/Engine/Plugins/Experimental/Mutable/Source/CustomizableObject/Public/MuCO/UnrealPortabilityHelpers.h
--- a/Engine/Plugins/Experimental/Mutable/Source/CustomizableObject/Public/MuCO/UnrealPortabilityHelpers.h
+++ b/Engine/Plugins/Experimental/Mutable/Source/CustomizableObject/Public/MuCO/UnrealPortabilityHelpers.h
@@ -3,109 +3,42 @@
 #pragma once
 
 
-#include "Engine/SkeletalMesh.h"
-#include "Engine/SkinnedAsset.h"
-#include "Engine/SkinnedAssetCommon.h"
-#include "Misc/Paths.h"
-#include "Rendering/SkeletalMeshRenderData.h"
-
-struct FBoneReference;
+#include "Runtime/Launch/Resources/Version.h"
 
-// Main engine branch should have 1 here. 
-// Projects with API modifications may set it 0.
-#define MUTABLE_CLEAN_ENGINE_BRANCH		1
+/** Helpers to ease portability across unreal engine versions */
 
-//-------------------------------------------------------------------------------------------------
-// Helpers to ease portability across unreal engine versions
-//-------------------------------------------------------------------------------------------------
+#if ENGINE_MAJOR_VERSION==5 && ENGINE_MINOR_VERSION>=1
 
-typedef FSkeletalMeshLODRenderData Helper_LODDataType;
-typedef FSkelMeshRenderSection Helper_SkelMeshRenderSection;
+#include "Engine/SkinnedAssetCommon.h"
+#include "Rendering/SkinWeightVertexBuffer.h"
+#include "Rendering/SkeletalMeshRenderData.h"
 
-inline FSkeletalMeshRenderData* Helper_GetResourceForRendering(const USkeletalMesh* Mesh)
-{
-	return Mesh->GetResourceForRendering();
-}
+#define UE_MUTABLE_GET_CLASS_PATHNAME(X)		X->GetClassPathName()
+#define UE_MUTABLE_GET_CLASSPATHS(X)			X.ClassPaths
+#define UE_MUTABLE_OBJECTPATH(X)				FSoftObjectPath(X)
+#define UE_MUTABLE_TOPLEVELASSETPATH(X,Y)		FTopLevelAssetPath(X, Y)
+#define UE_MUTABLE_ASSETCLASS(X)				X.AssetClassPath
+#define UE_MUTABLE_GETOBJECTPATH(X)				X.GetObjectPathString()
 
-inline TIndirectArray<FSkeletalMeshLODRenderData>* Helper_GetLODDataPtr(FSkeletalMeshRenderData* Resource)
-{
-	return &(Resource->LODRenderData);
-}
+#define UE_MUTABLE_GETSKINNEDASSET(X)			X->GetSkinnedAsset()
+#define UE_MUTABLE_SETSKINNEDASSET(X,Y)			X->SetSkinnedAsset(Y)
+#define UE_MUTABLE_GETSKELETALMESHASSET(X)		X->GetSkeletalMeshAsset()
 
-inline TIndirectArray<FSkeletalMeshLODRenderData>& Helper_GetLODData(FSkeletalMeshRenderData* Resource)
-{
-	return Resource->LODRenderData;
-}
+#elif ENGINE_MAJOR_VERSION==5 && ENGINE_MINOR_VERSION==0
 
-inline const TIndirectArray<FSkeletalMeshLODRenderData>& Helper_GetLODData(const FSkeletalMeshRenderData* Resource)
-{
-	return Resource->LODRenderData;
-}
+#define USkinnedAsset USkeletalMesh
 
-inline TIndirectArray<FSkeletalMeshLODRenderData>& Helper_GetLODData(const USkinnedAsset* Mesh)
-{
-	return Mesh->GetResourceForRendering()->LODRenderData;
-}
+#define UE_MUTABLE_GET_CLASS_PATHNAME(X)		X->GetFName()
+#define UE_MUTABLE_GET_CLASSPATHS(X)			X.ClassNames
+#define UE_MUTABLE_OBJECTPATH(X)				*X
+#define UE_MUTABLE_TOPLEVELASSETPATH(X,Y)		(Y)
+#define UE_MUTABLE_ASSETCLASS(X)				X.AssetClass
+#define UE_MUTABLE_GETOBJECTPATH(X)				X.ObjectPath.ToString()
 
-inline TArray<FSkelMeshRenderSection>& Helper_GetLODRenderSections(const USkinnedAsset* Mesh, int LODIndex)
-{
-	return  Mesh->GetResourceForRendering()->LODRenderData[LODIndex].RenderSections;
-}
+#define UE_MUTABLE_GETSKINNEDASSET(X)			X->SkeletalMesh
+#define UE_MUTABLE_SETSKINNEDASSET(X,Y)			X->SetSkeletalMesh(Y)
+#define UE_MUTABLE_GETSKELETALMESHASSET(X)		X->SkeletalMesh
 
-#if WITH_EDITORONLY_DATA
-inline FSkeletalMeshModel* Helper_GetImportedModel(const USkinnedAsset* Mesh)
-{
-	return Mesh->GetImportedModel();
-}
 #endif
 
-inline int32 Helper_LODGetTotalFaces(const FSkeletalMeshLODRenderData* Resource)
-{
-	return Resource->GetTotalFaces();
-}
-
-inline uint32 Helper_LODGetNumTexCoords(const FSkeletalMeshLODRenderData* Resource)
-{
-	return Resource->GetNumTexCoords();
-}
-
-inline const TArray<FSkelMeshRenderSection>* Helper_LODGetSections(const FSkeletalMeshLODRenderData* Resource)
-{
-	return &(Resource->RenderSections);
-}
-
-inline TArray<struct FSkeletalMeshLODInfo>& Helper_GetLODInfoArray(USkinnedAsset* Mesh)
-{
-	return Mesh->GetLODInfoArray();
-}
 
-inline const TArray<struct FSkeletalMeshLODInfo>& Helper_GetLODInfoArray(const USkinnedAsset* Mesh)
-{
-	return Mesh->GetLODInfoArray();
-}
-
-inline int Helper_GetLODNum(const USkinnedAsset* Mesh)
-{
-	return Mesh->GetLODNum();
-}
-
-inline const TArray<FBoneReference>& Helper_GetLODInfoBonesToRemove(const USkinnedAsset* Mesh, int LODIndex)
-{
-	return Mesh->GetLODInfo(LODIndex)->BonesToRemove;
-}
-
-inline TArray<FBoneReference>& Helper_GetLODInfoBonesToRemove(USkinnedAsset* Mesh, int LODIndex)
-{
-	return Mesh->GetLODInfo(LODIndex)->BonesToRemove;
-}
-
-inline FString Helper_GetSavedDir()
-{
-	return FPaths::ConvertRelativePathToFull(FPaths::ProjectSavedDir());
-}
-
-
-#if UE_ENABLE_INCLUDE_ORDER_DEPRECATED_IN_5_2
-#include "Rendering/SkeletalMeshModel.h"
-#include "Runtime/Launch/Resources/Version.h"
-#endif
diff --git a/Engine/Plugins/Experimental/Mutable/Source/CustomizableObjectEditor/Classes/MuCOE/Nodes/CustomizableObjectNodeMeshMorph.h b/Engine/Plugins/Experimental/Mutable/Source/CustomizableObjectEditor/Classes/MuCOE/Nodes/CustomizableObjectNodeMeshMorph.h
--- a/Engine/Plugins/Experimental/Mutable/Source/CustomizableObjectEditor/Classes/MuCOE/Nodes/CustomizableObjectNodeMeshMorph.h
+++ b/Engine/Plugins/Experimental/Mutable/Source/CustomizableObjectEditor/Classes/MuCOE/Nodes/CustomizableObjectNodeMeshMorph.h
@@ -3,6 +3,7 @@
 #pragma once
 
 #include "MuCOE/Nodes/CustomizableObjectNode.h"
+#include "MuCOE/Nodes/CustomizableObjectNodeMeshReshapeCommon.h"
 
 #include "CustomizableObjectNodeMeshMorph.generated.h"
 
diff --git a/Engine/Plugins/Experimental/Mutable/Source/CustomizableObjectEditor/Private/MuCOE/CustomizableObjectCustomSettings.cpp b/Engine/Plugins/Experimental/Mutable/Source/CustomizableObjectEditor/Private/MuCOE/CustomizableObjectCustomSettings.cpp
--- a/Engine/Plugins/Experimental/Mutable/Source/CustomizableObjectEditor/Private/MuCOE/CustomizableObjectCustomSettings.cpp
+++ b/Engine/Plugins/Experimental/Mutable/Source/CustomizableObjectEditor/Private/MuCOE/CustomizableObjectCustomSettings.cpp
@@ -9,6 +9,7 @@
 #include "MuCOE/CustomizableObjectEditorViewportClient.h"
 #include "MuCOE/CustomizableObjectEditorViewportLights.h"
 #include "MuCOE/SCustomizableObjectEditorViewport.h"
+#include "MuCOE/UnrealEditorPortabilityHelpers.h"
 #include "PropertyCustomizationHelpers.h"
 #include "Subsystems/ImportSubsystem.h"
 #include "Widgets/Colors/SColorBlock.h"
@@ -113,7 +114,7 @@ void SCustomizableObjectCustomSettings::Construct(const FArguments& InArgs, UCus
 				[
 					SAssignNew(LightComboButton, SComboButton)
 					.ButtonStyle(FAppStyle::Get(), "PropertyEditor.AssetComboStyle")
-					.ForegroundColor(FAppStyle::GetColor("PropertyEditor.AssetName.ColorAndOpacity"))
+					.ForegroundColor(UE_MUTABLE_GET_COLOR("PropertyEditor.AssetName.ColorAndOpacity"))
 					.OnGetMenuContent(this, &SCustomizableObjectCustomSettings::GetLightComboButtonContent)
 					.ContentPadding(2.0f)
 					.ButtonContent()
@@ -161,7 +162,7 @@ void SCustomizableObjectCustomSettings::Construct(const FArguments& InArgs, UCus
 			.Padding(0.f, 2.f)
 			[
 				SNew(SBorder)
-				.BorderImage(FAppStyle::GetBrush("ToolPanel.GroupBorder"))
+				.BorderImage(UE_MUTABLE_GET_BRUSH("ToolPanel.GroupBorder"))
 				.Padding(6.f)
 				[
 					SNew(SVerticalBox)
diff --git a/Engine/Plugins/Experimental/Mutable/Source/CustomizableObjectEditor/Private/MuCOE/CustomizableObjectEditor.cpp b/Engine/Plugins/Experimental/Mutable/Source/CustomizableObjectEditor/Private/MuCOE/CustomizableObjectEditor.cpp
--- a/Engine/Plugins/Experimental/Mutable/Source/CustomizableObjectEditor/Private/MuCOE/CustomizableObjectEditor.cpp
+++ b/Engine/Plugins/Experimental/Mutable/Source/CustomizableObjectEditor/Private/MuCOE/CustomizableObjectEditor.cpp
@@ -19,6 +19,7 @@
 #include "MuCO/CustomizableObjectInstance.h"
 #include "MuCO/CustomizableObjectSystem.h"
 #include "MuCO/CustomizableSkeletalComponent.h"
+#include "MuCO/UnrealPortabilityHelpers.h"
 #include "MuCOE/CustomizableObjectBakeHelpers.h"
 #include "MuCOE/CustomizableObjectCustomSettings.h"
 #include "MuCOE/CustomizableObjectEditorActions.h"
@@ -988,7 +989,7 @@ void FCustomizableObjectEditor::RefreshMaterialNodesInAllChildren()
 	// and adding the group name for filtering only for those COs of the selected group
 	FAssetRegistryModule& AssetRegistryModule = FModuleManager::LoadModuleChecked<FAssetRegistryModule>("AssetRegistry");
 	TArray<FAssetData> OutAssetData;
-	AssetRegistryModule.Get().GetAssetsByClass(FTopLevelAssetPath(TEXT("/Script/CustomizableObject"), TEXT("CustomizableObject")), OutAssetData);
+	AssetRegistryModule.Get().GetAssetsByClass(UE_MUTABLE_TOPLEVELASSETPATH(TEXT("/Script/CustomizableObject"), TEXT("CustomizableObject")), OutAssetData);
 
 	for (auto Itr = OutAssetData.CreateIterator(); Itr; ++Itr)
 	{
diff --git a/Engine/Plugins/Experimental/Mutable/Source/CustomizableObjectEditor/Private/MuCOE/CustomizableObjectEditorViewportClient.cpp b/Engine/Plugins/Experimental/Mutable/Source/CustomizableObjectEditor/Private/MuCOE/CustomizableObjectEditorViewportClient.cpp
--- a/Engine/Plugins/Experimental/Mutable/Source/CustomizableObjectEditor/Private/MuCOE/CustomizableObjectEditorViewportClient.cpp
+++ b/Engine/Plugins/Experimental/Mutable/Source/CustomizableObjectEditor/Private/MuCOE/CustomizableObjectEditorViewportClient.cpp
@@ -521,7 +521,7 @@ FCustomizableObjectEditorViewportClient::FCustomizableObjectEditorViewportClient
 void FCustomizableObjectEditorViewportClient::UpdateCameraSetup()
 {
 	static FRotator CustomOrbitRotation(-33.75, -135, 0);
-	if ( (SkeletalMeshComponents.Num() && SkeletalMeshComponents[0].IsValid() && SkeletalMeshComponents[0]->GetSkinnedAsset())
+	if ( (SkeletalMeshComponents.Num() && SkeletalMeshComponents[0].IsValid() && UE_MUTABLE_GETSKINNEDASSET(SkeletalMeshComponents[0]))
 		||
 		(StaticMeshComponent.IsValid() && StaticMeshComponent->GetStaticMesh()) )
 	{
@@ -737,11 +737,11 @@ void FCustomizableObjectEditorViewportClient::Draw(const FSceneView* View, FPrim
 
 		// Start Plane
 		DrawDirectionalArrow(PDI, PlaneMatrix, FColor::Red, MorphLength, MorphLength * 0.1f, 0, 0.1f);
-		DrawBox(PDI, PlaneMatrix, FVector(0.01f, PlaneRadius1, PlaneRadius1), Helper_GetMaterialProxy(ClipMorphMaterial), 0);
+		DrawBox(PDI, PlaneMatrix, FVector(0.01f, PlaneRadius1, PlaneRadius1), ClipMorphMaterial->GetRenderProxy(), 0);
 
 		// End Plane + Ellipse
 		PlaneMatrix.SetOrigin(ClipMorphOrigin + ClipMorphOffset + ClipMorphNormal * MorphLength);
-		DrawBox(PDI, PlaneMatrix, FVector(0.01f, PlaneRadius2, PlaneRadius2), Helper_GetMaterialProxy(ClipMorphMaterial), 0);
+		DrawBox(PDI, PlaneMatrix, FVector(0.01f, PlaneRadius2, PlaneRadius2), ClipMorphMaterial->GetRenderProxy(), 0);
 		DrawEllipse(PDI, ClipMorphOrigin + ClipMorphOffset + ClipMorphNormal * MorphLength, ClipMorphXAxis, ClipMorphYAxis, FColor::Red, Radius1, Radius2, 15, 1, 0.f, 0, false);
 	}
 
@@ -798,8 +798,8 @@ void FCustomizableObjectEditorViewportClient::Draw(const FSceneView* View, FPrim
 				FMatrix Mat1 = Mat;
 				Mat0.SetOrigin(Location0);
 				Mat1.SetOrigin(Location1);
-				DrawCylinderArc(PDI, Mat0, FVector(0.0f, 0.0f, 0.0f), FVector(0, 1, 0), FVector(0, 0, 1), FVector(1, 0, 0),  CylinderRadius, CylinderHalfHeight * 0.1f, 16, Helper_GetMaterialProxy(TransparentPlaneMaterialXY), SDPG_World, FColor(255, 85, 0, 192), GizmoProxy.Value.Angle);
-				DrawCylinderArc(PDI, Mat1, FVector(0.0f, 0.0f, 0.0f), FVector(0, 1, 0), FVector(0, 0, 1), FVector(1, 0, 0), CylinderRadius, CylinderHalfHeight * 0.1f, 16, Helper_GetMaterialProxy(TransparentPlaneMaterialXY), SDPG_World, FColor(255, 85, 0, 192), GizmoProxy.Value.Angle);
+				DrawCylinderArc(PDI, Mat0, FVector(0.0f, 0.0f, 0.0f), FVector(0, 1, 0), FVector(0, 0, 1), FVector(1, 0, 0),  CylinderRadius, CylinderHalfHeight * 0.1f, 16, TransparentPlaneMaterialXY->GetRenderProxy(), SDPG_World, FColor(255, 85, 0, 192), GizmoProxy.Value.Angle);
+				DrawCylinderArc(PDI, Mat1, FVector(0.0f, 0.0f, 0.0f), FVector(0, 1, 0), FVector(0, 0, 1), FVector(1, 0, 0), CylinderRadius, CylinderHalfHeight * 0.1f, 16, TransparentPlaneMaterialXY->GetRenderProxy(), SDPG_World, FColor(255, 85, 0, 192), GizmoProxy.Value.Angle);
 				break;
 			}
 			case ECustomizableObjectProjectorType::Wrapping:
@@ -839,7 +839,7 @@ void FCustomizableObjectEditorViewportClient::Draw(FViewport* InViewport, FCanva
 	// Defensive check to avoid unreal crashing inside render if the mesh is degenereated
 	for (TWeakObjectPtr<USkeletalMeshComponent> SkeletalMeshComponent : SkeletalMeshComponents)
 	{
-		if (SkeletalMeshComponent.IsValid() && SkeletalMeshComponent->GetSkinnedAsset() && Helper_GetLODInfoArray(SkeletalMeshComponent->GetSkinnedAsset()).Num() == 0)
+		if (SkeletalMeshComponent.IsValid() && UE_MUTABLE_GETSKINNEDASSET(SkeletalMeshComponent) && UE_MUTABLE_GETSKINNEDASSET(SkeletalMeshComponent)->GetLODInfoArray().Num() == 0)
 		{
 			SkeletalMeshComponent->SetSkeletalMesh(nullptr);
 		}
@@ -1072,7 +1072,7 @@ void FCustomizableObjectEditorViewportClient::DrawUVs(FViewport* InViewport, FCa
 
 		for (TWeakObjectPtr<USkeletalMeshComponent> SkeletalMeshComponent : SkeletalMeshComponents)
 		{
-			if (!SkeletalMeshComponent.IsValid() || !SkeletalMeshComponent->GetSkinnedAsset() || CurrentComponentIndex != ComponentIndex)
+			if (!SkeletalMeshComponent.IsValid() || !UE_MUTABLE_GETSKINNEDASSET(SkeletalMeshComponent) || CurrentComponentIndex != ComponentIndex)
 			{
 				CurrentComponentIndex++;
 				continue;
@@ -1080,7 +1080,7 @@ void FCustomizableObjectEditorViewportClient::DrawUVs(FViewport* InViewport, FCa
 
 			bool bFoundMaterial = false;
 
-			const FSkeletalMeshRenderData* MeshRes = SkeletalMeshComponent->GetSkinnedAsset()->GetResourceForRendering();
+			const FSkeletalMeshRenderData* MeshRes = UE_MUTABLE_GETSKINNEDASSET(SkeletalMeshComponent)->GetResourceForRendering();
 			if (UVChannel < (int32)MeshRes->LODRenderData[LODLevel].GetNumTexCoords())
 			{
 				// Find material index from name
@@ -1543,9 +1543,9 @@ void FCustomizableObjectEditorViewportClient::ResetCamera()
 	float MaxSphereRadius = 0.0f;
 	for (const TWeakObjectPtr<UDebugSkelMeshComponent>& SkeletalMeshComponent : SkeletalMeshComponents)
 	{
-		if (SkeletalMeshComponent->GetSkinnedAsset())
+		if (UE_MUTABLE_GETSKINNEDASSET(SkeletalMeshComponent))
 		{
-			MaxSphereRadius = FMath::Max(MaxSphereRadius, SkeletalMeshComponent->GetSkinnedAsset()->GetBounds().SphereRadius);
+			MaxSphereRadius = FMath::Max(MaxSphereRadius, UE_MUTABLE_GETSKINNEDASSET(SkeletalMeshComponent)->GetBounds().SphereRadius);
 		}
 	}
 	
@@ -2027,8 +2027,8 @@ void FCustomizableObjectEditorViewportClient::SetAnimation(UAnimationAsset* Anim
 	for (TWeakObjectPtr<USkeletalMeshComponent> SkeletalMeshComponent : SkeletalMeshComponents)
 	{
 		if (SkeletalMeshComponent.IsValid() && Animation != nullptr
-			&& SkeletalMeshComponent->GetSkinnedAsset() != nullptr
-			&& SkeletalMeshComponent->GetSkinnedAsset()->GetSkeleton()->IsCompatibleForEditor(Animation->GetSkeleton())
+			&& UE_MUTABLE_GETSKINNEDASSET(SkeletalMeshComponent) != nullptr
+			&& UE_MUTABLE_GETSKINNEDASSET(SkeletalMeshComponent)->GetSkeleton() == Animation->GetSkeleton()
 			)
 		{
 			SetRealtime(true);
@@ -2147,7 +2147,7 @@ void FCustomizableObjectEditorViewportClient::SetFloorOffset(float NewValue)
 {
 	for (TWeakObjectPtr<USkeletalMeshComponent> SkeletalMeshComponent : SkeletalMeshComponents)
 	{
-		USkeletalMesh* Mesh = SkeletalMeshComponent.IsValid() ? Cast<USkeletalMesh>(SkeletalMeshComponent->GetSkinnedAsset()) : nullptr;
+		USkeletalMesh* Mesh = SkeletalMeshComponent.IsValid() ? Cast<USkeletalMesh>(UE_MUTABLE_GETSKINNEDASSET(SkeletalMeshComponent)) : nullptr;
 
 		if (Mesh)
 		{
@@ -2704,9 +2704,9 @@ void FCustomizableObjectEditorViewportClient::BakeInstance(UCustomizableObjectIn
 				USkeletalMesh* SkeletalMesh = Cast<USkeletalMesh>(DupObject);
 				if (SkeletalMesh)
 				{
-					Helper_GetLODInfoArray(SkeletalMesh) = Helper_GetLODInfoArray(Mesh);
+					SkeletalMesh->GetLODInfoArray() = Mesh->GetLODInfoArray();
 
-					Helper_GetImportedModel(SkeletalMesh)->SkeletalMeshModelGUID = FGuid::NewGuid();
+					SkeletalMesh->GetImportedModel()->SkeletalMeshModelGUID = FGuid::NewGuid();
 
 					// Generate render data
 					SkeletalMesh->Build();
@@ -2899,9 +2899,9 @@ void FCustomizableObjectEditorViewportClient::ShowInstanceGeometryInformation(FC
 	// Show total number of triangles and vertices
 	for (TWeakObjectPtr<USkeletalMeshComponent> SkeletalMeshComponent : SkeletalMeshComponents)
 	{
-		if (SkeletalMeshComponent.IsValid() && SkeletalMeshComponent->GetSkinnedAsset())
+		if (SkeletalMeshComponent.IsValid() && UE_MUTABLE_GETSKINNEDASSET(SkeletalMeshComponent))
 		{
-			const FSkeletalMeshRenderData* MeshRes = SkeletalMeshComponent->GetSkinnedAsset()->GetResourceForRendering();
+			const FSkeletalMeshRenderData* MeshRes = UE_MUTABLE_GETSKINNEDASSET(SkeletalMeshComponent)->GetResourceForRendering();
 			int32 NumTriangles;
 			int32 NumVertices;
 			int32 NumLODLevel = MeshRes->LODRenderData.Num();
@@ -3399,7 +3399,7 @@ void SMutableSelectFolderDlg::Construct(const FArguments& InArgs)
 		.Padding(2)
 		[
 			SNew(SBorder)
-			.BorderImage(FAppStyle::GetBrush("ToolPanel.GroupBorder"))
+			.BorderImage(UE_MUTABLE_GET_BRUSH("ToolPanel.GroupBorder"))
 		[
 			SNew(SVerticalBox)
 
@@ -3469,14 +3469,14 @@ void SMutableSelectFolderDlg::Construct(const FArguments& InArgs)
 		.Padding(5)
 		[
 			SNew(SUniformGridPanel)
-			.SlotPadding(FAppStyle::GetMargin("StandardDialog.SlotPadding"))
-		.MinDesiredSlotWidth(FAppStyle::GetFloat("StandardDialog.MinDesiredSlotWidth"))
-		.MinDesiredSlotHeight(FAppStyle::GetFloat("StandardDialog.MinDesiredSlotHeight"))
+			.SlotPadding(UE_MUTABLE_GET_MARGIN("StandardDialog.SlotPadding"))
+		.MinDesiredSlotWidth(UE_MUTABLE_GET_FLOAT("StandardDialog.MinDesiredSlotWidth"))
+		.MinDesiredSlotHeight(UE_MUTABLE_GET_FLOAT("StandardDialog.MinDesiredSlotHeight"))
 		+ SUniformGridPanel::Slot(0, 0)
 		[
 			SNew(SButton)
 			.HAlign(HAlign_Center)
-		.ContentPadding(FAppStyle::GetMargin("StandardDialog.ContentPadding"))
+		.ContentPadding(UE_MUTABLE_GET_MARGIN("StandardDialog.ContentPadding"))
 		.Text(LOCTEXT("OK", "OK"))
 		.OnClicked(this, &SMutableSelectFolderDlg::OnButtonClick, EAppReturnType::Ok)
 		]
@@ -3484,7 +3484,7 @@ void SMutableSelectFolderDlg::Construct(const FArguments& InArgs)
 		[
 			SNew(SButton)
 			.HAlign(HAlign_Center)
-		.ContentPadding(FAppStyle::GetMargin("StandardDialog.ContentPadding"))
+		.ContentPadding(UE_MUTABLE_GET_MARGIN("StandardDialog.ContentPadding"))
 		.Text(LOCTEXT("Cancel", "Cancel"))
 		.OnClicked(this, &SMutableSelectFolderDlg::OnButtonClick, EAppReturnType::Cancel)
 		]
diff --git a/Engine/Plugins/Experimental/Mutable/Source/CustomizableObjectEditor/Private/MuCOE/CustomizableObjectIdentifierCustomization.cpp b/Engine/Plugins/Experimental/Mutable/Source/CustomizableObjectEditor/Private/MuCOE/CustomizableObjectIdentifierCustomization.cpp
--- a/Engine/Plugins/Experimental/Mutable/Source/CustomizableObjectEditor/Private/MuCOE/CustomizableObjectIdentifierCustomization.cpp
+++ b/Engine/Plugins/Experimental/Mutable/Source/CustomizableObjectEditor/Private/MuCOE/CustomizableObjectIdentifierCustomization.cpp
@@ -7,7 +7,7 @@
 #include "MuCOE/GraphTraversal.h"
 #include "MuCOE/Nodes/CustomizableObjectNodeObject.h"
 #include "PropertyCustomizationHelpers.h"
-
+#include "MuCO/UnrealPortabilityHelpers.h"
 
 #define LOCTEXT_NAMESPACE "FCustomizableObjectIdentifierCustomization"
 
@@ -38,7 +38,7 @@ void FCustomizableObjectIdentifierCustomization::CustomizeHeader(TSharedRef<clas
 
 	FAssetRegistryModule& AssetRegistryModule = FModuleManager::LoadModuleChecked<FAssetRegistryModule>("AssetRegistry");
 	TArray<FAssetData> OutAssetData;
-	AssetRegistryModule.Get().GetAssetsByClass(FTopLevelAssetPath(TEXT("/Script/CustomizableObject"), TEXT("CustomizableObject")), OutAssetData);
+	AssetRegistryModule.Get().GetAssetsByClass(UE_MUTABLE_TOPLEVELASSETPATH(TEXT("/Script/CustomizableObject"), TEXT("CustomizableObject")), OutAssetData);
 	SelectedCustomizableObject = nullptr;
 
 	FString StringGuid;
diff --git a/Engine/Plugins/Experimental/Mutable/Source/CustomizableObjectEditor/Private/MuCOE/CustomizableObjectInstanceEditor.cpp b/Engine/Plugins/Experimental/Mutable/Source/CustomizableObjectEditor/Private/MuCOE/CustomizableObjectInstanceEditor.cpp
--- a/Engine/Plugins/Experimental/Mutable/Source/CustomizableObjectEditor/Private/MuCOE/CustomizableObjectInstanceEditor.cpp
+++ b/Engine/Plugins/Experimental/Mutable/Source/CustomizableObjectEditor/Private/MuCOE/CustomizableObjectInstanceEditor.cpp
@@ -702,8 +702,8 @@ void FCustomizableObjectInstanceEditor::OnChangeMesh()
 	UCustomizableObjectInstance* SelectedMesh = GEditor->GetSelectedObjects()->GetTop<UCustomizableObjectInstance>();
 	if(SelectedMesh && SelectedMesh != CustomizableObjectInstance)
 	{
-		Helper_GetEditorSubsystem()->NotifyEditorClosed(this);
-		Helper_GetEditorSubsystem()->NotifyAssetOpened(SelectedMesh, this);
+		GEditor->GetEditorSubsystem<UAssetEditorSubsystem>()->NotifyEditorClosed(this);
+		GEditor->GetEditorSubsystem<UAssetEditorSubsystem>()->NotifyAssetOpened(SelectedMesh, this);
 
 		SetEditorMesh(SelectedMesh);
 //		LevelOfDetailSettings->MeshChanged();
diff --git a/Engine/Plugins/Experimental/Mutable/Source/CustomizableObjectEditor/Private/MuCOE/CustomizableObjectInstanceFactory.cpp b/Engine/Plugins/Experimental/Mutable/Source/CustomizableObjectEditor/Private/MuCOE/CustomizableObjectInstanceFactory.cpp
--- a/Engine/Plugins/Experimental/Mutable/Source/CustomizableObjectEditor/Private/MuCOE/CustomizableObjectInstanceFactory.cpp
+++ b/Engine/Plugins/Experimental/Mutable/Source/CustomizableObjectEditor/Private/MuCOE/CustomizableObjectInstanceFactory.cpp
@@ -11,7 +11,7 @@
 #include "MuCO/CustomizableSkeletalComponent.h"
 #include "MuCO/CustomizableSkeletalMeshActor.h"
 #include "MuCOE/Nodes/CustomizableObjectNodeObject.h"
-
+#include "MuCO/UnrealPortabilityHelpers.h"
 
 #define LOCTEXT_NAMESPACE "CustomizableObjectInstanceFactory"
 
@@ -123,7 +123,7 @@ void UCustomizableObjectInstanceFactory::PostSpawnActor(UObject* Asset, AActor*
 			if (USkeletalMeshComponent* SkeletalMeshComp = NewCSMActor->GetSkeletalMeshComponentAt(ComponentIndex))
 			{
 				SkeletalMeshComp->UnregisterComponent();
-				SkeletalMeshComp->SetSkinnedAsset(SkeletalMesh);
+				UE_MUTABLE_SETSKINNEDASSET(SkeletalMeshComp,SkeletalMesh);
 
 				if (ComponentIndex == 0 && NewCSMActor->GetWorld()->IsGameWorld())
 				{
diff --git a/Engine/Plugins/Experimental/Mutable/Source/CustomizableObjectEditor/Private/MuCOE/CustomizableObjectWidget.cpp b/Engine/Plugins/Experimental/Mutable/Source/CustomizableObjectEditor/Private/MuCOE/CustomizableObjectWidget.cpp
--- a/Engine/Plugins/Experimental/Mutable/Source/CustomizableObjectEditor/Private/MuCOE/CustomizableObjectWidget.cpp
+++ b/Engine/Plugins/Experimental/Mutable/Source/CustomizableObjectEditor/Private/MuCOE/CustomizableObjectWidget.cpp
@@ -169,7 +169,7 @@ void FCustomizableObjectWidget::Render_Axis( const FSceneView* View, FPrimitiveD
 		const float CylinderRadius = 1.2f;
 		const FVector Offset( 0,0,HalfHeight );
 
-		const FMaterialRenderProxy* Proxy = Helper_GetMaterialProxy(InMaterial);
+		const FMaterialRenderProxy* Proxy = InMaterial->GetRenderProxy();
 
 		switch( InAxis )
 		{
@@ -216,7 +216,7 @@ void FCustomizableObjectWidget::Render_Axis( const FSceneView* View, FPrimitiveD
 void FCustomizableObjectWidget::Render_Cube( FPrimitiveDrawInterface* PDI, const FMatrix& InMatrix, const UMaterialInterface* InMaterial, float InScale )
 {
 	const FMatrix CubeToWorld = FScaleMatrix(FVector(InScale,InScale,InScale)) * InMatrix;
-	DrawBox( PDI, CubeToWorld, FVector(1,1,1), Helper_GetMaterialProxy(InMaterial), SDPG_Foreground );
+	DrawBox( PDI, CubeToWorld, FVector(1,1,1), InMaterial->GetRenderProxy(), SDPG_Foreground );
 }
 
 
@@ -241,7 +241,7 @@ void FCustomizableObjectWidget::Render_Quad( FPrimitiveDrawInterface* PDI, const
 	MeshBuilder.AddTriangle(2,1,3);
 
 	// Draw the arrow mesh.
-	MeshBuilder.Draw(PDI, CubeToWorld, Helper_GetMaterialProxy(InMaterial), SDPG_Foreground, 0.f);
+	MeshBuilder.Draw(PDI, CubeToWorld, InMaterial->GetRenderProxy(), SDPG_Foreground, 0.f);
 
 #undef CUBE_FACE
 }
@@ -376,7 +376,7 @@ void FCustomizableObjectWidget::Render_Translate( const FSceneView* View, FPrimi
 		const FVector CameraYAxis = View->ViewMatrices.GetViewMatrix().GetColumn(1);
 		const FVector CameraZAxis = View->ViewMatrices.GetViewMatrix().GetColumn(2);
 
-		DrawSphere( PDI, InLocation, FRotator::ZeroRotator, FVector( 4.0f * Scale ), 10, 5, Helper_GetMaterialProxy(XYZMaterial), SDPG_Foreground );
+		DrawSphere( PDI, InLocation, FRotator::ZeroRotator, FVector( 4.0f * Scale ), 10, 5, XYZMaterial->GetRenderProxy(), SDPG_Foreground );
 
 		PDI->SetHitProxy( NULL );
 	}
@@ -998,7 +998,7 @@ void FCustomizableObjectWidget::DrawThickArc (const FThickArcParams& InParams, c
 		MeshBuilder.AddTriangle(VertexIndex+1, InnerVertexStartIndex+VertexIndex+1, InnerVertexStartIndex+VertexIndex);
 	}
 
-	MeshBuilder.Draw(InParams.PDI, FMatrix::Identity, Helper_GetMaterialProxy(InParams.Material),SDPG_Foreground,0.f);
+	MeshBuilder.Draw(InParams.PDI, FMatrix::Identity, InParams.Material->GetRenderProxy(),SDPG_Foreground,0.f);
 }
 
 /**
@@ -1057,7 +1057,7 @@ void FCustomizableObjectWidget::DrawSnapMarker(FPrimitiveDrawInterface* PDI, con
 
 		MeshBuilder.AddTriangle(0, 1, 2);
 		MeshBuilder.AddTriangle(1, 3, 2);
-		MeshBuilder.Draw(PDI, FMatrix::Identity, Helper_GetMaterialProxy(TransparentPlaneMaterialXY),SDPG_Foreground,0.f);
+		MeshBuilder.Draw(PDI, FMatrix::Identity, TransparentPlaneMaterialXY->GetRenderProxy(),SDPG_Foreground,0.f);
 	}
 }
 
@@ -1114,7 +1114,7 @@ void FCustomizableObjectWidget::DrawStartStopMarker(FPrimitiveDrawInterface* PDI
 		}
 
 		MeshBuilder.AddTriangle(0, 1, 2);
-		MeshBuilder.Draw(PDI, FMatrix::Identity, Helper_GetMaterialProxy(TransparentPlaneMaterialXY),SDPG_Foreground,0.f);
+		MeshBuilder.Draw(PDI, FMatrix::Identity, TransparentPlaneMaterialXY->GetRenderProxy(),SDPG_Foreground,0.f);
 	}
 }
 
diff --git a/Engine/Plugins/Experimental/Mutable/Source/CustomizableObjectEditor/Private/MuCOE/GenerateMutableSource/GenerateMutableSource.cpp b/Engine/Plugins/Experimental/Mutable/Source/CustomizableObjectEditor/Private/MuCOE/GenerateMutableSource/GenerateMutableSource.cpp
--- a/Engine/Plugins/Experimental/Mutable/Source/CustomizableObjectEditor/Private/MuCOE/GenerateMutableSource/GenerateMutableSource.cpp
+++ b/Engine/Plugins/Experimental/Mutable/Source/CustomizableObjectEditor/Private/MuCOE/GenerateMutableSource/GenerateMutableSource.cpp
@@ -399,7 +399,7 @@ UTexture2D* FindReferenceImage(const UEdGraphPin* Pin, FMutableGraphGenerationCo
 		{
 			if (const UEdGraphPin* ConnectedPin = FollowInputPin(*TypedNodeLayer->LayerPin(LayerIndex)))
 			{
-				if (ConnectedPin->PinType.PinCategory == Helper_GetPinCategory(Schema->PC_Image))
+				if (ConnectedPin->PinType.PinCategory == Schema->PC_Image)
 				{
 					Result = FindReferenceImage(ConnectedPin, GenerationContext);
 				}
@@ -535,7 +535,7 @@ UTexture2D* FindReferenceImage(const UEdGraphPin* Pin, FMutableGraphGenerationCo
 
 	else if (const UCustomizableObjectNodeTable* TypedNodeTable = Cast<UCustomizableObjectNodeTable>(Node))
 	{
-		if (Pin->PinType.PinCategory == Helper_GetPinCategory(Schema->PC_MaterialAsset))
+		if (Pin->PinType.PinCategory == Schema->PC_MaterialAsset)
 		{
 			Result = TypedNodeTable->FindReferenceTextureParameter(Pin, GenerationContext.CurrentMaterialTableParameter);
 		}
@@ -1266,7 +1266,7 @@ void PopulateReferenceSkeletalMeshesData(FMutableGraphGenerationContext& Generat
 				LODData.LODInfo.bAllowCPUAccess = LODInfo.bAllowCPUAccess;
 
 				// Copy Render data settings from the reference skeletal mesh
-				const FSkeletalMeshLODRenderData& ReferenceLODModel = Helper_GetLODData(RefSkeletalMesh)[LODIndex];
+				const FSkeletalMeshLODRenderData& ReferenceLODModel = RefSkeletalMesh->GetResourceForRendering()->LODRenderData[LODIndex];
 				LODData.RenderData.bIsLODOptional = ReferenceLODModel.bIsLODOptional;
 				LODData.RenderData.bStreamedDataInlined = ReferenceLODModel.bStreamedDataInlined;
 			}
diff --git a/Engine/Plugins/Experimental/Mutable/Source/CustomizableObjectEditor/Private/MuCOE/GenerateMutableSource/GenerateMutableSourceImage.cpp b/Engine/Plugins/Experimental/Mutable/Source/CustomizableObjectEditor/Private/MuCOE/GenerateMutableSource/GenerateMutableSourceImage.cpp
--- a/Engine/Plugins/Experimental/Mutable/Source/CustomizableObjectEditor/Private/MuCOE/GenerateMutableSource/GenerateMutableSourceImage.cpp
+++ b/Engine/Plugins/Experimental/Mutable/Source/CustomizableObjectEditor/Private/MuCOE/GenerateMutableSource/GenerateMutableSourceImage.cpp
@@ -305,7 +305,7 @@ mu::NodeImagePtr GenerateMutableSourceImage(const UEdGraphPin* Pin, FMutableGrap
 					GenerationContext.Compiler->CompilerLog(LOCTEXT("ModulateWithoutMask", "Texture layer effect uses Modulate without a mask. It will replace everything below it!"), Node);
 				}
 				
-				if (OtherPin->PinType.PinCategory == Helper_GetPinCategory(Schema->PC_Image))
+				if (OtherPin->PinType.PinCategory == Schema->PC_Image)
 				{
 					mu::NodeImagePtr BlendNode = GenerateMutableSourceImage(OtherPin, GenerationContext, MaxTextureSize);
 
@@ -317,7 +317,7 @@ mu::NodeImagePtr GenerateMutableSourceImage(const UEdGraphPin* Pin, FMutableGrap
 					Result = LayerNode;
 				}
 
-				else if (OtherPin->PinType.PinCategory == Helper_GetPinCategory(Schema->PC_Color))
+				else if (OtherPin->PinType.PinCategory == Schema->PC_Color)
 				{
 					mu::NodeColourPtr ColorNode = GenerateMutableSourceColor(OtherPin, GenerationContext);
 
@@ -807,7 +807,7 @@ mu::NodeImagePtr GenerateMutableSourceImage(const UEdGraphPin* Pin, FMutableGrap
 	}
 	
 	// If the node is a plain colour node, generate an image out of it
-	else if (Pin->PinType.PinCategory == Helper_GetPinCategory(Schema->PC_Color))
+	else if (Pin->PinType.PinCategory == Schema->PC_Color)
 	{
 		mu::NodeColourPtr ColorNode = GenerateMutableSourceColor(Pin, GenerationContext);
 
diff --git a/Engine/Plugins/Experimental/Mutable/Source/CustomizableObjectEditor/Private/MuCOE/GenerateMutableSource/GenerateMutableSourceMesh.cpp b/Engine/Plugins/Experimental/Mutable/Source/CustomizableObjectEditor/Private/MuCOE/GenerateMutableSource/GenerateMutableSourceMesh.cpp
--- a/Engine/Plugins/Experimental/Mutable/Source/CustomizableObjectEditor/Private/MuCOE/GenerateMutableSource/GenerateMutableSourceMesh.cpp
+++ b/Engine/Plugins/Experimental/Mutable/Source/CustomizableObjectEditor/Private/MuCOE/GenerateMutableSource/GenerateMutableSourceMesh.cpp
@@ -44,6 +44,7 @@
 #include "MuT/NodeMeshTable.h"
 #include "MuT/NodeMeshVariation.h"
 #include "PhysicsEngine/PhysicsAsset.h"
+#include "Engine/SkinnedAssetCommon.h"
 
 #define LOCTEXT_NAMESPACE "CustomizableObjectEditor"
 
@@ -71,9 +72,9 @@ TObjectPtr<const USkeletalMesh> GetMeshWithBoneRemovalApplied(TObjectPtr<USkelet
 	}
 
 	CacheEntry.bHasBonesToRemove = false;
-	for (int32 LODIndex = 0; LODIndex < Helper_GetLODNum(InSkeletalMesh); ++LODIndex)
+	for (int32 LODIndex = 0; LODIndex < InSkeletalMesh->GetLODNum(); ++LODIndex)
 	{
-		if (Helper_GetLODInfoBonesToRemove(InSkeletalMesh, LODIndex).Num() > 0)
+		if (InSkeletalMesh->GetLODInfo(LODIndex)->BonesToRemove.Num() > 0)
 		{
 			CacheEntry.bHasBonesToRemove = true;
 			break;
@@ -96,7 +97,7 @@ TObjectPtr<const USkeletalMesh> GetMeshWithBoneRemovalApplied(TObjectPtr<USkelet
 		CacheEntry.Mesh = (USkeletalMesh*)StaticDuplicateObject(InSkeletalMesh, GetTransientPackage(), InSkeletalMesh->GetFName(), EObjectFlags::RF_Transient);
 		check(CacheEntry.Mesh);
 
-		if (Helper_GetLODNum(InSkeletalMesh) < Helper_GetLODNum(CacheEntry.Mesh))
+		if (InSkeletalMesh->GetLODNum() < CacheEntry.Mesh->GetLODNum())
 		{
 			GenerationContext.Compiler->CompilerLog(LOCTEXT("ReferenceLODMismatch", "The reference mesh has less LODs than the generated mesh."),
 				CurrentNode,
@@ -116,9 +117,9 @@ TObjectPtr<const USkeletalMesh> GetMeshWithBoneRemovalApplied(TObjectPtr<USkelet
 	TSet<int32> RemovedBoneIndices;
 	{
 		TArray<FName> BoneNamesToRemove;
-		BoneNamesToRemove.Reset(Helper_GetLODInfoBonesToRemove(InSkeletalMesh, InLODIndex).Num());
+		BoneNamesToRemove.Reset(InSkeletalMesh->GetLODInfo(InLODIndex)->BonesToRemove.Num());
 		
-		for (const FBoneReference& BoneRef : Helper_GetLODInfoBonesToRemove(InSkeletalMesh, InLODIndex))
+		for (const FBoneReference& BoneRef : InSkeletalMesh->GetLODInfo(InLODIndex)->BonesToRemove)
 		{
 			BoneNamesToRemove.Add(BoneRef.BoneName);
 		}
@@ -598,7 +599,7 @@ mu::MeshPtr ConvertSkeletalMeshToMutable(USkeletalMesh* InSkeletalMesh, const TS
 		SkeletalMesh = GetMeshWithBoneRemovalApplied(InSkeletalMesh, LOD, GenerationContext, CurrentNode, RemovedBonesActiveParentIndices);
 	}
 
-	const FSkeletalMeshModel* ImportedModel = Helper_GetImportedModel(SkeletalMesh);
+	const FSkeletalMeshModel* ImportedModel = SkeletalMesh->GetImportedModel();
 
 	// Check in case the data has changed
 	if (!(ImportedModel
@@ -1799,7 +1800,7 @@ mu::MeshPtr ConvertStaticMeshToMutable(UStaticMesh* StaticMesh, int LOD, int Mat
 
 		// Position buffer
 		{
-			const FPositionVertexBuffer& VertexBuffer = Helper_GetStaticMeshPositionVertexBuffer(StaticMesh, LOD);
+			const FPositionVertexBuffer& VertexBuffer = StaticMesh->GetRenderData()->LODResources[LOD].VertexBuffers.PositionVertexBuffer;
 
 			const int ElementSize = 12;
 			const int ChannelCount = 1;
@@ -1818,7 +1819,7 @@ mu::MeshPtr ConvertStaticMeshToMutable(UStaticMesh* StaticMesh, int LOD, int Mat
 
 		// Tangent buffer
 		{
-			FStaticMeshVertexBuffer& VertexBuffer = Helper_GetStaticMeshVertexBuffer(StaticMesh, LOD);
+			FStaticMeshVertexBuffer& VertexBuffer = StaticMesh->GetRenderData()->LODResources[LOD].VertexBuffers.StaticMeshVertexBuffer;
 
 			MESH_BUFFER_SEMANTIC Semantics[2];
 			int SemanticIndices[2];
@@ -2043,13 +2044,13 @@ mu::MeshPtr BuildMorphedMutableMesh(const UEdGraphPin* BaseSourcePin, const FStr
 		// See if we need to select another LOD because of automatic LOD generation
 		if (GenerationContext.CurrentAutoLODStrategy == ECustomizableObjectAutomaticLODStrategy::AutomaticFromMesh)
 		{
-			if (SkeletalMesh && Helper_GetImportedModel(SkeletalMesh))
+			if (SkeletalMesh && SkeletalMesh->GetImportedModel())
 			{
 				// Add the CurrentLOD being generated to the first LOD used from this SkeletalMesh
 				int32 CurrentLOD = LODIndex + GenerationContext.CurrentLOD;
 
 				// If the mesh has additional LODs and they use the same material, use them.
-				FSkeletalMeshModel* ImportedModel = Helper_GetImportedModel(SkeletalMesh);
+				FSkeletalMeshModel* ImportedModel = SkeletalMesh->GetImportedModel();
 				if (ImportedModel->LODModels.Num() > CurrentLOD
 					&&
 					ImportedModel->LODModels[CurrentLOD].Sections.Num() > SectionIndex)
@@ -2090,7 +2091,7 @@ mu::MeshPtr BuildMorphedMutableMesh(const UEdGraphPin* BaseSourcePin, const FStr
 				int posOffset = MorphedSourceMesh->GetVertexBuffers().GetChannelOffset(posBuf, posChannel);
 				uint8* posBuffer = MorphedSourceMesh->GetVertexBuffers().GetBufferData(posBuf) + posOffset;
 
-				uint32 MaterialVertexStart = (uint32)Helper_GetImportedModel(SkeletalMesh)->LODModels[LODIndex].Sections[SectionIndex].GetVertexBufferIndex();
+				uint32 MaterialVertexStart = (uint32)SkeletalMesh->GetImportedModel()->LODModels[LODIndex].Sections[SectionIndex].GetVertexBufferIndex();
 				uint32 MeshVertexCount = (uint32)MorphedSourceMesh->GetVertexBuffers().GetElementCount();
 
 				const TArray<FMorphTargetLODModel>& MorphLODModels = MorphTarget->GetMorphLODModels();
@@ -2854,13 +2855,13 @@ mu::NodeMeshPtr GenerateMutableSourceMesh(const UEdGraphPin * Pin,
 			// See if we need to select another LOD because of automatic LOD generation
 			if (GenerationContext.CurrentAutoLODStrategy == ECustomizableObjectAutomaticLODStrategy::AutomaticFromMesh)
 			{
-				if (TypedNodeSkel->SkeletalMesh && Helper_GetImportedModel(TypedNodeSkel->SkeletalMesh))
+				if (TypedNodeSkel->SkeletalMesh && TypedNodeSkel->SkeletalMesh->GetImportedModel())
 				{
 					// Add the CurrentLOD being generated to the first LOD used from this SkeletalMesh
 					int32 CurrentLOD = LOD + GenerationContext.CurrentLOD; 
 					
 					// If the mesh has additional LODs and they use the same material, use them.
-					FSkeletalMeshModel* ImportedModel = Helper_GetImportedModel(TypedNodeSkel->SkeletalMesh);
+					FSkeletalMeshModel* ImportedModel = TypedNodeSkel->SkeletalMesh->GetImportedModel();
 					if (ImportedModel->LODModels.Num() > CurrentLOD
 						&&
 						ImportedModel->LODModels[CurrentLOD].Sections.Num() > SectionIndex)
@@ -2948,7 +2949,7 @@ mu::NodeMeshPtr GenerateMutableSourceMesh(const UEdGraphPin * Pin,
 
 				if (GenerationContext.Options.bSkinWeightProfilesEnabled)
 				{
-					FSkeletalMeshModel* ImportModel = Helper_GetImportedModel(TypedNodeSkel->SkeletalMesh);
+					FSkeletalMeshModel* ImportModel = TypedNodeSkel->SkeletalMesh->GetImportedModel();
 
 					const int32 SkinWeightProfilesCount = GenerationContext.SkinWeightProfilesInfo.Num();
 					for (int32 ProfileIndex = 0; ProfileIndex < SkinWeightProfilesCount; ++ProfileIndex)
@@ -2982,10 +2983,10 @@ mu::NodeMeshPtr GenerateMutableSourceMesh(const UEdGraphPin * Pin,
 				}
 
 				MeshData.bHasVertexColors = TypedNodeSkel->SkeletalMesh->GetHasVertexColors();
-				FSkeletalMeshModel* importModel = Helper_GetImportedModel(TypedNodeSkel->SkeletalMesh);
+				FSkeletalMeshModel* importModel = TypedNodeSkel->SkeletalMesh->GetImportedModel();
 				MeshData.NumTexCoordChannels = importModel->LODModels[LOD].NumTexCoords;
 				MeshData.MaxBoneIndexTypeSizeBytes = importModel->LODModels[LOD].RequiredBones.Num() > 256 ? 2 : 1;
-				MeshData.MaxNumBonesPerVertex = Helper_GetMaxBoneInfluences(importModel->LODModels[LOD]);
+				MeshData.MaxNumBonesPerVertex = importModel->LODModels[LOD].GetMaxBoneInfluences();
 				
 				// When mesh data is combined we will get an upper and lower bound of the number of triangles.
 				MeshData.MaxNumTriangles = importModel->LODModels[LOD].Sections[SectionIndex].NumTriangles;
@@ -3180,7 +3181,7 @@ mu::NodeMeshPtr GenerateMutableSourceMesh(const UEdGraphPin * Pin,
 
 					// Checking if it's a valid pin
 					if (MorphPin->Direction == EEdGraphPinDirection::EGPD_Output 
-						|| MorphPin->PinType.PinCategory != Helper_GetPinCategory(Schema->PC_Float) 
+						|| MorphPin->PinType.PinCategory != Schema->PC_Float 
 						|| !MorphPins[PinIndex]->LinkedTo.Num())
 					{
 						continue;
@@ -3584,9 +3585,9 @@ mu::NodeMeshPtr GenerateMutableSourceMesh(const UEdGraphPin * Pin,
 
 			if (GenerationContext.CurrentAutoLODStrategy == ECustomizableObjectAutomaticLODStrategy::AutomaticFromMesh)
 			{
-				if (SkeletalMesh && Helper_GetImportedModel(SkeletalMesh))
+				if (SkeletalMesh && SkeletalMesh->GetImportedModel())
 				{
-					FSkeletalMeshModel* ImportedModel = Helper_GetImportedModel(SkeletalMesh);
+					FSkeletalMeshModel* ImportedModel = SkeletalMesh->GetImportedModel();
 
 					CurrentLOD += GenerationContext.CurrentLOD;
 
@@ -3631,11 +3632,11 @@ mu::NodeMeshPtr GenerateMutableSourceMesh(const UEdGraphPin * Pin,
 			{
 				// TODO: this should be made for all the meshes of the Column to support meshes with different values
 				// Filling Mesh Data
-				FSkeletalMeshModel* importModel = Helper_GetImportedModel(SkeletalMesh);
+				FSkeletalMeshModel* importModel = SkeletalMesh->GetImportedModel();
 				MeshData.bHasVertexColors = SkeletalMesh->GetHasVertexColors();
 				MeshData.NumTexCoordChannels = importModel->LODModels[CurrentLOD].NumTexCoords;
 				MeshData.MaxBoneIndexTypeSizeBytes = importModel->LODModels[CurrentLOD].RequiredBones.Num() > 256 ? 2 : 1;
-				MeshData.MaxNumBonesPerVertex = Helper_GetMaxBoneInfluences(importModel->LODModels[CurrentLOD]);
+				MeshData.MaxNumBonesPerVertex = importModel->LODModels[CurrentLOD].GetMaxBoneInfluences();
 
 				// When mesh data is combined we will get an upper and lower bound of the number of triangles.
 				MeshData.MaxNumTriangles = importModel->LODModels[CurrentLOD].Sections[MaterialIndex].NumTriangles;
diff --git a/Engine/Plugins/Experimental/Mutable/Source/CustomizableObjectEditor/Private/MuCOE/GenerateMutableSource/GenerateMutableSourceTable.cpp b/Engine/Plugins/Experimental/Mutable/Source/CustomizableObjectEditor/Private/MuCOE/GenerateMutableSource/GenerateMutableSourceTable.cpp
--- a/Engine/Plugins/Experimental/Mutable/Source/CustomizableObjectEditor/Private/MuCOE/GenerateMutableSource/GenerateMutableSourceTable.cpp
+++ b/Engine/Plugins/Experimental/Mutable/Source/CustomizableObjectEditor/Private/MuCOE/GenerateMutableSource/GenerateMutableSourceTable.cpp
@@ -150,7 +150,7 @@ void FillTableColumn(const UCustomizableObjectNodeTable* TableNode,	mu::TablePtr
 			// Parameter used for LOD differences
 			int32 CurrentLOD = LOD;
 
-			int NumLODs = Helper_GetLODInfoArray(SkeletalMesh).Num();
+			int NumLODs = SkeletalMesh->GetLODInfoArray().Num();
 
 			if (NumLODs <= CurrentLOD)
 			{
@@ -161,8 +161,8 @@ void FillTableColumn(const UCustomizableObjectNodeTable* TableNode,	mu::TablePtr
 				GenerationContext.Compiler->CompilerLog(FText::FromString(msg), TableNode);
 			}
 
-			int32 NumMaterials = Helper_GetImportedModel(SkeletalMesh)->LODModels[CurrentLOD].Sections.Num();
-			int32 ReferenceNumMaterials = Helper_GetImportedModel(ReferenceSkeletalMesh)->LODModels[CurrentLOD].Sections.Num();
+			int32 NumMaterials = SkeletalMesh->GetImportedModel()->LODModels[CurrentLOD].Sections.Num();
+			int32 ReferenceNumMaterials = ReferenceSkeletalMesh->GetImportedModel()->LODModels[CurrentLOD].Sections.Num();
 
 			if (NumMaterials != ReferenceNumMaterials)
 			{
diff --git a/Engine/Plugins/Experimental/Mutable/Source/CustomizableObjectEditor/Private/MuCOE/MutableMeshPreviewUtils.cpp b/Engine/Plugins/Experimental/Mutable/Source/CustomizableObjectEditor/Private/MuCOE/MutableMeshPreviewUtils.cpp
--- a/Engine/Plugins/Experimental/Mutable/Source/CustomizableObjectEditor/Private/MuCOE/MutableMeshPreviewUtils.cpp
+++ b/Engine/Plugins/Experimental/Mutable/Source/CustomizableObjectEditor/Private/MuCOE/MutableMeshPreviewUtils.cpp
@@ -10,6 +10,7 @@
 #include "MuCO/UnrealConversionUtils.h"
 #include "MuCO/UnrealPortabilityHelpers.h"
 #include "MuR/OpMeshFormat.h"
+#include "Engine/SkeletalMesh.h"
 #include "UObject/Package.h"
 
 class UMaterialInterface;
@@ -549,14 +550,14 @@ namespace MutableMeshPreviewUtils
 			OutSkeletalMesh->GetMaterials().SetNum(1);
 			OutSkeletalMesh->GetMaterials()[0] = UnrealMaterial;
 
-			Helper_GetLODInfoArray(OutSkeletalMesh)[0].LODMaterialMap.SetNumZeroed(1);
+			OutSkeletalMesh->GetLODInfoArray()[0].LODMaterialMap.SetNumZeroed(1);
 
 			// Add RenderSections for each surface in the mesh
 			if (InMutableMesh)
 			{
 				for (int32 SurfaceIndex = 0; SurfaceIndex < InMutableMesh->GetSurfaceCount(); ++SurfaceIndex)
 				{
-					new(Helper_GetLODRenderSections(OutSkeletalMesh, 0)) Helper_SkelMeshRenderSection();
+					new(OutSkeletalMesh->GetResourceForRendering()->LODRenderData[0].RenderSections) FSkelMeshRenderSection();
 				}
 			}
 		}
@@ -589,7 +590,7 @@ namespace MutableMeshPreviewUtils
 				InMutableMesh,
 				0); 
 			
-			FSkeletalMeshLODRenderData& LODModel = Helper_GetLODData(OutSkeletalMesh)[0];
+			FSkeletalMeshLODRenderData& LODModel = OutSkeletalMesh->GetResourceForRendering()->LODRenderData[0];
 
 			// Generate the active bones array based on the bones found on BoneMap
 			TArray<uint16> ActiveBones;
@@ -713,25 +714,28 @@ namespace MutableMeshPreviewUtils
 
 		// Prepare the mesh to be filled with rendering data (buffers)
 		{
-			Helper_GetLODData(GeneratedSkeletalMesh).Add(( Helper_GetLODData(GeneratedSkeletalMesh).Num()) ? &Helper_GetLODData(GeneratedSkeletalMesh)[0] : new Helper_LODDataType());
-			Helper_GetLODInfoArray(GeneratedSkeletalMesh).Add(FSkeletalMeshLODInfo());
+			GeneratedSkeletalMesh->GetResourceForRendering()->LODRenderData.Add(( GeneratedSkeletalMesh->GetResourceForRendering()->LODRenderData.Num()) ? &GeneratedSkeletalMesh->GetResourceForRendering()->LODRenderData[0] : new FSkeletalMeshLODRenderData());
+			GeneratedSkeletalMesh->GetLODInfoArray().Add(FSkeletalMeshLODInfo());
+
+			FSkeletalMeshLODInfo& LastLODInfo = GeneratedSkeletalMesh->GetLODInfoArray().Last();
 
 	#if WITH_EDITORONLY_DATA
-			Helper_GetLODInfoArray(GeneratedSkeletalMesh).Last().BuildSettings.bUseFullPrecisionUVs = true;
+			LastLODInfo.BuildSettings.bUseFullPrecisionUVs = true;
 	#endif
 
 			constexpr int32 LODIndex = 0;
-			if (Helper_GetLODInfoArray(InReferenceSkeletalMesh).IsValidIndex(LODIndex))
+			if (InReferenceSkeletalMesh->GetLODInfoArray().IsValidIndex(LODIndex))
 			{
-				Helper_GetLODInfoArray(GeneratedSkeletalMesh).Last().ScreenSize = Helper_GetLODInfoArray(InReferenceSkeletalMesh)[LODIndex].ScreenSize;
-				Helper_GetLODInfoArray(GeneratedSkeletalMesh).Last().LODHysteresis = Helper_GetLODInfoArray(InReferenceSkeletalMesh)[LODIndex].LODHysteresis;
-				Helper_GetLODInfoArray(GeneratedSkeletalMesh).Last().bSupportUniformlyDistributedSampling = Helper_GetLODInfoArray(InReferenceSkeletalMesh)[LODIndex].bSupportUniformlyDistributedSampling;
-				Helper_GetLODInfoArray(GeneratedSkeletalMesh).Last().bAllowCPUAccess = Helper_GetLODInfoArray(InReferenceSkeletalMesh)[LODIndex].bAllowCPUAccess;
+				const FSkeletalMeshLODInfo& RefLODIndex = InReferenceSkeletalMesh->GetLODInfoArray()[LODIndex];
+				LastLODInfo.ScreenSize = RefLODIndex.ScreenSize;
+				LastLODInfo.LODHysteresis = RefLODIndex.LODHysteresis;
+				LastLODInfo.bSupportUniformlyDistributedSampling = RefLODIndex.bSupportUniformlyDistributedSampling;
+				LastLODInfo.bAllowCPUAccess = RefLODIndex.bAllowCPUAccess;
 			}
 			else
 			{
-				Helper_GetLODInfoArray(GeneratedSkeletalMesh).Last().ScreenSize = 0.3f / (LODIndex + 1);
-				Helper_GetLODInfoArray(GeneratedSkeletalMesh).Last().LODHysteresis = 0.02f;
+				LastLODInfo.ScreenSize = 0.3f / (LODIndex + 1);
+				LastLODInfo.LODHysteresis = 0.02f;
 				UE_LOG(LogMutable, Warning, TEXT("Setting default values for LOD ScreenSize and LODHysteresis because the values cannot be found in the reference mesh"));
 			}
 		}
diff --git a/Engine/Plugins/Experimental/Mutable/Source/CustomizableObjectEditor/Private/MuCOE/Nodes/CustomizableObjectNodeCurve.cpp b/Engine/Plugins/Experimental/Mutable/Source/CustomizableObjectEditor/Private/MuCOE/Nodes/CustomizableObjectNodeCurve.cpp
--- a/Engine/Plugins/Experimental/Mutable/Source/CustomizableObjectEditor/Private/MuCOE/Nodes/CustomizableObjectNodeCurve.cpp
+++ b/Engine/Plugins/Experimental/Mutable/Source/CustomizableObjectEditor/Private/MuCOE/Nodes/CustomizableObjectNodeCurve.cpp
@@ -49,14 +49,14 @@ void UCustomizableObjectNodeCurve::Serialize(FArchive& Ar)
 			if (UEdGraphPin* CurvePin = FindPin(FString::Printf(TEXT("Curve %d"), i)))
 			{
 				FString NewPinName = *(FString("Curve ") + Curves[i].CurveName.ToString());
-				Helper_SetPinName( CurvePin->PinName, NewPinName);
+				CurvePin->PinName = FName(NewPinName);
 				CurvePin->Modify(true);
 			}
 			else if (OutputPins.IsValidIndex(i) && OutputPins[i] && Helper_GetPinName(OutputPins[i]) != (FString("Curve ") + Curves[i].CurveName.ToString()))
 			{
 				CurvePin = OutputPins[i];
 				FString NewPinName = *(FString("Curve ") + Curves[i].CurveName.ToString());
-				Helper_SetPinName(CurvePin->PinName, NewPinName);
+				CurvePin->PinName = FName(NewPinName);
 				CurvePin->Modify(true);
 			}
 		}
diff --git a/Engine/Plugins/Experimental/Mutable/Source/CustomizableObjectEditor/Private/MuCOE/Nodes/CustomizableObjectNodeEditMaterialBaseDetails.cpp b/Engine/Plugins/Experimental/Mutable/Source/CustomizableObjectEditor/Private/MuCOE/Nodes/CustomizableObjectNodeEditMaterialBaseDetails.cpp
--- a/Engine/Plugins/Experimental/Mutable/Source/CustomizableObjectEditor/Private/MuCOE/Nodes/CustomizableObjectNodeEditMaterialBaseDetails.cpp
+++ b/Engine/Plugins/Experimental/Mutable/Source/CustomizableObjectEditor/Private/MuCOE/Nodes/CustomizableObjectNodeEditMaterialBaseDetails.cpp
@@ -6,6 +6,7 @@
 #include "IDetailsView.h"
 #include "MuCOE/Nodes/CustomizableObjectNodeEditMaterialBase.h"
 #include "MuCOE/Nodes/CustomizableObjectNodeMaterial.h"
+#include "MuCOE/UnrealEditorPortabilityHelpers.h"
 #include "PropertyCustomizationHelpers.h"
 #include "Widgets/Input/STextComboBox.h"
 
@@ -61,7 +62,7 @@ void FCustomizableObjectNodeEditMaterialBaseDetails::CustomizeDetails(IDetailLay
 				.CustomWidget()
 				[
 					SNew(SBorder)
-					.BorderImage(FAppStyle::GetBrush("NoBorder"))
+					.BorderImage(UE_MUTABLE_GET_BRUSH("NoBorder"))
 					.Padding(FMargin(0.0f, 0.0f, 10.0f, 0.0f))
 					[
 						SAssignNew(LayoutComboBox, STextComboBox)
diff --git a/Engine/Plugins/Experimental/Mutable/Source/CustomizableObjectEditor/Private/MuCOE/Nodes/CustomizableObjectNodeExternalPinDetails.cpp b/Engine/Plugins/Experimental/Mutable/Source/CustomizableObjectEditor/Private/MuCOE/Nodes/CustomizableObjectNodeExternalPinDetails.cpp
--- a/Engine/Plugins/Experimental/Mutable/Source/CustomizableObjectEditor/Private/MuCOE/Nodes/CustomizableObjectNodeExternalPinDetails.cpp
+++ b/Engine/Plugins/Experimental/Mutable/Source/CustomizableObjectEditor/Private/MuCOE/Nodes/CustomizableObjectNodeExternalPinDetails.cpp
@@ -7,6 +7,7 @@
 #include "MuCOE/CustomizableObjectEditorUtilities.h"
 #include "MuCOE/Nodes/CustomizableObjectNodeExposePin.h"
 #include "MuCOE/Nodes/CustomizableObjectNodeExternalPin.h"
+#include "MuCOE/UnrealEditorPortabilityHelpers.h"
 #include "PropertyCustomizationHelpers.h"
 #include "ScopedTransaction.h"
 #include "Widgets/Input/STextComboBox.h"
@@ -76,7 +77,7 @@ void FCustomizableObjectNodeExternalPinDetails::CustomizeDetails( IDetailLayoutB
 			ExternalCategory.AddCustomRow(LOCTEXT("FCustomizableObjectNodeExternalPinDetails", "Blocks"))
 			[
 				SNew(SBorder)
-					.BorderImage(FAppStyle::GetBrush("NoBorder"))
+					.BorderImage(UE_MUTABLE_GET_BRUSH("NoBorder"))
 					.Padding(FMargin(0.0f, 0.0f, 10.0f, 0.0f))
 					[
 						SNew(STextComboBox)
diff --git a/Engine/Plugins/Experimental/Mutable/Source/CustomizableObjectEditor/Private/MuCOE/Nodes/CustomizableObjectNodeGroupProjectorParameter.cpp b/Engine/Plugins/Experimental/Mutable/Source/CustomizableObjectEditor/Private/MuCOE/Nodes/CustomizableObjectNodeGroupProjectorParameter.cpp
--- a/Engine/Plugins/Experimental/Mutable/Source/CustomizableObjectEditor/Private/MuCOE/Nodes/CustomizableObjectNodeGroupProjectorParameter.cpp
+++ b/Engine/Plugins/Experimental/Mutable/Source/CustomizableObjectEditor/Private/MuCOE/Nodes/CustomizableObjectNodeGroupProjectorParameter.cpp
@@ -31,7 +31,7 @@ TArray<FGroupProjectorParameterImage> UCustomizableObjectNodeGroupProjectorParam
 
 	int32 NameIndex = 0;
 
-	for (auto RowIt = Helper_GetRowMap(OptionImagesDataTable).CreateConstIterator(); RowIt; ++RowIt)
+	for (TMap<FName, uint8*>::TConstIterator RowIt = OptionImagesDataTable->GetRowMap().CreateConstIterator(); RowIt; ++RowIt)
 	{
 		uint8* RowData = RowIt.Value();
 		FString PropertyValue(TEXT(""));
diff --git a/Engine/Plugins/Experimental/Mutable/Source/CustomizableObjectEditor/Private/MuCOE/Nodes/CustomizableObjectNodeMaterial.cpp b/Engine/Plugins/Experimental/Mutable/Source/CustomizableObjectEditor/Private/MuCOE/Nodes/CustomizableObjectNodeMaterial.cpp
--- a/Engine/Plugins/Experimental/Mutable/Source/CustomizableObjectEditor/Private/MuCOE/Nodes/CustomizableObjectNodeMaterial.cpp
+++ b/Engine/Plugins/Experimental/Mutable/Source/CustomizableObjectEditor/Private/MuCOE/Nodes/CustomizableObjectNodeMaterial.cpp
@@ -9,6 +9,7 @@
 #include "MuCOE/CustomizableObjectEditor_Deprecated.h"
 #include "MuCOE/EdGraphSchema_CustomizableObject.h"
 #include "MuCOE/GraphTraversal.h"
+#include "MuCOE/UnrealEditorPortabilityHelpers.h"
 #include "MuCOE/Nodes/CustomizableObjectNodeCopyMaterial.h"
 #include "MuCOE/Nodes/CustomizableObjectNodeSkeletalMesh.h"
 #include "MuCOE/Nodes/CustomizableObjectNodeStaticMesh.h"
@@ -752,9 +753,14 @@ int32 UCustomizableObjectNodeMaterial::GetImageUVLayout(const int32 ImageIndex)
 
 int32 UCustomizableObjectNodeMaterial::GetNumParameters(const EMaterialParameterType Type) const
 {
+	ensure(Material);
 	if (Material)
 	{
+#if ENGINE_MAJOR_VERSION==5 && ENGINE_MINOR_VERSION>=1
 		return Material->GetCachedExpressionData().GetParameterTypeEntry(Type).ParameterInfoSet.Num();
+#else
+		return Material->GetCachedExpressionData().Parameters.GetNumParameters(Type);
+#endif
 	}
 	else
 	{
@@ -767,6 +773,7 @@ FGuid UCustomizableObjectNodeMaterial::GetParameterId(const EMaterialParameterTy
 {
 	const FMaterialCachedExpressionData& Data = Material->GetCachedExpressionData();
 
+#if ENGINE_MAJOR_VERSION==5 && ENGINE_MINOR_VERSION>=1
 	if (Data.EditorOnlyData)
 	{
 		if (Data.EditorOnlyData->EditorEntries[(int32)Type].EditorInfo.Num() != 0)
@@ -774,6 +781,10 @@ FGuid UCustomizableObjectNodeMaterial::GetParameterId(const EMaterialParameterTy
 			return Data.EditorOnlyData->EditorEntries[(int32)Type].EditorInfo[ParameterIndex].ExpressionGuid;
 		}
 	}
+#else
+	const FMaterialCachedParameterEntry& Entry = Data.Parameters.GetParameterTypeEntry(Type);
+	return Entry.EditorInfo[ParameterIndex].ExpressionGuid;
+#endif
 	
 	return FGuid();
 }
@@ -781,6 +792,13 @@ FGuid UCustomizableObjectNodeMaterial::GetParameterId(const EMaterialParameterTy
 
 FName UCustomizableObjectNodeMaterial::GetParameterName(const EMaterialParameterType Type, const int32 ParameterIndex) const
 {
+	if (!Material)
+	{
+		return FName();
+	}
+
+#if ENGINE_MAJOR_VERSION==5 && ENGINE_MINOR_VERSION>=1
+
 	const FMaterialCachedParameterEntry& Entry = Material->GetCachedExpressionData().GetParameterTypeEntry(Type);
 
 	for (TSet<FMaterialParameterInfo>::TConstIterator It(Entry.ParameterInfoSet); It; ++It)
@@ -796,12 +814,29 @@ FName UCustomizableObjectNodeMaterial::GetParameterName(const EMaterialParameter
 	// The parameter should exist
 	check(false);
 
+#else
+
+	TArray<FGuid> ParameterIds;
+	TArray<FMaterialParameterInfo> ParameterInfo;
+	Material->GetAllParameterInfoOfType(Type, ParameterInfo, ParameterIds);
+
+	return ParameterInfo.IsValidIndex(ParameterIndex) ? ParameterInfo[ParameterIndex].Name : FName();
+
+#endif
+
 	return FName();
 }
 
 
 int32 UCustomizableObjectNodeMaterial::GetParameterLayerIndex(const EMaterialParameterType Type, const int32 ParameterIndex) const
 {
+	if (!Material)
+	{
+		return -1;
+	}
+
+#if ENGINE_MAJOR_VERSION==5 && ENGINE_MINOR_VERSION>=1
+
 	const FMaterialCachedParameterEntry& Entry = Material->GetCachedExpressionData().GetParameterTypeEntry(Type);
 
 	for (TSet<FMaterialParameterInfo>::TConstIterator It(Entry.ParameterInfoSet); It; ++It)
@@ -817,60 +852,102 @@ int32 UCustomizableObjectNodeMaterial::GetParameterLayerIndex(const EMaterialPar
 	// The parameter should exist
 	check(false);
 
+#else
+
+	TArray<FGuid> ParameterIds;
+	TArray<FMaterialParameterInfo> ParameterInfo;
+	Material->GetAllParameterInfoOfType(Type, ParameterInfo, ParameterIds);
+
+	return ParameterInfo.IsValidIndex(ParameterIndex) ? ParameterInfo[ParameterIndex].Index : -1;
+
+#endif
+
 	return -1;
 }
 
 
 FText UCustomizableObjectNodeMaterial::GetParameterLayerName(const EMaterialParameterType Type, const int32 ParameterIndex) const
 {
-	const FMaterialCachedParameterEntry& Entry = Material->GetCachedExpressionData().GetParameterTypeEntry(Type);
-
-	for (TSet<FMaterialParameterInfo>::TConstIterator It(Entry.ParameterInfoSet); It; ++It)
+	ensure(Material);
+	if (!Material)
 	{
-		const int32 IteratorIndex = It.GetId().AsInteger();
+		return FText();
+	}
 
-		if (IteratorIndex == ParameterIndex)
-		{
-			FMaterialLayersFunctions LayersValue;
-			Material->GetMaterialLayers(LayersValue);
+	int32 LayerIndex = GetParameterLayerIndex(Type,ParameterIndex);
 
-			return LayersValue.EditorOnly.LayerNames[(*It).Index];
-		}
-	}
-	
-	// The parameter should exist
-	check(false);
+	FMaterialLayersFunctions LayersValue;
+	Material->GetMaterialLayers(LayersValue);
+
+#if ENGINE_MAJOR_VERSION==5 && ENGINE_MINOR_VERSION>=1
+
+	return LayersValue.EditorOnly.LayerNames.IsValidIndex(LayerIndex) ? LayersValue.EditorOnly.LayerNames[LayerIndex] : FText();
+
+#else
 
-	return FText();
+	return LayersValue.LayerNames.IsValidIndex(LayerIndex) ? LayersValue.LayerNames[LayerIndex] : FText();
+
+#endif
 }
 
 
 bool UCustomizableObjectNodeMaterial::HasParameter(const FGuid& ParameterId) const
 {
-	if (Material)
+	ensure(Material);
+	if (!Material)
+	{
+		return false;
+	}
+
+#if ENGINE_MAJOR_VERSION==5 && ENGINE_MINOR_VERSION>=1
+
+	for (const EMaterialParameterType Type : ParameterTypes)
 	{
-		for (const EMaterialParameterType Type : ParameterTypes)
+		const FMaterialCachedExpressionData& Data = Material->GetCachedExpressionData();
+		const FMaterialCachedParameterEntry& Entry = Data.GetParameterTypeEntry(Type);
+
+		if (!Data.EditorOnlyData || Data.EditorOnlyData->EditorEntries[(int32)Type].EditorInfo.IsEmpty())
+		{
+			continue;
+		}
+
+		for (TSet<FMaterialParameterInfo>::TConstIterator It(Entry.ParameterInfoSet); It; ++It)
 		{
-			const FMaterialCachedExpressionData& Data = Material->GetCachedExpressionData();
-			const FMaterialCachedParameterEntry& Entry = Data.GetParameterTypeEntry(Type);
+			const int32 IteratorIndex = It.GetId().AsInteger();
+			const FGuid& ParamGUid = Data.EditorOnlyData->EditorEntries[(int32)Type].EditorInfo[IteratorIndex].ExpressionGuid;
 
-			if (!Data.EditorOnlyData || Data.EditorOnlyData->EditorEntries[(int32)Type].EditorInfo.IsEmpty())
+			if (ParamGUid == ParameterId)
 			{
-				continue;
+				return true;
 			}
+		}
+	}
 
-			for (TSet<FMaterialParameterInfo>::TConstIterator It(Entry.ParameterInfoSet); It; ++It)
-			{
-				const int32 IteratorIndex = It.GetId().AsInteger();
-				const FGuid& ParamGUid = Data.EditorOnlyData->EditorEntries[(int32)Type].EditorInfo[IteratorIndex].ExpressionGuid;
+#else
 
-				if (ParamGUid == ParameterId)
-				{
-					return true;
-				}
+	const FMaterialCachedExpressionData& Data = Material->GetCachedExpressionData();
+	for (const FMaterialCachedParameterEntry& Entry : Data.Parameters.RuntimeEntries)
+	{
+		for (const FMaterialCachedParameterEditorInfo& Info : Entry.EditorInfo)
+		{
+			if (Info.ExpressionGuid == ParameterId)
+			{
+				return true;
 			}
 		}
 	}
+	for (const FMaterialCachedParameterEntry& Entry : Data.Parameters.EditorOnlyEntries)
+	{
+		for (const FMaterialCachedParameterEditorInfo& Info : Entry.EditorInfo)
+		{
+			if (Info.ExpressionGuid == ParameterId)
+			{
+				return true;
+			}
+		}
+	}
+
+#endif
 
 	return false;
 }
diff --git a/Engine/Plugins/Experimental/Mutable/Source/CustomizableObjectEditor/Private/MuCOE/Nodes/CustomizableObjectNodeMeshClipMorph.cpp b/Engine/Plugins/Experimental/Mutable/Source/CustomizableObjectEditor/Private/MuCOE/Nodes/CustomizableObjectNodeMeshClipMorph.cpp
--- a/Engine/Plugins/Experimental/Mutable/Source/CustomizableObjectEditor/Private/MuCOE/Nodes/CustomizableObjectNodeMeshClipMorph.cpp
+++ b/Engine/Plugins/Experimental/Mutable/Source/CustomizableObjectEditor/Private/MuCOE/Nodes/CustomizableObjectNodeMeshClipMorph.cpp
@@ -91,7 +91,7 @@ void UCustomizableObjectNodeMeshClipMorph::PostEditChangeProperty(FPropertyChang
 {
 	FProperty* PropertyThatChanged = PropertyChangedEvent.Property;
 
-	const FName PropertyName = Helper_GetPropertyName(PropertyChangedEvent);
+	const FName PropertyName = PropertyChangedEvent.GetPropertyName();
 
 	if (PropertyName == "bLocalStartOffset")
 	{
diff --git a/Engine/Plugins/Experimental/Mutable/Source/CustomizableObjectEditor/Private/MuCOE/Nodes/CustomizableObjectNodeMeshClipMorphDetails.cpp b/Engine/Plugins/Experimental/Mutable/Source/CustomizableObjectEditor/Private/MuCOE/Nodes/CustomizableObjectNodeMeshClipMorphDetails.cpp
--- a/Engine/Plugins/Experimental/Mutable/Source/CustomizableObjectEditor/Private/MuCOE/Nodes/CustomizableObjectNodeMeshClipMorphDetails.cpp
+++ b/Engine/Plugins/Experimental/Mutable/Source/CustomizableObjectEditor/Private/MuCOE/Nodes/CustomizableObjectNodeMeshClipMorphDetails.cpp
@@ -9,6 +9,7 @@
 #include "MuCOE/CustomizableObjectGraph.h"
 #include "MuCOE/GraphTraversal.h"
 #include "MuCOE/Nodes/CustomizableObjectNodeMeshClipMorph.h"
+#include "MuCOE/UnrealEditorPortabilityHelpers.h"
 #include "PropertyCustomizationHelpers.h"
 #include "Widgets/Input/SCheckBox.h"
 #include "Widgets/Input/STextComboBox.h"
@@ -105,7 +106,7 @@ void FCustomizableObjectNodeMeshClipMorphDetails::CustomizeDetails(IDetailLayout
 				.CustomWidget()
 				[
 					SNew(SBorder)
-					.BorderImage(FAppStyle::GetBrush("NoBorder"))
+					.BorderImage(UE_MUTABLE_GET_BRUSH("NoBorder"))
 					.Padding(FMargin(0.0f, 0.0f, 10.0f, 0.0f))
 					.HAlign(EHorizontalAlignment::HAlign_Fill)
 					[
@@ -135,7 +136,7 @@ void FCustomizableObjectNodeMeshClipMorphDetails::CustomizeDetails(IDetailLayout
 				.CustomWidget()
 				[
 					SNew(SBorder)
-					.BorderImage(FAppStyle::GetBrush("NoBorder"))
+					.BorderImage(UE_MUTABLE_GET_BRUSH("NoBorder"))
 					.Padding(FMargin(0.0f, 0.0f, 10.0f, 0.0f))
 					.HAlign(EHorizontalAlignment::HAlign_Fill)
 					[
diff --git a/Engine/Plugins/Experimental/Mutable/Source/CustomizableObjectEditor/Private/MuCOE/Nodes/CustomizableObjectNodeMeshClipWithMeshDetails.cpp b/Engine/Plugins/Experimental/Mutable/Source/CustomizableObjectEditor/Private/MuCOE/Nodes/CustomizableObjectNodeMeshClipWithMeshDetails.cpp
--- a/Engine/Plugins/Experimental/Mutable/Source/CustomizableObjectEditor/Private/MuCOE/Nodes/CustomizableObjectNodeMeshClipWithMeshDetails.cpp
+++ b/Engine/Plugins/Experimental/Mutable/Source/CustomizableObjectEditor/Private/MuCOE/Nodes/CustomizableObjectNodeMeshClipWithMeshDetails.cpp
@@ -6,6 +6,7 @@
 #include "IDetailsView.h"
 #include "MuCOE/CustomizableObjectEditorUtilities.h"
 #include "MuCOE/GraphTraversal.h"
+#include "MuCOE/UnrealEditorPortabilityHelpers.h"
 #include "MuCOE/Nodes/CustomizableObjectNodeMaterial.h"
 #include "MuCOE/Nodes/CustomizableObjectNodeMeshClipWithMesh.h"
 #include "MuCOE/Nodes/CustomizableObjectNodeStaticMesh.h"
@@ -365,7 +366,7 @@ void FCustomizableObjectNodeMeshClipWithMeshDetails::CreateTagView()
 			.ToolTipText(LOCTEXT("AddNewTag", "Add New Tag"))
 			[
 				SNew(SImage)
-				.Image(FAppStyle::GetBrush(TEXT("Plus")))	
+				.Image(UE_MUTABLE_GET_BRUSH(TEXT("Plus")))	
 			]
 		]
 	];
@@ -426,7 +427,7 @@ void STagView::Construct(const FArguments& InArgs)
 			.ToolTipText(LOCTEXT("DeleteTag", "Delete Tag"))
 			[
 				SNew(SImage)
-				.Image(FAppStyle::GetBrush(TEXT("Cross")))
+				.Image(UE_MUTABLE_GET_BRUSH(TEXT("Cross")))
 			]
 		]
 	];
diff --git a/Engine/Plugins/Experimental/Mutable/Source/CustomizableObjectEditor/Private/MuCOE/Nodes/CustomizableObjectNodeMeshMorphDetails.cpp b/Engine/Plugins/Experimental/Mutable/Source/CustomizableObjectEditor/Private/MuCOE/Nodes/CustomizableObjectNodeMeshMorphDetails.cpp
--- a/Engine/Plugins/Experimental/Mutable/Source/CustomizableObjectEditor/Private/MuCOE/Nodes/CustomizableObjectNodeMeshMorphDetails.cpp
+++ b/Engine/Plugins/Experimental/Mutable/Source/CustomizableObjectEditor/Private/MuCOE/Nodes/CustomizableObjectNodeMeshMorphDetails.cpp
@@ -10,6 +10,7 @@
 #include "MuCOE/Nodes/CustomizableObjectNodeMeshMorph.h"
 #include "MuCOE/Nodes/CustomizableObjectNodeSkeletalMesh.h"
 #include "MuCOE/Nodes/CustomizableObjectNodeTable.h"
+#include "MuCOE/UnrealEditorPortabilityHelpers.h"
 #include "PropertyCustomizationHelpers.h"
 #include "Widgets/Input/STextComboBox.h"
 
@@ -82,7 +83,7 @@ void FCustomizableObjectNodeMeshMorphDetails::CustomizeDetails( IDetailLayoutBui
 				.CustomWidget()
 				[
 					SNew(SBorder)
-					.BorderImage(FAppStyle::GetBrush("NoBorder"))
+					.BorderImage(UE_MUTABLE_GET_BRUSH("NoBorder"))
 					.Padding(FMargin(0.0f, 0.0f, 10.0f, 0.0f))
 					[
 						SNew(STextComboBox)
diff --git a/Engine/Plugins/Experimental/Mutable/Source/CustomizableObjectEditor/Private/MuCOE/Nodes/CustomizableObjectNodeMorphMaterialDetails.cpp b/Engine/Plugins/Experimental/Mutable/Source/CustomizableObjectEditor/Private/MuCOE/Nodes/CustomizableObjectNodeMorphMaterialDetails.cpp
--- a/Engine/Plugins/Experimental/Mutable/Source/CustomizableObjectEditor/Private/MuCOE/Nodes/CustomizableObjectNodeMorphMaterialDetails.cpp
+++ b/Engine/Plugins/Experimental/Mutable/Source/CustomizableObjectEditor/Private/MuCOE/Nodes/CustomizableObjectNodeMorphMaterialDetails.cpp
@@ -10,6 +10,7 @@
 #include "MuCOE/Nodes/CustomizableObjectNodeMaterial.h"
 #include "MuCOE/Nodes/CustomizableObjectNodeMorphMaterial.h"
 #include "MuCOE/Nodes/CustomizableObjectNodeSkeletalMesh.h"
+#include "MuCOE/UnrealEditorPortabilityHelpers.h"
 #include "PropertyCustomizationHelpers.h"
 #include "Widgets/Input/STextComboBox.h"
 
@@ -74,7 +75,7 @@ void FCustomizableObjectNodeMorphMaterialDetails::CustomizeDetails( IDetailLayou
 			.CustomWidget()
 			[
 				SNew(SBorder)
-				.BorderImage(FAppStyle::GetBrush("NoBorder"))
+				.BorderImage(UE_MUTABLE_GET_BRUSH("NoBorder"))
 				.Padding(FMargin(0.0f, 0.0f, 10.0f, 0.0f))
 				[
 					SNew( SHorizontalBox )
diff --git a/Engine/Plugins/Experimental/Mutable/Source/CustomizableObjectEditor/Private/MuCOE/Nodes/CustomizableObjectNodeObjectDetails.cpp b/Engine/Plugins/Experimental/Mutable/Source/CustomizableObjectEditor/Private/MuCOE/Nodes/CustomizableObjectNodeObjectDetails.cpp
--- a/Engine/Plugins/Experimental/Mutable/Source/CustomizableObjectEditor/Private/MuCOE/Nodes/CustomizableObjectNodeObjectDetails.cpp
+++ b/Engine/Plugins/Experimental/Mutable/Source/CustomizableObjectEditor/Private/MuCOE/Nodes/CustomizableObjectNodeObjectDetails.cpp
@@ -12,6 +12,7 @@
 #include "MuCOE/Nodes/CustomizableObjectNodeObjectGroup.h"
 #include "MuCOE/SCustomizableObjectNodeObjectRTMorphTargetOverride.h"
 #include "MuCOE/Widgets/CustomizableObjectLODReductionSettings.h"
+#include "MuCOE/UnrealEditorPortabilityHelpers.h"
 #include "PropertyCustomizationHelpers.h"
 #include "ScopedTransaction.h"
 #include "Widgets/Input/SComboButton.h"
@@ -115,7 +116,7 @@ void FCustomizableObjectNodeObjectDetails::CustomizeDetails( IDetailLayoutBuilde
 					.CustomWidget()
 					[
 						SNew(SBorder)
-						.BorderImage(FAppStyle::GetBrush("NoBorder"))
+						.BorderImage(UE_MUTABLE_GET_BRUSH("NoBorder"))
 						.Padding(FMargin(0.0f, 0.0f, 10.0f, 0.0f))
 						[
 							SNew(STextComboBox)
diff --git a/Engine/Plugins/Experimental/Mutable/Source/CustomizableObjectEditor/Private/MuCOE/Nodes/CustomizableObjectNodeParentedMaterialDetails.cpp b/Engine/Plugins/Experimental/Mutable/Source/CustomizableObjectEditor/Private/MuCOE/Nodes/CustomizableObjectNodeParentedMaterialDetails.cpp
--- a/Engine/Plugins/Experimental/Mutable/Source/CustomizableObjectEditor/Private/MuCOE/Nodes/CustomizableObjectNodeParentedMaterialDetails.cpp
+++ b/Engine/Plugins/Experimental/Mutable/Source/CustomizableObjectEditor/Private/MuCOE/Nodes/CustomizableObjectNodeParentedMaterialDetails.cpp
@@ -10,6 +10,7 @@
 #include "MuCOE/Nodes/CustomizableObjectNodeExtendMaterial.h"
 #include "MuCOE/Nodes/CustomizableObjectNodeObject.h"
 #include "MuCOE/Nodes/CustomizableObjectNodeSkeletalMesh.h"
+#include "MuCOE/UnrealEditorPortabilityHelpers.h"
 #include "PropertyCustomizationHelpers.h"
 #include "ScopedTransaction.h"
 #include "Widgets/Input/STextComboBox.h"
@@ -115,7 +116,7 @@ void FCustomizableObjectNodeParentedMaterialDetails::CustomizeDetails(IDetailLay
 				.CustomWidget()
 				[
 					SNew(SBorder)
-					.BorderImage(FAppStyle::GetBrush("NoBorder"))
+					.BorderImage(UE_MUTABLE_GET_BRUSH("NoBorder"))
 					.Padding(FMargin(0.0f, 0.0f, 10.0f, 0.0f))
 					[
 						SNew(STextComboBox)
diff --git a/Engine/Plugins/Experimental/Mutable/Source/CustomizableObjectEditor/Private/MuCOE/Nodes/CustomizableObjectNodeProjectorConstant.cpp b/Engine/Plugins/Experimental/Mutable/Source/CustomizableObjectEditor/Private/MuCOE/Nodes/CustomizableObjectNodeProjectorConstant.cpp
--- a/Engine/Plugins/Experimental/Mutable/Source/CustomizableObjectEditor/Private/MuCOE/Nodes/CustomizableObjectNodeProjectorConstant.cpp
+++ b/Engine/Plugins/Experimental/Mutable/Source/CustomizableObjectEditor/Private/MuCOE/Nodes/CustomizableObjectNodeProjectorConstant.cpp
@@ -30,7 +30,7 @@ void UCustomizableObjectNodeProjectorConstant::PostEditChangeProperty(FPropertyC
 	{
 	}
 
-	const FName PropertyName = Helper_GetPropertyName(PropertyChangedEvent);
+	const FName PropertyName = PropertyChangedEvent.GetPropertyName();
 
 	if (PropertyName == "ProjectionType")
 	{
diff --git a/Engine/Plugins/Experimental/Mutable/Source/CustomizableObjectEditor/Private/MuCOE/Nodes/CustomizableObjectNodeProjectorParameter.cpp b/Engine/Plugins/Experimental/Mutable/Source/CustomizableObjectEditor/Private/MuCOE/Nodes/CustomizableObjectNodeProjectorParameter.cpp
--- a/Engine/Plugins/Experimental/Mutable/Source/CustomizableObjectEditor/Private/MuCOE/Nodes/CustomizableObjectNodeProjectorParameter.cpp
+++ b/Engine/Plugins/Experimental/Mutable/Source/CustomizableObjectEditor/Private/MuCOE/Nodes/CustomizableObjectNodeProjectorParameter.cpp
@@ -13,7 +13,7 @@ void UCustomizableObjectNodeProjectorParameter::PostEditChangeProperty(FProperty
 {
 	FProperty* PropertyThatChanged = PropertyChangedEvent.Property;
 
-	const FName PropertyName = Helper_GetPropertyName(PropertyChangedEvent);
+	const FName PropertyName = PropertyChangedEvent.GetPropertyName();
 
 	if (PropertyName == "ProjectionType")
 	{
diff --git a/Engine/Plugins/Experimental/Mutable/Source/CustomizableObjectEditor/Private/MuCOE/Nodes/CustomizableObjectNodeProjectorParameterDetails.cpp b/Engine/Plugins/Experimental/Mutable/Source/CustomizableObjectEditor/Private/MuCOE/Nodes/CustomizableObjectNodeProjectorParameterDetails.cpp
--- a/Engine/Plugins/Experimental/Mutable/Source/CustomizableObjectEditor/Private/MuCOE/Nodes/CustomizableObjectNodeProjectorParameterDetails.cpp
+++ b/Engine/Plugins/Experimental/Mutable/Source/CustomizableObjectEditor/Private/MuCOE/Nodes/CustomizableObjectNodeProjectorParameterDetails.cpp
@@ -11,6 +11,7 @@
 #include "MuCOE/ICustomizableObjectEditor.h"
 #include "MuCOE/Nodes/CustomizableObjectNodeProjectorConstant.h"
 #include "MuCOE/Nodes/CustomizableObjectNodeProjectorParameter.h"
+#include "MuCOE/UnrealEditorPortabilityHelpers.h"
 #include "PropertyCustomizationHelpers.h"
 #include "Widgets/Input/SButton.h"
 #include "Widgets/Input/STextComboBox.h"
@@ -145,7 +146,7 @@ void FCustomizableObjectNodeProjectorParameterDetails::CustomizeDetails( IDetail
 			.CustomWidget()
 			[
 				SNew(SBorder)
-				.BorderImage(FAppStyle::GetBrush("NoBorder"))
+				.BorderImage(UE_MUTABLE_GET_BRUSH("NoBorder"))
 				.Padding(FMargin(0.0f, 0.0f, 10.0f, 0.0f))
 				.HAlign(EHorizontalAlignment::HAlign_Fill)
 				[
diff --git a/Engine/Plugins/Experimental/Mutable/Source/CustomizableObjectEditor/Private/MuCOE/Nodes/CustomizableObjectNodeSkeletalMesh.cpp b/Engine/Plugins/Experimental/Mutable/Source/CustomizableObjectEditor/Private/MuCOE/Nodes/CustomizableObjectNodeSkeletalMesh.cpp
--- a/Engine/Plugins/Experimental/Mutable/Source/CustomizableObjectEditor/Private/MuCOE/Nodes/CustomizableObjectNodeSkeletalMesh.cpp
+++ b/Engine/Plugins/Experimental/Mutable/Source/CustomizableObjectEditor/Private/MuCOE/Nodes/CustomizableObjectNodeSkeletalMesh.cpp
@@ -117,11 +117,11 @@ void UCustomizableObjectNodeSkeletalMesh::AllocateDefaultPins(UCustomizableObjec
 
 	if (SkeletalMesh)
 	{
-		int NumLODs = Helper_GetLODInfoArray(SkeletalMesh).Num();
+		int NumLODs = SkeletalMesh->GetLODInfoArray().Num();
 		TMap<FString, TArray<int32>> MaterialNameUseCount; // Stores how many times the material name has been used per LOD
 		TMap<FString, TArray<int32>> ImageNameUseCount; // Stores how many times the image name has been used per LOD
 
-		if (const FSkeletalMeshModel* ImportedModel = Helper_GetImportedModel(SkeletalMesh))
+		if (const FSkeletalMeshModel* ImportedModel = SkeletalMesh->GetImportedModel())
 		{
 			for (int32 LODIndex = 0; LODIndex < NumLODs; ++LODIndex)
 			{
@@ -328,7 +328,7 @@ UTexture2D* UCustomizableObjectNodeSkeletalMesh::FindTextureForPin(const UEdGrap
 		return nullptr;
 	}
 
-	if (const FSkeletalMeshModel* ImportedModel = Helper_GetImportedModel(SkeletalMesh))
+	if (const FSkeletalMeshModel* ImportedModel = SkeletalMesh->GetImportedModel())
 	{
 		for (int LODIndex = 0; LODIndex < LODs.Num(); ++LODIndex)
 		{
@@ -523,13 +523,13 @@ bool UCustomizableObjectNodeSkeletalMesh::IsNodeOutDatedAndNeedsRefresh()
 
 	if (SkeletalMesh)
 	{
-		int32 numLods = Helper_GetLODNum(SkeletalMesh);
+		int32 numLods = SkeletalMesh->GetLODNum();
 		if (LODs.Num() != numLods)
 		{
 			Result = true;
 		}
 
-		if (const FSkeletalMeshModel* ImportedModel = Helper_GetImportedModel(SkeletalMesh))
+		if (const FSkeletalMeshModel* ImportedModel = SkeletalMesh->GetImportedModel())
 		{
 			for (int LODIndex = 0; !Result && LODIndex < LODs.Num(); ++LODIndex)
 			{
@@ -684,7 +684,7 @@ FSkeletalMaterial* UCustomizableObjectNodeSkeletalMesh::GetSkeletalMaterialFor(c
 	}
 	else
 	{
-		FSkeletalMeshModel* AuxImportedModel = Helper_GetImportedModel(SkeletalMesh);
+		FSkeletalMeshModel* AuxImportedModel = SkeletalMesh->GetImportedModel();
 
 		if (AuxImportedModel)
 		{
diff --git a/Engine/Plugins/Experimental/Mutable/Source/CustomizableObjectEditor/Private/MuCOE/Nodes/CustomizableObjectNodeStaticMesh.cpp b/Engine/Plugins/Experimental/Mutable/Source/CustomizableObjectEditor/Private/MuCOE/Nodes/CustomizableObjectNodeStaticMesh.cpp
--- a/Engine/Plugins/Experimental/Mutable/Source/CustomizableObjectEditor/Private/MuCOE/Nodes/CustomizableObjectNodeStaticMesh.cpp
+++ b/Engine/Plugins/Experimental/Mutable/Source/CustomizableObjectEditor/Private/MuCOE/Nodes/CustomizableObjectNodeStaticMesh.cpp
@@ -141,12 +141,16 @@ void UCustomizableObjectNodeStaticMesh::AllocateDefaultPins(UCustomizableObjectN
 				{
 					UMaterial* Material = MaterialInterface->GetMaterial();
 
-					Helper_ForEachTextureParameterName(Material, [&](const FString& ImageName)
-						{
-							Pin = CustomCreatePin(EGPD_Output, Schema->PC_Image, *ImageName);
-							Pin->bDefaultValueIsIgnored = true;
-							LODs[i].Materials[MaterialIndex].ImagePinsRef.Add(Pin);
-						});
+					TArray<FMaterialParameterInfo> ImageNames;
+					TArray<FGuid> ImageIds;
+					Material->GetAllTextureParameterInfo(ImageNames, ImageIds);
+					for (int32 ImageIndex = 0; ImageIndex < ImageNames.Num(); ++ImageIndex)
+					{
+						FString ImageName = ImageNames[ImageIndex].Name.ToString();
+						Pin = CustomCreatePin(EGPD_Output, Schema->PC_Image, *ImageName);
+						Pin->bDefaultValueIsIgnored = true;
+						LODs[i].Materials[MaterialIndex].ImagePinsRef.Add(Pin);
+					}
 				}
 			}
 		}
@@ -217,7 +221,7 @@ UTexture2D* UCustomizableObjectNodeStaticMesh::FindTextureForPin(const UEdGraphP
 					if (LODs[LODIndex].Materials[MaterialIndex].ImagePinsRef[ImageIndex].Get() == Pin)
 					{
 						UTexture* Texture=nullptr;
-						MaterialInterface->GetTextureParameterValue(Helper_GetPinFName(Pin), Texture);
+						MaterialInterface->GetTextureParameterValue( Pin?Pin->PinName:FName(), Texture);
 						return Cast<UTexture2D>(Texture);
 					}
 				}
diff --git a/Engine/Plugins/Experimental/Mutable/Source/CustomizableObjectEditor/Private/MuCOE/Nodes/CustomizableObjectNodeTable.cpp b/Engine/Plugins/Experimental/Mutable/Source/CustomizableObjectEditor/Private/MuCOE/Nodes/CustomizableObjectNodeTable.cpp
--- a/Engine/Plugins/Experimental/Mutable/Source/CustomizableObjectEditor/Private/MuCOE/Nodes/CustomizableObjectNodeTable.cpp
+++ b/Engine/Plugins/Experimental/Mutable/Source/CustomizableObjectEditor/Private/MuCOE/Nodes/CustomizableObjectNodeTable.cpp
@@ -240,12 +240,12 @@ void UCustomizableObjectNodeTable::GenerateMeshPins(UObject* Mesh, FString Name)
 {
 	if (USkeletalMesh* SkeletalMesh = Cast<USkeletalMesh>(Mesh))
 	{
-		int NumLODs = Helper_GetLODInfoArray(SkeletalMesh).Num();
+		int NumLODs = SkeletalMesh->GetLODInfoArray().Num();
 		const UEdGraphSchema_CustomizableObject* Schema = GetDefault<UEdGraphSchema_CustomizableObject>();
 
 		for (int32 LODIndex = 0; LODIndex < NumLODs; ++LODIndex)
 		{
-			int32 NumMaterials = Helper_GetImportedModel(SkeletalMesh)->LODModels[LODIndex].Sections.Num();
+			int32 NumMaterials = SkeletalMesh->GetImportedModel()->LODModels[LODIndex].Sections.Num();
 
 			for (int32 MatIndex = 0; MatIndex < NumMaterials; ++MatIndex)
 			{
@@ -261,7 +261,7 @@ void UCustomizableObjectNodeTable::GenerateMeshPins(UObject* Mesh, FString Name)
 				MeshPin->PinFriendlyName = FText::FromString(TableMeshPinName);
 				MeshPin->SafeSetHidden(false);
 
-				FSkeletalMeshModel* importModel = Helper_GetImportedModel(SkeletalMesh);
+				FSkeletalMeshModel* importModel = SkeletalMesh->GetImportedModel();
 				if (importModel && importModel->LODModels.Num() > LODIndex)
 				{
 					const uint32 NumberOfUVLayouts = importModel->LODModels[LODIndex].NumTexCoords;
@@ -392,11 +392,11 @@ bool UCustomizableObjectNodeTable::IsNodeOutDatedAndNeedsRefresh()
 				{
 					if (USkeletalMesh* SkeletalMesh = Cast<USkeletalMesh>(Object))
 					{
-						int NumLODs = Helper_GetLODInfoArray(SkeletalMesh).Num();
+						int NumLODs = SkeletalMesh->GetLODInfoArray().Num();
 
 						for (int32 LODIndex = 0; LODIndex < NumLODs; ++LODIndex)
 						{
-							int32 NumMaterials = Helper_GetImportedModel(SkeletalMesh)->LODModels[LODIndex].Sections.Num();
+							int32 NumMaterials = SkeletalMesh->GetImportedModel()->LODModels[LODIndex].Sections.Num();
 
 							for (int32 MatIndex = 0; MatIndex < NumMaterials; ++MatIndex)
 							{
diff --git a/Engine/Plugins/Experimental/Mutable/Source/CustomizableObjectEditor/Private/MuCOE/Nodes/CustomizableObjectNodeTableDetails.cpp b/Engine/Plugins/Experimental/Mutable/Source/CustomizableObjectEditor/Private/MuCOE/Nodes/CustomizableObjectNodeTableDetails.cpp
--- a/Engine/Plugins/Experimental/Mutable/Source/CustomizableObjectEditor/Private/MuCOE/Nodes/CustomizableObjectNodeTableDetails.cpp
+++ b/Engine/Plugins/Experimental/Mutable/Source/CustomizableObjectEditor/Private/MuCOE/Nodes/CustomizableObjectNodeTableDetails.cpp
@@ -12,6 +12,7 @@
 #include "MuCOE/CustomizableObjectLayout.h"
 #include "MuCOE/Nodes/CustomizableObjectNodeTable.h"
 #include "MuCOE/SCustomizableObjectNodeLayoutBlocksEditor.h"
+#include "MuCOE/UnrealEditorPortabilityHelpers.h"
 #include "Widgets/Input/STextComboBox.h"
 #include "Widgets/Input/SButton.h"
 
@@ -70,7 +71,7 @@ void FCustomizableObjectNodeTableDetails::CustomizeDetails(const TSharedPtr<IDet
 				+ SHorizontalBox::Slot()
 				[
 					SNew(SBorder)
-					.BorderImage(FAppStyle::GetBrush("NoBorder"))
+					.BorderImage(UE_MUTABLE_GET_BRUSH("NoBorder"))
 					.Padding(FMargin(0.0f, 0.0f, 10.0f, 0.0f))
 					[
 						SAssignNew(AnimMeshColumnComboBox, STextComboBox)
@@ -108,7 +109,7 @@ void FCustomizableObjectNodeTableDetails::CustomizeDetails(const TSharedPtr<IDet
 				+ SHorizontalBox::Slot()
 				[
 					SNew(SBorder)
-					.BorderImage(FAppStyle::GetBrush("NoBorder"))
+					.BorderImage(UE_MUTABLE_GET_BRUSH("NoBorder"))
 					.Padding(FMargin(0.0f, 0.0f, 10.0f, 0.0f))
 					[
 						SAssignNew(AnimComboBox, STextComboBox)
@@ -146,7 +147,7 @@ void FCustomizableObjectNodeTableDetails::CustomizeDetails(const TSharedPtr<IDet
 				+ SHorizontalBox::Slot()
 				[
 					SNew(SBorder)
-					.BorderImage(FAppStyle::GetBrush("NoBorder"))
+					.BorderImage(UE_MUTABLE_GET_BRUSH("NoBorder"))
 					.Padding(FMargin(0.0f, 0.0f, 10.0f, 0.0f))
 					[
 						SAssignNew(AnimSlotComboBox, STextComboBox)
@@ -184,7 +185,7 @@ void FCustomizableObjectNodeTableDetails::CustomizeDetails(const TSharedPtr<IDet
 				+ SHorizontalBox::Slot()
 				[
 					SNew(SBorder)
-					.BorderImage(FAppStyle::GetBrush("NoBorder"))
+					.BorderImage(UE_MUTABLE_GET_BRUSH("NoBorder"))
 					.Padding(FMargin(0.0f, 0.0f, 10.0f, 0.0f))
 					[
 						SAssignNew(AnimTagsComboBox, STextComboBox)
@@ -238,7 +239,7 @@ void FCustomizableObjectNodeTableDetails::CustomizeDetails(const TSharedPtr<IDet
 				+ SHorizontalBox::Slot()
 				[
 					SNew(SBorder)
-					.BorderImage(FAppStyle::GetBrush("NoBorder"))
+					.BorderImage(UE_MUTABLE_GET_BRUSH("NoBorder"))
 					.Padding(FMargin(0.0f, 0.0f, 10.0f, 0.0f))
 					[
 						SAssignNew(LayoutMeshColumnComboBox, STextComboBox)
diff --git a/Engine/Plugins/Experimental/Mutable/Source/CustomizableObjectEditor/Private/MuCOE/Nodes/CustomizableObjectNodeTextureLayer.cpp b/Engine/Plugins/Experimental/Mutable/Source/CustomizableObjectEditor/Private/MuCOE/Nodes/CustomizableObjectNodeTextureLayer.cpp
--- a/Engine/Plugins/Experimental/Mutable/Source/CustomizableObjectEditor/Private/MuCOE/Nodes/CustomizableObjectNodeTextureLayer.cpp
+++ b/Engine/Plugins/Experimental/Mutable/Source/CustomizableObjectEditor/Private/MuCOE/Nodes/CustomizableObjectNodeTextureLayer.cpp
@@ -120,8 +120,8 @@ bool UCustomizableObjectNodeTextureLayer::CanConnect(const UEdGraphPin* InOwnedI
 	{
 		if (LayerPin(LayerIndex) == InOwnedInputPin )
 		{
-			return (InOutputPin->PinType.PinCategory==Helper_GetPinCategory(Schema->PC_Image))
-				|| (InOutputPin->PinType.PinCategory==Helper_GetPinCategory(Schema->PC_Color));
+			return (InOutputPin->PinType.PinCategory==Schema->PC_Image)
+				|| (InOutputPin->PinType.PinCategory==Schema->PC_Color);
 		}
 	}
 
diff --git a/Engine/Plugins/Experimental/Mutable/Source/CustomizableObjectEditor/Private/MuCOE/Nodes/SCustomizableObjectNodeMaterialPinImage.cpp b/Engine/Plugins/Experimental/Mutable/Source/CustomizableObjectEditor/Private/MuCOE/Nodes/SCustomizableObjectNodeMaterialPinImage.cpp
--- a/Engine/Plugins/Experimental/Mutable/Source/CustomizableObjectEditor/Private/MuCOE/Nodes/SCustomizableObjectNodeMaterialPinImage.cpp
+++ b/Engine/Plugins/Experimental/Mutable/Source/CustomizableObjectEditor/Private/MuCOE/Nodes/SCustomizableObjectNodeMaterialPinImage.cpp
@@ -4,6 +4,7 @@
 
 #include "MuCOE/GraphTraversal.h"
 #include "MuCOE/Nodes/CustomizableObjectNodeMaterial.h"
+#include "MuCOE/UnrealEditorPortabilityHelpers.h"
 #include "Widgets/Input/SEditableTextBox.h"
 #include "Widgets/SToolTip.h"
 
@@ -25,9 +26,9 @@ void SCustomizableObjectNodeMaterialPinImage::Construct(const FArguments& InArgs
 	SetToolTip(TooltipWidget);
 
 	// Cache pin icons.
-	CachedPinMutableConnected = FAppStyle::GetBrush(TEXT("Graph.Pin.Connected"));
-	CachedPinMutableDisconnected = FAppStyle::GetBrush(TEXT("Graph.Pin.Disconnected"));
-	CachedPinPassthroughDisconnected = FAppStyle::GetBrush(TEXT("Graph.ExecPin.Disconnected"));
+	CachedPinMutableConnected = UE_MUTABLE_GET_BRUSH(TEXT("Graph.Pin.Connected"));
+	CachedPinMutableDisconnected = UE_MUTABLE_GET_BRUSH(TEXT("Graph.Pin.Disconnected"));
+	CachedPinPassthroughDisconnected = UE_MUTABLE_GET_BRUSH(TEXT("Graph.ExecPin.Disconnected"));
 }
 
 
diff --git a/Engine/Plugins/Experimental/Mutable/Source/CustomizableObjectEditor/Private/MuCOE/PinViewer/SPinViewerPinDetails.cpp b/Engine/Plugins/Experimental/Mutable/Source/CustomizableObjectEditor/Private/MuCOE/PinViewer/SPinViewerPinDetails.cpp
--- a/Engine/Plugins/Experimental/Mutable/Source/CustomizableObjectEditor/Private/MuCOE/PinViewer/SPinViewerPinDetails.cpp
+++ b/Engine/Plugins/Experimental/Mutable/Source/CustomizableObjectEditor/Private/MuCOE/PinViewer/SPinViewerPinDetails.cpp
@@ -4,6 +4,7 @@
 
 #include "MuCOE/PinViewer/SPinViewerDetailRowIndent.h"
 #include "MuCOE/Widgets/SMutableExpandableTableRow.h"
+#include "MuCOE/UnrealEditorPortabilityHelpers.h"
 #include "Widgets/Text/STextBlock.h"
 
 class SWidget;
@@ -60,7 +61,7 @@ TSharedPtr<SWidget> SPinViewerPinDetails::AddRow(const FText& Text, TSharedRef<S
 				.SizeRule(SSplitter::ESizeRule::FractionOfParent)
 				[
 					 SNew(SBorder)
-					 .BorderImage(FAppStyle::GetBrush("NoBorder"))
+					 .BorderImage(UE_MUTABLE_GET_BRUSH("NoBorder"))
 					 .VAlign(VAlign_Center)
 					 .Padding(FMargin(16, 0, 0, 0))
 					 [
@@ -71,7 +72,7 @@ TSharedPtr<SWidget> SPinViewerPinDetails::AddRow(const FText& Text, TSharedRef<S
 				+SSplitter::Slot()
 				[
 					SNew(SBorder)
-					.BorderImage(FAppStyle::GetBrush("NoBorder"))
+					.BorderImage(UE_MUTABLE_GET_BRUSH("NoBorder"))
 					.Padding(FMargin(4, 2, 0, 2))
 					[
 						Widget
diff --git a/Engine/Plugins/Experimental/Mutable/Source/CustomizableObjectEditor/Private/MuCOE/SCustomizableInstanceProperties.cpp b/Engine/Plugins/Experimental/Mutable/Source/CustomizableObjectEditor/Private/MuCOE/SCustomizableInstanceProperties.cpp
--- a/Engine/Plugins/Experimental/Mutable/Source/CustomizableObjectEditor/Private/MuCOE/SCustomizableInstanceProperties.cpp
+++ b/Engine/Plugins/Experimental/Mutable/Source/CustomizableObjectEditor/Private/MuCOE/SCustomizableInstanceProperties.cpp
@@ -12,6 +12,7 @@
 #include "MuCOE/CustomizableInstanceDetails.h"
 #include "MuCOE/CustomizableObjectEditorUtilities.h"
 #include "MuCOE/SMutableTextSearchBox.h"
+#include "MuCOE/UnrealEditorPortabilityHelpers.h"
 #include "Serialization/BufferArchive.h"
 #include "Slate/DeferredCleanupSlateBrush.h"
 #include "Widgets/Colors/SColorBlock.h"
@@ -837,7 +838,7 @@ void SCustomizableInstanceProperties::AddParameter(int32 ParamIndexInObject)
 					[
 						SNew(SImage)
 						// todo: editor style not available in game
-						//.Image(FAppStyle::GetBrush(TEXT("LevelEditor.ViewOptions.Small")))
+						//.Image(UE_MUTABLE_GET_BRUSH(TEXT("LevelEditor.ViewOptions.Small")))
 					]
 				]
 			
@@ -853,7 +854,7 @@ void SCustomizableInstanceProperties::AddParameter(int32 ParamIndexInObject)
 					[
 						SNew(SImage)
 						// todo: editor style not available in game
-						// .Image(FAppStyle::GetBrush(TEXT("LevelEditor.TranslateMode.Small")))
+						// .Image(UE_MUTABLE_GET_BRUSH(TEXT("LevelEditor.TranslateMode.Small")))
 					]
 				]
 			];
@@ -2308,7 +2309,7 @@ void SCreateProfileParameters::Construct(const FArguments& InArgs)
 		.Padding(2)
 		[
 			SNew(SBorder)
-			.BorderImage(FAppStyle::GetBrush("ToolPanel.GroupBorder"))
+			.BorderImage(UE_MUTABLE_GET_BRUSH("ToolPanel.GroupBorder"))
 		[
 			SNew(SVerticalBox)
 	+ SVerticalBox::Slot()
@@ -2337,14 +2338,14 @@ void SCreateProfileParameters::Construct(const FArguments& InArgs)
 		.Padding(5)
 		[
 			SNew(SUniformGridPanel)
-			.SlotPadding(FAppStyle::GetMargin("StandardDialog.SlotPadding"))
-			.MinDesiredSlotWidth(FAppStyle::GetFloat("StandardDialog.MinDesiredSlotWidth"))
-			.MinDesiredSlotHeight(FAppStyle::GetFloat("StandardDialog.MinDesiredSlotHeight"))
+			.SlotPadding(UE_MUTABLE_GET_MARGIN("StandardDialog.SlotPadding"))
+			.MinDesiredSlotWidth(UE_MUTABLE_GET_FLOAT("StandardDialog.MinDesiredSlotWidth"))
+			.MinDesiredSlotHeight(UE_MUTABLE_GET_FLOAT("StandardDialog.MinDesiredSlotHeight"))
 		+ SUniformGridPanel::Slot(0, 0)
 		[
 			SNew(SButton)
 			.HAlign(HAlign_Center)
-			.ContentPadding(FAppStyle::GetMargin("StandardDialog.ContentPadding"))
+			.ContentPadding(UE_MUTABLE_GET_MARGIN("StandardDialog.ContentPadding"))
 			.Text(LOCTEXT("OK", "OK"))
 			.OnClicked(this, &SCreateProfileParameters::OnButtonClick, EAppReturnType::Ok)
 		]
@@ -2352,7 +2353,7 @@ void SCreateProfileParameters::Construct(const FArguments& InArgs)
 		[
 			SNew(SButton)
 			.HAlign(HAlign_Center)
-			.ContentPadding(FAppStyle::GetMargin("StandardDialog.ContentPadding"))
+			.ContentPadding(UE_MUTABLE_GET_MARGIN("StandardDialog.ContentPadding"))
 			.Text(LOCTEXT("Cancel", "Cancel"))
 			.OnClicked(this, &SCreateProfileParameters::OnButtonClick, EAppReturnType::Cancel)
 		]
diff --git a/Engine/Plugins/Experimental/Mutable/Source/CustomizableObjectEditor/Private/MuCOE/SCustomizableObjectEditorAdvancedPreviewSettings.cpp b/Engine/Plugins/Experimental/Mutable/Source/CustomizableObjectEditor/Private/MuCOE/SCustomizableObjectEditorAdvancedPreviewSettings.cpp
--- a/Engine/Plugins/Experimental/Mutable/Source/CustomizableObjectEditor/Private/MuCOE/SCustomizableObjectEditorAdvancedPreviewSettings.cpp
+++ b/Engine/Plugins/Experimental/Mutable/Source/CustomizableObjectEditor/Private/MuCOE/SCustomizableObjectEditorAdvancedPreviewSettings.cpp
@@ -7,6 +7,7 @@
 #include "IContentBrowserSingleton.h"
 #include "Misc/MessageDialog.h"
 #include "Misc/Paths.h"
+#include "MuCOE/UnrealEditorPortabilityHelpers.h"
 #include "MuCOE/CustomizableObjectEditor.h"
 #include "MuCOE/CustomizableObjectEditorLevelProfile.h"
 #include "Widgets/Input/SButton.h"
@@ -28,7 +29,7 @@ void SLevelSelectWindow::Construct(const FArguments& InArgs)
 	DefaultSettings = InArgs._DefaultSettings.Get();
 
 	FAssetPickerConfig AssetPickerConfig;
-	AssetPickerConfig.Filter.ClassPaths.Add(FTopLevelAssetPath(TEXT("/Script/Engine"), TEXT("World")));
+	UE_MUTABLE_GET_CLASSPATHS(AssetPickerConfig.Filter).Add(UE_MUTABLE_TOPLEVELASSETPATH(TEXT("/Script/Engine"), TEXT("World")));
 	AssetPickerConfig.OnAssetSelected = FOnAssetSelected::CreateRaw(this, &SLevelSelectWindow::OnLevelAssetSelected);
 	AssetPickerConfig.ThumbnailScale = 0.4f;
 	FContentBrowserModule& ContentBrowserModule = FModuleManager::LoadModuleChecked<FContentBrowserModule>("ContentBrowser");
@@ -44,14 +45,14 @@ void SLevelSelectWindow::Construct(const FArguments& InArgs)
 			.Padding(2)
 			[
 				SNew(SBorder)
-				.BorderImage(FAppStyle::GetBrush("ToolPanel.GroupBorder"))
+				.BorderImage(UE_MUTABLE_GET_BRUSH("ToolPanel.GroupBorder"))
 				[
 					SNew(SVerticalBox)
 					+ SVerticalBox::Slot()
 					.FillHeight(1)
 					[
 						SNew(SBorder)
-						.BorderImage(FAppStyle::GetBrush("ToolPanel.GroupBorder"))
+						.BorderImage(UE_MUTABLE_GET_BRUSH("ToolPanel.GroupBorder"))
 						[
 							ContentBrowserModule.Get().CreateAssetPicker(AssetPickerConfig)
 						]
@@ -64,14 +65,14 @@ void SLevelSelectWindow::Construct(const FArguments& InArgs)
 			.Padding(5)
 			[
 				SNew(SUniformGridPanel)
-				.SlotPadding(FAppStyle::GetMargin("StandardDialog.SlotPadding"))
-				.MinDesiredSlotWidth(FAppStyle::GetFloat("StandardDialog.MinDesiredSlotWidth"))
-				.MinDesiredSlotHeight(FAppStyle::GetFloat("StandardDialog.MinDesiredSlotHeight"))
+				.SlotPadding(UE_MUTABLE_GET_MARGIN("StandardDialog.SlotPadding"))
+				.MinDesiredSlotWidth(UE_MUTABLE_GET_FLOAT("StandardDialog.MinDesiredSlotWidth"))
+				.MinDesiredSlotHeight(UE_MUTABLE_GET_FLOAT("StandardDialog.MinDesiredSlotHeight"))
 				+ SUniformGridPanel::Slot(0, 0)
 				[
 					SNew(SButton)
 					.HAlign(HAlign_Center)
-					.ContentPadding(FAppStyle::GetMargin("StandardDialog.ContentPadding"))
+					.ContentPadding(UE_MUTABLE_GET_MARGIN("StandardDialog.ContentPadding"))
 					.Text(LOCTEXT("OK", "OK"))
 					.OnClicked(this, &SLevelSelectWindow::OnButtonClick, EAppReturnType::Ok)
 				]
@@ -79,7 +80,7 @@ void SLevelSelectWindow::Construct(const FArguments& InArgs)
 				[
 					SNew(SButton)
 					.HAlign(HAlign_Center)
-					.ContentPadding(FAppStyle::GetMargin("StandardDialog.ContentPadding"))
+					.ContentPadding(UE_MUTABLE_GET_MARGIN("StandardDialog.ContentPadding"))
 					.Text(LOCTEXT("Cancel", "Cancel"))
 					.OnClicked(this, &SLevelSelectWindow::OnButtonClick, EAppReturnType::Cancel)
 				]
@@ -97,13 +98,13 @@ FReply SLevelSelectWindow::OnButtonClick(EAppReturnType::Type ButtonID)
 
 	if (UserResponse == EAppReturnType::Ok)
 	{
-		if (LevelAssetData.AssetClassPath != FTopLevelAssetPath(TEXT("/Script/Engine"), TEXT("World")))
+		if (UE_MUTABLE_ASSETCLASS(LevelAssetData) != UE_MUTABLE_TOPLEVELASSETPATH(TEXT("/Script/Engine"), TEXT("World")))
 		{
 			FMessageDialog::Open(EAppMsgType::Ok, LOCTEXT("PostProcessNameValueMissing", "Please select a Level asset"));
 			return FReply::Handled();
 		}
 
-		LevelAssetPath = LevelAssetData.GetObjectPathString();
+		LevelAssetPath = UE_MUTABLE_GETOBJECTPATH(LevelAssetData);
 
 		Result = LevelProfileManage::LoadProfileUObjectNames(LevelAssetPath, ArrayDirectionalLightName, ArrayPostProcessName, ArraySkyLightName);
 
@@ -214,19 +215,19 @@ void SUObjectSelecWindow::Construct(const FArguments& InArgs)
 			.Padding(2)
 			[
 				SNew(SBorder)
-				.BorderImage(FAppStyle::GetBrush("ToolPanel.GroupBorder"))
+				.BorderImage(UE_MUTABLE_GET_BRUSH("ToolPanel.GroupBorder"))
 				[
 					SNew(SVerticalBox)
 					+ SVerticalBox::Slot()
 					.AutoHeight()
 					[
 						SNew(SBorder)
-						.BorderImage(FAppStyle::GetBrush("ToolPanel.GroupBorder"))
+						.BorderImage(UE_MUTABLE_GET_BRUSH("ToolPanel.GroupBorder"))
 						[
 							SNew(SUniformGridPanel)
-							.SlotPadding(FAppStyle::GetMargin("StandardDialog.SlotPadding"))
-							.MinDesiredSlotWidth(FAppStyle::GetFloat("StandardDialog.MinDesiredSlotWidth"))
-							.MinDesiredSlotHeight(FAppStyle::GetFloat("StandardDialog.MinDesiredSlotHeight"))
+							.SlotPadding(UE_MUTABLE_GET_MARGIN("StandardDialog.SlotPadding"))
+							.MinDesiredSlotWidth(UE_MUTABLE_GET_FLOAT("StandardDialog.MinDesiredSlotWidth"))
+							.MinDesiredSlotHeight(UE_MUTABLE_GET_FLOAT("StandardDialog.MinDesiredSlotHeight"))
 							+ SUniformGridPanel::Slot(0, 0)
 							[
 								SNew(STextBlock)
@@ -245,12 +246,12 @@ void SUObjectSelecWindow::Construct(const FArguments& InArgs)
 					.AutoHeight()
 					[
 						SNew(SBorder)
-						.BorderImage(FAppStyle::GetBrush("ToolPanel.GroupBorder"))
+						.BorderImage(UE_MUTABLE_GET_BRUSH("ToolPanel.GroupBorder"))
 						[
 							SNew(SUniformGridPanel)
-							.SlotPadding(FAppStyle::GetMargin("StandardDialog.SlotPadding"))
-							.MinDesiredSlotWidth(FAppStyle::GetFloat("StandardDialog.MinDesiredSlotWidth"))
-							.MinDesiredSlotHeight(FAppStyle::GetFloat("StandardDialog.MinDesiredSlotHeight"))
+							.SlotPadding(UE_MUTABLE_GET_MARGIN("StandardDialog.SlotPadding"))
+							.MinDesiredSlotWidth(UE_MUTABLE_GET_FLOAT("StandardDialog.MinDesiredSlotWidth"))
+							.MinDesiredSlotHeight(UE_MUTABLE_GET_FLOAT("StandardDialog.MinDesiredSlotHeight"))
 							+ SUniformGridPanel::Slot(0, 0)
 							[
 								SNew(STextBlock)
@@ -268,12 +269,12 @@ void SUObjectSelecWindow::Construct(const FArguments& InArgs)
 					.AutoHeight()
 					[
 						SNew(SBorder)
-						.BorderImage(FAppStyle::GetBrush("ToolPanel.GroupBorder"))
+						.BorderImage(UE_MUTABLE_GET_BRUSH("ToolPanel.GroupBorder"))
 						[
 							SNew(SUniformGridPanel)
-							.SlotPadding(FAppStyle::GetMargin("StandardDialog.SlotPadding"))
-							.MinDesiredSlotWidth(FAppStyle::GetFloat("StandardDialog.MinDesiredSlotWidth"))
-							.MinDesiredSlotHeight(FAppStyle::GetFloat("StandardDialog.MinDesiredSlotHeight"))
+							.SlotPadding(UE_MUTABLE_GET_MARGIN("StandardDialog.SlotPadding"))
+							.MinDesiredSlotWidth(UE_MUTABLE_GET_FLOAT("StandardDialog.MinDesiredSlotWidth"))
+							.MinDesiredSlotHeight(UE_MUTABLE_GET_FLOAT("StandardDialog.MinDesiredSlotHeight"))
 							+ SUniformGridPanel::Slot(0, 0)
 							[
 								SNew(STextBlock)
@@ -291,12 +292,12 @@ void SUObjectSelecWindow::Construct(const FArguments& InArgs)
 					.AutoHeight()
 					[
 						SNew(SBorder)
-						.BorderImage(FAppStyle::GetBrush("ToolPanel.GroupBorder"))
+						.BorderImage(UE_MUTABLE_GET_BRUSH("ToolPanel.GroupBorder"))
 						[
 							SNew(SUniformGridPanel)
-							.SlotPadding(FAppStyle::GetMargin("StandardDialog.SlotPadding"))
-							.MinDesiredSlotWidth(FAppStyle::GetFloat("StandardDialog.MinDesiredSlotWidth"))
-							.MinDesiredSlotHeight(FAppStyle::GetFloat("StandardDialog.MinDesiredSlotHeight"))
+							.SlotPadding(UE_MUTABLE_GET_MARGIN("StandardDialog.SlotPadding"))
+							.MinDesiredSlotWidth(UE_MUTABLE_GET_FLOAT("StandardDialog.MinDesiredSlotWidth"))
+							.MinDesiredSlotHeight(UE_MUTABLE_GET_FLOAT("StandardDialog.MinDesiredSlotHeight"))
 							+ SUniformGridPanel::Slot(0, 0)
 							[
 								SNew(STextBlock)
@@ -318,14 +319,14 @@ void SUObjectSelecWindow::Construct(const FArguments& InArgs)
 			.Padding(5)
 			[
 				SNew(SUniformGridPanel)
-				.SlotPadding(FAppStyle::GetMargin("StandardDialog.SlotPadding"))
-				.MinDesiredSlotWidth(FAppStyle::GetFloat("StandardDialog.MinDesiredSlotWidth"))
-				.MinDesiredSlotHeight(FAppStyle::GetFloat("StandardDialog.MinDesiredSlotHeight"))
+				.SlotPadding(UE_MUTABLE_GET_MARGIN("StandardDialog.SlotPadding"))
+				.MinDesiredSlotWidth(UE_MUTABLE_GET_FLOAT("StandardDialog.MinDesiredSlotWidth"))
+				.MinDesiredSlotHeight(UE_MUTABLE_GET_FLOAT("StandardDialog.MinDesiredSlotHeight"))
 				+ SUniformGridPanel::Slot(0, 0)
 				[
 					SNew(SButton)
 					.HAlign(HAlign_Center)
-					.ContentPadding(FAppStyle::GetMargin("StandardDialog.ContentPadding"))
+					.ContentPadding(UE_MUTABLE_GET_MARGIN("StandardDialog.ContentPadding"))
 					.Text(LOCTEXT("OK", "OK"))
 					.OnClicked(this, &SUObjectSelecWindow::OnButtonClick, EAppReturnType::Ok)
 				]
@@ -333,7 +334,7 @@ void SUObjectSelecWindow::Construct(const FArguments& InArgs)
 				[
 					SNew(SButton)
 					.HAlign(HAlign_Center)
-					.ContentPadding(FAppStyle::GetMargin("StandardDialog.ContentPadding"))
+					.ContentPadding(UE_MUTABLE_GET_MARGIN("StandardDialog.ContentPadding"))
 					.Text(LOCTEXT("Cancel", "Cancel"))
 					.OnClicked(this, &SUObjectSelecWindow::OnButtonClick, EAppReturnType::Cancel)
 				]
@@ -439,9 +440,9 @@ void SCustomizableObjectEditorAdvancedPreviewSettings::Construct(const FArgument
 			.HAlign(HAlign_Fill)
 			[
 				SNew(SUniformGridPanel)
-				.SlotPadding(FAppStyle::GetMargin("StandardDialog.SlotPadding"))
-				.MinDesiredSlotWidth(FAppStyle::GetFloat("StandardDialog.MinDesiredSlotWidth"))
-				.MinDesiredSlotHeight(FAppStyle::GetFloat("StandardDialog.MinDesiredSlotHeight"))
+				.SlotPadding(UE_MUTABLE_GET_MARGIN("StandardDialog.SlotPadding"))
+				.MinDesiredSlotWidth(UE_MUTABLE_GET_FLOAT("StandardDialog.MinDesiredSlotWidth"))
+				.MinDesiredSlotHeight(UE_MUTABLE_GET_FLOAT("StandardDialog.MinDesiredSlotHeight"))
 				+ SUniformGridPanel::Slot(0, 0)
 				[
 					SNew(STextBlock)
diff --git a/Engine/Plugins/Experimental/Mutable/Source/CustomizableObjectEditor/Private/MuCOE/SCustomizableObjectEditorPerformanceReport.cpp b/Engine/Plugins/Experimental/Mutable/Source/CustomizableObjectEditor/Private/MuCOE/SCustomizableObjectEditorPerformanceReport.cpp
--- a/Engine/Plugins/Experimental/Mutable/Source/CustomizableObjectEditor/Private/MuCOE/SCustomizableObjectEditorPerformanceReport.cpp
+++ b/Engine/Plugins/Experimental/Mutable/Source/CustomizableObjectEditor/Private/MuCOE/SCustomizableObjectEditorPerformanceReport.cpp
@@ -13,6 +13,7 @@
 #include "MuCO/CustomizableObjectInstance.h"
 #include "MuCOE/CustomizableObjectEditorStyle.h"
 #include "MuCOE/SCustomizableObjectEditorTextureAnalyzer.h"
+#include "MuCOE/UnrealEditorPortabilityHelpers.h"
 #include "PropertyCustomizationHelpers.h"
 #include "Widgets/Input/SComboButton.h"
 #include "Widgets/Input/SNumericEntryBox.h"
@@ -1447,7 +1448,7 @@ const FSlateBrush* SPerformanceReportWorstCaseRow::GetSelectedRowHighlightColor(
 {
 	const bool bIsSelected = PerformanceReport && ReportInstance == PerformanceReport->CurrentReportWorstCase;
 	const bool bIsUpToDate = bIsSelected && ReportInstance && ReportInstance->WorstCaseInstance == PerformanceReport->CurrentReportInstance;
-	return bIsUpToDate? FAppStyle::GetBrush("Brushes.Highlight") : FAppStyle::GetBrush("Brushes.AccentGray");
+	return bIsUpToDate? UE_MUTABLE_GET_BRUSH("Brushes.Highlight") : UE_MUTABLE_GET_BRUSH("Brushes.AccentGray");
 }
 
 TSharedRef<ITableRow> SPerformanceReportWorstCaseRow::GenerateStyledRow(UWorstCasePerformanceReportInstance* ReportedInstance, const TSharedRef<STableViewBase>& OwnerTable)
diff --git a/Engine/Plugins/Experimental/Mutable/Source/CustomizableObjectEditor/Private/MuCOE/SCustomizableObjectEditorViewport.cpp b/Engine/Plugins/Experimental/Mutable/Source/CustomizableObjectEditor/Private/MuCOE/SCustomizableObjectEditorViewport.cpp
--- a/Engine/Plugins/Experimental/Mutable/Source/CustomizableObjectEditor/Private/MuCOE/SCustomizableObjectEditorViewport.cpp
+++ b/Engine/Plugins/Experimental/Mutable/Source/CustomizableObjectEditor/Private/MuCOE/SCustomizableObjectEditorViewport.cpp
@@ -778,7 +778,7 @@ TSharedRef<SWidget> SCustomizableObjectEditorViewportTabBody::BuildToolBar()
 		[
 			SNew(SBorder)
 			.Padding(0)
-		.BorderImage(FAppStyle::GetBrush("NoBorder"))
+		.BorderImage(UE_MUTABLE_GET_BRUSH("NoBorder"))
 		.IsEnabled(FSlateApplication::Get().GetNormalExecutionAttribute())
 		[
 			CommandToolbarBuilder.MakeWidget()
@@ -913,12 +913,12 @@ void SCustomizableObjectEditorViewportTabBody::GenerateUVMaterialOptions()
 
 	for (UDebugSkelMeshComponent* PreviewSkeletalMeshComponent : PreviewSkeletalMeshComponents)
 	{
-		if (PreviewSkeletalMeshComponent != nullptr && PreviewSkeletalMeshComponent->GetSkinnedAsset() != nullptr && PreviewSkeletalMeshComponent->GetSkinnedAsset()->GetResourceForRendering() != nullptr)
+		if (PreviewSkeletalMeshComponent != nullptr && UE_MUTABLE_GETSKINNEDASSET(PreviewSkeletalMeshComponent) != nullptr && UE_MUTABLE_GETSKINNEDASSET(PreviewSkeletalMeshComponent)->GetResourceForRendering() != nullptr)
 		{
 			TMap<FString, int32> MaterialLODs;
 
 			// Add Suffix "__X" for materials with duplicated names and Suffing " LOD_X" for multiple LODs.
-			const FSkeletalMeshRenderData* MeshRes = PreviewSkeletalMeshComponent->GetSkinnedAsset()->GetResourceForRendering();
+			const FSkeletalMeshRenderData* MeshRes = UE_MUTABLE_GETSKINNEDASSET(PreviewSkeletalMeshComponent)->GetResourceForRendering();
 			const TArray<UMaterialInterface*> Materials = PreviewSkeletalMeshComponent->GetMaterials();
 			for (UMaterialInterface* m : Materials)
 			{
@@ -1029,10 +1029,10 @@ void SCustomizableObjectEditorViewportTabBody::GenerateUVChannelOptions()
 
 	UDebugSkelMeshComponent* PreviewSkeletalMeshComponent = PreviewSkeletalMeshComponents[ComponentIndex];
 
-	if (PreviewSkeletalMeshComponent != nullptr && PreviewSkeletalMeshComponent->GetSkinnedAsset() != nullptr
-		&& PreviewSkeletalMeshComponent->GetSkinnedAsset()->GetResourceForRendering() != nullptr)
+	if (PreviewSkeletalMeshComponent != nullptr && UE_MUTABLE_GETSKINNEDASSET(PreviewSkeletalMeshComponent) != nullptr
+		&& UE_MUTABLE_GETSKINNEDASSET(PreviewSkeletalMeshComponent)->GetResourceForRendering() != nullptr)
 	{
-		const FSkeletalMeshRenderData* MeshRes = PreviewSkeletalMeshComponent->GetSkinnedAsset()->GetResourceForRendering();
+		const FSkeletalMeshRenderData* MeshRes = UE_MUTABLE_GETSKINNEDASSET(PreviewSkeletalMeshComponent)->GetResourceForRendering();
 		
 		int32 UVChannels = MeshRes->LODRenderData[LODIndex].GetNumTexCoords();
 		
@@ -1067,10 +1067,10 @@ FReply SCustomizableObjectEditorViewportTabBody::OnDrop(const FGeometry& MyGeome
 {
 	if (FAssetDragDropOp* DragDropOp = DragDropEvent.GetOperationAs<FAssetDragDropOp>().Get())
 	{
-		if (Helper_GetAssets(DragDropOp).Num())
+		if (DragDropOp->GetAssets().Num())
 		{
 			// This cast also includes UPoseAsset assets.
-			UAnimationAsset* AnimationAsset = Cast<UAnimationAsset>(Helper_GetAssets(DragDropOp)[0].GetAsset());
+			UAnimationAsset* AnimationAsset = Cast<UAnimationAsset>(DragDropOp->GetAssets()[0].GetAsset());
 			if (AnimationAsset)
 			{
 				LevelViewportClient->SetAnimation(AnimationAsset, EAnimationMode::AnimationSingleNode);
@@ -1090,9 +1090,9 @@ int32 SCustomizableObjectEditorViewportTabBody::GetLODModelCount() const
 
 	for (UDebugSkelMeshComponent* PreviewComponent : GetSkeletalMeshComponents())
 	{
-		if (PreviewComponent && PreviewComponent->GetSkinnedAsset())
+		if (PreviewComponent && UE_MUTABLE_GETSKINNEDASSET(PreviewComponent))
 		{
-			const TIndirectArray<FSkeletalMeshLODRenderData>& LODModels = Helper_GetLODData(PreviewComponent->GetSkinnedAsset());
+			const TIndirectArray<FSkeletalMeshLODRenderData>& LODModels = UE_MUTABLE_GETSKINNEDASSET(PreviewComponent)->GetResourceForRendering()->LODRenderData;
 			LODModelCount = FMath::Max(LODModelCount, LODModels.Num());
 		}
 	}
diff --git a/Engine/Plugins/Experimental/Mutable/Source/CustomizableObjectEditor/Private/MuCOE/SCustomizableObjectEditorViewportToolBar.cpp b/Engine/Plugins/Experimental/Mutable/Source/CustomizableObjectEditor/Private/MuCOE/SCustomizableObjectEditorViewportToolBar.cpp
--- a/Engine/Plugins/Experimental/Mutable/Source/CustomizableObjectEditor/Private/MuCOE/SCustomizableObjectEditorViewportToolBar.cpp
+++ b/Engine/Plugins/Experimental/Mutable/Source/CustomizableObjectEditor/Private/MuCOE/SCustomizableObjectEditorViewportToolBar.cpp
@@ -15,6 +15,7 @@
 #include "MuCOE/CustomizableObjectEditorViewportMenuCommands.h"
 #include "MuCOE/ICustomizableObjectInstanceEditor.h"
 #include "MuCOE/SCustomizableObjectEditorViewport.h"
+#include "MuCOE/UnrealEditorPortabilityHelpers.h"
 #include "Preferences/PersonaOptions.h"
 #include "SEditorViewportViewMenu.h"
 #include "SViewportToolBarComboMenu.h"
@@ -50,7 +51,7 @@ struct FSlateBrush;
 //
 //		TSharedPtr<SWidget> ExtraWidget =
 //			SNew(SBorder)
-//			.BorderImage(FAppStyle::GetBrush("FilledBorder"))
+//			.BorderImage(UE_MUTABLE_GET_BRUSH("FilledBorder"))
 //			[
 //				SNew(SHorizontalBox)
 //				+ SHorizontalBox::Slot()
@@ -165,7 +166,7 @@ struct FSlateBrush;
 //
 //		TSharedPtr<SWidget> ExtraWidget = SNew(STextBlock)
 //			.Text(AnimViewportPtr.Pin().ToSharedRef(), &SCustomizableObjectEditorViewport::GetWindStrengthLabel)
-//			.Font(FAppStyle::GetFontStyle(TEXT("MenuItem.Font")));
+//			.Font(UE_MUTABLE_GET_FONTSTYLE(TEXT("MenuItem.Font")));
 //
 //		this->ChildSlot
 //			[
@@ -228,7 +229,7 @@ struct FSlateBrush;
 //
 //		TSharedPtr<SWidget> ExtraWidget = SNew(STextBlock)
 //			.Text(AnimViewportPtr.Pin().ToSharedRef(), &SCustomizableObjectEditorViewport::GetGravityScaleLabel)
-//			.Font(FAppStyle::GetFontStyle(TEXT("MenuItem.Font")));
+//			.Font(UE_MUTABLE_GET_FONTSTYLE(TEXT("MenuItem.Font")));
 //
 //		this->ChildSlot
 //			[
@@ -387,7 +388,7 @@ void SCustomizableObjectEditorViewportToolBar::Construct(const FArguments& InArg
 	//		SNew(SEditorViewportToolbarMenu)
 	//		.ParentToolBar(SharedThis(this))
 	//	.Label(this, &SCustomizableObjectEditorViewportToolBar::GetPlaybackMenuLabel)
-	//	.LabelIcon(FAppStyle::GetBrush("AnimViewportMenu.PlayBackSpeed"))
+	//	.LabelIcon(UE_MUTABLE_GET_BRUSH("AnimViewportMenu.PlayBackSpeed"))
 	//	.OnGetMenuContent(this, &SCustomizableObjectEditorViewportToolBar::GeneratePlaybackMenu)
 	//	]
 
@@ -408,7 +409,7 @@ void SCustomizableObjectEditorViewportToolBar::Construct(const FArguments& InArg
 	FLinearColor ButtonColor2 = FLinearColor(0.2f, 0.2f, 0.2f, 0.75f);
 	FLinearColor TextColor1 = FLinearColor(1.0f, 1.0f, 1.0f, 1.0f);
 	FLinearColor TextColor2 = FLinearColor(0.8f, 0.8f, 0.8f, 0.8f);
-	FSlateFontInfo Info = FAppStyle::GetFontStyle("BoldFont");
+	FSlateFontInfo Info = UE_MUTABLE_GET_FONTSTYLE("BoldFont");
 	Info.Size += 26;
 
 	ChildSlot
@@ -418,8 +419,8 @@ void SCustomizableObjectEditorViewportToolBar::Construct(const FArguments& InArg
 			.VAlign(VAlign_Top)
 			[
 				SNew(SBorder)
-				.BorderImage(FAppStyle::GetBrush("NoBorder"))
-				.ForegroundColor(FAppStyle::GetSlateColor(DefaultForegroundName))
+				.BorderImage(UE_MUTABLE_GET_BRUSH("NoBorder"))
+				.ForegroundColor(UE_MUTABLE_GET_SLATECOLOR(DefaultForegroundName))
 				[
 					SNew(SVerticalBox)
 					+ SVerticalBox::Slot()
@@ -659,7 +660,7 @@ TSharedRef<SWidget> SCustomizableObjectEditorViewportToolBar::GenerateViewMenu()
 //			const float FOVMax = 170.f;
 //
 //			TSharedPtr<SWidget> FOVWidget = SNew(SSpinBox<float>)
-//				.Font(FAppStyle::GetFontStyle(TEXT("MenuItem.Font")))
+//				.Font(UE_MUTABLE_GET_FONTSTYLE(TEXT("MenuItem.Font")))
 //				.MinValue(FOVMin)
 //				.MaxValue(FOVMax)
 //				.Value(this, &SCustomizableObjectEditorViewportToolBar::OnGetFOVValue)
@@ -758,7 +759,7 @@ TSharedRef<SWidget> SCustomizableObjectEditorViewportToolBar::GenerateViewMenu()
 //	MenuBuilder.BeginSection("AnimViewportFloorOffset", LOCTEXT("Viewport_FloorOffsetLabel", "Floor Height Offset"));
 //	{
 //		TSharedPtr<SWidget> FloorOffsetWidget = SNew(SNumericEntryBox<float>)
-//			.Font(FAppStyle::GetFontStyle(TEXT("MenuItem.Font")))
+//			.Font(UE_MUTABLE_GET_FONTSTYLE(TEXT("MenuItem.Font")))
 //			.Value(this, &SCustomizableObjectEditorViewportToolBar::OnGetFloorOffset)
 //			.OnValueChanged(this, &SCustomizableObjectEditorViewportToolBar::OnFloorOffsetChanged)
 //			.ToolTipText(LOCTEXT("FloorOffsetToolTip", "Height offset for the floor mesh (stored per-mesh)"));
@@ -1077,7 +1078,7 @@ TSharedRef<SWidget> SCustomizableObjectEditorViewportToolBar::GenerateViewportOp
 				[
 					SNew(STextBlock)
 					.Text(LOCTEXT("OptionsMenu_CameraOptions_CameraSpeed_Text", "Camera Speed"))
-					.Font(FAppStyle::GetFontStyle(TEXT("MenuItem.Font")))
+					.Font(UE_MUTABLE_GET_FONTSTYLE(TEXT("MenuItem.Font")))
 				]
 
 				+SVerticalBox::Slot().AutoHeight()
@@ -1267,7 +1268,7 @@ const FSlateBrush* SCustomizableObjectEditorViewportToolBar::GetCameraMenuLabelI
 		}
 	}
 
-	return FAppStyle::GetBrush(Icon);
+	return UE_MUTABLE_GET_BRUSH(Icon);
 }
 
 float SCustomizableObjectEditorViewportToolBar::OnGetFOVValue() const
@@ -1437,7 +1438,7 @@ TSharedRef<SWidget> SCustomizableObjectEditorViewportToolBar::GenerateFOVMenu()
 			.WidthOverride( 100.0f )
 			[
 				SNew(SSpinBox<float>)
-				.Font( FAppStyle::GetFontStyle( TEXT( "MenuItem.Font" ) ) )
+				.Font( UE_MUTABLE_GET_FONTSTYLE( TEXT( "MenuItem.Font" ) ) )
 				.MinValue(FOVMin)
 				.MaxValue(FOVMax)
 				.Value( this, &SCustomizableObjectEditorViewportToolBar::OnGetFOVValue )
diff --git a/Engine/Plugins/Experimental/Mutable/Source/CustomizableObjectEditor/Private/MuCOE/SCustomizableObjectHighresScreenshot.cpp b/Engine/Plugins/Experimental/Mutable/Source/CustomizableObjectEditor/Private/MuCOE/SCustomizableObjectHighresScreenshot.cpp
--- a/Engine/Plugins/Experimental/Mutable/Source/CustomizableObjectEditor/Private/MuCOE/SCustomizableObjectHighresScreenshot.cpp
+++ b/Engine/Plugins/Experimental/Mutable/Source/CustomizableObjectEditor/Private/MuCOE/SCustomizableObjectHighresScreenshot.cpp
@@ -6,6 +6,7 @@
 #include "Framework/Docking/TabManager.h"
 #include "MuCOE/CustomizableObjectEditorViewportClient.h"
 #include "MuCOE/CustomizableObjectPreviewScene.h"
+#include "MuCOE/UnrealEditorPortabilityHelpers.h"
 #include "Widgets/Input/SButton.h"
 #include "Widgets/Input/SNumericEntryBox.h"
 #include "Widgets/Input/SSlider.h"
@@ -35,7 +36,7 @@ void SCustomizableObjectHighresScreenshot::Construct( const FArguments& InArgs )
 			+SVerticalBox::Slot()
 			[
 				SNew(SBorder)
-				.BorderImage(FAppStyle::GetBrush("ToolPanel.GroupBorder"))
+				.BorderImage(UE_MUTABLE_GET_BRUSH("ToolPanel.GroupBorder"))
 				[
 					SNew(SVerticalBox)
 					+SVerticalBox::Slot()
@@ -156,13 +157,13 @@ void SCustomizableObjectHighresScreenshot::Construct( const FArguments& InArgs )
 				.AutoHeight()
 				[
 					SNew(SImage)
-					.Image(FAppStyle::GetBrush("HighresScreenshot.WarningStrip"))
+					.Image(UE_MUTABLE_GET_BRUSH("HighresScreenshot.WarningStrip"))
 				]
 				+SVerticalBox::Slot()
 				.AutoHeight()
 				[
 					SNew( SBorder )
-					.BorderImage(FAppStyle::GetBrush("ToolPanel.GroupBorder"))
+					.BorderImage(UE_MUTABLE_GET_BRUSH("ToolPanel.GroupBorder"))
 					[
 						SNew( STextBlock )
 						.Text( NSLOCTEXT("CustomizableObjectEditor", "CaptureWarningText", "Due to the high system requirements of a high resolution screenshot, very large multipliers might cause the graphics driver to become unresponsive and possibly crash. In these circumstances, please try using a lower multiplier") )
@@ -173,14 +174,14 @@ void SCustomizableObjectHighresScreenshot::Construct( const FArguments& InArgs )
 				.AutoHeight()
 				[
 					SNew(SImage)
-					.Image(FAppStyle::GetBrush("HighresScreenshot.WarningStrip"))
+					.Image(UE_MUTABLE_GET_BRUSH("HighresScreenshot.WarningStrip"))
 				]
 			]
 			+SVerticalBox::Slot()
 			.AutoHeight()
 			[
 				SNew(SBorder)
-				.BorderImage(FAppStyle::GetBrush("ToolPanel.GroupBorder"))
+				.BorderImage(UE_MUTABLE_GET_BRUSH("ToolPanel.GroupBorder"))
 				[
 					SNew(SHorizontalBox)
 					+SHorizontalBox::Slot()
@@ -194,7 +195,7 @@ void SCustomizableObjectHighresScreenshot::Construct( const FArguments& InArgs )
 						.OnClicked( this, &SCustomizableObjectHighresScreenshot::OnCaptureClicked )
 						[
 							SNew(SImage)
-							.Image(FAppStyle::GetBrush("HighresScreenshot.Capture"))
+							.Image(UE_MUTABLE_GET_BRUSH("HighresScreenshot.Capture"))
 						]
 					]
 				]
diff --git a/Engine/Plugins/Experimental/Mutable/Source/CustomizableObjectEditor/Private/MuCOE/SCustomizableObjectLayoutGrid.cpp b/Engine/Plugins/Experimental/Mutable/Source/CustomizableObjectEditor/Private/MuCOE/SCustomizableObjectLayoutGrid.cpp
--- a/Engine/Plugins/Experimental/Mutable/Source/CustomizableObjectEditor/Private/MuCOE/SCustomizableObjectLayoutGrid.cpp
+++ b/Engine/Plugins/Experimental/Mutable/Source/CustomizableObjectEditor/Private/MuCOE/SCustomizableObjectLayoutGrid.cpp
@@ -8,6 +8,7 @@
 #include "Layout/WidgetPath.h"
 // Required for engine branch preprocessor defines.
 #include "MuCO/UnrealPortabilityHelpers.h"
+#include "MuCOE/UnrealEditorPortabilityHelpers.h"
 #include "Rendering/SlateRenderer.h"
 #include "RenderingThread.h"
 #include "Styling/SlateTypes.h"
@@ -442,7 +443,7 @@ FReply SCustomizableObjectLayoutGrid::OnMouseButtonDown( const FGeometry& MyGeom
 							.Value(this, &SCustomizableObjectLayoutGrid::GetBlockPriortyValue)
 							.UndeterminedString(LOCTEXT("MultipleValues", "Multiples Values"))
 							.OnValueChanged(this, &SCustomizableObjectLayoutGrid::OnBlockPriorityChanged)
-							.EditableTextBoxStyle(&FAppStyle::GetWidgetStyle<FEditableTextBoxStyle>("NormalEditableTextBox"))
+							.EditableTextBoxStyle(&UE_MUTABLE_GET_WIDGETSTYLE<FEditableTextBoxStyle>("NormalEditableTextBox"))
 						]
 					, FText::FromString("Block Priority"), true);
 
@@ -1178,11 +1179,7 @@ void FUVCanvasDrawer::DrawRenderThread(class FRHICommandListImmediate& RHICmdLis
 	{
 		RenderTarget->SetRenderTargetTexture(*(FTexture2DRHIRef*)InWindowBackBuffer);
 
-#if MUTABLE_CLEAN_ENGINE_BRANCH
 		FCanvas Canvas(RenderTarget, nullptr, FGameTime(), GMaxRHIFeatureLevel);
-#else
-		FCanvas Canvas(RenderTarget, nullptr, 0, 0, 0, GMaxRHIFeatureLevel);
-#endif
 		
 		Canvas.SetRenderTargetRect(RenderTarget->GetViewRect());
 		Canvas.SetRenderTargetScissorRect(RenderTarget->GetClippingRect());
diff --git a/Engine/Plugins/Experimental/Mutable/Source/CustomizableObjectEditor/Private/MuCOE/SCustomizableObjectNodeLayoutBlocksEditor.cpp b/Engine/Plugins/Experimental/Mutable/Source/CustomizableObjectEditor/Private/MuCOE/SCustomizableObjectNodeLayoutBlocksEditor.cpp
--- a/Engine/Plugins/Experimental/Mutable/Source/CustomizableObjectEditor/Private/MuCOE/SCustomizableObjectNodeLayoutBlocksEditor.cpp
+++ b/Engine/Plugins/Experimental/Mutable/Source/CustomizableObjectEditor/Private/MuCOE/SCustomizableObjectNodeLayoutBlocksEditor.cpp
@@ -7,6 +7,7 @@
 #include "Framework/MultiBox/MultiBoxBuilder.h"
 #include "MuCOE/CustomizableObjectEditorStyle.h"
 #include "MuCOE/SCustomizableObjectLayoutGrid.h"
+#include "MuCOE/UnrealEditorPortabilityHelpers.h"
 #include "Widgets/Input/STextComboBox.h"
 #include "Widgets/Layout/SGridPanel.h"
 #include "Widgets/SToolTip.h"
@@ -101,7 +102,7 @@ void SCustomizableObjectNodeLayoutBlocksEditor::SetCurrentLayout(UCustomizableOb
 		.Padding(3.0f,5.0f,0.0f,5.0f)
 		[
 			SNew(SImage)
-			.Image(FAppStyle::GetBrush(TEXT("Icons.Info")))
+			.Image(UE_MUTABLE_GET_BRUSH(TEXT("Icons.Info")))
 			.ToolTip(GenerateInfoToolTip())
 		]
 
@@ -225,7 +226,7 @@ TSharedRef<SWidget> SCustomizableObjectNodeLayoutBlocksEditor::BuildLayoutToolBa
 		[
 			SNew(SBorder)
 			.Padding(2.0f)
-			.BorderImage(FAppStyle::GetBrush("NoBorder"))
+			.BorderImage(UE_MUTABLE_GET_BRUSH("NoBorder"))
 			.IsEnabled( FSlateApplication::Get().GetNormalExecutionAttribute() )
 			[
 				LayoutToolbarBuilder.MakeWidget()
diff --git a/Engine/Plugins/Experimental/Mutable/Source/CustomizableObjectEditor/Private/MuCOE/SCustomizableObjectNodeLayoutBlocksSelector.cpp b/Engine/Plugins/Experimental/Mutable/Source/CustomizableObjectEditor/Private/MuCOE/SCustomizableObjectNodeLayoutBlocksSelector.cpp
--- a/Engine/Plugins/Experimental/Mutable/Source/CustomizableObjectEditor/Private/MuCOE/SCustomizableObjectNodeLayoutBlocksSelector.cpp
+++ b/Engine/Plugins/Experimental/Mutable/Source/CustomizableObjectEditor/Private/MuCOE/SCustomizableObjectNodeLayoutBlocksSelector.cpp
@@ -12,6 +12,7 @@
 #include "MuCOE/Nodes/CustomizableObjectNodeSkeletalMesh.h"
 #include "MuCOE/Nodes/CustomizableObjectNodeTable.h"
 #include "MuCOE/SCustomizableObjectLayoutGrid.h"
+#include "MuCOE/UnrealEditorPortabilityHelpers.h"
 #include "Widgets/Layout/SBorder.h"
 #include "Widgets/SBoxPanel.h"
 #include "Widgets/Text/STextBlock.h"
@@ -178,7 +179,7 @@ TSharedRef<SWidget> SCustomizableObjectNodeLayoutBlocksSelector::BuildLayoutTool
 		[
 			SNew(SBorder)
 			.Padding(0)
-			.BorderImage(FAppStyle::GetBrush("NoBorder"))
+			.BorderImage(UE_MUTABLE_GET_BRUSH("NoBorder"))
 			.IsEnabled( FSlateApplication::Get().GetNormalExecutionAttribute() )
 			[
 				LayoutToolbarBuilder.MakeWidget()
diff --git a/Engine/Plugins/Experimental/Mutable/Source/CustomizableObjectEditor/Private/MuCOE/SMutableCodeViewer.cpp b/Engine/Plugins/Experimental/Mutable/Source/CustomizableObjectEditor/Private/MuCOE/SMutableCodeViewer.cpp
--- a/Engine/Plugins/Experimental/Mutable/Source/CustomizableObjectEditor/Private/MuCOE/SMutableCodeViewer.cpp
+++ b/Engine/Plugins/Experimental/Mutable/Source/CustomizableObjectEditor/Private/MuCOE/SMutableCodeViewer.cpp
@@ -22,6 +22,7 @@
 #include "MuCOE/SMutableScalarViewer.h"
 #include "MuCOE/SMutableSkeletonViewer.h"
 #include "MuCOE/SMutableStringViewer.h"
+#include "MuCOE/UnrealEditorPortabilityHelpers.h"
 #include "MuR/SystemPrivate.h"
 #include "MuT/ErrorLogPrivate.h"
 #include "MuT/Streams.h"
@@ -659,7 +660,7 @@ void SMutableCodeViewer::Construct(const FArguments& InArgs, const TSharedPtr<mu
 				.FillHeight(1.0f)
 				[
 					SNew(SBorder)
-					.BorderImage(FAppStyle::GetBrush("ToolPanel.GroupBorder"))
+					.BorderImage(UE_MUTABLE_GET_BRUSH("ToolPanel.GroupBorder"))
 					.Padding(FMargin(4.0f, 4.0f))
 					[
 						SNew(SHorizontalBox)
@@ -778,7 +779,7 @@ void SMutableCodeViewer::Construct(const FArguments& InArgs, const TSharedPtr<mu
 				.Value(0.72f)
 				[
 					SAssignNew(PreviewBorder, SBorder)
-					.BorderImage(FAppStyle::GetBrush("ToolPanel.GroupBorder"))
+					.BorderImage(UE_MUTABLE_GET_BRUSH("ToolPanel.GroupBorder"))
 					.Padding(FMargin(4.0f, 4.0f))					
 				]
 			]
diff --git a/Engine/Plugins/Experimental/Mutable/Source/CustomizableObjectEditor/Private/MuCOE/SMutableGraphViewer.cpp b/Engine/Plugins/Experimental/Mutable/Source/CustomizableObjectEditor/Private/MuCOE/SMutableGraphViewer.cpp
--- a/Engine/Plugins/Experimental/Mutable/Source/CustomizableObjectEditor/Private/MuCOE/SMutableGraphViewer.cpp
+++ b/Engine/Plugins/Experimental/Mutable/Source/CustomizableObjectEditor/Private/MuCOE/SMutableGraphViewer.cpp
@@ -11,6 +11,7 @@
 #include "MuCOE/CustomizableObjectCompileRunnable.h"
 #include "MuCOE/CustomizableObjectEditorStyle.h"
 #include "MuCOE/SMutableCodeViewer.h"
+#include "MuCOE/UnrealEditorPortabilityHelpers.h"
 #include "MuT/Streams.h"
 #include "Widgets/Docking/SDockTab.h"
 #include "Widgets/Input/STextComboBox.h"
@@ -214,7 +215,7 @@ void SMutableGraphViewer::Construct(const FArguments& InArgs, const mu::NodePtr&
 			.Value(0.25f)
 			[
 				SNew(SBorder)
-				.BorderImage(FAppStyle::GetBrush("ToolPanel.GroupBorder"))
+				.BorderImage(UE_MUTABLE_GET_BRUSH("ToolPanel.GroupBorder"))
 				.Padding(FMargin(4.0f, 4.0f))
 				[
 					SAssignNew(TreeView, STreeView<TSharedPtr<FMutableGraphTreeElement>>)
@@ -237,7 +238,7 @@ void SMutableGraphViewer::Construct(const FArguments& InArgs, const mu::NodePtr&
 			.Value(0.75f)
 			[
 				SNew(SBorder)
-				.BorderImage(FAppStyle::GetBrush("ToolPanel.GroupBorder"))
+				.BorderImage(UE_MUTABLE_GET_BRUSH("ToolPanel.GroupBorder"))
 				.Padding(FMargin(4.0f, 4.0f))
 				//[
 				//	SubjectsTreeView->AsShared()
diff --git a/Engine/Plugins/Experimental/Mutable/Source/CustomizableObjectEditor/Private/MuCOE/SMutableObjectViewer.cpp b/Engine/Plugins/Experimental/Mutable/Source/CustomizableObjectEditor/Private/MuCOE/SMutableObjectViewer.cpp
--- a/Engine/Plugins/Experimental/Mutable/Source/CustomizableObjectEditor/Private/MuCOE/SMutableObjectViewer.cpp
+++ b/Engine/Plugins/Experimental/Mutable/Source/CustomizableObjectEditor/Private/MuCOE/SMutableObjectViewer.cpp
@@ -10,6 +10,7 @@
 #include "MuCOE/CustomizableObjectEditorStyle.h"
 #include "MuCOE/SMutableCodeViewer.h"
 #include "MuCOE/SMutableGraphViewer.h"
+#include "MuCOE/UnrealEditorPortabilityHelpers.h"
 #include "Widgets/Docking/SDockTab.h"
 #include "Widgets/Input/STextComboBox.h"
 #include "Widgets/Input/SNumericDropDown.h"
@@ -197,7 +198,7 @@ void SMutableObjectViewer::Construct(const FArguments& InArgs, UCustomizableObje
 		.VAlign(VAlign_Fill)
 		[
 			SNew(SBorder)
-			.BorderImage(FAppStyle::GetBrush("ToolPanel.GroupBorder"))
+			.BorderImage(UE_MUTABLE_GET_BRUSH("ToolPanel.GroupBorder"))
 			.Padding(FMargin(4.0f, 4.0f))
 			[
 				SAssignNew(TreeView, STreeView<TSharedPtr<FMutableObjectTreeElement>>)
diff --git a/Engine/Plugins/Experimental/Mutable/Source/CustomizableObjectEditor/Private/MuCOE/SMutableParametersWidget.cpp b/Engine/Plugins/Experimental/Mutable/Source/CustomizableObjectEditor/Private/MuCOE/SMutableParametersWidget.cpp
--- a/Engine/Plugins/Experimental/Mutable/Source/CustomizableObjectEditor/Private/MuCOE/SMutableParametersWidget.cpp
+++ b/Engine/Plugins/Experimental/Mutable/Source/CustomizableObjectEditor/Private/MuCOE/SMutableParametersWidget.cpp
@@ -3,6 +3,7 @@
 #include "MuCOE/SMutableParametersWidget.h"
 
 #include "MuCOE/SMutableTextSearchBox.h"
+#include "MuCOE/UnrealEditorPortabilityHelpers.h"
 #include "MuR/ModelPrivate.h"
 #include "MuR/ParametersPrivate.h"
 #include "Templates/SharedPointer.h"
@@ -334,7 +335,7 @@ void SMutableParametersWidget::GenerateAndAttachParameterSlate(const int32 Param
 				[
 					SNew(STextBlock)
 					.Text(FText::FromString(TEXT("Location")))
-					.Font(FAppStyle::GetFontStyle(TEXT("PropertyWindow.NormalFont")))
+					.Font(UE_MUTABLE_GET_FONTSTYLE(TEXT("PropertyWindow.NormalFont")))
 				]
 
 				+SGridPanel::Slot(1, 0)
@@ -359,7 +360,7 @@ void SMutableParametersWidget::GenerateAndAttachParameterSlate(const int32 Param
 				[
 					SNew(STextBlock)
 					.Text(FText::FromString(TEXT("Scale")))
-					.Font(FAppStyle::GetFontStyle("PropertyWindow.NormalFont"))
+					.Font(UE_MUTABLE_GET_FONTSTYLE("PropertyWindow.NormalFont"))
 				]
 
 				+ SGridPanel::Slot(1, 1)
diff --git a/Engine/Plugins/Experimental/Mutable/Source/CustomizableObjectEditor/Private/MuCOE/SMutableTextSearchBox.cpp b/Engine/Plugins/Experimental/Mutable/Source/CustomizableObjectEditor/Private/MuCOE/SMutableTextSearchBox.cpp
--- a/Engine/Plugins/Experimental/Mutable/Source/CustomizableObjectEditor/Private/MuCOE/SMutableTextSearchBox.cpp
+++ b/Engine/Plugins/Experimental/Mutable/Source/CustomizableObjectEditor/Private/MuCOE/SMutableTextSearchBox.cpp
@@ -2,6 +2,7 @@
 
 #include "MuCOE/SMutableTextSearchBox.h"
 
+#include "MuCOE/UnrealEditorPortabilityHelpers.h"
 #include "Framework/Application/SlateApplication.h"
 #include "Framework/Views/TableViewMetadata.h"
 #include "Layout/WidgetPath.h"
@@ -57,7 +58,7 @@ void SMutableTextSearchBox::Construct( const FArguments& InArgs )
 					.MenuContent
 						(
 							SNew(SBorder)
-							.BorderImage(FAppStyle::GetBrush("Menu.Background"))
+							.BorderImage(UE_MUTABLE_GET_BRUSH("Menu.Background"))
 							[
 								SNew(SBox)
 								.WidthOverride(175)	// to enforce some minimum width, ideally we define the minimum, not a fixed width
diff --git a/Engine/Plugins/Experimental/Mutable/Source/CustomizableObjectEditor/Private/MuCOE/StressTest.cpp b/Engine/Plugins/Experimental/Mutable/Source/CustomizableObjectEditor/Private/MuCOE/StressTest.cpp
--- a/Engine/Plugins/Experimental/Mutable/Source/CustomizableObjectEditor/Private/MuCOE/StressTest.cpp
+++ b/Engine/Plugins/Experimental/Mutable/Source/CustomizableObjectEditor/Private/MuCOE/StressTest.cpp
@@ -8,6 +8,7 @@
 #include "MuCO/CustomizableSkeletalComponent.h"
 #include "MuCO/UnrealPortabilityHelpers.h"
 #include "MuCOE/StatePerformanceTesting.h"
+#include "Engine/SkeletalMesh.h"
 #include "UObject/Package.h"
 
 class FSkeletalMeshLODRenderData;
@@ -29,9 +30,9 @@ void ULiveInstance::AddInstanceInformation()
 		return;
 	}
 
-	FSkeletalMeshRenderData* RenderData = Helper_GetResourceForRendering(Instance->GetSkeletalMesh());
-	TIndirectArray<FSkeletalMeshLODRenderData>* ArrayLODRenderData = Helper_GetLODDataPtr(RenderData);
-	int32 MaxIndexLOD = Helper_GetLODDataPtr(RenderData)->Num();
+	FSkeletalMeshRenderData* RenderData = Instance->GetSkeletalMesh()->GetResourceForRendering();
+	TIndirectArray<FSkeletalMeshLODRenderData>* ArrayLODRenderData = &RenderData->LODRenderData;
+	int32 MaxIndexLOD = RenderData->LODRenderData.Num();
 
 	for (int32 i = 0; i < MaxIndexLOD; ++i)
 	{
@@ -43,17 +44,17 @@ void ULiveInstance::AddInstanceInformation()
 
 		for (int32 c = 0; c < Instance->SkeletalMeshes.Num(); ++c)
 		{
-			RenderData = Helper_GetResourceForRendering(Instance->GetSkeletalMesh(c));
-			ArrayLODRenderData = Helper_GetLODDataPtr(RenderData);
-			MaxIndexLOD = Helper_GetLODDataPtr(RenderData)->Num();
-			CurrentLODNumFaces = Helper_LODGetTotalFaces(&((*ArrayLODRenderData)[i]));
-			const TArray<struct FSkeletalMeshLODInfo>& LODInfoArray = Helper_GetLODInfoArray(Instance->GetSkeletalMesh(c));
+			RenderData = Instance->GetSkeletalMesh(c)->GetResourceForRendering();
+			ArrayLODRenderData = &RenderData->LODRenderData;
+			MaxIndexLOD = RenderData->LODRenderData.Num();
+			CurrentLODNumFaces = ((*ArrayLODRenderData)[i]).GetTotalFaces();
+			const TArray<struct FSkeletalMeshLODInfo>& LODInfoArray = Instance->GetSkeletalMesh(c)->GetLODInfoArray();
 			if (LODInfoArray.IsValidIndex(i))
 			{
 				const FSkeletalMeshLODInfo& LODInfo = LODInfoArray[i];
 
-				MapMeasuredData["Number of faces"].AddMeasureData(Helper_LODGetTotalFaces(&((*ArrayLODRenderData)[i])));
-				MapMeasuredData["Number of texture coordinates"].AddMeasureData(Helper_LODGetNumTexCoords(&((*ArrayLODRenderData)[i])));
+				MapMeasuredData["Number of faces"].AddMeasureData((*ArrayLODRenderData)[i].GetTotalFaces());
+				MapMeasuredData["Number of texture coordinates"].AddMeasureData((*ArrayLODRenderData)[i].GetNumTexCoords());
 				const TArray<FSkeletalMaterial>& Materials = Instance->GetSkeletalMesh(c)->GetMaterials();
 				MapMeasuredData["Number of materials"].AddMeasureData(LODInfo.LODMaterialMap.Num());
 
diff --git a/Engine/Plugins/Experimental/Mutable/Source/CustomizableObjectEditor/Private/MuCOE/Thumbnails/CustomizableObjectThumbnailRenderer.cpp b/Engine/Plugins/Experimental/Mutable/Source/CustomizableObjectEditor/Private/MuCOE/Thumbnails/CustomizableObjectThumbnailRenderer.cpp
--- a/Engine/Plugins/Experimental/Mutable/Source/CustomizableObjectEditor/Private/MuCOE/Thumbnails/CustomizableObjectThumbnailRenderer.cpp
+++ b/Engine/Plugins/Experimental/Mutable/Source/CustomizableObjectEditor/Private/MuCOE/Thumbnails/CustomizableObjectThumbnailRenderer.cpp
@@ -71,7 +71,7 @@ void UCustomizableObjectThumbnailRenderer::Draw(UObject* Object, int32 X, int32
 				PreviewCustomizableSkeletalComponent->UpdateSkeletalMeshAsync();
 				USkeletalMesh* SkeletalMesh = PreviewCustomizableSkeletalComponent->GetSkeletalMesh();
 
-				if (SkeletalMesh && Helper_GetLODInfoArray(SkeletalMesh).Num())
+				if (SkeletalMesh && SkeletalMesh->GetLODInfoArray().Num())
 				{
 					ThumbnailScene->SetSkeletalMesh(SkeletalMesh);
 
diff --git a/Engine/Plugins/Experimental/Mutable/Source/CustomizableObjectEditor/Public/MuCOE/UnrealEditorPortabilityHelpers.h b/Engine/Plugins/Experimental/Mutable/Source/CustomizableObjectEditor/Public/MuCOE/UnrealEditorPortabilityHelpers.h
--- a/Engine/Plugins/Experimental/Mutable/Source/CustomizableObjectEditor/Public/MuCOE/UnrealEditorPortabilityHelpers.h
+++ b/Engine/Plugins/Experimental/Mutable/Source/CustomizableObjectEditor/Public/MuCOE/UnrealEditorPortabilityHelpers.h
@@ -14,117 +14,44 @@
 #include "Editor.h"
 #include "Engine/StaticMesh.h"
 #include "Engine/DataTable.h"
-#include <exception>
-#include <functional>
+#include "Engine/Texture.h"
+#include "Runtime/Launch/Resources/Version.h"
 
-class FPositionVertexBuffer;
-class FStaticMeshVertexBuffer;
 
 //---------------------------------------------------------------
 // Helpers to ease portability across unreal engine versions
 //---------------------------------------------------------------
 
+#if ENGINE_MAJOR_VERSION==5 && ENGINE_MINOR_VERSION>=1
 
-inline bool Helper_DoSectionsUse16BitBoneIndex( const FSkeletalMeshLODModel& m )
-{
-	for (int32 SectionIdx = 0; SectionIdx < m.Sections.Num(); ++SectionIdx)
-	{
-		if (m.Sections[SectionIdx].Use16BitBoneIndex())
-		{
-			return true;
-		}
-	}
-
-	return false;
-}
-
-inline int Helper_GetMaxBoneInfluences(const FSkeletalMeshLODModel& m)
-{
-	return m.GetMaxBoneInfluences();
-}
-
-
-inline FMaterialRenderProxy* Helper_GetMaterialProxy( const UMaterialInterface* Material )
-{
-	return Material->GetRenderProxy();
-}
+#define UE_MUTABLE_GET_BRUSH			FAppStyle::GetBrush
+#define UE_MUTABLE_GET_FLOAT			FAppStyle::GetFloat
+#define UE_MUTABLE_GET_MARGIN			FAppStyle::GetMargin
+#define UE_MUTABLE_GET_FONTSTYLE		FAppStyle::GetFontStyle
+#define UE_MUTABLE_GET_WIDGETSTYLE		FAppStyle::GetWidgetStyle
+#define UE_MUTABLE_GET_COLOR			FAppStyle::GetColor
+#define UE_MUTABLE_GET_SLATECOLOR		FAppStyle::GetSlateColor
 
+#elif ENGINE_MAJOR_VERSION==5 && ENGINE_MINOR_VERSION==0
 
-inline const TArray<FAssetData>& Helper_GetAssets(FAssetDragDropOp* DragDropOp)
-{
-	return DragDropOp->GetAssets();
-}
+#define UE_MUTABLE_GET_BRUSH			FAppStyle::Get().GetBrush
+#define UE_MUTABLE_GET_FLOAT			FAppStyle::Get().GetFloat
+#define UE_MUTABLE_GET_MARGIN			FAppStyle::Get().GetMargin
+#define UE_MUTABLE_GET_FONTSTYLE		FAppStyle::Get().GetFontStyle
+#define UE_MUTABLE_GET_WIDGETSTYLE		FAppStyle::Get().GetWidgetStyle
+#define UE_MUTABLE_GET_COLOR			FAppStyle::Get().GetColor
+#define UE_MUTABLE_GET_SLATECOLOR		FAppStyle::Get().GetSlateColor
 
-inline FName Helper_GetPropertyName(FPropertyChangedEvent& PropertyChangedEvent)
-{
-	return PropertyChangedEvent.GetPropertyName();
-}
+#endif
 
-inline FName Helper_GetPinFName(const UEdGraphPin* Pin)
-{
-	return Pin ? Pin->PinName: FName();
-}
 
 inline FString Helper_GetPinName(const UEdGraphPin* Pin)
 {
 	return Pin ? Pin->PinName.ToString() : FString();
 }
 
-inline void Helper_SetPinName( FName& To, const FString& From )
-{
-	To = FName(*From);
-}
-
-inline const FName& Helper_GetPinCategory(const FName& Cat)
-{
-	return Cat;
-}
-
-inline void Helper_GetMeshIndices(TArray<uint32>& Indices, USkeletalMesh* SkeletalMesh, int LODIndex)
-{
-	Indices = Helper_GetImportedModel(SkeletalMesh)->LODModels[LODIndex].IndexBuffer;
-}
-
-
-inline void Helper_ForEachTextureParameterName(const UMaterial* Material, const std::function<void(const FString&)>& f)
-{
-	TArray<FMaterialParameterInfo> ImageNames;
-	TArray<FGuid> ImageIds;
-	Material->GetAllTextureParameterInfo(ImageNames, ImageIds);
-	for (int32 ImageIndex = 0; ImageIndex < ImageNames.Num(); ++ImageIndex)
-	{
-		FString ImageName = ImageNames[ImageIndex].Name.ToString();
-		f(ImageName);
-	}
-}
-
-
-inline FStaticMeshVertexBuffer& Helper_GetStaticMeshVertexBuffer(UStaticMesh* StaticMesh, int LODIndex)
-{
-	return StaticMesh->GetRenderData()->LODResources[LODIndex].VertexBuffers.StaticMeshVertexBuffer;
-}
-
-
-inline const FPositionVertexBuffer& Helper_GetStaticMeshPositionVertexBuffer(UStaticMesh* StaticMesh, int LODIndex)
-{
-	return StaticMesh->GetRenderData()->LODResources[LODIndex].VertexBuffers.PositionVertexBuffer;
-}
-
-inline const TMap<FName, uint8*>& Helper_GetRowMap(UDataTable* Table)
-{
-	return Table->GetRowMap();
-}
-
-inline UAssetEditorSubsystem* Helper_GetEditorSubsystem()
-{
-	return GEditor->GetEditorSubsystem<UAssetEditorSubsystem>();
-}
-
-
-
 
 #if UE_ENABLE_INCLUDE_ORDER_DEPRECATED_IN_5_2
 #include "Framework/Commands/InputChord.h"
 #include "MuCO/ICustomizableObjectModule.h"
-#include "Runtime/Launch/Resources/Version.h"
 #endif
diff --git a/Engine/Plugins/Experimental/Mutable/Source/CustomizableObjectPopulation/Private/MuCOP/CustomizableObjectPopulationModule.cpp b/Engine/Plugins/Experimental/Mutable/Source/CustomizableObjectPopulation/Private/MuCOP/CustomizableObjectPopulationModule.cpp
--- a/Engine/Plugins/Experimental/Mutable/Source/CustomizableObjectPopulation/Private/MuCOP/CustomizableObjectPopulationModule.cpp
+++ b/Engine/Plugins/Experimental/Mutable/Source/CustomizableObjectPopulation/Private/MuCOP/CustomizableObjectPopulationModule.cpp
@@ -5,6 +5,7 @@
 #include "AssetRegistry/AssetRegistryModule.h"
 #include "Interfaces/IPluginManager.h"
 #include "MuCO/CustomizableObject.h"
+#include "MuCO/UnrealPortabilityHelpers.h"
 #include "MuCOE/ICustomizableObjectPopulationModule.h"
 #include "MuCOP/CustomizableObjectPopulation.h"
 #include "MuCOP/CustomizableObjectPopulationGenerator.h"
@@ -68,7 +69,7 @@ void FCustomizableObjectPopulationModule::RecompilePopulations(UCustomizableObje
 	{
 		// Creating a filter to search only population classes
 		FARFilter COFilter;
-		COFilter.ClassPaths.Add(UCustomizableObjectPopulation::StaticClass()->GetClassPathName());
+		UE_MUTABLE_GET_CLASSPATHS(COFilter).Add(UE_MUTABLE_GET_CLASS_PATHNAME(UCustomizableObjectPopulation::StaticClass()));
 
 		for (const FName AssetName : COAssetNames)
 		{
diff --git a/Engine/Plugins/Experimental/Mutable/Source/CustomizableObjectPopulationEditor/Private/MuCOPE/CustomizableObjectPopulationClassDetails.cpp b/Engine/Plugins/Experimental/Mutable/Source/CustomizableObjectPopulationEditor/Private/MuCOPE/CustomizableObjectPopulationClassDetails.cpp
--- a/Engine/Plugins/Experimental/Mutable/Source/CustomizableObjectPopulationEditor/Private/MuCOPE/CustomizableObjectPopulationClassDetails.cpp
+++ b/Engine/Plugins/Experimental/Mutable/Source/CustomizableObjectPopulationEditor/Private/MuCOPE/CustomizableObjectPopulationClassDetails.cpp
@@ -2,6 +2,7 @@
 
 #include "MuCOPE/CustomizableObjectPopulationClassDetails.h"
 
+#include "MuCOE/UnrealEditorPortabilityHelpers.h"
 #include "Curves/CurveLinearColor.h"
 #include "DetailLayoutBuilder.h"
 #include "Framework/Views/TableViewMetadata.h"
@@ -250,7 +251,7 @@ void SPopulationClassTagManager::BuildAllowlistWidget()
 			.ToolTipText(LOCTEXT("AddAllowlistTag", "Add Allow List Tag"))
 			[
 				SNew(SImage)
-				.Image(FAppStyle::GetBrush(TEXT("Plus")))
+				.Image(UE_MUTABLE_GET_BRUSH(TEXT("Plus")))
 			]
 		]
 	];
@@ -307,7 +308,7 @@ void SPopulationClassTagManager::BuildBlocklistWidget()
 			.ToolTipText(LOCTEXT("AddBlocklistTag", "Add Block List Tag"))
 			[
 				SNew(SImage)
-				.Image(FAppStyle::GetBrush(TEXT("Plus")))
+				.Image(UE_MUTABLE_GET_BRUSH(TEXT("Plus")))
 			]
 		]
 	];
@@ -435,7 +436,7 @@ void SPopulationClassTagList::Construct(const FArguments& InArgs)
 		.ToolTipText(LOCTEXT("RemoveTag", "Remove Tag"))
 		[
 			SNew(SImage)
-			.Image(FAppStyle::GetBrush(TEXT("Cross")))
+			.Image(UE_MUTABLE_GET_BRUSH(TEXT("Cross")))
 		]
 	];
 
@@ -607,7 +608,7 @@ void SPopulationClassCharacteristic::Construct(const FArguments& InArgs)
 				.ToolTipText(LOCTEXT("AddConstraint", "Add a New Constraint"))
 				[
 					SNew(SImage)
-					.Image(FAppStyle::GetBrush(TEXT("Plus")))
+					.Image(UE_MUTABLE_GET_BRUSH(TEXT("Plus")))
 				]
 			];
 		}
@@ -689,7 +690,7 @@ void SPopulationClassCharacteristic::Construct(const FArguments& InArgs)
 				.OnClicked(this, &SPopulationClassCharacteristic::OnRemoveCharacteristicButtonPressed)
 				[
 					SNew(SImage)
-					.Image(FAppStyle::GetBrush(TEXT("Cross")))
+					.Image(UE_MUTABLE_GET_BRUSH(TEXT("Cross")))
 				]
 			]
 
@@ -861,7 +862,7 @@ void SPopulationClassConstraint::Construct(const FArguments& InArgs)
 							.Visibility(multiConstraint ? EVisibility::Visible : EVisibility::Collapsed)
 							[
 								SNew(SImage)
-								.Image(FAppStyle::GetBrush(TEXT("Cross")))
+								.Image(UE_MUTABLE_GET_BRUSH(TEXT("Cross")))
 							]
 						]
 				]
@@ -1667,7 +1668,7 @@ void SRangeWidget::Construct(const FArguments& InArgs)
 					.OnClicked(this, &SRangeWidget::OnRemoveRangeButtonPressed)
 					[
 						SNew(SImage)
-						.Image(FAppStyle::GetBrush(TEXT("Cross")))
+						.Image(UE_MUTABLE_GET_BRUSH(TEXT("Cross")))
 					]
 				]
 			]
diff --git a/Engine/Plugins/Experimental/Mutable/Source/CustomizableObjectPopulationEditor/Private/MuCOPE/CustomizableObjectPopulationClassEditor.cpp b/Engine/Plugins/Experimental/Mutable/Source/CustomizableObjectPopulationEditor/Private/MuCOPE/CustomizableObjectPopulationClassEditor.cpp
--- a/Engine/Plugins/Experimental/Mutable/Source/CustomizableObjectPopulationEditor/Private/MuCOPE/CustomizableObjectPopulationClassEditor.cpp
+++ b/Engine/Plugins/Experimental/Mutable/Source/CustomizableObjectPopulationEditor/Private/MuCOPE/CustomizableObjectPopulationClassEditor.cpp
@@ -29,6 +29,7 @@
 #include "MuCOPE/CustomizableObjectPopulationEditorModule.h"
 #include "MuCOPE/SCustomizableObjectPopulationClassTagsTool.h"
 #include "MuCOPE/SCustomizableObjectPopulationEditorViewport.h"
+#include "MuCOE/UnrealEditorPortabilityHelpers.h"
 #include "PropertyEditorModule.h"
 #include "SCurveEditor.h"
 #include "Widgets/Docking/SDockTab.h"
@@ -979,7 +980,7 @@ void FCustomizableObjectPopulationClassEditor::GeneratePopulationClassInstances(
 			{
 				//Search for asset with the same name
 				FString ObjectPath = FilePath + FString(".") + ObjectName;
-				AssetData = AssetRegistryModule.Get().GetAssetByObjectPath(FSoftObjectPath(ObjectPath));
+				AssetData = AssetRegistryModule.Get().GetAssetByObjectPath(UE_MUTABLE_OBJECTPATH(ObjectPath));
 
 				// Popup window
 				if (AssetData.IsValid() && AssetData.GetClass() == UCustomizableObjectInstance::StaticClass() && RetType != EAppReturnType::YesAll)
@@ -1097,7 +1098,7 @@ void SSelectPopulationClassFolderDlg::Construct(const FArguments& InArgs)
 		.Padding(2)
 		[
 			SNew(SBorder)
-			.BorderImage(FAppStyle::GetBrush("ToolPanel.GroupBorder"))
+			.BorderImage(UE_MUTABLE_GET_BRUSH("ToolPanel.GroupBorder"))
 			[
 				SNew(SVerticalBox)
 
@@ -1141,15 +1142,15 @@ void SSelectPopulationClassFolderDlg::Construct(const FArguments& InArgs)
 		.Padding(5)
 		[
 			SNew(SUniformGridPanel)
-			.SlotPadding(FAppStyle::GetMargin("StandardDialog.SlotPadding"))
-			.MinDesiredSlotWidth(FAppStyle::GetFloat("StandardDialog.MinDesiredSlotWidth"))
-			.MinDesiredSlotHeight(FAppStyle::GetFloat("StandardDialog.MinDesiredSlotHeight"))
+			.SlotPadding(UE_MUTABLE_GET_MARGIN("StandardDialog.SlotPadding"))
+			.MinDesiredSlotWidth(UE_MUTABLE_GET_FLOAT("StandardDialog.MinDesiredSlotWidth"))
+			.MinDesiredSlotHeight(UE_MUTABLE_GET_FLOAT("StandardDialog.MinDesiredSlotHeight"))
 
 			+ SUniformGridPanel::Slot(0, 0)
 			[
 				SNew(SButton)
 				.HAlign(HAlign_Center)
-				.ContentPadding(FAppStyle::GetMargin("StandardDialog.ContentPadding"))
+				.ContentPadding(UE_MUTABLE_GET_MARGIN("StandardDialog.ContentPadding"))
 				.Text(LOCTEXT("OK", "OK"))
 				.OnClicked(this, &SSelectPopulationClassFolderDlg::OnButtonClick, EAppReturnType::Ok)
 			]
@@ -1158,7 +1159,7 @@ void SSelectPopulationClassFolderDlg::Construct(const FArguments& InArgs)
 			[
 				SNew(SButton)
 				.HAlign(HAlign_Center)
-				.ContentPadding(FAppStyle::GetMargin("StandardDialog.ContentPadding"))
+				.ContentPadding(UE_MUTABLE_GET_MARGIN("StandardDialog.ContentPadding"))
 				.Text(LOCTEXT("Cancel", "Cancel"))
 				.OnClicked(this, &SSelectPopulationClassFolderDlg::OnButtonClick, EAppReturnType::Cancel)
 			]
diff --git a/Engine/Plugins/Experimental/Mutable/Source/CustomizableObjectPopulationEditor/Private/MuCOPE/CustomizableObjectPopulationEditor.cpp b/Engine/Plugins/Experimental/Mutable/Source/CustomizableObjectPopulationEditor/Private/MuCOPE/CustomizableObjectPopulationEditor.cpp
--- a/Engine/Plugins/Experimental/Mutable/Source/CustomizableObjectPopulationEditor/Private/MuCOPE/CustomizableObjectPopulationEditor.cpp
+++ b/Engine/Plugins/Experimental/Mutable/Source/CustomizableObjectPopulationEditor/Private/MuCOPE/CustomizableObjectPopulationEditor.cpp
@@ -26,6 +26,7 @@
 #include "MuCOPE/CustomizableObjectPopulationEditorActions.h"
 #include "MuCOPE/CustomizableObjectPopulationEditorModule.h"
 #include "MuCOPE/SCustomizableObjectPopulationEditorViewport.h"
+#include "MuCOE/UnrealEditorPortabilityHelpers.h"
 #include "PropertyEditorModule.h"
 #include "Widgets/Docking/SDockTab.h"
 #include "Widgets/Input/SButton.h"
@@ -614,7 +615,7 @@ void FCustomizableObjectPopulationEditor::GeneratePopulationInstances()
 			{
 				//Search for asset with the same name
 				FString ObjectPath = FilePath + FString(".") + ObjectName;
-				AssetData = AssetRegistryModule.Get().GetAssetByObjectPath(FSoftObjectPath(ObjectPath));
+				AssetData = AssetRegistryModule.Get().GetAssetByObjectPath(UE_MUTABLE_OBJECTPATH(ObjectPath));
 
 				// Popup window
 				if (AssetData.IsValid() && AssetData.GetClass() == UCustomizableObjectInstance::StaticClass() && RetType != EAppReturnType::YesAll)
@@ -734,7 +735,7 @@ void SSelectPopulationFolderDlg::Construct(const FArguments& InArgs)
 		.Padding(2)
 		[
 			SNew(SBorder)
-			.BorderImage(FAppStyle::GetBrush("ToolPanel.GroupBorder"))
+			.BorderImage(UE_MUTABLE_GET_BRUSH("ToolPanel.GroupBorder"))
 			[
 				SNew(SVerticalBox)
 
@@ -778,15 +779,15 @@ void SSelectPopulationFolderDlg::Construct(const FArguments& InArgs)
 		.Padding(5)
 		[
 			SNew(SUniformGridPanel)
-			.SlotPadding(FAppStyle::GetMargin("StandardDialog.SlotPadding"))
-			.MinDesiredSlotWidth(FAppStyle::GetFloat("StandardDialog.MinDesiredSlotWidth"))
-			.MinDesiredSlotHeight(FAppStyle::GetFloat("StandardDialog.MinDesiredSlotHeight"))
+			.SlotPadding(UE_MUTABLE_GET_MARGIN("StandardDialog.SlotPadding"))
+			.MinDesiredSlotWidth(UE_MUTABLE_GET_FLOAT("StandardDialog.MinDesiredSlotWidth"))
+			.MinDesiredSlotHeight(UE_MUTABLE_GET_FLOAT("StandardDialog.MinDesiredSlotHeight"))
 			
 			+ SUniformGridPanel::Slot(0, 0)
 			[
 				SNew(SButton)
 				.HAlign(HAlign_Center)
-				.ContentPadding(FAppStyle::GetMargin("StandardDialog.ContentPadding"))
+				.ContentPadding(UE_MUTABLE_GET_MARGIN("StandardDialog.ContentPadding"))
 				.Text(LOCTEXT("OK", "OK"))
 				.OnClicked(this, &SSelectPopulationFolderDlg::OnButtonClick, EAppReturnType::Ok)
 			]
@@ -795,7 +796,7 @@ void SSelectPopulationFolderDlg::Construct(const FArguments& InArgs)
 			[
 				SNew(SButton)
 				.HAlign(HAlign_Center)
-				.ContentPadding(FAppStyle::GetMargin("StandardDialog.ContentPadding"))
+				.ContentPadding(UE_MUTABLE_GET_MARGIN("StandardDialog.ContentPadding"))
 				.Text(LOCTEXT("Cancel", "Cancel"))
 				.OnClicked(this, &SSelectPopulationFolderDlg::OnButtonClick, EAppReturnType::Cancel)
 			]
diff --git a/Engine/Plugins/Experimental/Mutable/Source/CustomizableObjectPopulationEditor/Private/MuCOPE/SCustomizableObjectPopulationClassTagManager.cpp b/Engine/Plugins/Experimental/Mutable/Source/CustomizableObjectPopulationEditor/Private/MuCOPE/SCustomizableObjectPopulationClassTagManager.cpp
--- a/Engine/Plugins/Experimental/Mutable/Source/CustomizableObjectPopulationEditor/Private/MuCOPE/SCustomizableObjectPopulationClassTagManager.cpp
+++ b/Engine/Plugins/Experimental/Mutable/Source/CustomizableObjectPopulationEditor/Private/MuCOPE/SCustomizableObjectPopulationClassTagManager.cpp
@@ -2,6 +2,7 @@
 
 #include "MuCOPE/SCustomizableObjectPopulationClassTagManager.h"
 
+#include "MuCOE/UnrealEditorPortabilityHelpers.h"
 #include "MuCO/CustomizableObject.h"
 #include "Widgets/Input/SButton.h"
 #include "Widgets/Input/SEditableTextBox.h"
@@ -101,7 +102,7 @@ void SCustomizableObjectPopulationClassTagsManager::BuildTagSelector()
 			.ToolTipText(LOCTEXT("AddTagSelector", "Add a Tag to the Customizable Object"))
 			[
 				SNew(SImage)
-				.Image(FAppStyle::GetBrush(TEXT("Plus")))
+				.Image(UE_MUTABLE_GET_BRUSH(TEXT("Plus")))
 			]
 		]
 	];
@@ -150,7 +151,7 @@ void SCustomizableObjectPopulationClassTagsManager::BuildTagSelector()
 				.ToolTipText(LOCTEXT("RemoveTag", "Remove Tag"))
 				[
 					SNew(SImage)
-					.Image(FAppStyle::GetBrush(TEXT("Cross")))
+					.Image(UE_MUTABLE_GET_BRUSH(TEXT("Cross")))
 				]
 			]
 		];
@@ -223,7 +224,7 @@ void SCustomizableObjectPopulationClassTagsManager::BuildTagManager()
 			.ToolTipText(LOCTEXT("AddTag", "Add a Tag to the Root Object"))
 			[
 				SNew(SImage)
-				.Image(FAppStyle::GetBrush(TEXT("Plus")))
+				.Image(UE_MUTABLE_GET_BRUSH(TEXT("Plus")))
 			]
 		]
 	];
@@ -261,7 +262,7 @@ void SCustomizableObjectPopulationClassTagsManager::BuildTagManager()
 				.ToolTipText(LOCTEXT("RemoveTag", "Remove Tag"))
 				[
 					SNew(SImage)
-					.Image(FAppStyle::GetBrush(TEXT("Cross")))
+					.Image(UE_MUTABLE_GET_BRUSH(TEXT("Cross")))
 				]
 			]
 		];
diff --git a/Engine/Plugins/Experimental/Mutable/Source/CustomizableObjectPopulationEditor/Private/MuCOPE/SCustomizableObjectPopulationClassTagsTool.cpp b/Engine/Plugins/Experimental/Mutable/Source/CustomizableObjectPopulationEditor/Private/MuCOPE/SCustomizableObjectPopulationClassTagsTool.cpp
--- a/Engine/Plugins/Experimental/Mutable/Source/CustomizableObjectPopulationEditor/Private/MuCOPE/SCustomizableObjectPopulationClassTagsTool.cpp
+++ b/Engine/Plugins/Experimental/Mutable/Source/CustomizableObjectPopulationEditor/Private/MuCOPE/SCustomizableObjectPopulationClassTagsTool.cpp
@@ -5,6 +5,7 @@
 #include "MuCO/CustomizableObject.h"
 #include "MuCOP/CustomizableObjectPopulationClass.h"
 #include "MuCOPE/CustomizableObjectPopulationClassDetails.h"
+#include "MuCOE/UnrealEditorPortabilityHelpers.h"
 #include "Widgets/Input/SButton.h"
 #include "Widgets/Input/SSearchBox.h"
 #include "Widgets/Input/STextComboBox.h"
@@ -92,14 +93,14 @@ void SCustomizableObjectPopulationClassTagsTool::GenerateChildSlot()
 					.InitiallyCollapsed(false)
 					.BorderBackgroundColor(FLinearColor(0.6f, 0.6f, 0.6f, 1.0f))
 					.BorderImage_Lambda([this]() { return GetExpandableAreaBorderImage(*TagsManagerExpandableArea); })
-					.BodyBorderImage(FAppStyle::GetBrush("ToolPanel.GroupBorder"))
+					.BodyBorderImage(UE_MUTABLE_GET_BRUSH("ToolPanel.GroupBorder"))
 					.BodyBorderBackgroundColor(FLinearColor::White)
 					.HeaderContent()
 					[
 						SNew(STextBlock)
 						.Margin(FMargin(0.0f,5.0f))
 						.Text(LOCTEXT("ClassTagManagementText", "Population Tag Management for the current Customizable Object"))
-						.Font(FAppStyle::GetFontStyle("DetailsView.CategoryFontStyle"))
+						.Font(UE_MUTABLE_GET_FONTSTYLE("DetailsView.CategoryFontStyle"))
 						.ShadowOffset(FVector2D(1.0f, 1.0f))
 					]
 					.BodyContent()
@@ -150,14 +151,14 @@ void SCustomizableObjectPopulationClassTagsTool::GenerateChildSlot()
 					.InitiallyCollapsed(false)
 					.BorderBackgroundColor(FLinearColor(0.6f, 0.6f, 0.6f, 1.0f))
 					.BorderImage_Lambda([this]() { return GetExpandableAreaBorderImage(*ParameterOptionsManagerExpandableArea); })
-					.BodyBorderImage(FAppStyle::GetBrush("ToolPanel.GroupBorder"))
+					.BodyBorderImage(UE_MUTABLE_GET_BRUSH("ToolPanel.GroupBorder"))
 					.BodyBorderBackgroundColor(FLinearColor::White)
 					.HeaderContent()
 					[
 						SNew(STextBlock)
 						.Margin(FMargin(0.0f, 5.0f))
 						.Text(LOCTEXT("ParameterOptionsManagerText", "Parameter Option Population Tags"))
-						.Font(FAppStyle::GetFontStyle("DetailsView.CategoryFontStyle"))
+						.Font(UE_MUTABLE_GET_FONTSTYLE("DetailsView.CategoryFontStyle"))
 						.ShadowOffset(FVector2D(1.0f, 1.0f))
 					]
 					.BodyContent()
@@ -274,7 +275,7 @@ void SCustomizableObjectPopulationClassTagsTool::GenerateTagsViewer()
 					.ButtonColorAndOpacity(FSlateColor(FLinearColor(1.0f,0.2f,0.2f)))
 					[
 						SNew(SImage)
-						.Image(FAppStyle::GetBrush(TEXT("Cross")))
+						.Image(UE_MUTABLE_GET_BRUSH(TEXT("Cross")))
 					]
 				]
 			];
@@ -304,7 +305,7 @@ void SCustomizableObjectPopulationClassTagsTool::GenerateTagsViewer()
 						.OnClicked(this, &SCustomizableObjectPopulationClassTagsTool::OnRemoveTagButtonPressed,i)
 						[
 							SNew(SImage)
-							.Image(FAppStyle::GetBrush(TEXT("Cross")))
+							.Image(UE_MUTABLE_GET_BRUSH(TEXT("Cross")))
 						]
 					]
 				];
@@ -661,9 +662,9 @@ const FSlateBrush* SCustomizableObjectPopulationClassTagsTool::GetExpandableArea
 {
 	if (Area.IsTitleHovered())
 	{
-		return Area.IsExpanded() ? FAppStyle::GetBrush("DetailsView.CategoryTop_Hovered") : FAppStyle::GetBrush("DetailsView.CollapsedCategory_Hovered");
+		return Area.IsExpanded() ? UE_MUTABLE_GET_BRUSH("DetailsView.CategoryTop_Hovered") : UE_MUTABLE_GET_BRUSH("DetailsView.CollapsedCategory_Hovered");
 	}
-	return Area.IsExpanded() ? FAppStyle::GetBrush("DetailsView.CategoryTop") : FAppStyle::GetBrush("DetailsView.CollapsedCategory");
+	return Area.IsExpanded() ? UE_MUTABLE_GET_BRUSH("DetailsView.CategoryTop") : UE_MUTABLE_GET_BRUSH("DetailsView.CollapsedCategory");
 }
 
 
@@ -749,7 +750,7 @@ void SSingleTagWidget::Construct(const FArguments& InArgs)
 				.OnClicked(this, &SSingleTagWidget::OnRemoveTagButtonPressed)
 				[
 					SNew(SImage)
-					.Image(FAppStyle::GetBrush(TEXT("Cross")))
+					.Image(UE_MUTABLE_GET_BRUSH(TEXT("Cross")))
 				]
 			]
 		]
diff --git a/Engine/Plugins/Experimental/Mutable/Source/CustomizableObjectPopulationEditor/Private/MuCOPE/SCustomizableObjectPopulationEditorViewportToolBar.cpp b/Engine/Plugins/Experimental/Mutable/Source/CustomizableObjectPopulationEditor/Private/MuCOPE/SCustomizableObjectPopulationEditorViewportToolBar.cpp
--- a/Engine/Plugins/Experimental/Mutable/Source/CustomizableObjectPopulationEditor/Private/MuCOPE/SCustomizableObjectPopulationEditorViewportToolBar.cpp
+++ b/Engine/Plugins/Experimental/Mutable/Source/CustomizableObjectPopulationEditor/Private/MuCOPE/SCustomizableObjectPopulationEditorViewportToolBar.cpp
@@ -8,6 +8,7 @@
 #include "LevelEditor.h"
 #include "Modules/ModuleManager.h"
 #include "MuCOE/CustomizableObjectEditorViewportMenuCommands.h"
+#include "MuCOE/UnrealEditorPortabilityHelpers.h"
 #include "MuCOPE/SCustomizableObjectPopulationEditorViewport.h"
 #include "SEditorViewportViewMenu.h"
 #include "Widgets/Input/SMenuAnchor.h"
@@ -56,7 +57,7 @@ void SCustomizableObjectPopulationEditorViewportToolBar::Construct(const FArgume
 	FLinearColor ButtonColor2 = FLinearColor(0.2f, 0.2f, 0.2f, 0.75f);
 	FLinearColor TextColor1 = FLinearColor(1.0f, 1.0f, 1.0f, 1.0f);
 	FLinearColor TextColor2 = FLinearColor(0.8f, 0.8f, 0.8f, 0.8f);
-	FSlateFontInfo Info = FAppStyle::GetFontStyle("BoldFont");
+	FSlateFontInfo Info = UE_MUTABLE_GET_FONTSTYLE("BoldFont");
 	Info.Size += 26;
 
 	ChildSlot
@@ -66,10 +67,10 @@ void SCustomizableObjectPopulationEditorViewportToolBar::Construct(const FArgume
 		.VAlign(VAlign_Top)
 		[
 			SNew(SBorder)
-			.BorderImage(FAppStyle::GetBrush("NoBorder"))
+			.BorderImage(UE_MUTABLE_GET_BRUSH("NoBorder"))
 		//// Color and opacity is changed based on whether or not the mouse cursor is hovering over the toolbar area
 		//.ColorAndOpacity(this, &SViewportToolBar::OnGetColorAndOpacity)
-		.ForegroundColor(FAppStyle::GetSlateColor(DefaultForegroundName))
+		.ForegroundColor(UE_MUTABLE_GET_SLATECOLOR(DefaultForegroundName))
 		[
 			SNew(SVerticalBox)
 			+ SVerticalBox::Slot()
@@ -196,7 +197,7 @@ TSharedRef<SWidget> SCustomizableObjectPopulationEditorViewportToolBar::Generate
 		.WidthOverride(100.0f)
 		[
 			SNew(SSpinBox<float>)
-			.Font(FAppStyle::GetFontStyle(TEXT("MenuItem.Font")))
+			.Font(UE_MUTABLE_GET_FONTSTYLE(TEXT("MenuItem.Font")))
 		.MinValue(FOVMin)
 		.MaxValue(FOVMax)
 		.Value(this, &SCustomizableObjectPopulationEditorViewportToolBar::OnGetFOVValue)
diff --git a/Engine/Plugins/Experimental/Mutable/Source/MutableTools/Private/MuT/CodeGenerator.cpp b/Engine/Plugins/Experimental/Mutable/Source/MutableTools/Private/MuT/CodeGenerator.cpp
--- a/Engine/Plugins/Experimental/Mutable/Source/MutableTools/Private/MuT/CodeGenerator.cpp
+++ b/Engine/Plugins/Experimental/Mutable/Source/MutableTools/Private/MuT/CodeGenerator.cpp
@@ -1425,9 +1425,9 @@ namespace mu
 
                             // Generate the image
                             IMAGE_STATE newState;
-                            newState.m_imageSize = FIntVector2(desc.m_size);
-							newState.m_imageRect.min = FIntVector2(rect.min);
-							newState.m_imageRect.size = FIntVector2(rect.size);
+                            newState.m_imageSize = UE::Math::TIntVector2<int32>(desc.m_size);
+							newState.m_imageRect.min = UE::Math::TIntVector2<int32>(rect.min);
+							newState.m_imageRect.size = UE::Math::TIntVector2<int32>(rect.size);
 							newState.m_layoutBlockId = pLayout->m_blocks[b].m_id;
                             newState.m_pLayout = pLayout;
                             m_imageState.Add( newState );
@@ -1530,7 +1530,7 @@ namespace mu
 
                                             // Generate the image block
                                             IMAGE_STATE newState;
-                                            newState.m_imageSize = FIntVector2(extendDesc.m_size);
+                                            newState.m_imageSize = UE::Math::TIntVector2<int32>(extendDesc.m_size);
                                             newState.m_imageRect = rect;
                                             newState.m_layoutBlockId = pExtendLayout->m_blocks[b].m_id;
                                             newState.m_pLayout = pExtendLayout;
diff --git a/Engine/Plugins/Experimental/Mutable/Source/MutableTools/Private/MuT/CodeGenerator_Image.cpp b/Engine/Plugins/Experimental/Mutable/Source/MutableTools/Private/MuT/CodeGenerator_Image.cpp
--- a/Engine/Plugins/Experimental/Mutable/Source/MutableTools/Private/MuT/CodeGenerator_Image.cpp
+++ b/Engine/Plugins/Experimental/Mutable/Source/MutableTools/Private/MuT/CodeGenerator_Image.cpp
@@ -1551,9 +1551,9 @@ namespace mu
 
             // We take whatever size will be produced
             FImageDesc desc = CalculateImageDesc( *node.m_pImage->GetBasePrivate() );
-            newState.m_imageSize = FIntVector2(desc.m_size);
-            newState.m_imageRect.min = FIntVector2(0,0);
-            newState.m_imageRect.size = FIntVector2(desc.m_size);
+            newState.m_imageSize = UE::Math::TIntVector2<int32>(desc.m_size);
+            newState.m_imageRect.min = UE::Math::TIntVector2<int32>(0,0);
+            newState.m_imageRect.size = UE::Math::TIntVector2<int32>(desc.m_size);
             newState.m_layoutBlockId = -1;
             m_imageState.Add( newState );
 
diff --git a/Engine/Plugins/Experimental/Mutable/Source/MutableTools/Private/MuT/CodeGenerator_Scalar.cpp b/Engine/Plugins/Experimental/Mutable/Source/MutableTools/Private/MuT/CodeGenerator_Scalar.cpp
--- a/Engine/Plugins/Experimental/Mutable/Source/MutableTools/Private/MuT/CodeGenerator_Scalar.cpp
+++ b/Engine/Plugins/Experimental/Mutable/Source/MutableTools/Private/MuT/CodeGenerator_Scalar.cpp
@@ -167,10 +167,10 @@ namespace mu
 					// We take whatever size will be produced
 					IMAGE_STATE newState;
 					FImageDesc desc = CalculateImageDesc(*node.m_additionalImages[a]->GetBasePrivate());
-					newState.m_imageSize = FIntVector2(desc.m_size);
+					newState.m_imageSize = UE::Math::TIntVector2<int32>(desc.m_size);
 					newState.m_imageRect.min[0] = 0;
 					newState.m_imageRect.min[1] = 0;
-					newState.m_imageRect.size = FIntVector2(desc.m_size);
+					newState.m_imageRect.size = UE::Math::TIntVector2<int32>(desc.m_size);
 					newState.m_layoutBlockId = -1;
 					m_imageState.Add(newState);
 
