[Mutable] Fix possible nullptr dereference on image displace operations.

- the null dereference could happen, for instance, if the source of a displacemnt happens to be a crop operation with empty rect.

#rnx

diff --git a/Engine/Plugins/Experimental/Mutable/Source/MutableRuntime/Private/MuR/CodeRunner.cpp b/Engine/Plugins/Experimental/Mutable/Source/MutableRuntime/Private/MuR/CodeRunner.cpp
--- a/Engine/Plugins/Experimental/Mutable/Source/MutableRuntime/Private/MuR/CodeRunner.cpp
+++ b/Engine/Plugins/Experimental/Mutable/Source/MutableRuntime/Private/MuR/CodeRunner.cpp
@@ -4206,6 +4206,12 @@ namespace mu
                 Ptr<const Image> pSource = GetMemory().GetImage( FCacheAddress(args.source,item) );
                 Ptr<const Image> pMap = GetMemory().GetImage( FCacheAddress(args.displacementMap,item) );
 
+				if (!pSource)
+				{
+					GetMemory().SetImage(item, nullptr);
+					break;
+				}
+
 				// TODO: This shouldn't happen: displacement maps cannot be scaled because their information
 				// is resolution sensitive (pixel offsets). If the size doesn't match, scale the source, apply 
 				// displacement and then unscale it.
