[mutable] Prevent customizable object meshes from disappearing when saving the level

#rnx

diff --git a/Engine/Plugins/Experimental/Mutable/Source/CustomizableObject/Classes/MuCO/CustomizableSkeletalComponent.h b/Engine/Plugins/Experimental/Mutable/Source/CustomizableObject/Classes/MuCO/CustomizableSkeletalComponent.h
--- a/Engine/Plugins/Experimental/Mutable/Source/CustomizableObject/Classes/MuCO/CustomizableSkeletalComponent.h
+++ b/Engine/Plugins/Experimental/Mutable/Source/CustomizableObject/Classes/MuCO/CustomizableSkeletalComponent.h
@@ -61,9 +61,6 @@ public:
 	// USceneComponent interface
 	virtual void BeginDestroy() override;
 
-	// UObject interface
-	virtual void PreSave(FObjectPreSaveContext ObjectSaveContext) override;
-
 #if WITH_EDITOR
 	// Used to generate instances outside the CustomizableObject editor and PIE
 	void UpdateDistFromComponentToLevelEditorCamera(const FVector& CameraPosition);
diff --git a/Engine/Plugins/Experimental/Mutable/Source/CustomizableObject/Private/MuCO/CustomizableSkeletalComponent.cpp b/Engine/Plugins/Experimental/Mutable/Source/CustomizableObject/Private/MuCO/CustomizableSkeletalComponent.cpp
--- a/Engine/Plugins/Experimental/Mutable/Source/CustomizableObject/Private/MuCO/CustomizableSkeletalComponent.cpp
+++ b/Engine/Plugins/Experimental/Mutable/Source/CustomizableObject/Private/MuCO/CustomizableSkeletalComponent.cpp
@@ -136,14 +136,6 @@ void UCustomizableSkeletalComponent::BeginDestroy()
 }
 
 
-void UCustomizableSkeletalComponent::PreSave(FObjectPreSaveContext ObjectSaveContext)
-{
-	SetSkeletalMesh(nullptr); // To prevent "Graph is linked to private object(s) in an external package" errors when saving in editor after having PIE and generated customized skeletal meshes while in PIE
-
-	Super::PreSave(ObjectSaveContext);
-}
-
-
 #if WITH_EDITOR
 
 void UCustomizableSkeletalComponent::UpdateDistFromComponentToLevelEditorCamera(const FVector& CameraPosition)
