Changed mu::UntypedMeshBufferIteratorConst to use member initializers and compiler-generated moves and copies, to avoid warnings about mismatched copy constructors and assignments.

diff --git a/Engine/Plugins/Experimental/Mutable/Source/MutableRuntime/Private/MuR/MeshPrivate.h b/Engine/Plugins/Experimental/Mutable/Source/MutableRuntime/Private/MuR/MeshPrivate.h
--- a/Engine/Plugins/Experimental/Mutable/Source/MutableRuntime/Private/MuR/MeshPrivate.h
+++ b/Engine/Plugins/Experimental/Mutable/Source/MutableRuntime/Private/MuR/MeshPrivate.h
@@ -367,13 +367,11 @@ namespace mu
 	{
 	public:
 
-		inline UntypedMeshBufferIteratorConst()
-		{
-			m_format = MBF_NONE;
-			m_components = 0;
-			m_elementSize = 0;
-			m_pBuf = nullptr;
-		}
+		UntypedMeshBufferIteratorConst() = default;
+		UntypedMeshBufferIteratorConst(UntypedMeshBufferIteratorConst&&) = default;
+		UntypedMeshBufferIteratorConst(const UntypedMeshBufferIteratorConst&) = default;
+		UntypedMeshBufferIteratorConst& operator=(UntypedMeshBufferIteratorConst&&) = default;
+		UntypedMeshBufferIteratorConst& operator=(const UntypedMeshBufferIteratorConst&) = default;
 
 		inline UntypedMeshBufferIteratorConst
 			(
@@ -417,14 +415,6 @@ namespace mu
 			}
 		}
 
-		inline UntypedMeshBufferIteratorConst( const UntypedMeshBufferIteratorConst& other )
-		{
-			m_elementSize = other.m_elementSize;
-			m_pBuf = other.m_pBuf;
-			m_format = other.m_format;
-			m_components = other.m_components;
-		}
-
         inline const uint8_t* ptr() const
 		{
 			return m_pBuf;
@@ -551,10 +541,10 @@ namespace mu
 
 	protected:
 
-		int m_elementSize;
-        const uint8_t* m_pBuf;
-		MESH_BUFFER_FORMAT m_format;
-		int m_components;
+		int m_elementSize = 0;
+        const uint8_t* m_pBuf = nullptr;
+		MESH_BUFFER_FORMAT m_format = MBF_NONE;
+		int m_components = 0;
 	};
 
 
