Removed redundant copy constructors from mu::vec* which were warning due to missing copy assignment operators.

diff --git a/Engine/Plugins/Experimental/Mutable/Source/MutableRuntime/Public/MuR/MutableMath.h b/Engine/Plugins/Experimental/Mutable/Source/MutableRuntime/Public/MuR/MutableMath.h
--- a/Engine/Plugins/Experimental/Mutable/Source/MutableRuntime/Public/MuR/MutableMath.h
+++ b/Engine/Plugins/Experimental/Mutable/Source/MutableRuntime/Public/MuR/MutableMath.h
@@ -168,14 +168,6 @@ namespace mu
         }
 
 
-        inline vec( const V& o )
-        {
-            for( int i=0; i<DIM; ++i )
-            {
-                m[i]=o.m[i];
-            }
-        }
-
 		//-----------------------------------------------------------------------------------------
 		static inline int GetDim()
 		{
@@ -476,9 +468,6 @@ namespace mu
 			this->m[1] = y;
 		}
 
-        //-----------------------------------------------------------------------------------------
-        inline vec2( const vec2& ) = default;
-
 		//-----------------------------------------------------------------------------------------
 		template<class SCALAR2>
         inline vec2( const vec<SCALAR2,2>& v )
