First version of Mutable extension API

diff --git a/Engine/Plugins/Experimental/Mutable/Mutable.uplugin b/Engine/Plugins/Experimental/Mutable/Mutable.uplugin
--- a/Engine/Plugins/Experimental/Mutable/Mutable.uplugin
+++ b/Engine/Plugins/Experimental/Mutable/Mutable.uplugin
@@ -44,6 +44,10 @@
 		{
 			"Name": "SkeletalMerging",
 			"Enabled": true
+		},
+		{
+			"Name": "StructUtils",
+			"Enabled": true
 		}
 	],
 
diff --git a/Engine/Plugins/Experimental/Mutable/Source/CustomizableObject/Classes/MuCO/CustomizableObject.h b/Engine/Plugins/Experimental/Mutable/Source/CustomizableObject/Classes/MuCO/CustomizableObject.h
--- a/Engine/Plugins/Experimental/Mutable/Source/CustomizableObject/Classes/MuCO/CustomizableObject.h
+++ b/Engine/Plugins/Experimental/Mutable/Source/CustomizableObject/Classes/MuCO/CustomizableObject.h
@@ -4,6 +4,8 @@
 
 #include "RHIDefinitions.h"
 #include "MuCO/CustomizableObjectClothingTypes.h"
+#include "MuCO/CustomizableObjectExtensionData.h"
+#include "MuCO/CustomizableObjectStreamedExtensionData.h"
 #include "MuCO/CustomizableObjectIdentifier.h"
 #include "MuCO/CustomizableObjectParameterTypeDefinitions.h"
 #include "MuCO/CustomizableObjectUIData.h"
@@ -1184,6 +1186,14 @@ public:
 	UPROPERTY()
 	TArray<FMutableSkinWeightProfileInfo> SkinWeightProfilesInfo;
 
+	// mu::ExtensionData::Index is an index into this array when mu::ExtensionData::Origin is ConstantAlwaysLoaded
+	UPROPERTY()
+	TArray<FCustomizableObjectExtensionData> AlwaysLoadedExtensionData;
+
+	// mu::ExtensionData::Index is an index into this array when mu::ExtensionData::Origin is ConstantStreamed
+	UPROPERTY()
+	TArray<FCustomizableObjectStreamedExtensionData> StreamedExtensionData;
+
 #if WITH_EDITORONLY_DATA
 
 	// Hide this property because it is not used yet.
@@ -1425,7 +1435,7 @@ private:
 	// This is a manual version number for the binary blobs in this asset.
 	// Increasing it invalidates all the previously compiled models.
 	// Warning: If while merging code both versions have changed, take the highest+1.
-	static const int32 CurrentSupportedVersion = 379;
+	static const int32 CurrentSupportedVersion = 380;
 
 public:
 
diff --git a/Engine/Plugins/Experimental/Mutable/Source/CustomizableObject/Classes/MuCO/CustomizableObjectExtension.h b/Engine/Plugins/Experimental/Mutable/Source/CustomizableObject/Classes/MuCO/CustomizableObjectExtension.h
new file mode 100644
--- /dev/null
+++ b/Engine/Plugins/Experimental/Mutable/Source/CustomizableObject/Classes/MuCO/CustomizableObjectExtension.h
@@ -0,0 +1,107 @@
+// Copyright Epic Games, Inc. All Rights Reserved.
+
+#pragma once
+
+#include "InstancedStruct.h"
+#include "UObject/Object.h"
+#include "UObject/WeakObjectPtr.h"
+
+#include "CustomizableObjectExtension.generated.h"
+
+class UCustomizableObjectInstance;
+class USkeletalMesh;
+class ACustomizableSkeletalMeshActor;
+
+/** A type of pin in the Mutable graph UI */
+struct FCustomizableObjectPinType
+{
+	/*
+	 * The identifier for this type, to be used internally
+	 *
+	 * Note that the same pin type names may be used by different extensions, so that extensions
+	 * can interoperate with each other using extension-defined pin types.
+	 * 
+	 * In other words, it's valid for one extension to create a new pin type and another extension
+	 * to create nodes that use that type.
+	 */
+	FName Name;
+
+	/** The display name for this type in the editor UI */
+	FText DisplayName;
+
+	/** The color that will be used in the editor UI for this pin and any wires connected to it */
+	FLinearColor Color;
+};
+
+/** An input pin that will be added to Object nodes */
+struct FObjectNodeInputPin
+{
+	/** This can be the name of a built-in pin type or an extension-defined FCustomizableObjectPinType */
+	FName PinType;
+
+	/**
+	 * The internal name for the pin, to disambiguate it from other pins.
+	 *
+	 * Ensure this name is unique for object node pins created by the same extension.
+	 * 
+	 * The system can automatically distinguish between pins with the same name across different
+	 * extensions, so this doesn't need to be a globally unique name.
+	 */
+	FName PinName;
+
+	/** The name that will be displayed for the pin in the editor UI */
+	FText DisplayName;
+
+	/**
+	 * Whether this pin accepts multiple inputs or not.
+	 *
+	 * Note that even if this is false, an Object node pin can still receive one input per Child
+	 * Object node, so the extension still needs to handle receiving multiple inputs for a single
+	 * pin.
+	 */
+	bool bIsArray = false;
+};
+
+/** An Object node input pin and the data that was passed into it by the Customizable Object graph */
+struct FInputPinDataContainer
+{
+	FInputPinDataContainer(const FObjectNodeInputPin& InPin, const FInstancedStruct& InData)
+		: Pin(InPin)
+		, Data(InData)
+	{
+	}
+
+	FObjectNodeInputPin Pin;
+	const FInstancedStruct& Data;
+};
+
+/**
+ * An extension that adds functionality to the Customizable Object system
+ *
+ * To create a new extension, make a subclass of this class and register it by calling
+ * ICustomizableObjectModule::Get().RegisterExtension().
+ */
+UCLASS()
+class CUSTOMIZABLEOBJECT_API UCustomizableObjectExtension : public UObject
+{
+	GENERATED_BODY() 
+public:
+	/** Returns any new pin types that are defined by this extension */
+	virtual TArray<FCustomizableObjectPinType> GetPinTypes() const { return TArray<FCustomizableObjectPinType>(); }
+
+	/** Returns the pins that this extension adds to Object nodes */
+	virtual TArray<FObjectNodeInputPin> GetAdditionalObjectNodePins() const { return TArray<FObjectNodeInputPin>(); }
+
+	/**
+	 * Called when a Skeletal Mesh asset is created
+	 *
+	 * @param InputPinData - The data for only the input pins *registered by this extension*. This
+	 * helps to enforce separation between the extensions, so that they don't depend on each other.
+	 * 
+	 * @param ComponentIndex - The component index of the Skeletal Mesh, for the case where the pin
+	 * data is associated with a particular component.
+	 * 
+	 * @param SkeletalMesh - The Skeletal Mesh that was created.
+	 */
+	virtual void OnSkeletalMeshCreated(const TArray<FInputPinDataContainer>& InputPinData, int32 ComponentIndex, USkeletalMesh* SkeletalMesh) const {}
+};
diff --git a/Engine/Plugins/Experimental/Mutable/Source/CustomizableObject/Classes/MuCO/CustomizableObjectExtensionData.h b/Engine/Plugins/Experimental/Mutable/Source/CustomizableObject/Classes/MuCO/CustomizableObjectExtensionData.h
new file mode 100644
--- /dev/null
+++ b/Engine/Plugins/Experimental/Mutable/Source/CustomizableObject/Classes/MuCO/CustomizableObjectExtensionData.h
@@ -0,0 +1,19 @@
+// Copyright Epic Games, Inc. All Rights Reserved.
+
+#pragma once
+
+#include "InstancedStruct.h"
+
+#include "CustomizableObjectExtensionData.generated.h"
+
+/** Some data of an extension-defined type generated by a Customizable Object Extension */
+USTRUCT()
+struct FCustomizableObjectExtensionData
+{
+	GENERATED_BODY()
+
+public:
+	/** The data itself. This could be any USTRUCT. */
+	UPROPERTY()
+	FInstancedStruct Data;
+};
diff --git a/Engine/Plugins/Experimental/Mutable/Source/CustomizableObject/Classes/MuCO/CustomizableObjectStreamedExtensionData.h b/Engine/Plugins/Experimental/Mutable/Source/CustomizableObject/Classes/MuCO/CustomizableObjectStreamedExtensionData.h
new file mode 100644
--- /dev/null
+++ b/Engine/Plugins/Experimental/Mutable/Source/CustomizableObject/Classes/MuCO/CustomizableObjectStreamedExtensionData.h
@@ -0,0 +1,87 @@
+// Copyright Epic Games, Inc. All Rights Reserved.
+
+#pragma once
+
+#include "MuCO/CustomizableObjectExtensionData.h"
+#include "UObject/Object.h"
+#include "UObject/SoftObjectPtr.h"
+
+#include "CustomizableObjectStreamedExtensionData.generated.h"
+
+/** A simple container that's used to store the FCustomizableObjectStreamedExtensionData in a package */
+UCLASS()
+class CUSTOMIZABLEOBJECT_API UCustomizableObjectExtensionDataContainer : public UObject
+{
+	GENERATED_BODY()
+
+public:
+	UPROPERTY()
+	FCustomizableObjectExtensionData Data;
+};
+
+/**
+ * Used to store extension data that can be streamed in and out from an external package.
+ *
+ * This allows you to load large extension data on demand, instead of having it always loaded.
+ */
+USTRUCT()
+struct CUSTOMIZABLEOBJECT_API FCustomizableObjectStreamedExtensionData
+{
+	GENERATED_BODY()
+
+public:
+	FCustomizableObjectStreamedExtensionData() = default;
+
+#if WITH_EDITOR
+	/** Streamed data can only be constructed in the editor */
+	FCustomizableObjectStreamedExtensionData(UCustomizableObjectExtensionDataContainer* InContainer);
+
+	/**
+	 * Clears the hard reference to the container, so that it's not forced to load as soon as this
+	 * struct is loaded, and replaces the soft reference with the provided one.
+	 * 
+	 * Note that the container may not have been moved to its new path yet, so don't try to load
+	 * the container from the given path here.
+	 */
+	void ConvertToSoftReferenceForCooking(const TSoftObjectPtr<UCustomizableObjectExtensionDataContainer>& NewContainerPath);
+#endif // WITH_EDITOR
+
+	/**
+	 * If this returns true, GetLoadedData can safely be called.
+	 *
+	 * The data will stay loaded until Unload is called.
+	 */
+	bool IsLoaded() const;
+
+	/**
+	 * Returns the data, as long as it's loaded.
+	 *
+	 * If the data isn't loaded, an assertion will fail.
+	 */
+	const FCustomizableObjectExtensionData& GetLoadedData() const;
+
+	/**
+	 * Release this struct's hard reference to the loaded data.
+	 *
+	 * This doesn't immediately remove it from memory, but will allow it to be deleted by the
+	 * garbage collector if there are no other references.
+	 *
+	 * Note that in editor builds streamed Extension Data must always be loaded due to the package
+	 * structure, and therefore this function has no effect in editor builds.
+	 */
+	void Unload();
+
+	/**
+	 * To stream this data in, load the object at the path returned by GetPath and then pass it
+	 * into NotifyLoaded.
+	 */
+	const TSoftObjectPtr<UCustomizableObjectExtensionDataContainer>& GetPath() const { return ContainerPath; }
+	void NotifyLoaded(const UCustomizableObjectExtensionDataContainer* LoadedContainer);
+
+private:
+	UPROPERTY()
+	TSoftObjectPtr<UCustomizableObjectExtensionDataContainer> ContainerPath;
+
+	UPROPERTY()
+	TObjectPtr<const UCustomizableObjectExtensionDataContainer> Container;
+};
diff --git a/Engine/Plugins/Experimental/Mutable/Source/CustomizableObject/CustomizableObject.Build.cs b/Engine/Plugins/Experimental/Mutable/Source/CustomizableObject/CustomizableObject.Build.cs
--- a/Engine/Plugins/Experimental/Mutable/Source/CustomizableObject/CustomizableObject.Build.cs
+++ b/Engine/Plugins/Experimental/Mutable/Source/CustomizableObject/CustomizableObject.Build.cs
@@ -43,6 +43,7 @@ public class CustomizableObject : ModuleRules
             "GameplayTags",
 			"MutableRuntime",
 			"AnimGraphRuntime",
+			"StructUtils"
 		});
 
         PrivateIncludePathModuleNames.AddRange(
diff --git a/Engine/Plugins/Experimental/Mutable/Source/CustomizableObject/Private/MuCO/CustomizableObjectInstance.cpp b/Engine/Plugins/Experimental/Mutable/Source/CustomizableObject/Private/MuCO/CustomizableObjectInstance.cpp
--- a/Engine/Plugins/Experimental/Mutable/Source/CustomizableObject/Private/MuCO/CustomizableObjectInstance.cpp
+++ b/Engine/Plugins/Experimental/Mutable/Source/CustomizableObject/Private/MuCO/CustomizableObjectInstance.cpp
@@ -2,6 +2,7 @@
 
 #include "MuCO/CustomizableObjectInstance.h"
 
+#include "Algo/Find.h"
 #include "Algo/MaxElement.h"
 #include "Animation/AnimClassInterface.h"
 #include "Animation/AnimBlueprintGeneratedClass.h"
@@ -17,10 +18,12 @@
 #include "Modules/ModuleManager.h"
 #include "MuCO/CustomizableInstanceLODManagement.h"
 #include "MuCO/CustomizableInstancePrivateData.h"
+#include "MuCO/CustomizableObjectExtension.h"
 #include "MuCO/CustomizableObjectMipDataProvider.h"
 #include "MuCO/CustomizableObjectPrivate.h"
 #include "MuCO/CustomizableSkeletalComponent.h"
 #include "MuCO/DefaultImageProvider.h"
+#include "MuCO/ICustomizableObjectModule.h"
 #include "MuCO/UnrealConversionUtils.h"
 #include "MuCO/UnrealPortabilityHelpers.h"
 #include "MuR/OpImagePixelFormat.h"
@@ -2064,6 +2067,95 @@ bool UCustomizableInstancePrivateData::UpdateSkeletalMesh_PostBeginUpdate0(UCust
 			break;
 		}
 
+		// Collate the extension data on the instance into groups based on the extension that
+		// produced it, so that we only need to call OnSkeletalMeshCreated once for each extension.
+		{
+			TMap<const UCustomizableObjectExtension*, TArray<FInputPinDataContainer>> ExtensionToExtensionData;
+
+			const TArrayView<const FRegisteredObjectNodeInputPin> ExtensionPins = ICustomizableObjectModule::Get().GetAdditionalObjectNodePins();
+
+			for (FInstanceUpdateData::FNamedExtensionData& ExtensionOutput : OperationData->InstanceUpdateData.ExtendedInputPins)
+			{
+				const FRegisteredObjectNodeInputPin* FoundPin =
+					Algo::FindBy(ExtensionPins, ExtensionOutput.Name, &FRegisteredObjectNodeInputPin::GlobalPinName);
+
+				if (!FoundPin)
+				{
+					// Failed to find the corresponding pin for this output
+					// 
+					// This may indicate that a plugin has been removed or renamed since the CO was compiled
+					UE_LOG(LogMutable, Error, TEXT("Failed to find Object node input pin with name %s"), *ExtensionOutput.Name.ToString());
+					continue;
+				}
+
+				const UCustomizableObjectExtension* Extension = FoundPin->Extension.Get();
+				if (!Extension)
+				{
+					// Extension is not loaded or not found
+					UE_LOG(LogMutable, Error, TEXT("Extension for Object node input pin %s is no longer valid"), *ExtensionOutput.Name.ToString());
+					continue;
+				}
+
+				if (ExtensionOutput.Data->Origin == mu::ExtensionData::EOrigin::Invalid)
+				{
+					// Null data was produced
+					//
+					// This can happen if a node produces an ExtensionData but doesn't initialize it
+					UE_LOG(LogMutable, Error, TEXT("Invalid data sent to Object node input pin %s"), *ExtensionOutput.Name.ToString());
+					continue;
+				}
+
+				TArray<FInputPinDataContainer>& ContainerArray = ExtensionToExtensionData.FindOrAdd(Extension);
+
+				const FCustomizableObjectExtensionData* ReferencedExtensionData = nullptr;
+				switch (ExtensionOutput.Data->Origin)
+				{
+					case mu::ExtensionData::EOrigin::ConstantAlwaysLoaded:
+					{
+						check(CustomizableObject->AlwaysLoadedExtensionData.IsValidIndex(ExtensionOutput.Data->Index));
+						ReferencedExtensionData = &CustomizableObject->AlwaysLoadedExtensionData[ExtensionOutput.Data->Index];
+					}
+					break;
+
+					case mu::ExtensionData::EOrigin::ConstantStreamed:
+					{
+						check(CustomizableObject->StreamedExtensionData.IsValidIndex(ExtensionOutput.Data->Index));
+						
+						const FCustomizableObjectStreamedExtensionData& StreamedData =
+							CustomizableObject->StreamedExtensionData[ExtensionOutput.Data->Index];
+
+						if (!StreamedData.IsLoaded())
+						{
+							// The data should have been loaded as part of executing the CO program.
+							//
+							// This could indicate a bug in the streaming logic.
+							UE_LOG(LogMutable, Error, TEXT("Customizable Object produced a streamed extension data that is not loaded: %s"),
+								*StreamedData.GetPath().ToString());
+
+							continue;
+						}
+
+						ReferencedExtensionData = &StreamedData.GetLoadedData();
+					}
+					break;
+
+					default:
+						unimplemented();
+				}
+
+				check(ReferencedExtensionData);
+		
+				ContainerArray.Emplace(FoundPin->InputPin, ReferencedExtensionData->Data);
+			}
+
+			// Now that we have an array of extension data for each extension, go through the extensions and
+			// give each one its data.
+			for (const TPair<const UCustomizableObjectExtension*, TArray<FInputPinDataContainer>>& Pair : ExtensionToExtensionData)
+			{
+				Pair.Key->OnSkeletalMeshCreated(Pair.Value, ComponentIndex, SkeletalMesh);
+			}
+		}
+
 		bHasSkeletalMesh = true;
 	}
 
diff --git a/Engine/Plugins/Experimental/Mutable/Source/CustomizableObject/Private/MuCO/CustomizableObjectModule.cpp b/Engine/Plugins/Experimental/Mutable/Source/CustomizableObject/Private/MuCO/CustomizableObjectModule.cpp
--- a/Engine/Plugins/Experimental/Mutable/Source/CustomizableObject/Private/MuCO/CustomizableObjectModule.cpp
+++ b/Engine/Plugins/Experimental/Mutable/Source/CustomizableObject/Private/MuCO/CustomizableObjectModule.cpp
@@ -1,8 +1,12 @@
 // Copyright Epic Games, Inc. All Rights Reserved.
 
+#include "Algo/Find.h"
 #include "Interfaces/IPluginManager.h"
 #include "Misc/ConfigCacheIni.h"
+#include "MuCO/CustomizableObject.h"
+#include "MuCO/CustomizableObjectExtension.h"
 #include "MuCO/ICustomizableObjectModule.h"
+#include "UObject/StrongObjectPtr.h"
 
 /**
  * Customizable Object module implementation (private)
@@ -18,6 +22,22 @@ public:
 	// ICustomizableObjectModule interface
 	FString GetPluginVersion() const override;
 	bool AreExtraBoneInfluencesEnabled() const override;
+	void RegisterExtension(TObjectPtr<const UCustomizableObjectExtension> Extension) override;
+	void UnregisterExtension(TObjectPtr<const UCustomizableObjectExtension> Extension) override;
+	TArrayView<const TObjectPtr<const UCustomizableObjectExtension>> GetRegisteredExtensions() const override;
+	TArrayView<const FRegisteredCustomizableObjectPinType> GetExtendedPinTypes() const override;
+	TArrayView<const FRegisteredObjectNodeInputPin> GetAdditionalObjectNodePins() const override;
+
+private:
+	void RefreshExtensionData();
+
+	// Ensure extensions aren't garbage collected
+	TArray<TStrongObjectPtr<const UCustomizableObjectExtension>> StrongExtensions;
+	// For returning from GetRegisteredExtensions
+	TArray<TObjectPtr<const UCustomizableObjectExtension>> Extensions;
+
+	TArray<FRegisteredCustomizableObjectPinType> ExtendedPinTypes;
+	TArray<FRegisteredObjectNodeInputPin> AdditionalObjectNodePins;
 };
 
 
@@ -60,3 +80,89 @@ bool FCustomizableObjectModule::AreExtraBoneInfluencesEnabled() const
 
 	return bAreExtraBoneInfluencesEnabled;
 }
+
+void FCustomizableObjectModule::RegisterExtension(TObjectPtr<const UCustomizableObjectExtension> Extension)
+{
+	check(IsInGameThread());
+	
+	StrongExtensions.Add(TStrongObjectPtr<const UCustomizableObjectExtension>(Extension));
+	Extensions.Add(Extension);
+	
+	RefreshExtensionData();
+}
+
+void FCustomizableObjectModule::UnregisterExtension(TObjectPtr<const UCustomizableObjectExtension> Extension)
+{
+	check(IsInGameThread());
+	
+	StrongExtensions.Remove(TStrongObjectPtr<const UCustomizableObjectExtension>(Extension));
+	Extensions.Remove(Extension);
+	
+	RefreshExtensionData();
+}
+
+TArrayView<const TObjectPtr<const UCustomizableObjectExtension>> FCustomizableObjectModule::GetRegisteredExtensions() const
+{
+	check(IsInGameThread());
+	return MakeArrayView(Extensions);
+}
+
+TArrayView<const FRegisteredCustomizableObjectPinType> FCustomizableObjectModule::GetExtendedPinTypes() const
+{
+	check(IsInGameThread());
+	return MakeArrayView(ExtendedPinTypes);
+}
+
+TArrayView<const FRegisteredObjectNodeInputPin> FCustomizableObjectModule::GetAdditionalObjectNodePins() const
+{
+	check(IsInGameThread());
+	return MakeArrayView(AdditionalObjectNodePins);
+}
+
+void FCustomizableObjectModule::RefreshExtensionData()
+{
+	ExtendedPinTypes.Reset();
+	AdditionalObjectNodePins.Reset();
+
+	for (const TObjectPtr<const UCustomizableObjectExtension>& Extension : Extensions)
+	{
+		for (const FCustomizableObjectPinType& PinType : Extension->GetPinTypes())
+		{
+			FRegisteredCustomizableObjectPinType& RegisteredPinType = ExtendedPinTypes.AddDefaulted_GetRef();
+
+			RegisteredPinType.Extension = TWeakObjectPtr<const UCustomizableObjectExtension>(Extension);
+			RegisteredPinType.PinType = PinType;
+		}
+
+		for (const FObjectNodeInputPin& Pin : Extension->GetAdditionalObjectNodePins())
+		{
+			FRegisteredObjectNodeInputPin RegisteredPin;
+			RegisteredPin.Extension = TWeakObjectPtr<const UCustomizableObjectExtension>(Extension);
+			// Generate a name that will be unique across extensions, to prevent extensions from
+			// unintentionally interfering with each other.
+			RegisteredPin.GlobalPinName = FName(Extension->GetPathName() + TEXT("__") + Pin.PinName.ToString());
+			RegisteredPin.InputPin = Pin;
+
+			const FRegisteredObjectNodeInputPin* MatchingPin =
+				Algo::FindBy(AdditionalObjectNodePins, RegisteredPin.GlobalPinName, &FRegisteredObjectNodeInputPin::GlobalPinName);
+
+			if (MatchingPin)
+			{
+				// The pin should only be in the list if its extension is valid
+				check(MatchingPin->Extension.Get());
+
+				UE_LOG(LogMutable, Error,
+					TEXT("Object node pin %s from extension %s has the same name as pin %s from extension %s. Please rename one of the two."),
+					*Pin.PinName.ToString(),
+					*Extension->GetPathName(),
+					*MatchingPin->InputPin.PinName.ToString(),
+					*MatchingPin->Extension.Get()->GetPathName());
+
+				// Don't register the clashing pin
+				continue;
+			}
+
+			AdditionalObjectNodePins.Add(RegisteredPin);
+		}
+	}
+}
diff --git a/Engine/Plugins/Experimental/Mutable/Source/CustomizableObject/Private/MuCO/CustomizableObjectStreamedExtensionData.cpp b/Engine/Plugins/Experimental/Mutable/Source/CustomizableObject/Private/MuCO/CustomizableObjectStreamedExtensionData.cpp
new file mode 100644
--- /dev/null
+++ b/Engine/Plugins/Experimental/Mutable/Source/CustomizableObject/Private/MuCO/CustomizableObjectStreamedExtensionData.cpp
@@ -0,0 +1,79 @@
+// Copyright Epic Games, Inc. All Rights Reserved.
+
+#include "MuCO/CustomizableObjectStreamedExtensionData.h"
+
+#include "ExternalPackageHelper.h"
+#include "MuCO/CustomizableObject.h"
+
+#include UE_INLINE_GENERATED_CPP_BY_NAME(CustomizableObjectStreamedExtensionData)
+
+#if WITH_EDITOR
+FCustomizableObjectStreamedExtensionData::FCustomizableObjectStreamedExtensionData(
+	UCustomizableObjectExtensionDataContainer* InContainer)
+{
+	check(IsInGameThread());
+	check(InContainer);
+
+	Container = InContainer;
+	ContainerPath = InContainer;
+}
+
+void FCustomizableObjectStreamedExtensionData::ConvertToSoftReferenceForCooking(
+	const TSoftObjectPtr<UCustomizableObjectExtensionDataContainer>& NewContainerPath)
+{
+	check(IsInGameThread());
+
+	// Remove the hard reference to the container, so that it can be unloaded
+	Container = nullptr;
+
+	ContainerPath = NewContainerPath;
+}
+#endif // WITH_EDITOR
+
+bool FCustomizableObjectStreamedExtensionData::IsLoaded() const
+{
+	check(IsInGameThread());
+
+	return Container != nullptr;
+}
+
+const FCustomizableObjectExtensionData& FCustomizableObjectStreamedExtensionData::GetLoadedData() const
+{
+	check(IsInGameThread());
+
+	check(Container);
+	return Container->Data;
+}
+
+void FCustomizableObjectStreamedExtensionData::Unload()
+{
+	check(IsInGameThread());
+
+	if (!Container)
+	{
+		// Already unloaded
+		return;
+	}
+
+	if (FPlatformProperties::RequiresCookedData())
+	{
+		// The container should have been cooked out to its own package
+		check(Container->GetOuter()->IsA<UPackage>());
+
+		Container = nullptr;
+	}
+	else
+	{
+		// In an uncooked build the container is stored in the CO package and the CO needs to
+		// maintain a hard reference to it to ensure that it isn't discarded on save.
+		check(Container->GetOuter()->IsA<UCustomizableObject>());
+	}
+}
+
+void FCustomizableObjectStreamedExtensionData::NotifyLoaded(const UCustomizableObjectExtensionDataContainer* LoadedContainer)
+{
+	check(IsInGameThread());
+	check(ContainerPath.Get() == LoadedContainer);
+
+	Container = LoadedContainer;
+}
diff --git a/Engine/Plugins/Experimental/Mutable/Source/CustomizableObject/Private/MuCO/CustomizableObjectSystem.cpp b/Engine/Plugins/Experimental/Mutable/Source/CustomizableObject/Private/MuCO/CustomizableObjectSystem.cpp
--- a/Engine/Plugins/Experimental/Mutable/Source/CustomizableObject/Private/MuCO/CustomizableObjectSystem.cpp
+++ b/Engine/Plugins/Experimental/Mutable/Source/CustomizableObject/Private/MuCO/CustomizableObjectSystem.cpp
@@ -21,6 +21,7 @@
 #include "MuCO/LogBenchmarkUtil.h"
 #include "MuCO/LogInformationUtil.h"
 #include "MuCO/UnrealBakeHelpers.h"
+#include "MuCO/UnrealExtensionDataStreamer.h"
 #include "MuCO/UnrealMutableImageProvider.h"
 #include "MuCO/UnrealMutableModelDiskStreamer.h"
 #include "MuCO/UnrealPortabilityHelpers.h"
@@ -330,7 +331,8 @@ void UCustomizableObjectSystem::InitSystem()
 	check(pSettings);
 	pSettings->SetProfile(false);
 	pSettings->SetStreamingCache(Private->LastStreamingMemorySize);
-	Private->MutableSystem = new mu::System(pSettings);
+	Private->ExtensionDataStreamer = new FUnrealExtensionDataStreamer(Private.ToSharedRef());
+	Private->MutableSystem = new mu::System(pSettings, Private->ExtensionDataStreamer);
 	check(Private->MutableSystem);
 
 	Private->Streamer = new FUnrealMutableModelBulkStreamer();
@@ -1433,6 +1435,22 @@ namespace impl
 				}
 			}
 		}
+
+		// Copy ExtensionData Object node input from the Instance to the InstanceUpdateData
+		for (int32 ExtensionDataIndex = 0; ExtensionDataIndex < Instance->GetExtensionDataCount(); ExtensionDataIndex++)
+		{
+			mu::ExtensionDataPtrConst ExtensionData;
+			const char* NameAnsi = nullptr;
+			Instance->GetExtensionData(ExtensionDataIndex, ExtensionData, NameAnsi);
+
+			check(ExtensionData);
+			check(NameAnsi);
+
+			FInstanceUpdateData::FNamedExtensionData& NewEntry = OperationData->InstanceUpdateData.ExtendedInputPins.AddDefaulted_GetRef();
+			NewEntry.Data = ExtensionData;
+			NewEntry.Name = NameAnsi;
+			check(NewEntry.Name != NAME_None);
+		}
 	}
 
 
@@ -2040,6 +2058,15 @@ namespace impl
 		FCustomizableObjectSystemPrivate* SystemPrivateData = System->GetPrivate();
 		check(SystemPrivateData != nullptr);
 
+		// TODO: If this is the first code that runs after the CO program has finished AND if it's
+		// guaranteed that the next CO program hasn't started yet, we need to call ClearActiveObject
+		// and CancelPendingLoads on SystemPrivateData->ExtensionDataStreamer.
+		//
+		// ExtensionDataStreamer->AreAnyLoadsPending should return false if the program succeeded.
+		//
+		// If the program aborted, AreAnyLoadsPending may return true, as the program doesn't cancel
+		// its own loads on exit (maybe it should?)
+
 		FMutableResourceCache& Cache = SystemPrivateData->GetObjectCache(CustomizableObject);
 
 		System->ProtectedCachedTextures.Reset(Cache.Images.Num());
@@ -2263,6 +2290,9 @@ namespace impl
 		check(SystemPrivateData->Streamer != nullptr);
 		SystemPrivateData->Streamer->PrepareStreamingForObject(CustomizableObject);
 
+		check(SystemPrivateData->ExtensionDataStreamer != nullptr);
+		SystemPrivateData->ExtensionDataStreamer->SetActiveObject(CustomizableObject);
+
 		CandidateInstance->CommitMinMaxLOD();
 
 		FString StateName = CandidateInstance->GetCustomizableObject()->GetStateName(CandidateInstance->GetState());
diff --git a/Engine/Plugins/Experimental/Mutable/Source/CustomizableObject/Private/MuCO/CustomizableObjectSystemPrivate.h b/Engine/Plugins/Experimental/Mutable/Source/CustomizableObject/Private/MuCO/CustomizableObjectSystemPrivate.h
--- a/Engine/Plugins/Experimental/Mutable/Source/CustomizableObject/Private/MuCO/CustomizableObjectSystemPrivate.h
+++ b/Engine/Plugins/Experimental/Mutable/Source/CustomizableObject/Private/MuCO/CustomizableObjectSystemPrivate.h
@@ -4,6 +4,7 @@
 
 #include "Containers/Queue.h"
 #include "MuCO/CustomizableObject.h"
+#include "MuCO/CustomizableObjectExtension.h"
 #include "Containers/Ticker.h"
 
 #include "MuCO/CustomizableObjectInstanceDescriptor.h"
@@ -418,6 +419,13 @@ struct FInstanceUpdateData
 
 	TArray<FSkeletonData> Skeletons;
 
+	struct FNamedExtensionData
+	{
+		mu::ExtensionDataPtrConst Data;
+		FName Name;
+	};
+	TArray<FNamedExtensionData> ExtendedInputPins;
+
 	/** */
 	void Clear()
 	{
@@ -428,6 +436,7 @@ struct FInstanceUpdateData
 		Scalars.Empty();
 		Vectors.Empty();
 		Skeletons.Empty();
+		ExtendedInputPins.Empty();
 	}
 };
 
@@ -542,6 +551,9 @@ public:
 	// Non-owned reference
 	class FUnrealMutableModelBulkStreamer* Streamer = nullptr;
 
+	// Non-owned reference
+	class FUnrealExtensionDataStreamer* ExtensionDataStreamer = nullptr;
+
 	// This object is responsible for providing custom images to mutable models (for image parameters)
 	// This object is called from the mutable thread, and it should only access data already safely submitted from
 	// the game thread and stored in FUnrealMutableImageProvider::GlobalExternalImages.
diff --git a/Engine/Plugins/Experimental/Mutable/Source/CustomizableObject/Private/MuCO/UnrealExtensionDataStreamer.cpp b/Engine/Plugins/Experimental/Mutable/Source/CustomizableObject/Private/MuCO/UnrealExtensionDataStreamer.cpp
new file mode 100644
--- /dev/null
+++ b/Engine/Plugins/Experimental/Mutable/Source/CustomizableObject/Private/MuCO/UnrealExtensionDataStreamer.cpp
@@ -0,0 +1,278 @@
+// Copyright Epic Games, Inc. All Rights Reserved.
+
+#include "MuCO/UnrealExtensionDataStreamer.h"
+
+#include "Algo/Find.h"
+#include "Algo/IndexOf.h"
+#include "Async/Async.h"
+#include "MuCO/CustomizableObject.h"
+#include "MuCO/CustomizableObjectSystem.h"
+#include "MuCO/CustomizableObjectSystemPrivate.h"
+#include "Templates/SharedPointer.h"
+
+FUnrealExtensionDataStreamer::FUnrealExtensionDataStreamer(const TSharedRef<FCustomizableObjectSystemPrivate>& InSystemPrivate)
+	: SystemPrivate(InSystemPrivate)
+{
+	check(IsInGameThread());
+
+	Mutex = new FCriticalSection();
+}
+
+FUnrealExtensionDataStreamer::~FUnrealExtensionDataStreamer()
+{
+	check(IsInGameThread());
+
+	CancelPendingLoads();
+
+	delete Mutex;
+}
+
+void FUnrealExtensionDataStreamer::SetActiveObject(UCustomizableObject* InObject)
+{
+	check(IsInGameThread());
+	check(InObject);
+
+	FScopeLock Lock(Mutex);
+
+	ActiveObject = InObject;
+}
+
+void FUnrealExtensionDataStreamer::ClearActiveObject()
+{
+	check(IsInGameThread());
+
+	FScopeLock Lock(Mutex);
+
+	ActiveObject = nullptr;
+}
+
+bool FUnrealExtensionDataStreamer::AreAnyLoadsPending() const
+{
+	FScopeLock Lock(Mutex);
+
+	return PendingLoads.Num() > 0;
+}
+
+void FUnrealExtensionDataStreamer::CancelPendingLoads()
+{
+	// Needs to be called from the Game thread, as FStreamableHandle::CancelHandle requires this.
+	//
+	// Restricting this to the Game thread also guarantees that any tasks cancelled by this
+	// function will definitely not have started yet, since they only run on the Game thread.
+	check(IsInGameThread());
+
+	FScopeLock Lock(Mutex);
+
+	bool bAnyTasksNeedCancelling = false;
+	for (const FPendingLoad& PendingLoad : PendingLoads)
+	{
+		if (!PendingLoad.TaskFuture.IsReady())
+		{
+			// This task was queued to be run on the Game thread but hasn't run yet
+			bAnyTasksNeedCancelling = true;
+			continue;
+		}
+
+		TSharedPtr<FStreamableHandle> Handle = PendingLoad.TaskFuture.Get();
+		if (Handle.IsValid())
+		{
+			Handle->CancelHandle();
+		}
+	}
+
+	if (bAnyTasksNeedCancelling)
+	{
+		// Set the shared cancel value to true so that whenever the tasks actually execute they
+		// will cancel immediately.
+		check(ShouldCancelPtr.IsValid());
+		*ShouldCancelPtr = true;
+
+		// The next task to be queued should allocate a new cancel value to ensure it doesn't
+		// interfere with the previous one.
+		ShouldCancelPtr = nullptr;
+	}
+
+	PendingLoads.Empty();
+}
+
+mu::ExtensionDataPtr FUnrealExtensionDataStreamer::CloneExtensionData(const mu::ExtensionDataPtrConst& Source)
+{
+	unimplemented();
+	return nullptr;
+}
+
+TSharedRef<const mu::FExtensionDataLoadHandle> FUnrealExtensionDataStreamer::StartLoad(
+	const mu::ExtensionDataPtrConst& Data,
+	TArray<mu::ExtensionDataPtrConst>& OutUnloadedConstants)
+{
+	// Note that this can be called from any thread
+
+	OutUnloadedConstants.Reset();
+
+	TSharedRef<mu::FExtensionDataLoadHandle> LoadHandle = MakeShared<mu::FExtensionDataLoadHandle>();
+
+	LoadHandle->Data = Data;
+	LoadHandle->LoadState = mu::FExtensionDataLoadHandle::ELoadState::Pending;
+
+	{
+		FScopeLock Lock(Mutex);
+
+		if (!ShouldCancelPtr.IsValid())
+		{
+			ShouldCancelPtr = MakeShared<bool>(false);
+		}
+
+		TFuture<TSharedPtr<FStreamableHandle>> LoadHandleFuture = Async(EAsyncExecution::TaskGraphMainThread,
+			[
+				ShouldCancelRef = ShouldCancelPtr.ToSharedRef(),
+				LoadHandle,
+				ObjectToLoadFor = ActiveObject,
+				SystemPrivate = SystemPrivate
+			]()
+			{
+				if (*ShouldCancelRef)
+				{
+					// Load was cancelled
+					LoadHandle->LoadState = mu::FExtensionDataLoadHandle::ELoadState::FailedToLoad;
+					return TSharedPtr<FStreamableHandle>();
+				}
+
+				return StartLoadOnGameThread(SystemPrivate, ObjectToLoadFor, LoadHandle);
+			});
+
+
+		PendingLoads.Emplace(LoadHandle, MoveTemp(LoadHandleFuture));
+	}
+
+	return LoadHandle;
+}
+
+// Note that this function only returns an FStreamableHandle if it needed to start a load.
+//
+// If the requested object is already loaded, this function can return nullptr and still be
+// considered successful.
+TSharedPtr<FStreamableHandle> FUnrealExtensionDataStreamer::StartLoadOnGameThread(
+	const TSharedRef<FCustomizableObjectSystemPrivate>& SystemPrivate,
+	const TWeakObjectPtr<UCustomizableObject>& ObjectToLoadFor,
+	const TSharedRef<mu::FExtensionDataLoadHandle>& LoadHandle)
+{
+	check(IsInGameThread());
+	check(LoadHandle->Data->Origin == mu::ExtensionData::EOrigin::ConstantStreamed);
+
+	UCustomizableObject* Object = ObjectToLoadFor.Get();
+	if (!Object)
+	{
+		// CO doesn't exist anymore
+		LoadHandle->LoadState = mu::FExtensionDataLoadHandle::ELoadState::FailedToLoad;
+		return nullptr;
+	}
+
+	if (!Object->StreamedExtensionData.IsValidIndex(LoadHandle->Data->Index))
+	{
+		// The compiled data appears to be out of sync with the CO's properties
+
+		UE_LOG(LogMutable, Error, TEXT("Couldn't find streamed Extension Data with index %d in %s. Compiled data may be stale."),
+			LoadHandle->Data->Index, *Object->GetFullName());
+
+		LoadHandle->LoadState = mu::FExtensionDataLoadHandle::ELoadState::FailedToLoad;
+		return nullptr;
+	}
+
+	FCustomizableObjectStreamedExtensionData& StreamedData = Object->StreamedExtensionData[LoadHandle->Data->Index];
+	if (StreamedData.IsLoaded())
+	{
+		// Already loaded
+
+		// Need to call NotifyLoadCompleted to mark the request as complete and remove it from
+		// PendingLoads.
+		SystemPrivate->ExtensionDataStreamer->NotifyLoadCompleted(Object, LoadHandle);
+
+		return nullptr;
+	}
+
+	// Note that this just checks if the path is non-null, NOT if the object is loaded
+	check(!StreamedData.GetPath().IsNull());
+
+	TArray<FSoftObjectPath> TargetsToStream;
+	TargetsToStream.Add(StreamedData.GetPath().ToSoftObjectPath());
+
+	TFunction<void()> OnLoadComplete = [SystemPrivate, ObjectToLoadFor, LoadHandle]()
+	{
+		UCustomizableObject* Object = ObjectToLoadFor.Get();
+		if (!Object)
+		{
+			// CO doesn't exist anymore
+			return;
+		}
+
+		check(SystemPrivate->ExtensionDataStreamer);
+		SystemPrivate->ExtensionDataStreamer->NotifyLoadCompleted(Object, LoadHandle);
+	};
+
+	const bool bManageActiveHandle = false;
+	const bool bStartStalled = false;
+
+	FString DebugString;
+#if UE_BUILD_DEBUG || UE_BUILD_DEVELOPMENT
+	DebugString = FString::Printf(TEXT("UnrealExtensionDataStreamer for %s"), *Object->GetPathName());
+#endif
+
+	FStreamableManager& Manager = UCustomizableObjectSystem::GetInstance()->GetStreamableManager();
+
+	return Manager.RequestAsyncLoad(
+		TargetsToStream,
+		MoveTemp(OnLoadComplete),
+		FStreamableManager::DefaultAsyncLoadPriority,
+		bManageActiveHandle,
+		bStartStalled,
+		DebugString);
+}
+
+void FUnrealExtensionDataStreamer::NotifyLoadCompleted(
+	UCustomizableObject* Object,
+	const TSharedRef<mu::FExtensionDataLoadHandle>& LoadHandle)
+{
+	check(IsInGameThread());
+	check(Object);
+
+	if (!Object->StreamedExtensionData.IsValidIndex(LoadHandle->Data->Index))
+	{
+		// The compiled data appears to be out of sync with the CO's properties
+
+		UE_LOG(LogMutable, Error, TEXT("Couldn't find streamed Extension Data with index %d in %s. Compiled data may be stale."),
+			LoadHandle->Data->Index, *Object->GetFullName());
+
+		LoadHandle->LoadState = mu::FExtensionDataLoadHandle::ELoadState::FailedToLoad;
+		return;
+	}
+
+	FCustomizableObjectStreamedExtensionData& StreamedData = Object->StreamedExtensionData[LoadHandle->Data->Index];
+
+	// The object could have been loaded by another request, in which case we can skip updating
+	// the StreamedData.
+	if (!StreamedData.IsLoaded())
+	{
+		const UCustomizableObjectExtensionDataContainer* LoadedObject = StreamedData.GetPath().Get();
+		if (!LoadedObject)
+		{
+			// Object wasn't loaded for some reason
+			LoadHandle->LoadState = mu::FExtensionDataLoadHandle::ELoadState::FailedToLoad;
+			return;
+		}
+
+		StreamedData.NotifyLoaded(LoadedObject);
+	}
+
+	LoadHandle->LoadState = mu::FExtensionDataLoadHandle::ELoadState::Loaded;
+
+	{
+		FScopeLock Lock(Mutex);
+
+		const int32 Index = Algo::IndexOfBy(PendingLoads, LoadHandle, &FPendingLoad::LoadHandle);
+		// If the load was previously cancelled, CancelHandle should have prevented this callback from
+		// running, so at this point the load should still be in PendingLoads.
+		check(Index != INDEX_NONE);
+
+		PendingLoads.RemoveAtSwap(Index);
+	}
+}
diff --git a/Engine/Plugins/Experimental/Mutable/Source/CustomizableObject/Public/MuCO/ICustomizableObjectModule.h b/Engine/Plugins/Experimental/Mutable/Source/CustomizableObject/Public/MuCO/ICustomizableObjectModule.h
--- a/Engine/Plugins/Experimental/Mutable/Source/CustomizableObject/Public/MuCO/ICustomizableObjectModule.h
+++ b/Engine/Plugins/Experimental/Mutable/Source/CustomizableObject/Public/MuCO/ICustomizableObjectModule.h
@@ -2,8 +2,24 @@
 
 #pragma once
 
+#include "Containers/ArrayView.h"
 #include "Modules/ModuleManager.h"
+#include "MuCO/CustomizableObjectExtension.h"
 
+struct FRegisteredCustomizableObjectPinType
+{
+	TWeakObjectPtr<const UCustomizableObjectExtension> Extension;
+	FCustomizableObjectPinType PinType;
+};
+
+struct FRegisteredObjectNodeInputPin
+{
+	TWeakObjectPtr<const UCustomizableObjectExtension> Extension;
+	/** A name for this pin that should be globally unique across extensions */
+	FName GlobalPinName;
+
+	FObjectNodeInputPin InputPin;
+};
 
 /**
  * The public interface of the CustomizableObject module
@@ -39,5 +55,24 @@ public:
 	// Return a string representing the plugin version.
 	virtual FString GetPluginVersion() const = 0;
 
+	/**
+	 * Extension functions
+	 *
+	 * These may only be called from the game thread
+	 */
+
+	virtual void RegisterExtension(TObjectPtr<const UCustomizableObjectExtension> Extension) = 0;
+	virtual void UnregisterExtension(TObjectPtr<const UCustomizableObjectExtension> Extension) = 0;
+	virtual TArrayView<const TObjectPtr<const UCustomizableObjectExtension>> GetRegisteredExtensions() const = 0;
+
+	/**
+	 * The results from these functions should only reference extensions that are still valid.
+	 *
+	 * If one of these functions returns data with an invalid weak pointer to a
+	 * UCustomizableObjectExtension, it means that the extension was unloaded without calling
+	 * UnregisterExtension.
+	 */
+	virtual TArrayView<const FRegisteredCustomizableObjectPinType> GetExtendedPinTypes() const = 0;
+	virtual TArrayView<const FRegisteredObjectNodeInputPin> GetAdditionalObjectNodePins() const = 0;
 };
 
diff --git a/Engine/Plugins/Experimental/Mutable/Source/CustomizableObject/Public/MuCO/UnrealExtensionDataStreamer.h b/Engine/Plugins/Experimental/Mutable/Source/CustomizableObject/Public/MuCO/UnrealExtensionDataStreamer.h
new file mode 100644
--- /dev/null
+++ b/Engine/Plugins/Experimental/Mutable/Source/CustomizableObject/Public/MuCO/UnrealExtensionDataStreamer.h
@@ -0,0 +1,98 @@
+// Copyright Epic Games, Inc. All Rights Reserved.
+
+#pragma once
+
+#include "Engine/StreamableManager.h"
+#include "MuR/ExtensionData.h"
+#include "MuR/ExtensionDataStreamer.h"
+#include "Templates/SharedPointerFwd.h"
+
+class FCustomizableObjectSystemPrivate;
+class UCustomizableObject;
+
+/** An implementation of mu::ExtensionDataStreamer designed to work with the Customizable Object integration */
+class CUSTOMIZABLEOBJECT_API FUnrealExtensionDataStreamer : public mu::ExtensionDataStreamer
+{
+public:
+	FUnrealExtensionDataStreamer(const TSharedRef<FCustomizableObjectSystemPrivate>& InSystemPrivate);
+	virtual ~FUnrealExtensionDataStreamer();
+
+	/** Not copyable because TFuture isn't copyable. This shouldn't be limiting in practise. */
+	FUnrealExtensionDataStreamer(const FUnrealExtensionDataStreamer& Other) = delete;
+	FUnrealExtensionDataStreamer& operator=(const FUnrealExtensionDataStreamer& Other) = delete;
+
+	/**
+	 * Note that *only* the functions from the mu::ExtensionDataStreamer interface may be called
+	 * from other threads. All other functions may only be called from the Game thread.
+	 */
+
+	/**
+	 * Set the Customizable Object that is currently running.
+	 *
+	 * StartLoad calls are assumed to come from the object passed in here.
+	 */
+	void SetActiveObject(UCustomizableObject* InObject);
+
+	/**
+	 * Clear the active object
+	 * 
+	 * Call this after a Customizable Object has finished running to help detect cases where
+	 * SetActiveObject wasn't called before running a different CO.
+	 */
+	void ClearActiveObject();
+
+	/** Returns true if there are any loads requested by StartLoad that haven't yet completed. */
+	bool AreAnyLoadsPending() const;
+
+	/** Cancels all loads requested by StartLoad that haven't yet completed. */
+	void CancelPendingLoads();
+
+	/** mu::ExtensionDataStreamer interface */
+	virtual mu::ExtensionDataPtr CloneExtensionData(const mu::ExtensionDataPtrConst& Source) override;
+	
+	virtual TSharedRef<const mu::FExtensionDataLoadHandle> StartLoad(
+		const mu::ExtensionDataPtrConst& Data,
+		TArray<mu::ExtensionDataPtrConst>& OutUnloadedConstants) override;
+
+private:
+	/** Queued up on the Game thread to start the async load */
+	static TSharedPtr<FStreamableHandle> StartLoadOnGameThread(
+		const TSharedRef<FCustomizableObjectSystemPrivate>& SystemPrivate,
+		const TWeakObjectPtr<UCustomizableObject>& ObjectToLoadFor,
+		const TSharedRef<mu::FExtensionDataLoadHandle>& LoadHandle);
+
+	/** Called when the async load has finished loading the Extension Data */
+	void NotifyLoadCompleted(UCustomizableObject* Object, const TSharedRef<mu::FExtensionDataLoadHandle>& LoadHandle);
+
+	TSharedRef<FCustomizableObjectSystemPrivate> SystemPrivate;
+
+	/** Guards PendingLoads, ShouldCancelPtr and ActiveObject */
+	FCriticalSection* Mutex = nullptr;
+
+	TWeakObjectPtr<UCustomizableObject> ActiveObject;
+	
+	struct FPendingLoad
+	{
+		FPendingLoad(
+			const TSharedRef<mu::FExtensionDataLoadHandle>& InLoadHandle,
+			TFuture<TSharedPtr<FStreamableHandle>>&& InTaskFuture)
+			: LoadHandle(InLoadHandle)
+			, TaskFuture(MoveTemp(InTaskFuture))
+		{
+		}
+
+		FPendingLoad(const FPendingLoad& Other) = delete;
+		FPendingLoad& operator=(const FPendingLoad& Other) = delete;
+
+		TSharedRef<mu::FExtensionDataLoadHandle> LoadHandle;
+		TFuture<TSharedPtr<FStreamableHandle>> TaskFuture;
+	};
+	/** All outstanding loads that haven't called NotifyLoadCompleted yet */
+	TArray<FPendingLoad> PendingLoads;
+
+	/**
+	 * This is shared by all outstanding loads and can be used to cancel them if their Game thread
+	 * task hasn't run yet.
+	 */
+	TSharedPtr<bool> ShouldCancelPtr;
+};
diff --git a/Engine/Plugins/Experimental/Mutable/Source/CustomizableObjectEditor/Classes/MuCOE/ICustomizableObjectExtensionNode.h b/Engine/Plugins/Experimental/Mutable/Source/CustomizableObjectEditor/Classes/MuCOE/ICustomizableObjectExtensionNode.h
new file mode 100644
--- /dev/null
+++ b/Engine/Plugins/Experimental/Mutable/Source/CustomizableObjectEditor/Classes/MuCOE/ICustomizableObjectExtensionNode.h
@@ -0,0 +1,26 @@
+// Copyright Epic Games, Inc. All Rights Reserved.
+
+#pragma once
+
+#include "MuT/NodeExtensionData.h"
+#include "UObject/Interface.h"
+
+#include "ICustomizableObjectExtensionNode.generated.h"
+
+class FExtensionDataCompilerInterface;
+
+/** Customizable Object graph nodes that output ExtensionData should implement this interface */
+UINTERFACE()
+class CUSTOMIZABLEOBJECTEDITOR_API UCustomizableObjectExtensionNode : public UInterface
+{
+	GENERATED_BODY()
+};
+
+class ICustomizableObjectExtensionNode
+{
+	GENERATED_BODY()
+
+public:
+	/** Generate a mu::Node that produces ExtensionData */
+	virtual mu::NodeExtensionDataPtr GenerateMutableNode(FExtensionDataCompilerInterface& CompilerInterface) const = 0;
+};
diff --git a/Engine/Plugins/Experimental/Mutable/Source/CustomizableObjectEditor/Classes/MuCOE/Nodes/CustomizableObjectNode.h b/Engine/Plugins/Experimental/Mutable/Source/CustomizableObjectEditor/Classes/MuCOE/Nodes/CustomizableObjectNode.h
--- a/Engine/Plugins/Experimental/Mutable/Source/CustomizableObjectEditor/Classes/MuCOE/Nodes/CustomizableObjectNode.h
+++ b/Engine/Plugins/Experimental/Mutable/Source/CustomizableObjectEditor/Classes/MuCOE/Nodes/CustomizableObjectNode.h
@@ -101,6 +101,14 @@ public:
 	/** Virtual implementation of RemovePin. Allows to do work before removing a pin.
 	 * Use this function instead of RemovePin. RemovePin does not removes possible attached PinData. */
 	virtual bool CustomRemovePin(UEdGraphPin& Pin);
+
+	/**
+	 * Subclasses should override this to return true and set OutCategory if this node should be
+	 * auto-added to the right-click context menu in the graph editor.
+	 * 
+	 * Some nodes are added manually in the graph editor code and don't need to do this.
+	 */
+	virtual bool ShouldAddToContextMenu(FText& OutCategory) const { return false; }
 	
 	void GetInputPins(TArray<class UEdGraphPin*>& OutInputPins) const;
 	void GetOutputPins(TArray<class UEdGraphPin*>& OutOutputPins) const;
diff --git a/Engine/Plugins/Experimental/Mutable/Source/CustomizableObjectEditor/Classes/MuCOE/Nodes/CustomizableObjectNodeExtensionDataConstant.h b/Engine/Plugins/Experimental/Mutable/Source/CustomizableObjectEditor/Classes/MuCOE/Nodes/CustomizableObjectNodeExtensionDataConstant.h
new file mode 100644
--- /dev/null
+++ b/Engine/Plugins/Experimental/Mutable/Source/CustomizableObjectEditor/Classes/MuCOE/Nodes/CustomizableObjectNodeExtensionDataConstant.h
@@ -0,0 +1,22 @@
+// Copyright Epic Games, Inc. All Rights Reserved.
+
+#pragma once
+
+#include "MuCOE/Nodes/CustomizableObjectNode.h"
+
+#include "CustomizableObjectNodeExtensionDataConstant.generated.h"
+
+/**
+ * A base class for nodes that output a constant ExtensionData.
+ *
+ * This is typically used to import an Unreal asset into the Customizable Object graph.
+ */
+UCLASS(abstract)
+class CUSTOMIZABLEOBJECTEDITOR_API UCustomizableObjectNodeExtensionDataConstant : public UCustomizableObjectNode
+{
+	GENERATED_BODY()
+
+public:
+	virtual bool IsAffectedByLOD() const override { return false; }
+};
+
diff --git a/Engine/Plugins/Experimental/Mutable/Source/CustomizableObjectEditor/Classes/MuCOE/Nodes/CustomizableObjectNodeObject.h b/Engine/Plugins/Experimental/Mutable/Source/CustomizableObjectEditor/Classes/MuCOE/Nodes/CustomizableObjectNodeObject.h
--- a/Engine/Plugins/Experimental/Mutable/Source/CustomizableObjectEditor/Classes/MuCOE/Nodes/CustomizableObjectNodeObject.h
+++ b/Engine/Plugins/Experimental/Mutable/Source/CustomizableObjectEditor/Classes/MuCOE/Nodes/CustomizableObjectNodeObject.h
@@ -202,7 +202,6 @@ public:
 
 	// UObject interface.
 	void PostEditChangeProperty(FPropertyChangedEvent& PropertyChangedEvent) override;
-	void BackwardsCompatibleFixup() override;
 
 	// EdGraphNode interface
 	FText GetNodeTitle(ENodeTitleType::Type TitleType) const override;
@@ -213,6 +212,7 @@ public:
 	void PostDuplicate(bool bDuplicateForPIE) override;
 
 	// UCustomizableObjectNode interface
+	virtual void BackwardsCompatibleFixup() override;
 	virtual void PostBackwardsCompatibleFixup() override;
 	void AllocateDefaultPins(UCustomizableObjectNodeRemapPins* RemapPins) override;
 
@@ -225,7 +225,7 @@ public:
 
 	UEdGraphPin* LODPin(int32 LODIndex) const
 	{
-		FString LODName = FString::Printf(TEXT("LOD %d "), LODIndex);
+		FString LODName = FString::Printf(TEXT("%s%d "), LODPinNamePrefix, LODIndex);
 		return FindPin(LODName);
 	}
 
@@ -235,7 +235,7 @@ public:
 
 		for (UEdGraphPin* Pin : GetAllNonOrphanPins())
 		{
-			if (Pin->GetName().StartsWith(TEXT("LOD ")))
+			if (Pin->GetName().StartsWith(LODPinNamePrefix))
 			{
 				Count++;
 			}
@@ -246,12 +246,12 @@ public:
 
 	UEdGraphPin* ChildrenPin() const
 	{
-		return FindPin(TEXT("Children"));
+		return FindPin(ChildrenPinName);
 	}
 
 	UEdGraphPin* OutputPin() const
 	{
-		return FindPin(TEXT("Object"));
+		return FindPin(OutputPinName);
 	}
 
 	/** Return the LOD which a LOD pin references to. Retrun -1 if a pin does not belong to any LOD. */
@@ -276,5 +276,12 @@ public:
 	// Node Details Support
 	int32 CurrentComponent = 0;
 	int32 CurrentLOD = 0;
+
+private:
+	static const FName ChildrenPinName;
+	static const FName OutputPinName;
+	static const TCHAR* LODPinNamePrefix;
+
+	static bool IsBuiltInPin(FName PinName);
 };
 
diff --git a/Engine/Plugins/Experimental/Mutable/Source/CustomizableObjectEditor/CustomizableObjectEditor.Build.cs b/Engine/Plugins/Experimental/Mutable/Source/CustomizableObjectEditor/CustomizableObjectEditor.Build.cs
--- a/Engine/Plugins/Experimental/Mutable/Source/CustomizableObjectEditor/CustomizableObjectEditor.Build.cs
+++ b/Engine/Plugins/Experimental/Mutable/Source/CustomizableObjectEditor/CustomizableObjectEditor.Build.cs
@@ -99,6 +99,7 @@ public class CustomizableObjectEditor : ModuleRules
 			new string[] {
 				"PropertyEditor",
 				"CustomizableObject",
+				"StructUtils",
 			}
 			);
 
diff --git a/Engine/Plugins/Experimental/Mutable/Source/CustomizableObjectEditor/Private/MuCOE/CustomizableObjectCompiler.cpp b/Engine/Plugins/Experimental/Mutable/Source/CustomizableObjectEditor/Private/MuCOE/CustomizableObjectCompiler.cpp
--- a/Engine/Plugins/Experimental/Mutable/Source/CustomizableObjectEditor/Private/MuCOE/CustomizableObjectCompiler.cpp
+++ b/Engine/Plugins/Experimental/Mutable/Source/CustomizableObjectEditor/Private/MuCOE/CustomizableObjectCompiler.cpp
@@ -1091,6 +1091,13 @@ void FCustomizableObjectCompiler::CompileInternal(UCustomizableObject* Object, c
 		Object->GroupNodeMap = GenerationContext.GroupNodeMap;
 		Object->ParameterUIDataMap = GenerationContext.ParameterUIDataMap;
 		Object->StateUIDataMap = GenerationContext.StateUIDataMap;
+		Object->AlwaysLoadedExtensionData = MoveTemp(GenerationContext.AlwaysLoadedExtensionData);
+
+		Object->StreamedExtensionData.Empty(GenerationContext.StreamedExtensionData.Num());
+		for (UCustomizableObjectExtensionDataContainer* Container : GenerationContext.StreamedExtensionData)
+		{
+			Object->StreamedExtensionData.Emplace(Container);
+		}
 
 #if WITH_EDITORONLY_DATA
 		Object->CustomizableObjectPathMap = GenerationContext.CustomizableObjectPathMap;
diff --git a/Engine/Plugins/Experimental/Mutable/Source/CustomizableObjectEditor/Private/MuCOE/CustomizableObjectCookPackageSplitter.cpp b/Engine/Plugins/Experimental/Mutable/Source/CustomizableObjectEditor/Private/MuCOE/CustomizableObjectCookPackageSplitter.cpp
new file mode 100644
--- /dev/null
+++ b/Engine/Plugins/Experimental/Mutable/Source/CustomizableObjectEditor/Private/MuCOE/CustomizableObjectCookPackageSplitter.cpp
@@ -0,0 +1,260 @@
+// Copyright Epic Games, Inc. All Rights Reserved.
+
+#include "MuCOE/CustomizableObjectCookPackageSplitter.h"
+
+#include "Algo/Find.h"
+#include "MuCO/CustomizableObject.h"
+#include "UObject/NameTypes.h"
+
+REGISTER_COOKPACKAGE_SPLITTER(FCustomizableObjectCookPackageSplitter, UCustomizableObject);
+
+namespace
+{
+// Look up a streamed Extension Data constant by name on a Customizable Object.
+//
+// Returns nullptr if not found.
+FCustomizableObjectStreamedExtensionData* FindStreamedExtensionData(
+	UCustomizableObject* Object,
+	const FString& ContainerName)
+{
+	return Algo::FindByPredicate(
+		Object->StreamedExtensionData,
+			[&ContainerName](const FCustomizableObjectStreamedExtensionData& StreamedData)
+			{
+				const FSoftObjectPath& Path = StreamedData.GetPath().ToSoftObjectPath();
+				
+				// ContainerName should match the last element of the path, which could be the
+				// sub-path string or the asset name.
+
+				if (Path.GetSubPathString().Len() > 0)
+				{
+					return Path.GetSubPathString() == ContainerName;
+				}
+
+				return Path.GetAssetName() == ContainerName;
+			});
+}
+
+// Finds a streamed Extension Data container by name and moves it to the given Outer
+bool MoveContainerToNewOuter(
+	UCustomizableObject* Object,
+	const FString& ContainerName,
+	UObject* NewOuter,
+	UCustomizableObjectExtensionDataContainer*& OutContainer)
+{
+	OutContainer = nullptr;
+
+	const FCustomizableObjectStreamedExtensionData* FoundData = FindStreamedExtensionData(Object, ContainerName);
+	if (!FoundData)
+	{
+		UE_LOG(LogMutable, Error, TEXT("Couldn't find streamed Extension Data container with name %s in array of %d entries"),
+			*ContainerName, Object->StreamedExtensionData.Num());
+
+		return false;
+	}
+
+	UCustomizableObjectExtensionDataContainer* Container = FoundData->GetPath().LoadSynchronous();
+	if (!Container)
+	{
+		UE_LOG(LogMutable, Error, TEXT("Failed to load streamed Extension Data container %s"),
+			*FoundData->GetPath().ToString());
+
+		return false;
+	}
+
+	// Ensure the target object doesn't exist
+	check(!FindObject<UObject>(NewOuter, *Container->GetName()));
+
+	// The Rename function moves the object into the given package
+	if (!Container->Rename(nullptr, NewOuter, REN_DontCreateRedirectors))
+	{
+		UE_LOG(LogMutable, Error, TEXT("Failed to move streamed Extension Data container %s into Outer %s"),
+			*Container->GetPathName(), *NewOuter->GetPathName());
+
+		return false;
+	}
+
+	OutContainer = Container;
+	return true;
+}
+}
+
+bool FCustomizableObjectCookPackageSplitter::ShouldSplit(UObject* SplitData)
+{
+	const UCustomizableObject* Object = CastChecked<UCustomizableObject>(SplitData);
+
+	return Object->StreamedExtensionData.Num() > 0;
+}
+
+TArray<ICookPackageSplitter::FGeneratedPackage> FCustomizableObjectCookPackageSplitter::GetGenerateList(
+	const UPackage* OwnerPackage,
+	const UObject* OwnerObject)
+{
+	const UCustomizableObject* Object = CastChecked<UCustomizableObject>(OwnerObject);
+
+	TArray<ICookPackageSplitter::FGeneratedPackage> Result;
+
+	// Generate a new package for each streamed Extension Data
+	for (const FCustomizableObjectStreamedExtensionData& StreamedData : Object->StreamedExtensionData)
+	{
+		// The StreamedData container path should be of the form
+		// OwnerPackageName.OwnerObjectName:ContainerName
+		const FSoftObjectPath& StreamedDataPath = StreamedData.GetPath().ToSoftObjectPath();
+		
+		// Check that the StreamedData container has the OwnerObject as its Outer
+		check(StreamedDataPath.GetWithoutSubPath() == FSoftObjectPath(OwnerObject));
+
+		// Check that the ContainerName is valid and that there isn't another Outer level between
+		// the OwnerObject and the container.
+		check(StreamedDataPath.GetSubPathString().Len() > 0);
+		check(!StreamedDataPath.GetSubPathString().Contains(SUBOBJECT_DELIMITER));
+
+		ICookPackageSplitter::FGeneratedPackage& Package = Result.AddDefaulted_GetRef();
+		// Because of the checks above, the container name must be unique within this Customizable
+		// Object, so it's safe to use as a package path.
+		Package.RelativePath = StreamedDataPath.GetSubPathString();
+		Package.SetCreateAsMap(false);
+
+		// To support iterative cooking, GenerationHash should only change when OwnerPackage
+		// changes.
+		//
+		// The simplest and fastest way to do this is to set it to OwnerPackage's GUID.
+		{
+			// Zero the hash, as we won't be writing all bytes of it below
+			Package.GenerationHash.Reset();
+			
+			// Ensure there's no padding within FGuid, as the padding could be uninitialized
+			static_assert(sizeof(FGuid) == sizeof(uint32) * 4);
+			// Ensure an FGuid will fit within the hash buffer
+			static_assert(sizeof(FGuid) <= sizeof(FBlake3Hash::ByteArray));
+
+			// Create a copy that we can get the address of
+			//
+			// Note that UPackage::GetGuid is deprecated but there's no replacement yet. The engine
+			// team has advised me to use it anyway, so it's necessary to suppress the deprecation
+			// warning here.
+			PRAGMA_DISABLE_DEPRECATION_WARNINGS
+			const FGuid OwnerGuid = OwnerPackage->GetGuid();
+			PRAGMA_ENABLE_DEPRECATION_WARNINGS
+
+			FMemory::Memcpy(Package.GenerationHash.GetBytes(), &OwnerGuid, sizeof(FGuid));
+
+			// If OwnerGuid is non-zero, the hash should now be non-zero
+			check(!OwnerGuid.IsValid() || !Package.GenerationHash.IsZero());
+		}
+	}
+
+	return Result;
+}
+
+bool FCustomizableObjectCookPackageSplitter::PreSaveGeneratorPackage(
+	UPackage* OwnerPackage,
+	UObject* OwnerObject,
+	const TArray<ICookPackageSplitter::FGeneratedPackageForPreSave>& PlaceholderPackages,
+	TArray<UPackage*>& OutKeepReferencedPackages)
+{
+	// The CO is just about to be saved (i.e. produce the cooked version of the asset), so this
+	// function needs to:
+	// 
+	// 1.	Move the streamed Extension Data out of the CO's package, so that it doesn't get saved
+	//		into the cooked package.
+	// 
+	// 2.	Remove hard references to the streamed data, so that it doesn't get loaded as soon as
+	//		the CO is loaded
+	//
+	// 3.	Fix up the soft references to point to the streamed data's new location
+
+	UCustomizableObject* Object = CastChecked<UCustomizableObject>(OwnerObject);
+
+	// There should be one generated package per streamed Extension Data
+	check(Object->StreamedExtensionData.Num() == PlaceholderPackages.Num());
+
+	// After the CO has been saved, the contract for ICookPackageSplitter states that we need to
+	// restore the CO back to how it was before, so we need to save some information to help with
+	// this.
+	SavedContainerNames.Reset();
+
+	for (const ICookPackageSplitter::FGeneratedPackageForPreSave& GeneratedPackage : PlaceholderPackages)
+	{
+		// Find the streamed data based on the name of the generated package
+		FCustomizableObjectStreamedExtensionData* FoundData = FindStreamedExtensionData(Object, GeneratedPackage.RelativePath);
+		if (!FoundData)
+		{
+			UE_LOG(LogMutable, Error, TEXT("Couldn't find streamed Extension Data container with name %s in array of %d entries"),
+				*GeneratedPackage.RelativePath, Object->StreamedExtensionData.Num());
+
+			return false;
+		}
+
+		// Move the streamed data to the generated package
+		UCustomizableObjectExtensionDataContainer* Container = nullptr;
+		if (!MoveContainerToNewOuter(Object, GeneratedPackage.RelativePath, GeneratedPackage.Package, Container))
+		{
+			return false;
+		}
+
+		SavedContainerNames.Add(GeneratedPackage.RelativePath);
+
+		// Remove the hard reference and set the soft reference to the streamed data's new location
+		FoundData->ConvertToSoftReferenceForCooking(TSoftObjectPtr<UCustomizableObjectExtensionDataContainer>(Container));
+	}
+
+	return true;
+}
+
+void FCustomizableObjectCookPackageSplitter::PostSaveGeneratorPackage(UPackage* OwnerPackage, UObject* OwnerObject)
+{
+	// Move the streamed data back into the CO's package and restore the StreamedExtensionData
+	// array on the CO to how it was before PreSaveGeneratorPackage.
+
+	UCustomizableObject* Object = CastChecked<UCustomizableObject>(OwnerObject);
+
+	TArray<FCustomizableObjectStreamedExtensionData> NewArray;
+	NewArray.Reserve(SavedContainerNames.Num());
+
+	for (const FString& ContainerName : SavedContainerNames)
+	{
+		UCustomizableObjectExtensionDataContainer* Container = nullptr;
+		MoveContainerToNewOuter(Object, ContainerName, OwnerObject, Container);
+
+		NewArray.Emplace(Container);
+	}
+
+	Object->StreamedExtensionData = NewArray;
+}
+
+bool FCustomizableObjectCookPackageSplitter::PopulateGeneratedPackage(
+	UPackage* OwnerPackage,
+	UObject* OwnerObject,
+	const FGeneratedPackageForPopulate& GeneratedPackage,
+	TArray<UObject*>& OutObjectsToMove,
+	TArray<UPackage*>& OutKeepReferencedPackages)
+{
+	// Move the container into its newly generated package
+
+	UCustomizableObject* Object = CastChecked<UCustomizableObject>(OwnerObject);
+
+	UCustomizableObjectExtensionDataContainer* Container = nullptr;
+	if (!MoveContainerToNewOuter(Object, GeneratedPackage.RelativePath, GeneratedPackage.Package, Container))
+	{
+		return false;
+	}
+
+	OutObjectsToMove.Add(Container);
+
+	return true;
+}
+
+void FCustomizableObjectCookPackageSplitter::PostSaveGeneratedPackage(
+	UPackage* OwnerPackage,
+	UObject* OwnerObject,
+	const FGeneratedPackageForPopulate& GeneratedPackage)
+{
+	// Now that the generated package has been saved/cooked, move the container back to the CO, so
+	// that everything is the same as it was before cooking.
+
+	UCustomizableObject* Object = CastChecked<UCustomizableObject>(OwnerObject);
+
+	UCustomizableObjectExtensionDataContainer* Container = nullptr;
+	MoveContainerToNewOuter(Object, GeneratedPackage.RelativePath, OwnerObject, Container);
+}
diff --git a/Engine/Plugins/Experimental/Mutable/Source/CustomizableObjectEditor/Private/MuCOE/CustomizableObjectCookPackageSplitter.h b/Engine/Plugins/Experimental/Mutable/Source/CustomizableObjectEditor/Private/MuCOE/CustomizableObjectCookPackageSplitter.h
new file mode 100644
--- /dev/null
+++ b/Engine/Plugins/Experimental/Mutable/Source/CustomizableObjectEditor/Private/MuCOE/CustomizableObjectCookPackageSplitter.h
@@ -0,0 +1,43 @@
+// Copyright Epic Games, Inc. All Rights Reserved.
+
+#pragma once
+
+#include "CookPackageSplitter.h"
+#include "MuCO/CustomizableObject.h"
+
+/** Handles splitting the streamable Extension Data constants into their own packages */
+class FCustomizableObjectCookPackageSplitter : public ICookPackageSplitter
+{
+public:
+	/** ICookPackageSplitter interface */
+	static bool ShouldSplit(UObject* SplitData);
+	static FString GetSplitterDebugName() { return TEXT("FCustomizableObjectCookPackageSplitter"); }
+
+	virtual TArray<ICookPackageSplitter::FGeneratedPackage> GetGenerateList(
+		const UPackage* OwnerPackage,
+		const UObject* OwnerObject) override;
+	
+	virtual bool PreSaveGeneratorPackage(
+		UPackage* OwnerPackage,
+		UObject* OwnerObject,
+		const TArray<ICookPackageSplitter::FGeneratedPackageForPreSave>& PlaceholderPackages,
+		TArray<UPackage*>& OutKeepReferencedPackages) override;
+
+	virtual void PostSaveGeneratorPackage(UPackage* OwnerPackage, UObject* OwnerObject) override;
+
+	virtual bool PopulateGeneratedPackage(
+		UPackage* OwnerPackage,
+		UObject* OwnerObject,
+		const FGeneratedPackageForPopulate& GeneratedPackage,
+		TArray<UObject*>& OutObjectsToMove,
+		TArray<UPackage*>& OutKeepReferencedPackages) override;
+
+	virtual void PostSaveGeneratedPackage(
+		UPackage* OwnerPackage,
+		UObject* OwnerObject,
+		const FGeneratedPackageForPopulate& GeneratedPackage) override;
+
+private:
+	TArray<TSoftObjectPtr<UCustomizableObjectExtensionDataContainer>> SavedStreamedExtensionData;
+	TArray<FString> SavedContainerNames;
+};
diff --git a/Engine/Plugins/Experimental/Mutable/Source/CustomizableObjectEditor/Private/MuCOE/EdGraphSchema_CustomizableObject.cpp b/Engine/Plugins/Experimental/Mutable/Source/CustomizableObjectEditor/Private/MuCOE/EdGraphSchema_CustomizableObject.cpp
--- a/Engine/Plugins/Experimental/Mutable/Source/CustomizableObjectEditor/Private/MuCOE/EdGraphSchema_CustomizableObject.cpp
+++ b/Engine/Plugins/Experimental/Mutable/Source/CustomizableObjectEditor/Private/MuCOE/EdGraphSchema_CustomizableObject.cpp
@@ -8,6 +8,8 @@
 #include "Framework/Commands/GenericCommands.h"
 #include "GraphEditorActions.h"
 #include "Materials/MaterialInterface.h"
+#include "MuCO/CustomizableObjectExtension.h"
+#include "MuCO/ICustomizableObjectModule.h"
 #include "MuCOE/CustomizableObjectEditorNodeContextCommands.h"
 #include "MuCOE/ICustomizableObjectEditor.h"
 #include "MuCOE/Nodes/CustomizableObjectNodeAnimationPose.h"
@@ -483,24 +485,52 @@ void UEdGraphSchema_CustomizableObject::GetGraphContextActions(FGraphContextMenu
 
 	{
 		// External Pin Nodes
-		const FName* PinTypes[] = { &PC_Material, &PC_Mesh, &PC_Image, &PC_Projector, &PC_GroupProjector, &PC_Color, &PC_Float, &PC_Bool, &PC_Enum, &PC_Stack };
-		
-		for (const FName* PinCategory : PinTypes)
+		TArray<FName> PinTypes({ PC_Material, PC_Mesh, PC_Image, PC_Projector, PC_GroupProjector, PC_Color, PC_Float, PC_Bool, PC_Enum, PC_Stack });
+
+		// Add pin types from extensions
+		for (const FRegisteredCustomizableObjectPinType& PinType : ICustomizableObjectModule::Get().GetExtendedPinTypes())
+		{
+			PinTypes.AddUnique(PinType.PinType.Name);
+		}
+
+		for (const FName& PinCategory : PinTypes)
 		{
 			UCustomizableObjectNodeExternalPin* CustomizableObjectNodeExternalPin = ContextMenuBuilder.CreateTemplateNode<UCustomizableObjectNodeExternalPin>();
-			CustomizableObjectNodeExternalPin->PinType = *PinCategory;
+			CustomizableObjectNodeExternalPin->PinType = PinCategory;
 			
 			AddNewNodeActionFiltered(CustomizableObjectNodeExternalPin, ContextMenuBuilder, TEXT("Import Pin"), GeneralGrouping, Filter);
 		}
 
-		for (const FName* PinCategory : PinTypes)
+		for (const FName& PinCategory : PinTypes)
 		{
 			UCustomizableObjectNodeExposePin* CustomizableObjectNodeExposePin = ContextMenuBuilder.CreateTemplateNode<UCustomizableObjectNodeExposePin>();
-			CustomizableObjectNodeExposePin->PinType = *PinCategory;
+			CustomizableObjectNodeExposePin->PinType = PinCategory;
 			
 			AddNewNodeActionFiltered(CustomizableObjectNodeExposePin, ContextMenuBuilder, TEXT("Export Pin"), GeneralGrouping, Filter);
 		}
 	}
+
+	// Search for all subclasses of UCustomizableObjectNode
+	//
+	// Iterate over the Class Default Objects instead of their corresponding UClasses, as this allows
+	// us to filter the TObjectIterator to UCustomizableObjectNode instead of UClass, which should
+	// produce far fewer results to iterate through.
+	for (TObjectIterator<UCustomizableObjectNode> It(RF_NoFlags); It; ++It)
+	{
+		const UCustomizableObjectNode* Node = *It;
+		if (!Node->HasAllFlags(RF_ClassDefaultObject))
+		{
+			// Only interested in CDOs
+			continue;
+		}
+
+		FText Category;
+		if (Node->ShouldAddToContextMenu(Category))
+		{
+			UCustomizableObjectNode* TemplateNode = ContextMenuBuilder.CreateTemplateNode<UCustomizableObjectNode>(Node->GetClass());
+			AddNewNodeActionFiltered(TemplateNode, ContextMenuBuilder, Category.ToString(), GeneralGrouping, Filter);
+		}
+	}
 	
 	{
 		UEdGraphNode_Comment* Node = NewObject<UEdGraphNode_Comment>();
@@ -645,6 +675,14 @@ FLinearColor UEdGraphSchema_CustomizableObject::GetPinTypeColor(const FName& Typ
 		return FLinearColor(1.000000f, 1.000000f, 1.000000f, 1.000000f);
 	}
 
+	for (const FRegisteredCustomizableObjectPinType& PinType : ICustomizableObjectModule::Get().GetExtendedPinTypes())
+	{
+		if (PinType.PinType.Name == TypeString)
+		{
+			return PinType.PinType.Color;
+		}
+	}
+
 	return FLinearColor(0.750000f, 0.600000f, 0.400000f, 1.000000f);
 }
 
@@ -1048,8 +1086,17 @@ FText UEdGraphSchema_CustomizableObject::GetPinCategoryName(const FName& PinCate
 	}
 	else
 	{
-		check(false); // Unknown pin category. Add the unknown category to the "switch".
-		return FText();
+		for (const FRegisteredCustomizableObjectPinType& PinType : ICustomizableObjectModule::Get().GetExtendedPinTypes())
+		{
+			if (PinType.PinType.Name == PinCategory)
+			{
+				return PinType.PinType.DisplayName;
+			}
+		}
+
+		// Need to fail gracefully here in case a plugin that was active when this graph was
+		// created is no longer loaded.
+		return LOCTEXT("Unknown_Pin_Category", "Unknown");
 	}
 }
 
diff --git a/Engine/Plugins/Experimental/Mutable/Source/CustomizableObjectEditor/Private/MuCOE/ExtensionDataCompilerInterface.cpp b/Engine/Plugins/Experimental/Mutable/Source/CustomizableObjectEditor/Private/MuCOE/ExtensionDataCompilerInterface.cpp
new file mode 100644
--- /dev/null
+++ b/Engine/Plugins/Experimental/Mutable/Source/CustomizableObjectEditor/Private/MuCOE/ExtensionDataCompilerInterface.cpp
@@ -0,0 +1,63 @@
+// Copyright Epic Games, Inc. All Rights Reserved.
+
+#include "MuCOE/ExtensionDataCompilerInterface.h"
+
+#include "InstancedStruct.h"
+#include "MuCO/CustomizableObjectStreamedExtensionData.h"
+#include "MuCOE/GenerateMutableSource/GenerateMutableSource.h"
+#include "MuR/ExtensionData.h"
+
+FExtensionDataCompilerInterface::FExtensionDataCompilerInterface(FMutableGraphGenerationContext& InGenerationContext)
+	: GenerationContext(InGenerationContext)
+{
+}
+
+mu::ExtensionDataPtrConst FExtensionDataCompilerInterface::MakeStreamedExtensionData(UCustomizableObjectExtensionDataContainer*& OutContainer)
+{
+	mu::ExtensionDataPtr Result = new mu::ExtensionData;
+	Result->Origin = mu::ExtensionData::EOrigin::ConstantStreamed;
+	Result->Index = GenerationContext.StreamedExtensionData.Num();
+
+	// Generate a deterministic name to help with deterministic cooking
+	const FString ContainerName = FString::Printf(TEXT("Streamed_%d"), Result->Index);
+
+	UObject* ExistingObject = FindObject<UObject>(GenerationContext.Object, *ContainerName);
+	if (ExistingObject)
+	{
+		// This must have been left behind from a previous compilation and hasn't been deleted by 
+		// GC yet.
+		//
+		// Move it into the transient package to get it out of the way.
+		ExistingObject->Rename(nullptr, GetTransientPackage(), REN_DontCreateRedirectors);
+
+		check(!FindObject<UObject>(GenerationContext.Object, *ContainerName));
+	}
+
+	check(GenerationContext.Object);
+	OutContainer = NewObject<UCustomizableObjectExtensionDataContainer>(
+		GenerationContext.Object,
+		FName(*ContainerName),
+		RF_Public);
+
+	GenerationContext.StreamedExtensionData.Add(OutContainer);
+
+	return Result;
+}
+
+mu::ExtensionDataPtrConst FExtensionDataCompilerInterface::MakeAlwaysLoadedExtensionData(FInstancedStruct&& Data)
+{
+	mu::ExtensionDataPtr Result = new mu::ExtensionData;
+	Result->Origin = mu::ExtensionData::EOrigin::ConstantAlwaysLoaded;
+	Result->Index = GenerationContext.AlwaysLoadedExtensionData.Num();
+
+	FCustomizableObjectExtensionData* CompileTimeExtensionData = &GenerationContext.AlwaysLoadedExtensionData.AddDefaulted_GetRef();
+	CompileTimeExtensionData->Data = MoveTemp(Data);
+
+	return Result;
+}
+
+UObject* FExtensionDataCompilerInterface::GetOuterForAlwaysLoadedObjects()
+{
+	check(GenerationContext.Object);
+	return GenerationContext.Object;
+}
diff --git a/Engine/Plugins/Experimental/Mutable/Source/CustomizableObjectEditor/Private/MuCOE/GenerateMutableSource/GenerateMutableSource.cpp b/Engine/Plugins/Experimental/Mutable/Source/CustomizableObjectEditor/Private/MuCOE/GenerateMutableSource/GenerateMutableSource.cpp
--- a/Engine/Plugins/Experimental/Mutable/Source/CustomizableObjectEditor/Private/MuCOE/GenerateMutableSource/GenerateMutableSource.cpp
+++ b/Engine/Plugins/Experimental/Mutable/Source/CustomizableObjectEditor/Private/MuCOE/GenerateMutableSource/GenerateMutableSource.cpp
@@ -6,12 +6,15 @@
 #include "Engine/SkeletalMeshSocket.h"
 #include "Engine/TextureLODSettings.h"
 #include "Interfaces/ITargetPlatform.h"
+#include "MuCO/CustomizableObjectExtension.h"
+#include "MuCO/ICustomizableObjectModule.h"
 #include "MuCOE/CustomizableObjectCompiler.h"
 #include "MuCOE/EdGraphSchema_CustomizableObject.h"
 #include "MuCOE/GenerateMutableSource/GenerateMutableSourceGroupProjector.h"
 #include "MuCOE/GenerateMutableSource/GenerateMutableSourceModifier.h"
 #include "MuCOE/GenerateMutableSource/GenerateMutableSourceSurface.h"
 #include "MuCOE/GraphTraversal.h"
+#include "MuCOE/ICustomizableObjectExtensionNode.h"
 #include "MuCOE/Nodes/CustomizableObjectNodeEditMaterialBase.h"
 #include "MuCOE/Nodes/CustomizableObjectNodeExtendMaterial.h"
 #include "MuCOE/Nodes/CustomizableObjectNodeMaterial.h"
@@ -115,6 +118,7 @@ void CheckNumOutputs(const UEdGraphPin& Pin, const FMutableGraphGenerationContex
 
 FMutableGraphGenerationContext::FMutableGraphGenerationContext(UCustomizableObject* InObject, FCustomizableObjectCompiler* InCompiler, const FCompilationOptions& InOptions)
 	: Object(InObject), Compiler(InCompiler), Options(InOptions)
+	, ExtensionDataCompilerInterface(*this)
 {
 	// Default flags for mesh generation nodes.
 	MeshGenerationFlags.Push(EMutableMeshConversionFlags::None);
@@ -867,6 +871,41 @@ mu::NodeObjectPtr GenerateMutableSource(const UEdGraphPin * Pin, FMutableGraphGe
 			}
 		}
 
+		// Generate inputs to Object node pins added by extensions
+		for (const FRegisteredObjectNodeInputPin& ExtensionInputPin : ICustomizableObjectModule::Get().GetAdditionalObjectNodePins())
+		{
+			const UEdGraphPin* GraphPin = TypedNodeObj->FindPin(ExtensionInputPin.GlobalPinName, EGPD_Input);
+			if (!GraphPin)
+			{
+				continue;
+			}
+
+			TArray<UEdGraphPin*> ConnectedPins = FollowInputPinArray(*GraphPin);
+
+			// If the pin isn't supposed to take more than one connection, ignore all but the first
+			// incoming connection.
+			if (!ExtensionInputPin.InputPin.bIsArray && ConnectedPins.Num() > 1)
+			{
+				FString Msg = FString::Printf(TEXT("Extension input %s has multiple incoming connections but is only expecting one connection."),
+					*ExtensionInputPin.InputPin.DisplayName.ToString());
+
+				GenerationContext.Compiler->CompilerLog(FText::FromString(Msg), Node, EMessageSeverity::Warning);
+			}
+
+			for (const UEdGraphPin* ConnectedPin : ConnectedPins)
+			{
+				const UEdGraphNode* ConnectedNode = ConnectedPin->GetOwningNode();
+				
+				if (const ICustomizableObjectExtensionNode* ExtensionNode = Cast<ICustomizableObjectExtensionNode>(ConnectedNode))
+				{
+					if (mu::NodeExtensionDataPtr GeneratedNode = ExtensionNode->GenerateMutableNode(GenerationContext.ExtensionDataCompilerInterface))
+					{
+						ObjectNode->AddExtensionDataNode(GeneratedNode, TCHAR_TO_ANSI(*ExtensionInputPin.GlobalPinName.ToString()));
+					}
+				}
+			}
+		}
+
 		// Children
 		TArray<UEdGraphPin*> ConnectedChildrenPins = FollowInputPinArray(*TypedNodeObj->ChildrenPin());
 		ObjectNode->SetChildCount(ConnectedChildrenPins.Num());
diff --git a/Engine/Plugins/Experimental/Mutable/Source/CustomizableObjectEditor/Private/MuCOE/GenerateMutableSource/GenerateMutableSource.h b/Engine/Plugins/Experimental/Mutable/Source/CustomizableObjectEditor/Private/MuCOE/GenerateMutableSource/GenerateMutableSource.h
--- a/Engine/Plugins/Experimental/Mutable/Source/CustomizableObjectEditor/Private/MuCOE/GenerateMutableSource/GenerateMutableSource.h
+++ b/Engine/Plugins/Experimental/Mutable/Source/CustomizableObjectEditor/Private/MuCOE/GenerateMutableSource/GenerateMutableSource.h
@@ -4,6 +4,8 @@
 
 #include "Animation/Skeleton.h"
 #include "Engine/SkeletalMesh.h"
+#include "MuCO/CustomizableObject.h"
+#include "MuCOE/ExtensionDataCompilerInterface.h"
 #include "MuCOE/Nodes/CustomizableObjectNodeMaterialBase.h"
 #include "MuCOE/Nodes/CustomizableObjectNodeObject.h"
 #include "MuR/MutableMemory.h"
@@ -752,6 +754,11 @@ struct FMutableGraphGenerationContext
 	};
 
 	TMap<TObjectPtr<class USkeletalMesh>, FMeshWithBoneRemovalApplied> MeshesWithBoneRemovalApplied;
+
+	/** Extension Data constants are collected here */
+	FExtensionDataCompilerInterface ExtensionDataCompilerInterface;
+	TArray<FCustomizableObjectExtensionData> AlwaysLoadedExtensionData;
+	TArray<UCustomizableObjectExtensionDataContainer*> StreamedExtensionData;
 };
 
 /** Pin Data scope wrapper. Pops the pin data on scope exit. */
diff --git a/Engine/Plugins/Experimental/Mutable/Source/CustomizableObjectEditor/Private/MuCOE/Nodes/CustomizableObjectNodeObject.cpp b/Engine/Plugins/Experimental/Mutable/Source/CustomizableObjectEditor/Private/MuCOE/Nodes/CustomizableObjectNodeObject.cpp
--- a/Engine/Plugins/Experimental/Mutable/Source/CustomizableObjectEditor/Private/MuCOE/Nodes/CustomizableObjectNodeObject.cpp
+++ b/Engine/Plugins/Experimental/Mutable/Source/CustomizableObjectEditor/Private/MuCOE/Nodes/CustomizableObjectNodeObject.cpp
@@ -4,6 +4,8 @@
 
 #include "Containers/Queue.h"
 #include "Logging/MessageLog.h"
+#include "MuCO/ICustomizableObjectModule.h"
+#include "MuCO/CustomizableObjectExtension.h"
 #include "MuCOE/CustomizableObjectEditorLogger.h"
 #include "MuCOE/CustomizableObjectGraph.h"
 #include "MuCOE/EdGraphSchema_CustomizableObject.h"
@@ -16,6 +18,9 @@ class UCustomizableObjectNodeRemapPins;
 
 #define LOCTEXT_NAMESPACE "CustomizableObjectEditor"
 
+const FName UCustomizableObjectNodeObject::ChildrenPinName(TEXT("Children"));
+const FName UCustomizableObjectNodeObject::OutputPinName(TEXT("Object"));
+const TCHAR* UCustomizableObjectNodeObject::LODPinNamePrefix = TEXT("LOD ");
 
 UCustomizableObjectNodeObject::UCustomizableObjectNodeObject()
 	: Super()
@@ -77,7 +82,7 @@ void UCustomizableObjectNodeObject::PostEditChangeProperty(FPropertyChangedEvent
 	}
 
 	FProperty* PropertyThatChanged = PropertyChangedEvent.Property;
-	if ( PropertyThatChanged && PropertyThatChanged->GetName() == TEXT("NumLODs"))
+	if (PropertyThatChanged && PropertyThatChanged->GetName() == TEXT("NumLODs"))
 	{
 		NumLODs = FMath::Clamp(NumLODs, 1, 64);
 
@@ -107,18 +112,32 @@ void UCustomizableObjectNodeObject::AllocateDefaultPins(UCustomizableObjectNodeR
 {
 	const UEdGraphSchema_CustomizableObject* Schema = GetDefault<UEdGraphSchema_CustomizableObject>();
 
+	// NOTE: Ensure all built-in pins are handled in UCustomizableObjectNodeObject::IsBuiltInPin
+
 	for (int32 i = 0; i < NumLODs; ++i)
 	{
-		FString LODName = FString::Printf( TEXT("LOD %d "), i );
+		FString LODName = FString::Printf(TEXT("%s%d "), LODPinNamePrefix, i);
 
 		UEdGraphPin* Pin = CustomCreatePin(EGPD_Input, Schema->PC_Material, FName(*LODName), true);
 		Pin->bDefaultValueIsIgnored = true;
 	}
 
-	UEdGraphPin* ChildrenPin = CustomCreatePin(EGPD_Input, Schema->PC_Object, FName("Children"), true );
+	UEdGraphPin* ChildrenPin = CustomCreatePin(EGPD_Input, Schema->PC_Object, ChildrenPinName, true );
 	ChildrenPin->bDefaultValueIsIgnored = true;
 
-	UEdGraphPin* OutputPin = CustomCreatePin(EGPD_Output, Schema->PC_Object, FName("Object"));
+	for (const FRegisteredObjectNodeInputPin& Pin : ICustomizableObjectModule::Get().GetAdditionalObjectNodePins())
+	{
+		// Use the global pin name here to prevent extensions using the same pin names from
+		// interfering with each other.
+		//
+		// This also prevents extension pins from clashing with the built-in pins from this node,
+		// such as "Object".
+		UEdGraphPin* GraphPin = CustomCreatePin(EGPD_Input, Pin.InputPin.PinType, Pin.GlobalPinName, Pin.InputPin.bIsArray);
+
+		GraphPin->PinFriendlyName = Pin.InputPin.DisplayName;
+	}
+
+	UEdGraphPin* OutputPin = CustomCreatePin(EGPD_Output, Schema->PC_Object, OutputPinName);
 
 	if (bIsBase)
 	{
@@ -253,6 +272,9 @@ void UCustomizableObjectNodeObject::PostBackwardsCompatibleFixup()
 		.Node(*this)
 		.Log();
 	}
+
+	// Reconstruct in case any extension pins have changed
+	ReconstructNode();
 }
 
 void UCustomizableObjectNodeObject::PostPasteNode()
@@ -304,5 +326,11 @@ bool UCustomizableObjectNodeObject::IsSingleOutputNode() const
 	return true;
 }
 
+bool UCustomizableObjectNodeObject::IsBuiltInPin(FName PinName)
+{
+	return PinName == ChildrenPinName
+		|| PinName == OutputPinName
+		|| PinName.ToString().StartsWith(LODPinNamePrefix);
+}
 
 #undef LOCTEXT_NAMESPACE
diff --git a/Engine/Plugins/Experimental/Mutable/Source/CustomizableObjectEditor/Public/MuCOE/ExtensionDataCompilerInterface.h b/Engine/Plugins/Experimental/Mutable/Source/CustomizableObjectEditor/Public/MuCOE/ExtensionDataCompilerInterface.h
new file mode 100644
--- /dev/null
+++ b/Engine/Plugins/Experimental/Mutable/Source/CustomizableObjectEditor/Public/MuCOE/ExtensionDataCompilerInterface.h
@@ -0,0 +1,67 @@
+// Copyright Epic Games, Inc. All Rights Reserved.
+
+#pragma once
+
+#include "MuR/ExtensionData.h"
+
+struct FInstancedStruct;
+class UCustomizableObjectExtensionDataContainer;
+struct FMutableGraphGenerationContext;
+
+/**
+ * An object that gets passed around during Customizable Object compilation to help set up
+ * Extension Data.
+ */
+class CUSTOMIZABLEOBJECTEDITOR_API FExtensionDataCompilerInterface
+{
+public:
+	FExtensionDataCompilerInterface(FMutableGraphGenerationContext& InGenerationContext);
+
+	/**
+	 * Register a new Extension Data compile-time constant that will be streamed in on demand.
+	 * 
+	 * This constant will be cooked out to its own package and will be loaded as needed.
+	 * 
+	 * Importantly, any UObjects that are referenced by the Extension Data that aren't in an asset
+	 * package should have the container returned from this function as their Outer. This ensures
+	 * that they get cooked into the correct package for streaming.
+	 *
+	 * For example, if the Extension Data has a dynamically generated USkeletalMesh that it
+	 * references, which doesn't belong to another asset, the USkeletalMesh should have the given
+	 * container as its Outer.
+	 * 
+	 * Assign the constant data to the FInstancedStruct within the container's Data member.
+	 * 
+	 * The provided FInstancedStruct will be visible to the garbage collector, so any object
+	 * references from UPROPERTYs in the instanced struct will be treated as hard references from
+	 * the Customizable Object.
+	 *
+	 * Don't save any references to objects under the container as string paths, e.g. using
+	 * TSoftObjectPtr, because the container will be moved to a different package during cooking
+	 * and your paths won't be automatically updated.
+	 */
+	mu::ExtensionDataPtrConst MakeStreamedExtensionData(UCustomizableObjectExtensionDataContainer*& OutContainer);
+
+	/**
+	 * Register a new Extension Data compile-time constant that will always be loaded.
+	 *
+	 * This constant will be embedded in the Customizable Object and hence will be loaded in memory
+	 * as long as the Customizable Object is loaded.
+	 * 
+	 * Any UObjects referenced by the provided instanced struct that aren't in an asset package
+	 * should be created with the return value of GetOuterForAlwaysLoadedObjects as their Outer.
+	 *
+	 * Move the constant data into this function and use the resulting mu::ExtensionData in the
+	 * mu::Node graph, e.g. set it as the value of a mu::NodeExtensionDataConstant.
+	 * 
+	 * As with MakeStreamedExtensionData, the provided instanced struct will be visible to the GC.
+	 */
+	mu::ExtensionDataPtrConst MakeAlwaysLoadedExtensionData(FInstancedStruct&& Data);
+
+	/** The Outer to use for objects owned by always-loaded Extension Data constants. */
+	UObject* GetOuterForAlwaysLoadedObjects();
+
+private:
+	FMutableGraphGenerationContext& GenerationContext;
+};
+
diff --git a/Engine/Plugins/Experimental/Mutable/Source/MutableRuntime/Private/MuR/CodeRunner.cpp b/Engine/Plugins/Experimental/Mutable/Source/MutableRuntime/Private/MuR/CodeRunner.cpp
--- a/Engine/Plugins/Experimental/Mutable/Source/MutableRuntime/Private/MuR/CodeRunner.cpp
+++ b/Engine/Plugins/Experimental/Mutable/Source/MutableRuntime/Private/MuR/CodeRunner.cpp
@@ -707,6 +707,53 @@ namespace mu
             break;
         }
 
+		case OP_TYPE::IN_ADDEXTENSIONDATA:
+		{
+			OP::InstanceAddExtensionDataArgs Args = pModel->GetPrivate()->m_program.GetOpArgs<OP::InstanceAddExtensionDataArgs>(item.At);
+			switch (item.Stage)
+			{
+				case 0:
+				{
+					// Must pass in an Instance op and ExtensionData op
+					check(Args.Instance);
+					check(Args.ExtensionData);
+
+					TArray<FScheduledOp> Dependencies;
+					Dependencies.Emplace(Args.Instance, item);
+					Dependencies.Emplace(Args.ExtensionData, item);
+
+					AddOp(FScheduledOp(item.At, item, 1), Dependencies);
+
+					break;
+				}
+
+				case 1:
+				{
+					// Assemble result
+					InstancePtrConst InstanceOpResult = GetMemory().GetInstance(FCacheAddress(Args.Instance, item));
+
+					InstancePtr Result = CloneOrTakeOver<Instance>(InstanceOpResult.get());
+
+					if (ExtensionDataPtrConst ExtensionData = GetMemory().GetExtensionData(FCacheAddress(Args.ExtensionData, item)))
+					{
+						const OP::ADDRESS NameAddress = Args.ExtensionDataName;
+						check(NameAddress < (uint32)pModel->GetPrivate()->m_program.m_constantStrings.Num());
+						const char* NameString = pModel->GetPrivate()->m_program.m_constantStrings[NameAddress].c_str();
+
+						Result->GetPrivate()->AddExtensionData(ExtensionData, NameString);
+					}
+
+					GetMemory().SetInstance(item, Result);
+					break;
+				}
+
+				default:
+					check(false);
+			}
+			
+			break;
+		}
+
         default:
                 check(false);
         }
@@ -883,6 +930,22 @@ namespace mu
 			break;
         }
 
+		case OP_TYPE::ED_CONSTANT:
+		{
+			OP::ResourceConstantArgs Args = pModel->GetPrivate()->m_program.GetOpArgs<OP::ResourceConstantArgs>(item.At);
+
+            const FProgram& Program = pModel->GetPrivate()->m_program;
+
+			// Assume the ROM has been loaded previously
+			ExtensionDataPtrConst SourceConst;
+			Program.GetExtensionDataConstant(Args.value, SourceConst);
+
+			check(SourceConst);
+
+            GetMemory().SetExtensionData(item, SourceConst);
+            break;
+		}
+
         default:
             if (type!=OP_TYPE::NONE)
             {
@@ -5353,7 +5416,8 @@ namespace mu
             break;
 
         case OP_TYPE::ME_CONSTANT:
-        case OP_TYPE::IM_CONSTANT:
+		case OP_TYPE::IM_CONSTANT:
+		case OP_TYPE::ED_CONSTANT:
             RunCode_ConstantResource(item, pModel);
             break;
 
diff --git a/Engine/Plugins/Experimental/Mutable/Source/MutableRuntime/Private/MuR/CodeRunner.h b/Engine/Plugins/Experimental/Mutable/Source/MutableRuntime/Private/MuR/CodeRunner.h
--- a/Engine/Plugins/Experimental/Mutable/Source/MutableRuntime/Private/MuR/CodeRunner.h
+++ b/Engine/Plugins/Experimental/Mutable/Source/MutableRuntime/Private/MuR/CodeRunner.h
@@ -7,6 +7,7 @@
 #include "Containers/ContainerAllocationPolicies.h"
 #include "HAL/PlatformCrt.h"
 #include "HAL/PlatformMath.h"
+#include "MuR/ExtensionDataStreamer.h"
 #include "MuR/Image.h"
 #include "MuR/Layout.h"
 #include "MuR/MutableMemory.h"
@@ -193,6 +194,31 @@ namespace  mu
 			int32 RomIndex = -1;
 		};
 
+		class FLoadExtensionDataTask : public CodeRunner::FIssuedTask
+		{
+		public:
+			FLoadExtensionDataTask(const FScheduledOp& InOp, const ExtensionDataPtrConst& InExtensionData, int32 InModelConstantIndex)
+				: FIssuedTask(InOp)
+				, Data(InExtensionData)
+				, ModelConstantIndex(InModelConstantIndex)
+			{
+				check(Data.get());
+				check(Data->Origin == ExtensionData::EOrigin::ConstantStreamed);
+			}
+
+			// FIssuedTask interface
+			bool Prepare(CodeRunner*, const TSharedPtr<const Model>&, bool& bOutFailed) override;
+			void Complete(CodeRunner*) override;
+			bool IsComplete(CodeRunner*) override;
+
+		private:
+			ExtensionDataPtrConst Data;
+			// The index of this constant in FProgram::m_constantExtensionData
+			int32 ModelConstantIndex;
+
+			TSharedPtr<const FExtensionDataLoadHandle> LoadHandle;
+		};
+
 		class FLoadImageRomsTask : public CodeRunner::FIssuedTask
 		{
 		public:
diff --git a/Engine/Plugins/Experimental/Mutable/Source/MutableRuntime/Private/MuR/CodeRunnerTasked.cpp b/Engine/Plugins/Experimental/Mutable/Source/MutableRuntime/Private/MuR/CodeRunnerTasked.cpp
--- a/Engine/Plugins/Experimental/Mutable/Source/MutableRuntime/Private/MuR/CodeRunnerTasked.cpp
+++ b/Engine/Plugins/Experimental/Mutable/Source/MutableRuntime/Private/MuR/CodeRunnerTasked.cpp
@@ -1,5 +1,6 @@
 // Copyright Epic Games, Inc. All Rights Reserved.
 
+#include "Algo/Find.h"
 #include "Async/TaskGraphInterfaces.h"
 #include "Containers/Array.h"
 #include "Containers/UnrealString.h"
@@ -12,6 +13,7 @@
 #include "Misc/AssertionMacros.h"
 #include "MuR/CodeRunner.h"
 #include "MuR/System.h"
+#include "MuR/ExtensionDataStreamer.h"
 #include "MuR/Image.h"
 #include "MuR/ImagePrivate.h"
 #include "MuR/Layout.h"
@@ -1964,6 +1966,86 @@ namespace mu
 	}
 
 
+	//---------------------------------------------------------------------------------------------
+	//---------------------------------------------------------------------------------------------
+	//---------------------------------------------------------------------------------------------
+	bool CodeRunner::FLoadExtensionDataTask::Prepare(CodeRunner* Runner, const TSharedPtr<const Model>& InModel, bool& bOutFailed)
+	{
+		MUTABLE_CPUPROFILER_SCOPE(FLoadExtensionDataTask_Prepare);
+
+		FProgram& Program = Runner->m_pModel->GetPrivate()->m_program;
+		check(Program.m_constantExtensionData.IsValidIndex(ModelConstantIndex));
+
+		const FExtensionDataConstant& Constant = Program.m_constantExtensionData[ModelConstantIndex];
+		check(Constant.Data == Data);
+
+		if (Constant.LoadState == FExtensionDataConstant::ELoadState::FailedToLoad)
+		{
+			// This already failed before, so don't waste time trying again.
+			//
+			// This behavior could be made an option in future if needed.
+			bOutFailed = true;
+			return false;
+		}
+
+		TArray<ExtensionDataPtrConst> UnloadedConstants;
+		LoadHandle = Runner->m_pSystem->GetExtensionDataStreamer()->StartLoad(Data, UnloadedConstants);
+		check(LoadHandle.IsValid());
+
+		// Update the load state of any constants that were unloaded
+		for (const ExtensionDataPtrConst& Unloaded : UnloadedConstants)
+		{
+			check(Unloaded->Origin == ExtensionData::EOrigin::ConstantStreamed);
+
+			FExtensionDataConstant* FoundConstant = Algo::FindBy(Program.m_constantExtensionData, Unloaded, &FExtensionDataConstant::Data);
+			check(FoundConstant);
+
+			FoundConstant->LoadState = FExtensionDataConstant::ELoadState::Unloaded;
+		}
+
+		// No worker thread work
+		return false;
+	}
+
+
+	//---------------------------------------------------------------------------------------------
+	void CodeRunner::FLoadExtensionDataTask::Complete(CodeRunner* Runner)
+	{
+		MUTABLE_CPUPROFILER_SCOPE(FLoadExtensionDataTask_Complete);
+
+		// Complete should only be called if the load is finished
+		check(LoadHandle->LoadState != FExtensionDataLoadHandle::ELoadState::Pending);
+
+		FProgram& Program = Runner->m_pModel->GetPrivate()->m_program;
+		check(Program.m_constantExtensionData.IsValidIndex(ModelConstantIndex));
+
+		FExtensionDataConstant& Constant = Program.m_constantExtensionData[ModelConstantIndex];
+		check(Constant.Data == Data);
+
+		// Update the load state in the Program
+		if (LoadHandle->LoadState == FExtensionDataLoadHandle::ELoadState::Loaded)
+		{
+			Constant.LoadState = FExtensionDataConstant::ELoadState::CurrentlyLoaded;
+		}
+		else
+		{
+			check(LoadHandle->LoadState == FExtensionDataLoadHandle::ELoadState::FailedToLoad);
+			Constant.LoadState = FExtensionDataConstant::ELoadState::FailedToLoad;
+		}
+
+		// Process the constant op normally, now that the data is loaded.
+		Runner->RunCode(Op, Runner->m_pParams, Runner->m_pModel.Get(), Runner->m_lodMask);
+	}
+
+
+	//---------------------------------------------------------------------------------------------
+	bool CodeRunner::FLoadExtensionDataTask::IsComplete(CodeRunner* Runner)
+	{
+		check(LoadHandle.IsValid());
+		return LoadHandle->LoadState != FExtensionDataLoadHandle::ELoadState::Pending;
+	}
+
+
 	//---------------------------------------------------------------------------------------------
 	//---------------------------------------------------------------------------------------------
 	//---------------------------------------------------------------------------------------------
@@ -2211,6 +2293,26 @@ namespace mu
 			}
 			break;
 		}
+		
+		case OP_TYPE::ED_CONSTANT:
+		{
+			OP::ResourceConstantArgs Args = program.GetOpArgs<OP::ResourceConstantArgs>(item.At);
+
+			const FExtensionDataConstant::ELoadState LoadState = program.m_constantExtensionData[Args.value].LoadState;
+
+			check(LoadState != FExtensionDataConstant::ELoadState::Invalid);
+			
+			if (LoadState == FExtensionDataConstant::ELoadState::AlwaysLoaded
+				|| LoadState == FExtensionDataConstant::ELoadState::CurrentlyLoaded)
+			{
+				// Already loaded. The rest of the constant code will run right away.
+			}
+			else
+			{
+				Issued = MakeShared<FLoadExtensionDataTask>(item, program.m_constantExtensionData[Args.value].Data, Args.value);
+			}
+			break;
+		}
 
 		case OP_TYPE::IM_CONSTANT:
 		{
diff --git a/Engine/Plugins/Experimental/Mutable/Source/MutableRuntime/Private/MuR/ExtensionData.cpp b/Engine/Plugins/Experimental/Mutable/Source/MutableRuntime/Private/MuR/ExtensionData.cpp
new file mode 100644
--- /dev/null
+++ b/Engine/Plugins/Experimental/Mutable/Source/MutableRuntime/Private/MuR/ExtensionData.cpp
@@ -0,0 +1,62 @@
+// Copyright Epic Games, Inc. All Rights Reserved.
+
+#include "MuR/ExtensionData.h"
+
+#include "MuR/MutableTrace.h"
+#include "MuR/SerialisationPrivate.h"
+#include "Templates/TypeHash.h"
+
+namespace mu
+{
+
+//---------------------------------------------------------------------------------------------
+void ExtensionData::Serialise(const ExtensionData* Data, OutputArchive& Archive)
+{
+	Archive << *Data;
+}
+
+
+//---------------------------------------------------------------------------------------------
+ExtensionDataPtr ExtensionData::StaticUnserialise(InputArchive& Archive)
+{
+	MUTABLE_CPUPROFILER_SCOPE(ExtensionDataUnserialise);
+	LLM_SCOPE_BYNAME(TEXT("MutableRuntime"));
+
+	ExtensionDataPtr Result = new ExtensionData();
+	Archive >> *Result;
+
+	return Result;
+}
+
+
+//---------------------------------------------------------------------------------------------
+uint32 ExtensionData::Hash() const
+{
+	uint32 Result = ::GetTypeHash(Index);
+	HashCombine(Result, ::GetTypeHash((uint8)Origin));
+
+	return Result;
+}
+
+
+//---------------------------------------------------------------------------------------------
+void ExtensionData::Serialise(OutputArchive& Archive) const
+{
+	Archive << Index;
+
+	uint8 OriginByte = (uint8)Origin;
+	Archive << OriginByte;
+}
+
+
+//---------------------------------------------------------------------------------------------
+void ExtensionData::Unserialise(InputArchive& Archive)
+{
+	Archive >> Index;
+
+	uint8 OriginByte;
+	Archive >> OriginByte;
+	Origin = (EOrigin)OriginByte;
+}
+
+}
diff --git a/Engine/Plugins/Experimental/Mutable/Source/MutableRuntime/Private/MuR/Instance.cpp b/Engine/Plugins/Experimental/Mutable/Source/MutableRuntime/Private/MuR/Instance.cpp
--- a/Engine/Plugins/Experimental/Mutable/Source/MutableRuntime/Private/MuR/Instance.cpp
+++ b/Engine/Plugins/Experimental/Mutable/Source/MutableRuntime/Private/MuR/Instance.cpp
@@ -472,6 +472,23 @@ namespace mu
     }
 
 
+    //---------------------------------------------------------------------------------------------
+	int32 Instance::GetExtensionDataCount() const
+	{
+		return m_pD->m_extensionData.Num();
+	}
+
+
+    //---------------------------------------------------------------------------------------------
+	void Instance::GetExtensionData(int32 Index, ExtensionDataPtrConst& OutExtensionData, const char*& OutName) const
+	{
+		check(m_pD->m_extensionData.IsValidIndex(Index));
+
+		OutExtensionData = m_pD->m_extensionData[Index].Data;
+		OutName = m_pD->m_extensionData[Index].Name.c_str();
+	}
+
+
     //---------------------------------------------------------------------------------------------
 	int Instance::Private::AddLOD()
 	{
@@ -693,5 +710,19 @@ namespace mu
 
         return result;
     }
+
+
+    //---------------------------------------------------------------------------------------------
+	void Instance::Private::AddExtensionData(ExtensionDataPtrConst Data, const char* Name)
+	{
+		check(Data);
+		check(Name);
+
+		NamedExtensionData& Entry = m_extensionData.AddDefaulted_GetRef();
+		Entry.Data = Data;
+		Entry.Name = Name;
+
+		check(Entry.Name.length() > 0);
+	}
 }
 
diff --git a/Engine/Plugins/Experimental/Mutable/Source/MutableRuntime/Private/MuR/InstancePrivate.h b/Engine/Plugins/Experimental/Mutable/Source/MutableRuntime/Private/MuR/InstancePrivate.h
--- a/Engine/Plugins/Experimental/Mutable/Source/MutableRuntime/Private/MuR/InstancePrivate.h
+++ b/Engine/Plugins/Experimental/Mutable/Source/MutableRuntime/Private/MuR/InstancePrivate.h
@@ -119,6 +119,11 @@ namespace mu
 		TArray<INSTANCE_COMPONENT, TInlineAllocator<4>> m_components;
     };
 
+	struct NamedExtensionData
+	{
+		ExtensionDataPtrConst Data;
+		string Name;
+	};
 
 	class Instance::Private : public Base
 	{
@@ -130,6 +135,9 @@ namespace mu
 		//!
 		TArray<INSTANCE_LOD,TInlineAllocator<4>> m_lods;
 
+		// Every entry must have a valid ExtensionData and name
+		TArray<NamedExtensionData> m_extensionData;
+
         Private()
         {
             m_id = 0;
@@ -153,5 +161,8 @@ namespace mu
         int AddVector( int lod, int comp, int surf, const FVector4f&, const char* strName );
         int AddScalar( int lod, int comp, int surf, float, const char* strName );
         int AddString( int lod, int comp, int surf, const char* strValue, const char* strName );
+		
+		// Data and Name must be non-null
+		void AddExtensionData(ExtensionDataPtrConst Data, const char* Name);
     };
 }
diff --git a/Engine/Plugins/Experimental/Mutable/Source/MutableRuntime/Private/MuR/ModelPrivate.h b/Engine/Plugins/Experimental/Mutable/Source/MutableRuntime/Private/MuR/ModelPrivate.h
--- a/Engine/Plugins/Experimental/Mutable/Source/MutableRuntime/Private/MuR/ModelPrivate.h
+++ b/Engine/Plugins/Experimental/Mutable/Source/MutableRuntime/Private/MuR/ModelPrivate.h
@@ -7,6 +7,8 @@
 
 #include "MuR/SerialisationPrivate.h"
 #include "MuR/Operations.h"
+#include "MuR/ExtensionData.h"
+#include "MuR/ExtensionDataStreamer.h"
 #include "MuR/ImagePrivate.h"
 #include "MuR/MeshPrivate.h"
 #include "MuR/ParametersPrivate.h"
@@ -66,6 +68,50 @@ namespace mu
 	};
 	MUTABLE_DEFINE_POD_SERIALISABLE(FImageLODRange);
 
+	struct FExtensionDataConstant
+	{
+		// This should always be valid, but if the state is Unloaded it won't be usable.
+		//
+		// Avoid storing references to this Data in Memory while the state is Unloaded.
+		ExtensionDataPtrConst Data;
+
+		enum class ELoadState : uint8
+		{
+			Invalid,
+			Unloaded,
+			FailedToLoad,
+			CurrentlyLoaded,
+			AlwaysLoaded
+		};
+
+		// This should be initialized to a valid load state when Data is set
+		ELoadState LoadState = ELoadState::Invalid;
+
+		inline void Serialise(OutputArchive& arch) const
+		{
+			arch << Data;
+		}
+		
+		inline void Unserialise(InputArchive& arch)
+		{
+			arch >> Data;
+
+			check(Data.get());
+			check(Data->Origin == ExtensionData::EOrigin::ConstantAlwaysLoaded
+				|| Data->Origin == ExtensionData::EOrigin::ConstantStreamed);
+
+			if (Data->Origin == ExtensionData::EOrigin::ConstantAlwaysLoaded)
+			{
+				LoadState = ELoadState::AlwaysLoaded;
+			}
+			else
+			{
+				// Streamed constants are assumed to be unloaded to start with
+				LoadState = ELoadState::Unloaded;
+			}
+		}
+	};
+
     //!
     struct FProgram
     {
@@ -188,6 +234,9 @@ namespace mu
         //! Constant mesh data: the first is the index in m_roms for each mesh or -1 if it is always loaded.
 		TArray<TPair<int32, Ptr<const Mesh>>> m_constantMeshes;
 
+		//! Constant ExtensionData
+		TArray<FExtensionDataConstant> m_constantExtensionData;
+
         //! Constant string data
 		TArray<string> m_constantStrings;
 
@@ -236,6 +285,7 @@ namespace mu
 			arch << m_constantImageLODIndices;
 			arch << m_constantImages;
 			arch << m_constantMeshes;
+			arch << m_constantExtensionData;
 			arch << m_constantStrings;
             arch << m_constantLayouts;
             arch << m_constantProjectors;
@@ -260,6 +310,7 @@ namespace mu
 			arch >> m_constantImageLODIndices;
 			arch >> m_constantImages;
 			arch >> m_constantMeshes;
+			arch >> m_constantExtensionData;
 			arch >> m_constantStrings;
             arch >> m_constantLayouts;
             arch >> m_constantProjectors;
@@ -422,6 +473,29 @@ namespace mu
 			return m_constantMeshes.Add(TPair<int32, Ptr<const Mesh>>( -1, pMesh.get() ));
 		}
 
+		OP::ADDRESS AddConstant(Ptr<const ExtensionData> Data)
+		{
+			// Ensure unique
+			for (int32 Index = 0; Index < m_constantExtensionData.Num(); Index++)
+			{
+				const ExtensionData* Candidate = m_constantExtensionData[Index].Data.get();
+				if (*Candidate == *Data)
+				{
+					return Index;
+				}
+			}
+
+			FExtensionDataConstant& NewConstant = m_constantExtensionData.AddDefaulted_GetRef();
+			NewConstant.Data = Data;
+
+			// The data is assumed to be loaded during compilation
+			NewConstant.LoadState =
+				Data->Origin == ExtensionData::EOrigin::ConstantAlwaysLoaded
+				? FExtensionDataConstant::ELoadState::AlwaysLoaded
+				: FExtensionDataConstant::ELoadState::CurrentlyLoaded;
+
+			return m_constantExtensionData.Num() - 1;
+		}
 
 		OP::ADDRESS AddConstant( Ptr<const Layout> pLayout )
         {
@@ -630,6 +704,15 @@ namespace mu
 			res = m_constantMeshes[ConstantIndex].Value;
 		}
 
+		void GetExtensionDataConstant(int32 ConstantIndex, ExtensionDataPtrConst& Result) const
+		{
+			const FExtensionDataConstant& Constant = m_constantExtensionData[ConstantIndex];
+
+			check(Constant.LoadState != FExtensionDataConstant::ELoadState::Unloaded);
+			check(Constant.Data.get());
+
+			Result = Constant.Data;
+		}
 
         inline OP_TYPE GetOpType( OP::ADDRESS at ) const
         {
diff --git a/Engine/Plugins/Experimental/Mutable/Source/MutableRuntime/Private/MuR/NullExtensionDataStreamer.cpp b/Engine/Plugins/Experimental/Mutable/Source/MutableRuntime/Private/MuR/NullExtensionDataStreamer.cpp
new file mode 100644
--- /dev/null
+++ b/Engine/Plugins/Experimental/Mutable/Source/MutableRuntime/Private/MuR/NullExtensionDataStreamer.cpp
@@ -0,0 +1,34 @@
+// Copyright Epic Games, Inc. All Rights Reserved.
+
+#include "MuR/NullExtensionDataStreamer.h"
+
+#include "MuR/ExtensionData.h"
+
+namespace mu
+{
+
+ExtensionDataPtr NullExtensionDataStreamer::CloneExtensionData(const ExtensionDataPtrConst& Source)
+{
+	// Note that these appear to be valid and should be good enough for testing, but they won't be
+	// usable since we're not copying the actual data pointed to by Source.
+
+	ExtensionDataPtr Result = new ExtensionData();
+	Result->Index = NextIndex;
+	NextIndex++;
+	Result->Origin = ExtensionData::EOrigin::Runtime;
+
+	return Result;
+}
+
+TSharedRef<const FExtensionDataLoadHandle> NullExtensionDataStreamer::StartLoad(const ExtensionDataPtrConst& Data, TArray<ExtensionDataPtrConst>& OutUnloadedConstants)
+{
+	OutUnloadedConstants.Reset();
+
+	TSharedRef<FExtensionDataLoadHandle> Result = MakeShared<FExtensionDataLoadHandle>();
+	Result->Data = Data;
+	Result->LoadState = FExtensionDataLoadHandle::ELoadState::FailedToLoad;
+
+	return Result;
+}
+
+}
diff --git a/Engine/Plugins/Experimental/Mutable/Source/MutableRuntime/Private/MuR/NullExtensionDataStreamer.h b/Engine/Plugins/Experimental/Mutable/Source/MutableRuntime/Private/MuR/NullExtensionDataStreamer.h
new file mode 100644
--- /dev/null
+++ b/Engine/Plugins/Experimental/Mutable/Source/MutableRuntime/Private/MuR/NullExtensionDataStreamer.h
@@ -0,0 +1,21 @@
+// Copyright Epic Games, Inc. All Rights Reserved.
+
+#pragma once
+
+#include "MuR/ExtensionDataStreamer.h"
+
+namespace mu
+{
+
+/** A stub implementation that can be used if there's no valid ExtensionDataStreamer */
+class NullExtensionDataStreamer : public ExtensionDataStreamer
+{
+public:
+	virtual ExtensionDataPtr CloneExtensionData(const ExtensionDataPtrConst& Source) override;
+	virtual TSharedRef<const FExtensionDataLoadHandle> StartLoad(const ExtensionDataPtrConst& Data, TArray<ExtensionDataPtrConst>& OutUnloadedConstants) override;
+
+private:
+	int32 NextIndex = 0;
+};
+
+}
diff --git a/Engine/Plugins/Experimental/Mutable/Source/MutableRuntime/Private/MuR/Operations.cpp b/Engine/Plugins/Experimental/Mutable/Source/MutableRuntime/Private/MuR/Operations.cpp
--- a/Engine/Plugins/Experimental/Mutable/Source/MutableRuntime/Private/MuR/Operations.cpp
+++ b/Engine/Plugins/Experimental/Mutable/Source/MutableRuntime/Private/MuR/Operations.cpp
@@ -17,130 +17,132 @@ namespace mu
     // clang-format off
     static const OP_DESC s_opDescs[OP_DESC_COUNT] =
 	{ 
-		// type			cached	supported base image formats
-        { DT_NONE,		false,	{ 0, 0, 0, 0, 0, 0, 0, 0, 0, 0, 0, 0, 0, 0, 0, 0, 0, 0, 0, 0, 0, 0, 0, 0, 0, 0, 0, 0, 0 }		},	// NONE
-
-		{ DT_BOOL,		false,	{ 0, 0, 0, 0, 0, 0, 0, 0, 0, 0, 0, 0, 0, 0, 0, 0, 0, 0, 0, 0, 0, 0, 0, 0, 0, 0, 0, 0, 0 }		},	// BO_CONSTANT
-		{ DT_INT,		false,	{ 0, 0, 0, 0, 0, 0, 0, 0, 0, 0, 0, 0, 0, 0, 0, 0, 0, 0, 0, 0, 0, 0, 0, 0, 0, 0, 0, 0, 0 }		},	// NU_CONSTANT
-		{ DT_SCALAR,	false,	{ 0, 0, 0, 0, 0, 0, 0, 0, 0, 0, 0, 0, 0, 0, 0, 0, 0, 0, 0, 0, 0, 0, 0, 0, 0, 0, 0, 0, 0 }		},	// SC_CONSTANT
-		{ DT_COLOUR,	false,	{ 0, 0, 0, 0, 0, 0, 0, 0, 0, 0, 0, 0, 0, 0, 0, 0, 0, 0, 0, 0, 0, 0, 0, 0, 0, 0, 0, 0, 0 }		},	// CO_CONSTANT
-		{ DT_IMAGE,		true,   { 1, 1, 1, 1, 1, 1, 1, 1, 1, 1, 1, 1, 0, 0, 0, 0, 0, 0, 0, 0, 0, 0, 0, 0, 0, 0, 0, 0, 0 }		},	// IM_CONSTANT
-		{ DT_MESH,		true,   { 0, 0, 0, 0, 0, 0, 0, 0, 0, 0, 0, 0, 0, 0, 0, 0, 0, 0, 0, 0, 0, 0, 0, 0, 0, 0, 0, 0, 0 }		},	// ME_CONSTANT
-		{ DT_LAYOUT,	false,	{ 0, 0, 0, 0, 0, 0, 0, 0, 0, 0, 0, 0, 0, 0, 0, 0, 0, 0, 0, 0, 0, 0, 0, 0, 0, 0, 0, 0, 0 }		},	// LA_CONSTANT
-		{ DT_PROJECTOR,	false,	{ 0, 0, 0, 0, 0, 0, 0, 0, 0, 0, 0, 0, 0, 0, 0, 0, 0, 0, 0, 0, 0, 0, 0, 0, 0, 0, 0, 0, 0 }		},	// PR_CONSTANT
-		{ DT_STRING,	false,	{ 0, 0, 0, 0, 0, 0, 0, 0, 0, 0, 0, 0, 0, 0, 0, 0, 0, 0, 0, 0, 0, 0, 0, 0, 0, 0, 0, 0, 0 }		},	// ST_CONSTANT
-
-		{ DT_BOOL,		false,	{ 0, 0, 0, 0, 0, 0, 0, 0, 0, 0, 0, 0, 0, 0, 0, 0, 0, 0, 0, 0, 0, 0, 0, 0, 0, 0, 0, 0, 0 }		},	// BO_PARAMETER
-		{ DT_INT,		false,	{ 0, 0, 0, 0, 0, 0, 0, 0, 0, 0, 0, 0, 0, 0, 0, 0, 0, 0, 0, 0, 0, 0, 0, 0, 0, 0, 0, 0, 0 }		},	// NU_PARAMETER
-		{ DT_SCALAR,	false,	{ 0, 0, 0, 0, 0, 0, 0, 0, 0, 0, 0, 0, 0, 0, 0, 0, 0, 0, 0, 0, 0, 0, 0, 0, 0, 0, 0, 0, 0 }		},	// SC_PARAMETER
-		{ DT_COLOUR,	false,	{ 0, 0, 0, 0, 0, 0, 0, 0, 0, 0, 0, 0, 0, 0, 0, 0, 0, 0, 0, 0, 0, 0, 0, 0, 0, 0, 0, 0, 0 }		},	// CO_PARAMETER
-		{ DT_PROJECTOR,	false,	{ 0, 0, 0, 0, 0, 0, 0, 0, 0, 0, 0, 0, 0, 0, 0, 0, 0, 0, 0, 0, 0, 0, 0, 0, 0, 0, 0, 0, 0 }		},	// PR_PARAMETER
-		{ DT_IMAGE,		false,	{ 0, 0, 0, 0, 0, 0, 0, 0, 0, 0, 0, 0, 0, 0, 0, 0, 0, 0, 0, 0, 0, 0, 0, 0, 0, 0, 0, 0, 0 }		},	// IM_PARAMETER
-		{ DT_STRING,	false,	{ 0, 0, 0, 0, 0, 0, 0, 0, 0, 0, 0, 0, 0, 0, 0, 0, 0, 0, 0, 0, 0, 0, 0, 0, 0, 0, 0, 0, 0 }		},	// ST_PARAMETER
-
-		{ DT_INT,		false,	{ 0, 0, 0, 0, 0, 0, 0, 0, 0, 0, 0, 0, 0, 0, 0, 0, 0, 0, 0, 0, 0, 0, 0, 0, 0, 0, 0, 0, 0 }		},	// NU_CONDITIONAL
-		{ DT_SCALAR,	false,	{ 0, 0, 0, 0, 0, 0, 0, 0, 0, 0, 0, 0, 0, 0, 0, 0, 0, 0, 0, 0, 0, 0, 0, 0, 0, 0, 0, 0, 0 }		},	// SC_CONDITIONAL
-		{ DT_COLOUR,	false,	{ 0, 0, 0, 0, 0, 0, 0, 0, 0, 0, 0, 0, 0, 0, 0, 0, 0, 0, 0, 0, 0, 0, 0, 0, 0, 0, 0, 0, 0 }		},	// CO_CONDITIONAL
-		{ DT_IMAGE,		true,	{ 1, 1, 1, 1, 1, 1, 1, 1, 1, 1, 1, 1, 1, 1, 1, 1, 1, 1, 1, 1, 1, 1, 0, 0, 1, 1, 0, 0, 0 }		},	// IM_CONDITIONAL
-		{ DT_MESH,		true,	{ 0, 0, 0, 0, 0, 0, 0, 0, 0, 0, 0, 0, 0, 0, 0, 0, 0, 0, 0, 0, 0, 0, 0, 0, 0, 0, 0, 0, 0 }		},	// ME_CONDITIONAL
-		{ DT_LAYOUT,	false,  { 0, 0, 0, 0, 0, 0, 0, 0, 0, 0, 0, 0, 0, 0, 0, 0, 0, 0, 0, 0, 0, 0, 0, 0, 0, 0, 0, 0, 0 }		},	// LA_CONDITIONAL
-		{ DT_INSTANCE,	false,	{ 0, 0, 0, 0, 0, 0, 0, 0, 0, 0, 0, 0, 0, 0, 0, 0, 0, 0, 0, 0, 0, 0, 0, 0, 0, 0, 0, 0, 0 }		},	// IN_CONDITIONAL
-
-		{ DT_INT,		false,	{ 0, 0, 0, 0, 0, 0, 0, 0, 0, 0, 0, 0, 0, 0, 0, 0, 0, 0, 0, 0, 0, 0, 0, 0, 0, 0, 0, 0, 0 }		},	// NU_SWITCH
-		{ DT_SCALAR,	false,	{ 0, 0, 0, 0, 0, 0, 0, 0, 0, 0, 0, 0, 0, 0, 0, 0, 0, 0, 0, 0, 0, 0, 0, 0, 0, 0, 0, 0, 0 }		},	// SC_SWITCH
-		{ DT_COLOUR,	false,	{ 0, 0, 0, 0, 0, 0, 0, 0, 0, 0, 0, 0, 0, 0, 0, 0, 0, 0, 0, 0, 0, 0, 0, 0, 0, 0, 0, 0, 0 }		},	// CO_SWITCH
-		{ DT_IMAGE,		true,	{ 1, 1, 1, 1, 1, 1, 1, 1, 1, 1, 1, 1, 1, 1, 1, 1, 1, 1, 1, 1, 1, 1, 0, 0, 1, 1, 0, 0, 0 }		},	// IM_SWITCH
-		{ DT_MESH,		true,	{ 0, 0, 0, 0, 0, 0, 0, 0, 0, 0, 0, 0, 0, 0, 0, 0, 0, 0, 0, 0, 0, 0, 0, 0, 0, 0, 0, 0, 0 }		},	// ME_SWITCH
-		{ DT_LAYOUT,	false,	{ 0, 0, 0, 0, 0, 0, 0, 0, 0, 0, 0, 0, 0, 0, 0, 0, 0, 0, 0, 0, 0, 0, 0, 0, 0, 0, 0, 0, 0 }		},	// LA_SWITCH
-		{ DT_INSTANCE,	false,	{ 0, 0, 0, 0, 0, 0, 0, 0, 0, 0, 0, 0, 0, 0, 0, 0, 0, 0, 0, 0, 0, 0, 0, 0, 0, 0, 0, 0, 0 }		},	// IN_SWITCH
-
-		{ DT_BOOL,		false,  { 0, 0, 0, 0, 0, 0, 0, 0, 0, 0, 0, 0, 0, 0, 0, 0, 0, 0, 0, 0, 0, 0, 0, 0, 0, 0, 0, 0, 0 }		},	// BO_LESS
-		{ DT_BOOL,		false,  { 0, 0, 0, 0, 0, 0, 0, 0, 0, 0, 0, 0, 0, 0, 0, 0, 0, 0, 0, 0, 0, 0, 0, 0, 0, 0, 0, 0, 0 }		},	// BO_EQUAL_SC_CONST
-		{ DT_BOOL,		false,  { 0, 0, 0, 0, 0, 0, 0, 0, 0, 0, 0, 0, 0, 0, 0, 0, 0, 0, 0, 0, 0, 0, 0, 0, 0, 0, 0, 0, 0 }		},	// BO_AND
-		{ DT_BOOL,		false,  { 0, 0, 0, 0, 0, 0, 0, 0, 0, 0, 0, 0, 0, 0, 0, 0, 0, 0, 0, 0, 0, 0, 0, 0, 0, 0, 0, 0, 0 }		},	// BO_OR
-		{ DT_BOOL,		false,  { 0, 0, 0, 0, 0, 0, 0, 0, 0, 0, 0, 0, 0, 0, 0, 0, 0, 0, 0, 0, 0, 0, 0, 0, 0, 0, 0, 0, 0 }		},	// BO_NOT
-
-		{ DT_SCALAR,	true,   { 0, 0, 0, 0, 0, 0, 0, 0, 0, 0, 0, 0, 0, 0, 0, 0, 0, 0, 0, 0, 0, 0, 0, 0, 0, 0, 0, 0, 0 }		},	// SC_MULTIPLYADD
-		{ DT_SCALAR,	true,   { 0, 0, 0, 0, 0, 0, 0, 0, 0, 0, 0, 0, 0, 0, 0, 0, 0, 0, 0, 0, 0, 0, 0, 0, 0, 0, 0, 0, 0 }		},	// SC_ARITHMETIC
-		{ DT_SCALAR,	true,   { 0, 0, 0, 0, 0, 0, 0, 0, 0, 0, 0, 0, 0, 0, 0, 0, 0, 0, 0, 0, 0, 0, 0, 0, 0, 0, 0, 0, 0 }		},	// SC_CURVE
-
-		{ DT_COLOUR,	false,	{ 0, 0, 0, 0, 0, 0, 0, 0, 0, 0, 0, 0, 0, 0, 0, 0, 0, 0, 0, 0, 0, 0, 0, 0, 0, 0, 0, 0, 0 }		},	// CO_SAMPLEIMAGE
-		{ DT_COLOUR,	false,	{ 0, 0, 0, 0, 0, 0, 0, 0, 0, 0, 0, 0, 0, 0, 0, 0, 0, 0, 0, 0, 0, 0, 0, 0, 0, 0, 0, 0, 0 }		},	// CO_SWIZZLE
-		{ DT_COLOUR,	false,	{ 0, 0, 0, 0, 0, 0, 0, 0, 0, 0, 0, 0, 0, 0, 0, 0, 0, 0, 0, 0, 0, 0, 0, 0, 0, 0, 0, 0, 0 }		},	// CO_IMAGESIZE
-		{ DT_COLOUR,	false,	{ 0, 0, 0, 0, 0, 0, 0, 0, 0, 0, 0, 0, 0, 0, 0, 0, 0, 0, 0, 0, 0, 0, 0, 0, 0, 0, 0, 0, 0 }		},	// CO_LAYOUTBLOCKTRANSFORM
-		{ DT_COLOUR,	false,	{ 0, 0, 0, 0, 0, 0, 0, 0, 0, 0, 0, 0, 0, 0, 0, 0, 0, 0, 0, 0, 0, 0, 0, 0, 0, 0, 0, 0, 0 }		},	// CO_FROMSCALARS
-		{ DT_COLOUR,	false,	{ 0, 0, 0, 0, 0, 0, 0, 0, 0, 0, 0, 0, 0, 0, 0, 0, 0, 0, 0, 0, 0, 0, 0, 0, 0, 0, 0, 0, 0 }		},	// CO_ARITHMETIC
-
-		{ DT_IMAGE,		true,	{ 0, 1, 1, 1, 0, 0, 0, 0, 0, 0, 0, 0, 0, 0, 0, 0, 0, 0, 0, 0, 0, 0, 0, 0, 0, 0, 0, 0, 0 }		},	// IM_LAYER
-		{ DT_IMAGE,		true,	{ 0, 1, 1, 1, 0, 0, 0, 0, 0, 0, 0, 0, 0, 0, 0, 0, 0, 0, 0, 0, 0, 0, 0, 0, 0, 0, 0, 0, 0 }		},	// IM_LAYERCOLOUR
-		{ DT_IMAGE,		true,	{ 1, 1, 1, 1, 1, 1, 1, 1, 1, 1, 1, 1, 0, 0, 0, 0, 0, 0, 0, 0, 0, 0, 0, 0, 0, 0, 0, 0, 0 }		},	// IM_PIXELFORMAT
-		{ DT_IMAGE,		true,	{ 1, 1, 1, 1, 0, 0, 0, 0, 0, 0, 0, 0, 0, 0, 0, 0, 0, 0, 0, 0, 0, 0, 0, 0, 0, 0, 0, 0, 0 }		},	// IM_MIPMAP
-		{ DT_IMAGE,		true,	{ 0, 1, 1, 1, 0, 0, 0, 0, 0, 0, 0, 0, 0, 0, 0, 0, 0, 0, 0, 0, 0, 0, 0, 0, 0, 0, 0, 0, 0 }		},	// IM_RESIZE
-		{ DT_IMAGE,		true,	{ 0, 1, 1, 1, 0, 0, 0, 0, 0, 0, 0, 0, 0, 0, 0, 0, 0, 0, 0, 0, 0, 0, 0, 0, 0, 0, 0, 0, 0 }		},	// IM_RESIZELIKE
-		{ DT_IMAGE,		true,	{ 0, 1, 1, 1, 0, 0, 0, 0, 0, 0, 0, 0, 0, 0, 0, 0, 0, 0, 0, 0, 0, 0, 0, 0, 0, 0, 0, 0, 0 }		},	// IM_RESIZEREL
-		{ DT_IMAGE,		true,   { 1, 1, 1, 1, 1, 1, 1, 1, 1, 1, 1, 1, 0, 0, 0, 0, 0, 0, 0, 0, 0, 0, 0, 0, 0, 0, 0, 0, 0 }		},	// IM_BLANKLAYOUT
-		{ DT_IMAGE,		true,	{ 1, 1, 1, 1, 1, 1, 1, 1, 1, 1, 1, 0, 0, 0, 1, 1, 1, 1, 1, 1, 1, 1, 0, 0, 1, 1, 0, 0, 0 }		},	// IM_COMPOSE
-		{ DT_IMAGE,		true,	{ 0, 1, 1, 0, 0, 0, 0, 0, 0, 0, 0, 0, 0, 0, 0, 0, 0, 0, 0, 0, 0, 0, 0, 0, 0, 0, 0, 0, 0 }		},	// IM_DIFFERENCE
-		{ DT_IMAGE,		true,	{ 0, 1, 1, 1, 0, 0, 0, 0, 0, 0, 0, 0, 0, 0, 0, 0, 0, 0, 0, 0, 0, 0, 0, 0, 0, 0, 0, 0, 0 }		},	// IM_INTERPOLATE
-		{ DT_IMAGE,		true,	{ 0, 1, 1, 1, 0, 0, 0, 0, 0, 0, 0, 0, 0, 0, 0, 0, 0, 0, 0, 0, 0, 0, 0, 0, 0, 0, 0, 0, 0 }		},	// IM_INTERPOLATE3
-		{ DT_IMAGE,		true,	{ 0, 1, 1, 0, 0, 0, 0, 0, 0, 0, 0, 0, 0, 0, 0, 0, 0, 0, 0, 0, 0, 0, 0, 0, 0, 0, 0, 0, 0 }		},	// IM_SATURATE
-		{ DT_IMAGE,		true,	{ 0, 1, 1, 0, 0, 0, 0, 0, 0, 0, 0, 0, 0, 0, 0, 0, 0, 0, 0, 0, 0, 0, 0, 0, 0, 0, 0, 0, 0 }		},	// IM_LUMINANCE
-		{ DT_IMAGE,		true,	{ 0, 1, 1, 1, 0, 0, 0, 0, 0, 0, 0, 0, 0, 0, 0, 0, 0, 0, 0, 0, 0, 0, 0, 0, 0, 0, 0, 0, 0 }		},	// IM_SWIZZLE
-		{ DT_IMAGE,		true,	{ 0, 1, 1, 1, 0, 0, 0, 0, 0, 0, 0, 0, 0, 0, 0, 0, 0, 0, 0, 0, 0, 0, 0, 0, 0, 0, 0, 0, 0 }		},	// IM_SELECTCOLOUR
-		{ DT_IMAGE,		true,	{ 0, 1, 1, 1, 0, 0, 0, 0, 0, 0, 0, 0, 0, 0, 0, 0, 0, 0, 0, 0, 0, 0, 0, 0, 0, 0, 0, 0, 0 }		},	// IM_COLOURMAP
-		{ DT_IMAGE,		true,	{ 0, 0, 0, 0, 0, 0, 0, 0, 0, 0, 0, 0, 0, 0, 0, 0, 0, 0, 0, 0, 0, 0, 0, 0, 0, 0, 0, 0, 0 }		},	// IM_GRADIENT
-		{ DT_IMAGE,		true,	{ 0, 1, 1, 1, 0, 0, 0, 0, 0, 0, 0, 0, 0, 0, 0, 0, 0, 0, 0, 0, 0, 0, 0, 0, 0, 0, 0, 0, 0 }		},	// IM_BINARISE
-		{ DT_IMAGE,		true,	{ 0, 0, 0, 0, 0, 0, 0, 0, 0, 0, 0, 0, 0, 0, 0, 0, 0, 0, 0, 0, 0, 0, 0, 0, 0, 0, 0, 0, 0 }		},	// IM_PLAINCOLOUR
-		{ DT_IMAGE,		true,	{ 0, 0, 0, 0, 0, 0, 0, 0, 0, 0, 0, 0, 0, 0, 0, 0, 0, 0, 0, 0, 0, 0, 0, 0, 0, 0, 0, 0, 0 }		},	// IM_GPU
-		{ DT_IMAGE,		true,	{ 0, 0, 0, 0, 0, 0, 0, 0, 0, 0, 0, 0, 0, 0, 0, 0, 0, 0, 0, 0, 0, 0, 0, 0, 0, 0, 0, 0, 0 }		},	// IM_CROP
-		{ DT_IMAGE,		true,	{ 0, 0, 0, 0, 0, 0, 0, 0, 0, 0, 0, 0, 0, 0, 0, 0, 0, 0, 0, 0, 0, 0, 0, 0, 0, 0, 0, 0, 0 }		},	// IM_PATCH
-		{ DT_IMAGE,		true,	{ 0, 0, 0, 0, 0, 0, 0, 0, 0, 0, 0, 0, 0, 0, 0, 0, 0, 0, 0, 0, 0, 0, 0, 0, 0, 0, 0, 0, 0 }		},	// IM_RASTERMESH
-		{ DT_IMAGE,		true,	{ 0, 0, 0, 0, 0, 0, 0, 0, 0, 0, 0, 0, 0, 0, 0, 0, 0, 0, 0, 0, 0, 0, 0, 0, 0, 0, 0, 0, 0 }		},	// IM_MAKEGROWMAP
-		{ DT_IMAGE,		true,	{ 0, 0, 0, 0, 0, 0, 0, 0, 0, 0, 0, 0, 0, 0, 0, 0, 0, 0, 0, 0, 0, 0, 0, 0, 0, 0, 0, 0, 0 }		},	// IM_DISPLACE
-		{ DT_IMAGE,		true,	{ 0, 0, 0, 0, 0, 0, 0, 0, 0, 0, 0, 0, 0, 0, 0, 0, 0, 0, 0, 0, 0, 0, 0, 0, 0, 0, 0, 0, 0 }		},	// IM_MULTILAYER
-		{ DT_IMAGE,		true,	{ 0, 1, 1, 1, 0, 0, 0, 0, 0, 0, 0, 0, 0, 0, 0, 0, 0, 0, 0, 0, 0, 0, 0, 0, 0, 0, 0, 0, 0 }		},	// IM_INVERT
-		{ DT_IMAGE,		true,	{ 0, 1, 1, 1, 0, 0, 0, 0, 0, 0, 0, 0, 0, 0, 0, 0, 0, 0, 0, 0, 0, 0, 0, 0, 0, 0, 0, 0, 0 }		},	// IM_NORMALCOMPOSITE
-		{ DT_IMAGE,		true,	{ 0, 1, 1, 1, 0, 0, 0, 0, 0, 0, 0, 0, 0, 0, 0, 0, 0, 0, 0, 0, 0, 0, 0, 0, 0, 0, 0, 0, 0 }		},	// IM_TRANSFORM
-
-		{ DT_MESH,		true,	{ 0, 0, 0, 0, 0, 0, 0, 0, 0, 0, 0, 0, 0, 0, 0, 0, 0, 0, 0, 0, 0, 0, 0, 0, 0, 0, 0, 0, 0 }		},	// ME_APPLYLAYOUT
-		{ DT_MESH,		true,	{ 0, 0, 0, 0, 0, 0, 0, 0, 0, 0, 0, 0, 0, 0, 0, 0, 0, 0, 0, 0, 0, 0, 0, 0, 0, 0, 0, 0, 0 }		},	// ME_DIFFERENCE
-		{ DT_MESH,		true,	{ 0, 0, 0, 0, 0, 0, 0, 0, 0, 0, 0, 0, 0, 0, 0, 0, 0, 0, 0, 0, 0, 0, 0, 0, 0, 0, 0, 0, 0 }		},	// ME_MORPH2
-		{ DT_MESH,		true,	{ 0, 0, 0, 0, 0, 0, 0, 0, 0, 0, 0, 0, 0, 0, 0, 0, 0, 0, 0, 0, 0, 0, 0, 0, 0, 0, 0, 0, 0 }		},	// ME_MERGE
-		{ DT_MESH,		true,	{ 0, 0, 0, 0, 0, 0, 0, 0, 0, 0, 0, 0, 0, 0, 0, 0, 0, 0, 0, 0, 0, 0, 0, 0, 0, 0, 0, 0, 0 }		},	// ME_INTERPOLATE
-		{ DT_MESH,		true,	{ 0, 0, 0, 0, 0, 0, 0, 0, 0, 0, 0, 0, 0, 0, 0, 0, 0, 0, 0, 0, 0, 0, 0, 0, 0, 0, 0, 0, 0 }		},	// ME_MASKCLIPMESH
-		{ DT_MESH,		true,	{ 0, 0, 0, 0, 0, 0, 0, 0, 0, 0, 0, 0, 0, 0, 0, 0, 0, 0, 0, 0, 0, 0, 0, 0, 0, 0, 0, 0, 0 }		},	// ME_MASKDIFF
-		{ DT_MESH,		true,	{ 0, 0, 0, 0, 0, 0, 0, 0, 0, 0, 0, 0, 0, 0, 0, 0, 0, 0, 0, 0, 0, 0, 0, 0, 0, 0, 0, 0, 0 }		},	// ME_REMOVEMASK
-		{ DT_MESH,		true,	{ 0, 0, 0, 0, 0, 0, 0, 0, 0, 0, 0, 0, 0, 0, 0, 0, 0, 0, 0, 0, 0, 0, 0, 0, 0, 0, 0, 0, 0 }		},	// ME_FORMAT
-		{ DT_MESH,		true,	{ 0, 0, 0, 0, 0, 0, 0, 0, 0, 0, 0, 0, 0, 0, 0, 0, 0, 0, 0, 0, 0, 0, 0, 0, 0, 0, 0, 0, 0 }		},	// ME_EXTRACTLAYOUTBLOCK
-		{ DT_MESH,		true,	{ 0, 0, 0, 0, 0, 0, 0, 0, 0, 0, 0, 0, 0, 0, 0, 0, 0, 0, 0, 0, 0, 0, 0, 0, 0, 0, 0, 0, 0 }		},	// ME_EXTRACTFACEGROUP
-		{ DT_MESH,		true,	{ 0, 0, 0, 0, 0, 0, 0, 0, 0, 0, 0, 0, 0, 0, 0, 0, 0, 0, 0, 0, 0, 0, 0, 0, 0, 0, 0, 0, 0 }		},	// ME_TRANSFORM
-		{ DT_MESH,		true,	{ 0, 0, 0, 0, 0, 0, 0, 0, 0, 0, 0, 0, 0, 0, 0, 0, 0, 0, 0, 0, 0, 0, 0, 0, 0, 0, 0, 0, 0 }		},	// ME_CLIPMORPHPLANE
-		{ DT_MESH,		true,	{ 0, 0, 0, 0, 0, 0, 0, 0, 0, 0, 0, 0, 0, 0, 0, 0, 0, 0, 0, 0, 0, 0, 0, 0, 0, 0, 0, 0, 0 }		},	// ME_CLIPWITHMESH
-		{ DT_MESH,		true,	{ 0, 0, 0, 0, 0, 0, 0, 0, 0, 0, 0, 0, 0, 0, 0, 0, 0, 0, 0, 0, 0, 0, 0, 0, 0, 0, 0, 0, 0 }		},	// ME_SETSKELETON
-		{ DT_MESH,		true,	{ 0, 0, 0, 0, 0, 0, 0, 0, 0, 0, 0, 0, 0, 0, 0, 0, 0, 0, 0, 0, 0, 0, 0, 0, 0, 0, 0, 0, 0 }		},	// ME_PROJECT
-		{ DT_MESH,		true,	{ 0, 0, 0, 0, 0, 0, 0, 0, 0, 0, 0, 0, 0, 0, 0, 0, 0, 0, 0, 0, 0, 0, 0, 0, 0, 0, 0, 0, 0 }		},	// ME_APPLYPOSE
-		{ DT_MESH,		true,	{ 0, 0, 0, 0, 0, 0, 0, 0, 0, 0, 0, 0, 0, 0, 0, 0, 0, 0, 0, 0, 0, 0, 0, 0, 0, 0, 0, 0, 0 }		},	// ME_REMAPINDICES
-		{ DT_MESH,		true,	{ 0, 0, 0, 0, 0, 0, 0, 0, 0, 0, 0, 0, 0, 0, 0, 0, 0, 0, 0, 0, 0, 0, 0, 0, 0, 0, 0, 0, 0 }		},	// ME_GEOMETRYOPERATION
-		{ DT_MESH,		true,	{ 0, 0, 0, 0, 0, 0, 0, 0, 0, 0, 0, 0, 0, 0, 0, 0, 0, 0, 0, 0, 0, 0, 0, 0, 0, 0, 0, 0, 0 }		},	// ME_BINDSHAPE
-		{ DT_MESH,		true,	{ 0, 0, 0, 0, 0, 0, 0, 0, 0, 0, 0, 0, 0, 0, 0, 0, 0, 0, 0, 0, 0, 0, 0, 0, 0, 0, 0, 0, 0 }		},	// ME_APPLYSHAPE
-		{ DT_MESH,		true,	{ 0, 0, 0, 0, 0, 0, 0, 0, 0, 0, 0, 0, 0, 0, 0, 0, 0, 0, 0, 0, 0, 0, 0, 0, 0, 0, 0, 0, 0 }		},	// ME_CLIPDEFORM
-		{ DT_MESH,		true,	{ 0, 0, 0, 0, 0, 0, 0, 0, 0, 0, 0, 0, 0, 0, 0, 0, 0, 0, 0, 0, 0, 0, 0, 0, 0, 0, 0, 0, 0 }		},	// ME_MORPHRESHAPE
-		{ DT_MESH,		true,	{ 0, 0, 0, 0, 0, 0, 0, 0, 0, 0, 0, 0, 0, 0, 0, 0, 0, 0, 0, 0, 0, 0, 0, 0, 0, 0, 0, 0, 0 }		},	// ME_OPTIMIZESKINNING
-
-		{ DT_INSTANCE,	false,	{ 0, 0, 0, 0, 0, 0, 0, 0, 0, 0, 0, 0, 0, 0, 0, 0, 0, 0, 0, 0, 0, 0, 0, 0, 0, 0, 0, 0, 0 }		},	// IN_ADDMESH
-		{ DT_INSTANCE,	false,	{ 0, 0, 0, 0, 0, 0, 0, 0, 0, 0, 0, 0, 0, 0, 0, 0, 0, 0, 0, 0, 0, 0, 0, 0, 0, 0, 0, 0, 0 }		},	// IN_ADDIMAGE
-		{ DT_INSTANCE,	false,	{ 0, 0, 0, 0, 0, 0, 0, 0, 0, 0, 0, 0, 0, 0, 0, 0, 0, 0, 0, 0, 0, 0, 0, 0, 0, 0, 0, 0, 0 }		},	// IN_ADDVECTOR
-		{ DT_INSTANCE,	false,	{ 0, 0, 0, 0, 0, 0, 0, 0, 0, 0, 0, 0, 0, 0, 0, 0, 0, 0, 0, 0, 0, 0, 0, 0, 0, 0, 0, 0, 0 }		},	// IN_ADDSCALAR
-		{ DT_INSTANCE,	false,	{ 0, 0, 0, 0, 0, 0, 0, 0, 0, 0, 0, 0, 0, 0, 0, 0, 0, 0, 0, 0, 0, 0, 0, 0, 0, 0, 0, 0, 0 }		},	// IN_ADDSTRING
-		{ DT_INSTANCE,	false,	{ 0, 0, 0, 0, 0, 0, 0, 0, 0, 0, 0, 0, 0, 0, 0, 0, 0, 0, 0, 0, 0, 0, 0, 0, 0, 0, 0, 0, 0 }		},	// IN_ADDSURFACE
-		{ DT_INSTANCE,	false,	{ 0, 0, 0, 0, 0, 0, 0, 0, 0, 0, 0, 0, 0, 0, 0, 0, 0, 0, 0, 0, 0, 0, 0, 0, 0, 0, 0, 0, 0 }		},	// IN_ADDCOMPONENT
-		{ DT_INSTANCE,	false,	{ 0, 0, 0, 0, 0, 0, 0, 0, 0, 0, 0, 0, 0, 0, 0, 0, 0, 0, 0, 0, 0, 0, 0, 0, 0, 0, 0, 0, 0 }		},	// IN_ADDLOD
-
-		{ DT_LAYOUT,	true,   { 0, 0, 0, 0, 0, 0, 0, 0, 0, 0, 0, 0, 0, 0, 0, 0, 0, 0, 0, 0, 0, 0, 0, 0, 0, 0, 0, 0, 0 }		},	// LA_PACK
-		{ DT_LAYOUT,	true,   { 0, 0, 0, 0, 0, 0, 0, 0, 0, 0, 0, 0, 0, 0, 0, 0, 0, 0, 0, 0, 0, 0, 0, 0, 0, 0, 0, 0, 0 }		},	// LA_MERGE
-		{ DT_LAYOUT,	true,   { 0, 0, 0, 0, 0, 0, 0, 0, 0, 0, 0, 0, 0, 0, 0, 0, 0, 0, 0, 0, 0, 0, 0, 0, 0, 0, 0, 0, 0 }		},	// LA_REMOVEBLOCKS
-		{ DT_LAYOUT,	true,   { 0, 0, 0, 0, 0, 0, 0, 0, 0, 0, 0, 0, 0, 0, 0, 0, 0, 0, 0, 0, 0, 0, 0, 0, 0, 0, 0, 0, 0 }		},	// LA_FROMMESH
-};
+		// type				cached	supported base image formats
+        { DT_NONE,			false,	{ 0, 0, 0, 0, 0, 0, 0, 0, 0, 0, 0, 0, 0, 0, 0, 0, 0, 0, 0, 0, 0, 0, 0, 0, 0, 0, 0, 0, 0 }		},	// NONE
+
+		{ DT_BOOL,			false,	{ 0, 0, 0, 0, 0, 0, 0, 0, 0, 0, 0, 0, 0, 0, 0, 0, 0, 0, 0, 0, 0, 0, 0, 0, 0, 0, 0, 0, 0 }		},	// BO_CONSTANT
+		{ DT_INT,			false,	{ 0, 0, 0, 0, 0, 0, 0, 0, 0, 0, 0, 0, 0, 0, 0, 0, 0, 0, 0, 0, 0, 0, 0, 0, 0, 0, 0, 0, 0 }		},	// NU_CONSTANT
+		{ DT_SCALAR,		false,	{ 0, 0, 0, 0, 0, 0, 0, 0, 0, 0, 0, 0, 0, 0, 0, 0, 0, 0, 0, 0, 0, 0, 0, 0, 0, 0, 0, 0, 0 }		},	// SC_CONSTANT
+		{ DT_COLOUR,		false,	{ 0, 0, 0, 0, 0, 0, 0, 0, 0, 0, 0, 0, 0, 0, 0, 0, 0, 0, 0, 0, 0, 0, 0, 0, 0, 0, 0, 0, 0 }		},	// CO_CONSTANT
+		{ DT_IMAGE,			true,   { 1, 1, 1, 1, 1, 1, 1, 1, 1, 1, 1, 1, 0, 0, 0, 0, 0, 0, 0, 0, 0, 0, 0, 0, 0, 0, 0, 0, 0 }		},	// IM_CONSTANT
+		{ DT_MESH,			true,   { 0, 0, 0, 0, 0, 0, 0, 0, 0, 0, 0, 0, 0, 0, 0, 0, 0, 0, 0, 0, 0, 0, 0, 0, 0, 0, 0, 0, 0 }		},	// ME_CONSTANT
+		{ DT_LAYOUT,		false,	{ 0, 0, 0, 0, 0, 0, 0, 0, 0, 0, 0, 0, 0, 0, 0, 0, 0, 0, 0, 0, 0, 0, 0, 0, 0, 0, 0, 0, 0 }		},	// LA_CONSTANT
+		{ DT_PROJECTOR,		false,	{ 0, 0, 0, 0, 0, 0, 0, 0, 0, 0, 0, 0, 0, 0, 0, 0, 0, 0, 0, 0, 0, 0, 0, 0, 0, 0, 0, 0, 0 }		},	// PR_CONSTANT
+		{ DT_STRING,		false,	{ 0, 0, 0, 0, 0, 0, 0, 0, 0, 0, 0, 0, 0, 0, 0, 0, 0, 0, 0, 0, 0, 0, 0, 0, 0, 0, 0, 0, 0 }		},	// ST_CONSTANT
+		{ DT_EXTENSION_DATA,true,	{ 0, 0, 0, 0, 0, 0, 0, 0, 0, 0, 0, 0, 0, 0, 0, 0, 0, 0, 0, 0, 0, 0, 0, 0, 0, 0, 0, 0, 0 }		},	// ED_CONSTANT
+
+		{ DT_BOOL,			false,	{ 0, 0, 0, 0, 0, 0, 0, 0, 0, 0, 0, 0, 0, 0, 0, 0, 0, 0, 0, 0, 0, 0, 0, 0, 0, 0, 0, 0, 0 }		},	// BO_PARAMETER
+		{ DT_INT,			false,	{ 0, 0, 0, 0, 0, 0, 0, 0, 0, 0, 0, 0, 0, 0, 0, 0, 0, 0, 0, 0, 0, 0, 0, 0, 0, 0, 0, 0, 0 }		},	// NU_PARAMETER
+		{ DT_SCALAR,		false,	{ 0, 0, 0, 0, 0, 0, 0, 0, 0, 0, 0, 0, 0, 0, 0, 0, 0, 0, 0, 0, 0, 0, 0, 0, 0, 0, 0, 0, 0 }		},	// SC_PARAMETER
+		{ DT_COLOUR,		false,	{ 0, 0, 0, 0, 0, 0, 0, 0, 0, 0, 0, 0, 0, 0, 0, 0, 0, 0, 0, 0, 0, 0, 0, 0, 0, 0, 0, 0, 0 }		},	// CO_PARAMETER
+		{ DT_PROJECTOR,		false,	{ 0, 0, 0, 0, 0, 0, 0, 0, 0, 0, 0, 0, 0, 0, 0, 0, 0, 0, 0, 0, 0, 0, 0, 0, 0, 0, 0, 0, 0 }		},	// PR_PARAMETER
+		{ DT_IMAGE,			false,	{ 0, 0, 0, 0, 0, 0, 0, 0, 0, 0, 0, 0, 0, 0, 0, 0, 0, 0, 0, 0, 0, 0, 0, 0, 0, 0, 0, 0, 0 }		},	// IM_PARAMETER
+		{ DT_STRING,		false,	{ 0, 0, 0, 0, 0, 0, 0, 0, 0, 0, 0, 0, 0, 0, 0, 0, 0, 0, 0, 0, 0, 0, 0, 0, 0, 0, 0, 0, 0 }		},	// ST_PARAMETER
+
+		{ DT_INT,			false,	{ 0, 0, 0, 0, 0, 0, 0, 0, 0, 0, 0, 0, 0, 0, 0, 0, 0, 0, 0, 0, 0, 0, 0, 0, 0, 0, 0, 0, 0 }		},	// NU_CONDITIONAL
+		{ DT_SCALAR,		false,	{ 0, 0, 0, 0, 0, 0, 0, 0, 0, 0, 0, 0, 0, 0, 0, 0, 0, 0, 0, 0, 0, 0, 0, 0, 0, 0, 0, 0, 0 }		},	// SC_CONDITIONAL
+		{ DT_COLOUR,		false,	{ 0, 0, 0, 0, 0, 0, 0, 0, 0, 0, 0, 0, 0, 0, 0, 0, 0, 0, 0, 0, 0, 0, 0, 0, 0, 0, 0, 0, 0 }		},	// CO_CONDITIONAL
+		{ DT_IMAGE,			true,	{ 1, 1, 1, 1, 1, 1, 1, 1, 1, 1, 1, 1, 1, 1, 1, 1, 1, 1, 1, 1, 1, 1, 0, 0, 1, 1, 0, 0, 0 }		},	// IM_CONDITIONAL
+		{ DT_MESH,			true,	{ 0, 0, 0, 0, 0, 0, 0, 0, 0, 0, 0, 0, 0, 0, 0, 0, 0, 0, 0, 0, 0, 0, 0, 0, 0, 0, 0, 0, 0 }		},	// ME_CONDITIONAL
+		{ DT_LAYOUT,		false,  { 0, 0, 0, 0, 0, 0, 0, 0, 0, 0, 0, 0, 0, 0, 0, 0, 0, 0, 0, 0, 0, 0, 0, 0, 0, 0, 0, 0, 0 }		},	// LA_CONDITIONAL
+		{ DT_INSTANCE,		false,	{ 0, 0, 0, 0, 0, 0, 0, 0, 0, 0, 0, 0, 0, 0, 0, 0, 0, 0, 0, 0, 0, 0, 0, 0, 0, 0, 0, 0, 0 }		},	// IN_CONDITIONAL
+
+		{ DT_INT,			false,	{ 0, 0, 0, 0, 0, 0, 0, 0, 0, 0, 0, 0, 0, 0, 0, 0, 0, 0, 0, 0, 0, 0, 0, 0, 0, 0, 0, 0, 0 }		},	// NU_SWITCH
+		{ DT_SCALAR,		false,	{ 0, 0, 0, 0, 0, 0, 0, 0, 0, 0, 0, 0, 0, 0, 0, 0, 0, 0, 0, 0, 0, 0, 0, 0, 0, 0, 0, 0, 0 }		},	// SC_SWITCH
+		{ DT_COLOUR,		false,	{ 0, 0, 0, 0, 0, 0, 0, 0, 0, 0, 0, 0, 0, 0, 0, 0, 0, 0, 0, 0, 0, 0, 0, 0, 0, 0, 0, 0, 0 }		},	// CO_SWITCH
+		{ DT_IMAGE,			true,	{ 1, 1, 1, 1, 1, 1, 1, 1, 1, 1, 1, 1, 1, 1, 1, 1, 1, 1, 1, 1, 1, 1, 0, 0, 1, 1, 0, 0, 0 }		},	// IM_SWITCH
+		{ DT_MESH,			true,	{ 0, 0, 0, 0, 0, 0, 0, 0, 0, 0, 0, 0, 0, 0, 0, 0, 0, 0, 0, 0, 0, 0, 0, 0, 0, 0, 0, 0, 0 }		},	// ME_SWITCH
+		{ DT_LAYOUT,		false,	{ 0, 0, 0, 0, 0, 0, 0, 0, 0, 0, 0, 0, 0, 0, 0, 0, 0, 0, 0, 0, 0, 0, 0, 0, 0, 0, 0, 0, 0 }		},	// LA_SWITCH
+		{ DT_INSTANCE,		false,	{ 0, 0, 0, 0, 0, 0, 0, 0, 0, 0, 0, 0, 0, 0, 0, 0, 0, 0, 0, 0, 0, 0, 0, 0, 0, 0, 0, 0, 0 }		},	// IN_SWITCH
+
+		{ DT_BOOL,			false,  { 0, 0, 0, 0, 0, 0, 0, 0, 0, 0, 0, 0, 0, 0, 0, 0, 0, 0, 0, 0, 0, 0, 0, 0, 0, 0, 0, 0, 0 }		},	// BO_LESS
+		{ DT_BOOL,			false,  { 0, 0, 0, 0, 0, 0, 0, 0, 0, 0, 0, 0, 0, 0, 0, 0, 0, 0, 0, 0, 0, 0, 0, 0, 0, 0, 0, 0, 0 }		},	// BO_EQUAL_SC_CONST
+		{ DT_BOOL,			false,  { 0, 0, 0, 0, 0, 0, 0, 0, 0, 0, 0, 0, 0, 0, 0, 0, 0, 0, 0, 0, 0, 0, 0, 0, 0, 0, 0, 0, 0 }		},	// BO_AND
+		{ DT_BOOL,			false,  { 0, 0, 0, 0, 0, 0, 0, 0, 0, 0, 0, 0, 0, 0, 0, 0, 0, 0, 0, 0, 0, 0, 0, 0, 0, 0, 0, 0, 0 }		},	// BO_OR
+		{ DT_BOOL,			false,  { 0, 0, 0, 0, 0, 0, 0, 0, 0, 0, 0, 0, 0, 0, 0, 0, 0, 0, 0, 0, 0, 0, 0, 0, 0, 0, 0, 0, 0 }		},	// BO_NOT
+
+		{ DT_SCALAR,		true,   { 0, 0, 0, 0, 0, 0, 0, 0, 0, 0, 0, 0, 0, 0, 0, 0, 0, 0, 0, 0, 0, 0, 0, 0, 0, 0, 0, 0, 0 }		},	// SC_MULTIPLYADD
+		{ DT_SCALAR,		true,   { 0, 0, 0, 0, 0, 0, 0, 0, 0, 0, 0, 0, 0, 0, 0, 0, 0, 0, 0, 0, 0, 0, 0, 0, 0, 0, 0, 0, 0 }		},	// SC_ARITHMETIC
+		{ DT_SCALAR,		true,   { 0, 0, 0, 0, 0, 0, 0, 0, 0, 0, 0, 0, 0, 0, 0, 0, 0, 0, 0, 0, 0, 0, 0, 0, 0, 0, 0, 0, 0 }		},	// SC_CURVE
+
+		{ DT_COLOUR,		false,	{ 0, 0, 0, 0, 0, 0, 0, 0, 0, 0, 0, 0, 0, 0, 0, 0, 0, 0, 0, 0, 0, 0, 0, 0, 0, 0, 0, 0, 0 }		},	// CO_SAMPLEIMAGE
+		{ DT_COLOUR,		false,	{ 0, 0, 0, 0, 0, 0, 0, 0, 0, 0, 0, 0, 0, 0, 0, 0, 0, 0, 0, 0, 0, 0, 0, 0, 0, 0, 0, 0, 0 }		},	// CO_SWIZZLE
+		{ DT_COLOUR,		false,	{ 0, 0, 0, 0, 0, 0, 0, 0, 0, 0, 0, 0, 0, 0, 0, 0, 0, 0, 0, 0, 0, 0, 0, 0, 0, 0, 0, 0, 0 }		},	// CO_IMAGESIZE
+		{ DT_COLOUR,		false,	{ 0, 0, 0, 0, 0, 0, 0, 0, 0, 0, 0, 0, 0, 0, 0, 0, 0, 0, 0, 0, 0, 0, 0, 0, 0, 0, 0, 0, 0 }		},	// CO_LAYOUTBLOCKTRANSFORM
+		{ DT_COLOUR,		false,	{ 0, 0, 0, 0, 0, 0, 0, 0, 0, 0, 0, 0, 0, 0, 0, 0, 0, 0, 0, 0, 0, 0, 0, 0, 0, 0, 0, 0, 0 }		},	// CO_FROMSCALARS
+		{ DT_COLOUR,		false,	{ 0, 0, 0, 0, 0, 0, 0, 0, 0, 0, 0, 0, 0, 0, 0, 0, 0, 0, 0, 0, 0, 0, 0, 0, 0, 0, 0, 0, 0 }		},	// CO_ARITHMETIC
+
+		{ DT_IMAGE,			true,	{ 0, 1, 1, 1, 0, 0, 0, 0, 0, 0, 0, 0, 0, 0, 0, 0, 0, 0, 0, 0, 0, 0, 0, 0, 0, 0, 0, 0, 0 }		},	// IM_LAYER
+		{ DT_IMAGE,			true,	{ 0, 1, 1, 1, 0, 0, 0, 0, 0, 0, 0, 0, 0, 0, 0, 0, 0, 0, 0, 0, 0, 0, 0, 0, 0, 0, 0, 0, 0 }		},	// IM_LAYERCOLOUR
+		{ DT_IMAGE,			true,	{ 1, 1, 1, 1, 1, 1, 1, 1, 1, 1, 1, 1, 0, 0, 0, 0, 0, 0, 0, 0, 0, 0, 0, 0, 0, 0, 0, 0, 0 }		},	// IM_PIXELFORMAT
+		{ DT_IMAGE,			true,	{ 1, 1, 1, 1, 0, 0, 0, 0, 0, 0, 0, 0, 0, 0, 0, 0, 0, 0, 0, 0, 0, 0, 0, 0, 0, 0, 0, 0, 0 }		},	// IM_MIPMAP
+		{ DT_IMAGE,			true,	{ 0, 1, 1, 1, 0, 0, 0, 0, 0, 0, 0, 0, 0, 0, 0, 0, 0, 0, 0, 0, 0, 0, 0, 0, 0, 0, 0, 0, 0 }		},	// IM_RESIZE
+		{ DT_IMAGE,			true,	{ 0, 1, 1, 1, 0, 0, 0, 0, 0, 0, 0, 0, 0, 0, 0, 0, 0, 0, 0, 0, 0, 0, 0, 0, 0, 0, 0, 0, 0 }		},	// IM_RESIZELIKE
+		{ DT_IMAGE,			true,	{ 0, 1, 1, 1, 0, 0, 0, 0, 0, 0, 0, 0, 0, 0, 0, 0, 0, 0, 0, 0, 0, 0, 0, 0, 0, 0, 0, 0, 0 }		},	// IM_RESIZEREL
+		{ DT_IMAGE,			true,   { 1, 1, 1, 1, 1, 1, 1, 1, 1, 1, 1, 1, 0, 0, 0, 0, 0, 0, 0, 0, 0, 0, 0, 0, 0, 0, 0, 0, 0 }		},	// IM_BLANKLAYOUT
+		{ DT_IMAGE,			true,	{ 1, 1, 1, 1, 1, 1, 1, 1, 1, 1, 1, 0, 0, 0, 1, 1, 1, 1, 1, 1, 1, 1, 0, 0, 1, 1, 0, 0, 0 }		},	// IM_COMPOSE
+		{ DT_IMAGE,			true,	{ 0, 1, 1, 0, 0, 0, 0, 0, 0, 0, 0, 0, 0, 0, 0, 0, 0, 0, 0, 0, 0, 0, 0, 0, 0, 0, 0, 0, 0 }		},	// IM_DIFFERENCE
+		{ DT_IMAGE,			true,	{ 0, 1, 1, 1, 0, 0, 0, 0, 0, 0, 0, 0, 0, 0, 0, 0, 0, 0, 0, 0, 0, 0, 0, 0, 0, 0, 0, 0, 0 }		},	// IM_INTERPOLATE
+		{ DT_IMAGE,			true,	{ 0, 1, 1, 1, 0, 0, 0, 0, 0, 0, 0, 0, 0, 0, 0, 0, 0, 0, 0, 0, 0, 0, 0, 0, 0, 0, 0, 0, 0 }		},	// IM_INTERPOLATE3
+		{ DT_IMAGE,			true,	{ 0, 1, 1, 0, 0, 0, 0, 0, 0, 0, 0, 0, 0, 0, 0, 0, 0, 0, 0, 0, 0, 0, 0, 0, 0, 0, 0, 0, 0 }		},	// IM_SATURATE
+		{ DT_IMAGE,			true,	{ 0, 1, 1, 0, 0, 0, 0, 0, 0, 0, 0, 0, 0, 0, 0, 0, 0, 0, 0, 0, 0, 0, 0, 0, 0, 0, 0, 0, 0 }		},	// IM_LUMINANCE
+		{ DT_IMAGE,			true,	{ 0, 1, 1, 1, 0, 0, 0, 0, 0, 0, 0, 0, 0, 0, 0, 0, 0, 0, 0, 0, 0, 0, 0, 0, 0, 0, 0, 0, 0 }		},	// IM_SWIZZLE
+		{ DT_IMAGE,			true,	{ 0, 1, 1, 1, 0, 0, 0, 0, 0, 0, 0, 0, 0, 0, 0, 0, 0, 0, 0, 0, 0, 0, 0, 0, 0, 0, 0, 0, 0 }		},	// IM_SELECTCOLOUR
+		{ DT_IMAGE,			true,	{ 0, 1, 1, 1, 0, 0, 0, 0, 0, 0, 0, 0, 0, 0, 0, 0, 0, 0, 0, 0, 0, 0, 0, 0, 0, 0, 0, 0, 0 }		},	// IM_COLOURMAP
+		{ DT_IMAGE,			true,	{ 0, 0, 0, 0, 0, 0, 0, 0, 0, 0, 0, 0, 0, 0, 0, 0, 0, 0, 0, 0, 0, 0, 0, 0, 0, 0, 0, 0, 0 }		},	// IM_GRADIENT
+		{ DT_IMAGE,			true,	{ 0, 1, 1, 1, 0, 0, 0, 0, 0, 0, 0, 0, 0, 0, 0, 0, 0, 0, 0, 0, 0, 0, 0, 0, 0, 0, 0, 0, 0 }		},	// IM_BINARISE
+		{ DT_IMAGE,			true,	{ 0, 0, 0, 0, 0, 0, 0, 0, 0, 0, 0, 0, 0, 0, 0, 0, 0, 0, 0, 0, 0, 0, 0, 0, 0, 0, 0, 0, 0 }		},	// IM_PLAINCOLOUR
+		{ DT_IMAGE,			true,	{ 0, 0, 0, 0, 0, 0, 0, 0, 0, 0, 0, 0, 0, 0, 0, 0, 0, 0, 0, 0, 0, 0, 0, 0, 0, 0, 0, 0, 0 }		},	// IM_GPU
+		{ DT_IMAGE,			true,	{ 0, 0, 0, 0, 0, 0, 0, 0, 0, 0, 0, 0, 0, 0, 0, 0, 0, 0, 0, 0, 0, 0, 0, 0, 0, 0, 0, 0, 0 }		},	// IM_CROP
+		{ DT_IMAGE,			true,	{ 0, 0, 0, 0, 0, 0, 0, 0, 0, 0, 0, 0, 0, 0, 0, 0, 0, 0, 0, 0, 0, 0, 0, 0, 0, 0, 0, 0, 0 }		},	// IM_PATCH
+		{ DT_IMAGE,			true,	{ 0, 0, 0, 0, 0, 0, 0, 0, 0, 0, 0, 0, 0, 0, 0, 0, 0, 0, 0, 0, 0, 0, 0, 0, 0, 0, 0, 0, 0 }		},	// IM_RASTERMESH
+		{ DT_IMAGE,			true,	{ 0, 0, 0, 0, 0, 0, 0, 0, 0, 0, 0, 0, 0, 0, 0, 0, 0, 0, 0, 0, 0, 0, 0, 0, 0, 0, 0, 0, 0 }		},	// IM_MAKEGROWMAP
+		{ DT_IMAGE,			true,	{ 0, 0, 0, 0, 0, 0, 0, 0, 0, 0, 0, 0, 0, 0, 0, 0, 0, 0, 0, 0, 0, 0, 0, 0, 0, 0, 0, 0, 0 }		},	// IM_DISPLACE
+		{ DT_IMAGE,			true,	{ 0, 0, 0, 0, 0, 0, 0, 0, 0, 0, 0, 0, 0, 0, 0, 0, 0, 0, 0, 0, 0, 0, 0, 0, 0, 0, 0, 0, 0 }		},	// IM_MULTILAYER
+		{ DT_IMAGE,			true,	{ 0, 1, 1, 1, 0, 0, 0, 0, 0, 0, 0, 0, 0, 0, 0, 0, 0, 0, 0, 0, 0, 0, 0, 0, 0, 0, 0, 0, 0 }		},	// IM_INVERT
+		{ DT_IMAGE,			true,	{ 0, 1, 1, 1, 0, 0, 0, 0, 0, 0, 0, 0, 0, 0, 0, 0, 0, 0, 0, 0, 0, 0, 0, 0, 0, 0, 0, 0, 0 }		},	// IM_NORMALCOMPOSITE
+		{ DT_IMAGE,			true,	{ 0, 1, 1, 1, 0, 0, 0, 0, 0, 0, 0, 0, 0, 0, 0, 0, 0, 0, 0, 0, 0, 0, 0, 0, 0, 0, 0, 0, 0 }		},	// IM_TRANSFORM
+
+		{ DT_MESH,			true,	{ 0, 0, 0, 0, 0, 0, 0, 0, 0, 0, 0, 0, 0, 0, 0, 0, 0, 0, 0, 0, 0, 0, 0, 0, 0, 0, 0, 0, 0 }		},	// ME_APPLYLAYOUT
+		{ DT_MESH,			true,	{ 0, 0, 0, 0, 0, 0, 0, 0, 0, 0, 0, 0, 0, 0, 0, 0, 0, 0, 0, 0, 0, 0, 0, 0, 0, 0, 0, 0, 0 }		},	// ME_DIFFERENCE
+		{ DT_MESH,			true,	{ 0, 0, 0, 0, 0, 0, 0, 0, 0, 0, 0, 0, 0, 0, 0, 0, 0, 0, 0, 0, 0, 0, 0, 0, 0, 0, 0, 0, 0 }		},	// ME_MORPH2
+		{ DT_MESH,			true,	{ 0, 0, 0, 0, 0, 0, 0, 0, 0, 0, 0, 0, 0, 0, 0, 0, 0, 0, 0, 0, 0, 0, 0, 0, 0, 0, 0, 0, 0 }		},	// ME_MERGE
+		{ DT_MESH,			true,	{ 0, 0, 0, 0, 0, 0, 0, 0, 0, 0, 0, 0, 0, 0, 0, 0, 0, 0, 0, 0, 0, 0, 0, 0, 0, 0, 0, 0, 0 }		},	// ME_INTERPOLATE
+		{ DT_MESH,			true,	{ 0, 0, 0, 0, 0, 0, 0, 0, 0, 0, 0, 0, 0, 0, 0, 0, 0, 0, 0, 0, 0, 0, 0, 0, 0, 0, 0, 0, 0 }		},	// ME_MASKCLIPMESH
+		{ DT_MESH,			true,	{ 0, 0, 0, 0, 0, 0, 0, 0, 0, 0, 0, 0, 0, 0, 0, 0, 0, 0, 0, 0, 0, 0, 0, 0, 0, 0, 0, 0, 0 }		},	// ME_MASKDIFF
+		{ DT_MESH,			true,	{ 0, 0, 0, 0, 0, 0, 0, 0, 0, 0, 0, 0, 0, 0, 0, 0, 0, 0, 0, 0, 0, 0, 0, 0, 0, 0, 0, 0, 0 }		},	// ME_REMOVEMASK
+		{ DT_MESH,			true,	{ 0, 0, 0, 0, 0, 0, 0, 0, 0, 0, 0, 0, 0, 0, 0, 0, 0, 0, 0, 0, 0, 0, 0, 0, 0, 0, 0, 0, 0 }		},	// ME_FORMAT
+		{ DT_MESH,			true,	{ 0, 0, 0, 0, 0, 0, 0, 0, 0, 0, 0, 0, 0, 0, 0, 0, 0, 0, 0, 0, 0, 0, 0, 0, 0, 0, 0, 0, 0 }		},	// ME_EXTRACTLAYOUTBLOCK
+		{ DT_MESH,			true,	{ 0, 0, 0, 0, 0, 0, 0, 0, 0, 0, 0, 0, 0, 0, 0, 0, 0, 0, 0, 0, 0, 0, 0, 0, 0, 0, 0, 0, 0 }		},	// ME_EXTRACTFACEGROUP
+		{ DT_MESH,			true,	{ 0, 0, 0, 0, 0, 0, 0, 0, 0, 0, 0, 0, 0, 0, 0, 0, 0, 0, 0, 0, 0, 0, 0, 0, 0, 0, 0, 0, 0 }		},	// ME_TRANSFORM
+		{ DT_MESH,			true,	{ 0, 0, 0, 0, 0, 0, 0, 0, 0, 0, 0, 0, 0, 0, 0, 0, 0, 0, 0, 0, 0, 0, 0, 0, 0, 0, 0, 0, 0 }		},	// ME_CLIPMORPHPLANE
+		{ DT_MESH,			true,	{ 0, 0, 0, 0, 0, 0, 0, 0, 0, 0, 0, 0, 0, 0, 0, 0, 0, 0, 0, 0, 0, 0, 0, 0, 0, 0, 0, 0, 0 }		},	// ME_CLIPWITHMESH
+		{ DT_MESH,			true,	{ 0, 0, 0, 0, 0, 0, 0, 0, 0, 0, 0, 0, 0, 0, 0, 0, 0, 0, 0, 0, 0, 0, 0, 0, 0, 0, 0, 0, 0 }		},	// ME_SETSKELETON
+		{ DT_MESH,			true,	{ 0, 0, 0, 0, 0, 0, 0, 0, 0, 0, 0, 0, 0, 0, 0, 0, 0, 0, 0, 0, 0, 0, 0, 0, 0, 0, 0, 0, 0 }		},	// ME_PROJECT
+		{ DT_MESH,			true,	{ 0, 0, 0, 0, 0, 0, 0, 0, 0, 0, 0, 0, 0, 0, 0, 0, 0, 0, 0, 0, 0, 0, 0, 0, 0, 0, 0, 0, 0 }		},	// ME_APPLYPOSE
+		{ DT_MESH,			true,	{ 0, 0, 0, 0, 0, 0, 0, 0, 0, 0, 0, 0, 0, 0, 0, 0, 0, 0, 0, 0, 0, 0, 0, 0, 0, 0, 0, 0, 0 }		},	// ME_REMAPINDICES
+		{ DT_MESH,			true,	{ 0, 0, 0, 0, 0, 0, 0, 0, 0, 0, 0, 0, 0, 0, 0, 0, 0, 0, 0, 0, 0, 0, 0, 0, 0, 0, 0, 0, 0 }		},	// ME_GEOMETRYOPERATION
+		{ DT_MESH,			true,	{ 0, 0, 0, 0, 0, 0, 0, 0, 0, 0, 0, 0, 0, 0, 0, 0, 0, 0, 0, 0, 0, 0, 0, 0, 0, 0, 0, 0, 0 }		},	// ME_BINDSHAPE
+		{ DT_MESH,			true,	{ 0, 0, 0, 0, 0, 0, 0, 0, 0, 0, 0, 0, 0, 0, 0, 0, 0, 0, 0, 0, 0, 0, 0, 0, 0, 0, 0, 0, 0 }		},	// ME_APPLYSHAPE
+		{ DT_MESH,			true,	{ 0, 0, 0, 0, 0, 0, 0, 0, 0, 0, 0, 0, 0, 0, 0, 0, 0, 0, 0, 0, 0, 0, 0, 0, 0, 0, 0, 0, 0 }		},	// ME_CLIPDEFORM
+		{ DT_MESH,			true,	{ 0, 0, 0, 0, 0, 0, 0, 0, 0, 0, 0, 0, 0, 0, 0, 0, 0, 0, 0, 0, 0, 0, 0, 0, 0, 0, 0, 0, 0 }		},	// ME_MORPHRESHAPE
+		{ DT_MESH,			true,	{ 0, 0, 0, 0, 0, 0, 0, 0, 0, 0, 0, 0, 0, 0, 0, 0, 0, 0, 0, 0, 0, 0, 0, 0, 0, 0, 0, 0, 0 }		},	// ME_OPTIMIZESKINNING
+
+		{ DT_INSTANCE,		false,	{ 0, 0, 0, 0, 0, 0, 0, 0, 0, 0, 0, 0, 0, 0, 0, 0, 0, 0, 0, 0, 0, 0, 0, 0, 0, 0, 0, 0, 0 }		},	// IN_ADDMESH
+		{ DT_INSTANCE,		false,	{ 0, 0, 0, 0, 0, 0, 0, 0, 0, 0, 0, 0, 0, 0, 0, 0, 0, 0, 0, 0, 0, 0, 0, 0, 0, 0, 0, 0, 0 }		},	// IN_ADDIMAGE
+		{ DT_INSTANCE,		false,	{ 0, 0, 0, 0, 0, 0, 0, 0, 0, 0, 0, 0, 0, 0, 0, 0, 0, 0, 0, 0, 0, 0, 0, 0, 0, 0, 0, 0, 0 }		},	// IN_ADDVECTOR
+		{ DT_INSTANCE,		false,	{ 0, 0, 0, 0, 0, 0, 0, 0, 0, 0, 0, 0, 0, 0, 0, 0, 0, 0, 0, 0, 0, 0, 0, 0, 0, 0, 0, 0, 0 }		},	// IN_ADDSCALAR
+		{ DT_INSTANCE,		false,	{ 0, 0, 0, 0, 0, 0, 0, 0, 0, 0, 0, 0, 0, 0, 0, 0, 0, 0, 0, 0, 0, 0, 0, 0, 0, 0, 0, 0, 0 }		},	// IN_ADDSTRING
+		{ DT_INSTANCE,		false,	{ 0, 0, 0, 0, 0, 0, 0, 0, 0, 0, 0, 0, 0, 0, 0, 0, 0, 0, 0, 0, 0, 0, 0, 0, 0, 0, 0, 0, 0 }		},	// IN_ADDSURFACE
+		{ DT_INSTANCE,		false,	{ 0, 0, 0, 0, 0, 0, 0, 0, 0, 0, 0, 0, 0, 0, 0, 0, 0, 0, 0, 0, 0, 0, 0, 0, 0, 0, 0, 0, 0 }		},	// IN_ADDCOMPONENT
+		{ DT_INSTANCE,		false,	{ 0, 0, 0, 0, 0, 0, 0, 0, 0, 0, 0, 0, 0, 0, 0, 0, 0, 0, 0, 0, 0, 0, 0, 0, 0, 0, 0, 0, 0 }		},	// IN_ADDLOD
+		{ DT_INSTANCE,		false,	{ 0, 0, 0, 0, 0, 0, 0, 0, 0, 0, 0, 0, 0, 0, 0, 0, 0, 0, 0, 0, 0, 0, 0, 0, 0, 0, 0, 0, 0 }		},	// IN_ADDEXTENSIONDATA
+
+		{ DT_LAYOUT,		true,   { 0, 0, 0, 0, 0, 0, 0, 0, 0, 0, 0, 0, 0, 0, 0, 0, 0, 0, 0, 0, 0, 0, 0, 0, 0, 0, 0, 0, 0 }		},	// LA_PACK
+		{ DT_LAYOUT,		true,   { 0, 0, 0, 0, 0, 0, 0, 0, 0, 0, 0, 0, 0, 0, 0, 0, 0, 0, 0, 0, 0, 0, 0, 0, 0, 0, 0, 0, 0 }		},	// LA_MERGE
+		{ DT_LAYOUT,		true,   { 0, 0, 0, 0, 0, 0, 0, 0, 0, 0, 0, 0, 0, 0, 0, 0, 0, 0, 0, 0, 0, 0, 0, 0, 0, 0, 0, 0, 0 }		},	// LA_REMOVEBLOCKS
+		{ DT_LAYOUT,		true,   { 0, 0, 0, 0, 0, 0, 0, 0, 0, 0, 0, 0, 0, 0, 0, 0, 0, 0, 0, 0, 0, 0, 0, 0, 0, 0, 0, 0, 0 }		},	// LA_FROMMESH
+	};
 
     // clang-format on
 
@@ -407,6 +409,7 @@ namespace mu
         case OP_TYPE::ME_CONSTANT:
         case OP_TYPE::LA_CONSTANT:
         case OP_TYPE::PR_CONSTANT:
+		case OP_TYPE::ED_CONSTANT:
         case OP_TYPE::BO_PARAMETER:
         case OP_TYPE::NU_PARAMETER:
         case OP_TYPE::SC_PARAMETER:
@@ -1102,6 +1105,16 @@ namespace mu
             break;
         }
 
+		case OP_TYPE::IN_ADDEXTENSIONDATA:
+		{
+			const OP::InstanceAddExtensionDataArgs Args = program.GetOpArgs<OP::InstanceAddExtensionDataArgs>(at);
+
+			f(Args.Instance);
+			f(Args.ExtensionData);
+
+			break;
+		}
+
         //-------------------------------------------------------------------------------------
         case OP_TYPE::LA_PACK:
         {
diff --git a/Engine/Plugins/Experimental/Mutable/Source/MutableRuntime/Private/MuR/Operations.h b/Engine/Plugins/Experimental/Mutable/Source/MutableRuntime/Private/MuR/Operations.h
--- a/Engine/Plugins/Experimental/Mutable/Source/MutableRuntime/Private/MuR/Operations.h
+++ b/Engine/Plugins/Experimental/Mutable/Source/MutableRuntime/Private/MuR/Operations.h
@@ -40,6 +40,7 @@ namespace mu
         LA_CONSTANT,
         PR_CONSTANT,
         ST_CONSTANT,
+		ED_CONSTANT,
 
         //! User parameter
         BO_PARAMETER,
@@ -340,6 +341,9 @@ namespace mu
         //! Add all LODs to an instance. This operation can only appear once in a model.
         IN_ADDLOD,
 
+		//! Add extension data to an instance
+		IN_ADDEXTENSIONDATA,
+
         //-----------------------------------------------------------------------------------------
         // Layout operations
         //-----------------------------------------------------------------------------------------
@@ -1015,6 +1019,19 @@ namespace mu
             ADDRESS lod[ MUTABLE_OP_MAX_ADD_COUNT ];
         };
 
+		struct InstanceAddExtensionDataArgs
+		{
+			// This is a reference to an op that produces the Instance that the ExtensionData will
+			// be added to.
+			ADDRESS Instance;
+			// An op that produces the ExtensionData to add to the Instance
+			ADDRESS ExtensionData;
+			// The name to associate with the ExtensionData
+			//
+			// This is an index into the string table
+			ADDRESS ExtensionDataName;
+		};
+
         //-------------------------------------------------------------------------------------
         struct LayoutPackArgs
         {
@@ -1176,6 +1193,7 @@ namespace mu
 		DT_INSTANCE,
 		DT_PROJECTOR,
 		DT_STRING,
+		DT_EXTENSION_DATA,
 
 		// Supporting data types : Never returned as an actual data type for any operation.
 		DT_MATRIX,
diff --git a/Engine/Plugins/Experimental/Mutable/Source/MutableRuntime/Private/MuR/System.cpp b/Engine/Plugins/Experimental/Mutable/Source/MutableRuntime/Private/MuR/System.cpp
--- a/Engine/Plugins/Experimental/Mutable/Source/MutableRuntime/Private/MuR/System.cpp
+++ b/Engine/Plugins/Experimental/Mutable/Source/MutableRuntime/Private/MuR/System.cpp
@@ -21,6 +21,7 @@
 #include "MuR/MutableMath.h"
 #include "MuR/MutableString.h"
 #include "MuR/MutableTrace.h"
+#include "MuR/NullExtensionDataStreamer.h"
 #include "MuR/Operations.h"
 #include "MuR/Parameters.h"
 #include "MuR/ParametersPrivate.h"
@@ -42,7 +43,7 @@ namespace mu
     //---------------------------------------------------------------------------------------------
     //---------------------------------------------------------------------------------------------
     //---------------------------------------------------------------------------------------------
-    System::System( const SettingsPtr& pInSettings )
+    System::System(const SettingsPtr& pInSettings, ExtensionDataStreamer* DataStreamer)
     {
 		LLM_SCOPE_BYNAME(TEXT("MutableRuntime"));
 
@@ -54,7 +55,7 @@ namespace mu
         }
 
         // Choose the implementation
-        m_pD = new System::Private( pSettings );
+        m_pD = new System::Private( pSettings, DataStreamer );
     }
 
 
@@ -79,7 +80,7 @@ namespace mu
 
 
     //---------------------------------------------------------------------------------------------
-    System::Private::Private( SettingsPtr pSettings )
+    System::Private::Private(SettingsPtr pSettings, ExtensionDataStreamer* DataStreamer)
     {
 		LLM_SCOPE_BYNAME(TEXT("MutableRuntime"));
 
@@ -89,6 +90,12 @@ namespace mu
         m_maxMemory = 0;
 	
 		m_modelCache.m_romBudget = pSettings->GetPrivate()->m_streamingCacheBytes;
+
+		if (!DataStreamer)
+		{
+			DataStreamer = new NullExtensionDataStreamer();
+		}
+		m_ExtensionDataStreamer = DataStreamer;
 	}
 
 
@@ -101,6 +108,9 @@ namespace mu
         delete m_pStreamInterface;
         m_pStreamInterface = nullptr;
 
+		delete m_ExtensionDataStreamer;
+		m_ExtensionDataStreamer = nullptr;
+
         delete m_pImageParameterGenerator;
         m_pImageParameterGenerator = nullptr;
     }
diff --git a/Engine/Plugins/Experimental/Mutable/Source/MutableRuntime/Private/MuR/SystemPrivate.h b/Engine/Plugins/Experimental/Mutable/Source/MutableRuntime/Private/MuR/SystemPrivate.h
--- a/Engine/Plugins/Experimental/Mutable/Source/MutableRuntime/Private/MuR/SystemPrivate.h
+++ b/Engine/Plugins/Experimental/Mutable/Source/MutableRuntime/Private/MuR/SystemPrivate.h
@@ -16,6 +16,7 @@
 
 namespace mu
 {
+	class ExtensionDataStreamer;
 
 	//! Reference-counted colour to be stored in the cache
 	class Colour : public RefCounted
@@ -730,6 +731,27 @@ namespace mu
 			return pResult;
 		}
 
+		Ptr<const ExtensionData> GetExtensionData(FCacheAddress at)
+		{
+			if (!at.At)
+			{
+				return nullptr;
+			}
+			
+			const TPair<int, Ptr<const RefCounted>>* d = m_resources.get_ptr(at);
+			if (!d)
+			{
+				return nullptr;
+			}
+
+			Ptr<const ExtensionData> pResult;
+			if (at.At)
+			{
+				pResult = (const ExtensionData*)d->Value.get();
+			}
+			return pResult;
+		}
+
 		void SetBool(FCacheAddress at, bool v)
 		{
 			check(at.At < m_resources.size_code());
@@ -786,6 +808,12 @@ namespace mu
 			m_resources[at] = TPair<int, Ptr<const RefCounted>>(2, v);
 		}
 
+		void SetExtensionData(FCacheAddress at, Ptr<const ExtensionData> v)
+		{
+			check(at.At < m_resources.size_code());
+			m_resources[at] = TPair<int, Ptr<const RefCounted>>(1, v);
+		}
+
 		void SetImage(FCacheAddress at, Ptr<const Image> v)
 		{
 			check(at.At < m_resources.size_code());
@@ -954,7 +982,7 @@ namespace mu
     {
     public:
 
-        Private( SettingsPtr pSettings );
+        Private(SettingsPtr pSettings, ExtensionDataStreamer* DataStreamer);
         virtual ~Private();
 
         //-----------------------------------------------------------------------------------------
@@ -977,6 +1005,8 @@ namespace mu
 		void SetStreamingCache(uint64 bytes) ;
 		void ClearStreamingCache() ;
 
+		ExtensionDataStreamer* GetExtensionDataStreamer() const { return m_ExtensionDataStreamer; }
+
         //!
         SettingsPtrConst m_pSettings;
 
@@ -1071,6 +1101,9 @@ namespace mu
 
 	private:
 
+		/** Owned by this system. */
+		ExtensionDataStreamer* m_ExtensionDataStreamer = nullptr;
+
 		/** This flag is turned on when a streaming error or similar happens. Results are not usable.
 		* This should only happen in-editor.
 		*/
diff --git a/Engine/Plugins/Experimental/Mutable/Source/MutableRuntime/Public/MuR/ExtensionData.h b/Engine/Plugins/Experimental/Mutable/Source/MutableRuntime/Public/MuR/ExtensionData.h
new file mode 100644
--- /dev/null
+++ b/Engine/Plugins/Experimental/Mutable/Source/MutableRuntime/Public/MuR/ExtensionData.h
@@ -0,0 +1,62 @@
+// Copyright Epic Games, Inc. All Rights Reserved.
+
+#pragma once
+
+#include "MuR/Ptr.h"
+#include "MuR/RefCounted.h"
+#include "MuR/Serialisation.h"
+#include "MuR/Types.h"
+
+namespace mu
+{
+	//! ExtensionData represents data types that Mutable doesn't support natively.
+	//!
+	//! Extensions can provide data, and functionality to operate on that data, without Mutable
+	//! needing to know what the data refers to.
+
+	class ExtensionData;
+	typedef Ptr<ExtensionData> ExtensionDataPtr;
+	typedef Ptr<const ExtensionData> ExtensionDataPtrConst;
+
+	class MUTABLERUNTIME_API ExtensionData : public RefCounted
+	{
+	public:
+		static void Serialise(const ExtensionData* Data, OutputArchive& Archive);
+		static ExtensionDataPtr StaticUnserialise(InputArchive& Archive);
+
+		//! A stable hash of the contents
+		uint32 Hash() const;
+
+		void Serialise(OutputArchive& Archive) const;
+		void Unserialise(InputArchive& Archive);
+
+		inline bool operator==(const ExtensionData& Other) const
+		{
+			const bool bResult =
+				Other.Index == Index
+				&& Other.Origin == Origin;
+
+			return bResult;
+		}
+
+		enum class EOrigin : uint8
+		{
+			//! An invalid value used to indicate that this ExtensionData hasn't been initialized
+			Invalid,
+
+			//! This ExtensionData is a compile-time constant that's always loaded into memory
+			ConstantAlwaysLoaded,
+
+			//! This ExtensionData is a compile-time constant that's streamed in from disk when needed
+			ConstantStreamed,
+
+			//! This ExtensionData was generated at runtime
+			Runtime
+		};
+
+		int16 Index = -1;
+
+		EOrigin Origin = EOrigin::Invalid;
+	};
+}
+
diff --git a/Engine/Plugins/Experimental/Mutable/Source/MutableRuntime/Public/MuR/ExtensionDataStreamer.h b/Engine/Plugins/Experimental/Mutable/Source/MutableRuntime/Public/MuR/ExtensionDataStreamer.h
new file mode 100644
--- /dev/null
+++ b/Engine/Plugins/Experimental/Mutable/Source/MutableRuntime/Public/MuR/ExtensionDataStreamer.h
@@ -0,0 +1,69 @@
+// Copyright Epic Games, Inc. All Rights Reserved.
+
+#pragma once
+
+#include "MuR/ExtensionData.h"
+#include "MuR/Ptr.h"
+#include "MuR/RefCounted.h"
+#include "MuR/Serialisation.h"
+#include "MuR/Types.h"
+
+namespace mu
+{
+
+//! Allows runtime code to track the progress of an Extension Data load
+class FExtensionDataLoadHandle
+{
+public:
+	ExtensionDataPtrConst Data;
+
+	enum class ELoadState : uint8
+	{
+		Pending,
+		FailedToLoad,
+		Loaded
+	};
+	ELoadState LoadState;
+};
+
+//! \brief Interface to allow the runtime to load ExtensionData for use
+//!
+//! An object implementing this interface must be provided to the System for any Models that use
+//! ExtensionData.
+//! 
+//! All functions must be safe to call from the Mutable thread.
+//! 
+//! \ingroup runtime
+class MUTABLERUNTIME_API ExtensionDataStreamer : public Base
+{
+public:
+
+	//! Ensure virtual destruction.
+	virtual ~ExtensionDataStreamer() = default;
+
+	//! Make a copy of the actual data referenced by the ExtensionData, and return a new
+	//! runtime-generated ExtensionData that references it.
+	//! 
+	//! Source can either be a constant or another runtime-generated ExtensionData.
+	//! 
+	//! The new ExtensionData should be completely independent of the source ExtensionData, such
+	//! that if the source instance is deleted, the new instance should still be valid.
+	virtual ExtensionDataPtr CloneExtensionData(const ExtensionDataPtrConst& Source) = 0;
+	
+	//! Start asynchronously loading the streamed constant data referenced by the given
+	//! ExtensionData.
+	//! 
+	//! \param Data The ExtensionData to load. Its Origin must be ConstantStreamed.
+	//! 
+	//! \param OutUnloadedConstants An array of ConstantStreamed ExtensionData that was unloaded
+	//! to make space for the requested data. Guaranteed to only contain valid pointers to
+	//! ConstantStreamed data from this program.
+	//! 
+	//! \return A handle that allows the caller to poll the load state. When the state is no longer
+	//! Pending, the load is either complete or has failed.
+	virtual TSharedRef<const FExtensionDataLoadHandle> StartLoad(
+		const ExtensionDataPtrConst& Data,
+		TArray<ExtensionDataPtrConst>& OutUnloadedConstants) = 0;
+};
+
+}
diff --git a/Engine/Plugins/Experimental/Mutable/Source/MutableRuntime/Public/MuR/Instance.h b/Engine/Plugins/Experimental/Mutable/Source/MutableRuntime/Public/MuR/Instance.h
--- a/Engine/Plugins/Experimental/Mutable/Source/MutableRuntime/Public/MuR/Instance.h
+++ b/Engine/Plugins/Experimental/Mutable/Source/MutableRuntime/Public/MuR/Instance.h
@@ -3,6 +3,7 @@
 #pragma once
 
 #include "HAL/PlatformMath.h"
+#include "MuR/ExtensionData.h"
 #include "MuR/Ptr.h"
 #include "MuR/RefCounted.h"
 #include "MuR/Types.h"
@@ -202,6 +203,14 @@ namespace mu
         //! \param sca Index of the string, from 0 to GetStringCount(lod,comp)-1
         const char* GetStringName( int lod, int comp, int surf, int str ) const;
 
+		//! Get the number of ExtensionData values in a component
+		int32 GetExtensionDataCount() const;
+
+		//! Get an ExtensionData value from a component
+		//! \param Index Index of the ExtensionData to fetch
+		//! \param OutExtensionData Receives the ExtensionData
+		//! \param OutName Receives the name associated with the ExtensionData. Guaranteed to be a valid string of non-zero length.
+		void GetExtensionData(int32 Index, ExtensionDataPtrConst& OutExtensionData, const char*& OutName) const;
 
         //-----------------------------------------------------------------------------------------
 		// Interface pattern
diff --git a/Engine/Plugins/Experimental/Mutable/Source/MutableRuntime/Public/MuR/Model.h b/Engine/Plugins/Experimental/Mutable/Source/MutableRuntime/Public/MuR/Model.h
--- a/Engine/Plugins/Experimental/Mutable/Source/MutableRuntime/Public/MuR/Model.h
+++ b/Engine/Plugins/Experimental/Mutable/Source/MutableRuntime/Public/MuR/Model.h
@@ -13,7 +13,7 @@
 //! data structures. Compiled models are not necessarily compatible when the runtime is updated,
 //! so this version number can be used externally to verify this. It is not used internally, and
 //! serializing models from different versions than this runtime will probably result in a crash.
-#define MUTABLE_COMPILED_MODEL_CODE_VERSION		uint32( 62 )
+#define MUTABLE_COMPILED_MODEL_CODE_VERSION		uint32( 63 )
 #define MUTABLE_PARAMETERS_VERSION              uint32( 2 )
 
 
diff --git a/Engine/Plugins/Experimental/Mutable/Source/MutableRuntime/Public/MuR/System.h b/Engine/Plugins/Experimental/Mutable/Source/MutableRuntime/Public/MuR/System.h
--- a/Engine/Plugins/Experimental/Mutable/Source/MutableRuntime/Public/MuR/System.h
+++ b/Engine/Plugins/Experimental/Mutable/Source/MutableRuntime/Public/MuR/System.h
@@ -53,6 +53,7 @@ namespace mu
 	class ModelStreamer;
 	class Parameters;
     class Mesh;
+	class ExtensionDataStreamer;
 
     class System;
     using SystemPtr=Ptr<System>;
@@ -111,7 +112,8 @@ namespace mu
 		//! Constructor of a system object to build data.
         //! \param Settings Optional class with the settings to use in this system. The default
         //! value configures a production-ready system.
-        System( const Ptr<Settings>& Settings = nullptr );
+		//! \param DataStreamer Optional interface to allow the Model to stream in ExtensionData from disk
+        System( const Ptr<Settings>& Settings = nullptr, ExtensionDataStreamer* DataStreamer = nullptr );
 
         //! Set a new provider for model data. The provider will become owned by this instance and
         //! destroyed when necessary.
diff --git a/Engine/Plugins/Experimental/Mutable/Source/MutableTools/Private/MuT/ASTOpAddExtensionData.cpp b/Engine/Plugins/Experimental/Mutable/Source/MutableTools/Private/MuT/ASTOpAddExtensionData.cpp
new file mode 100644
--- /dev/null
+++ b/Engine/Plugins/Experimental/Mutable/Source/MutableTools/Private/MuT/ASTOpAddExtensionData.cpp
@@ -0,0 +1,98 @@
+// Copyright Epic Games, Inc. All Rights Reserved.
+
+#include "MuT/ASTOpAddExtensionData.h"
+
+#include "HAL/PlatformMath.h"
+#include "MuR/ModelPrivate.h"
+#include "MuR/RefCounted.h"
+#include "MuT/StreamsPrivate.h"
+
+#include <utility>
+
+namespace mu
+{
+
+
+//-------------------------------------------------------------------------------------------------
+//-------------------------------------------------------------------------------------------------
+//-------------------------------------------------------------------------------------------------
+ASTOpAddExtensionData::ASTOpAddExtensionData()
+	: Instance(this)
+	, ExtensionData(this)
+{
+}
+
+
+ASTOpAddExtensionData::~ASTOpAddExtensionData()
+{
+	// Explicit call needed to avoid recursive destruction
+	ASTOp::RemoveChildren();
+}
+
+bool ASTOpAddExtensionData::IsEqual(const ASTOp& OtherUntyped) const
+{
+	if (const ASTOpAddExtensionData* Other = dynamic_cast<const ASTOpAddExtensionData*>(&OtherUntyped))
+	{
+		return Instance == Other->Instance
+			&& ExtensionData == Other->ExtensionData
+			&& ExtensionDataName == Other->ExtensionDataName;
+	}
+
+	return false;
+}
+
+uint64 ASTOpAddExtensionData::Hash() const
+{
+	uint64 Result = std::hash<uint64>()(uint64(OP_TYPE::IN_ADDEXTENSIONDATA));
+	
+	hash_combine(Result, Instance.child().get());
+	hash_combine(Result, ExtensionData.child().get());
+	hash_combine(Result, ExtensionDataName);
+	
+	return Result;
+}
+
+mu::Ptr<ASTOp> ASTOpAddExtensionData::Clone(MapChildFuncRef MapChild) const
+{
+	Ptr<ASTOpAddExtensionData> NewInstance = new ASTOpAddExtensionData();
+
+	NewInstance->Instance = MapChild(Instance.child());
+	NewInstance->ExtensionData = MapChild(ExtensionData.child());
+	NewInstance->ExtensionDataName = ExtensionDataName;
+
+	return NewInstance;
+}
+
+void ASTOpAddExtensionData::ForEachChild(const TFunctionRef<void(ASTChild&)> F)
+{
+	F(Instance);
+	F(ExtensionData);
+}
+
+void ASTOpAddExtensionData::Link(FProgram& Program, const FLinkerOptions*)
+{
+	// Already linked?
+	if (linkedAddress)
+	{
+		return;
+	}
+
+	OP::InstanceAddExtensionDataArgs Args;
+	FMemory::Memset(Args, 0);
+
+	check(Instance->linkedAddress);
+	Args.Instance = Instance->linkedAddress;
+
+	check(ExtensionData->linkedAddress);
+	Args.ExtensionData = ExtensionData->linkedAddress;
+
+	check(ExtensionDataName.length() > 0);
+	Args.ExtensionDataName = Program.AddConstant(ExtensionDataName);
+
+	linkedAddress = (OP::ADDRESS)Program.m_opAddress.Num();
+	Program.m_opAddress.Add((uint32_t)Program.m_byteCode.Num());
+	AppendCode(Program.m_byteCode, OP_TYPE::IN_ADDEXTENSIONDATA);
+	AppendCode(Program.m_byteCode, Args);
+}
+
+}
diff --git a/Engine/Plugins/Experimental/Mutable/Source/MutableTools/Private/MuT/ASTOpAddExtensionData.h b/Engine/Plugins/Experimental/Mutable/Source/MutableTools/Private/MuT/ASTOpAddExtensionData.h
new file mode 100644
--- /dev/null
+++ b/Engine/Plugins/Experimental/Mutable/Source/MutableTools/Private/MuT/ASTOpAddExtensionData.h
@@ -0,0 +1,45 @@
+// Copyright Epic Games, Inc. All Rights Reserved.
+
+#pragma once
+
+#include "HAL/Platform.h"
+#include "Math/NumericLimits.h"
+#include "MuR/Operations.h"
+#include "MuR/Ptr.h"
+#include "MuT/AST.h"
+
+
+namespace mu
+{
+struct FProgram;
+
+
+//---------------------------------------------------------------------------------------------
+//! Adds a named ExtensionData to an Instance
+//---------------------------------------------------------------------------------------------
+class ASTOpAddExtensionData : public ASTOp
+{
+public:
+
+	ASTChild Instance;
+	ASTChild ExtensionData;
+	string ExtensionDataName;
+
+public:
+
+	ASTOpAddExtensionData();
+	ASTOpAddExtensionData(const ASTOpAddExtensionData&) = delete;
+	~ASTOpAddExtensionData();
+
+	OP_TYPE GetOpType() const override { return OP_TYPE::IN_ADDEXTENSIONDATA; }
+	uint64 Hash() const override;
+	void ForEachChild(const TFunctionRef<void(ASTChild&)> F) override;
+	bool IsEqual(const ASTOp& OtherUntyped) const override;
+	Ptr<ASTOp> Clone(MapChildFuncRef MapChild) const override;
+	void Link(FProgram& Program, const FLinkerOptions* Options) override;
+};
+
+
+
+}
+
diff --git a/Engine/Plugins/Experimental/Mutable/Source/MutableTools/Private/MuT/ASTOpConstantExtensionData.cpp b/Engine/Plugins/Experimental/Mutable/Source/MutableTools/Private/MuT/ASTOpConstantExtensionData.cpp
new file mode 100644
--- /dev/null
+++ b/Engine/Plugins/Experimental/Mutable/Source/MutableTools/Private/MuT/ASTOpConstantExtensionData.cpp
@@ -0,0 +1,63 @@
+// Copyright Epic Games, Inc. All Rights Reserved.
+
+#include "MuT/ASTOpConstantExtensionData.h"
+
+#include "HAL/UnrealMemory.h"
+#include "MuR/ModelPrivate.h"
+#include "MuT/StreamsPrivate.h"
+
+#include <utility>
+
+
+namespace mu
+{
+
+
+//-------------------------------------------------------------------------------------------------
+bool ASTOpConstantExtensionData::IsEqual(const ASTOp& OtherUntyped) const
+{
+	if (const ASTOpConstantExtensionData* Other = dynamic_cast<const ASTOpConstantExtensionData*>(&OtherUntyped))
+	{
+		return Other->Value == Value;
+	}
+
+	return false;
+}
+
+
+//-------------------------------------------------------------------------------------------------
+mu::Ptr<ASTOp> ASTOpConstantExtensionData::Clone(MapChildFuncRef MapChild) const
+{
+	Ptr<ASTOpConstantExtensionData> Result = new ASTOpConstantExtensionData();
+
+	Result->Value = Value;
+
+	return Result;
+}
+
+
+//-------------------------------------------------------------------------------------------------
+uint64 ASTOpConstantExtensionData::Hash() const
+{
+	return std::hash<uint64>()(Value.get() != nullptr ? Value->Hash() : 0ll);
+}
+
+
+//-------------------------------------------------------------------------------------------------
+void ASTOpConstantExtensionData::Link(FProgram& Program, const FLinkerOptions* Options)
+{
+	if (!linkedAddress)
+	{
+		OP::ResourceConstantArgs Args;
+		FMemory::Memset(Args, 0);
+
+		Args.value = Program.AddConstant(Value);
+
+		linkedAddress = (OP::ADDRESS)Program.m_opAddress.Num();
+		Program.m_opAddress.Add((uint32_t)Program.m_byteCode.Num());
+		AppendCode(Program.m_byteCode, OP_TYPE::ED_CONSTANT);
+		AppendCode(Program.m_byteCode, Args);
+	}
+}
+
+}
\ No newline at end of file
diff --git a/Engine/Plugins/Experimental/Mutable/Source/MutableTools/Private/MuT/ASTOpConstantExtensionData.h b/Engine/Plugins/Experimental/Mutable/Source/MutableTools/Private/MuT/ASTOpConstantExtensionData.h
new file mode 100644
--- /dev/null
+++ b/Engine/Plugins/Experimental/Mutable/Source/MutableTools/Private/MuT/ASTOpConstantExtensionData.h
@@ -0,0 +1,40 @@
+// Copyright Epic Games, Inc. All Rights Reserved.
+
+#pragma once
+
+#include "HAL/Platform.h"
+#include "MuR/Operations.h"
+#include "MuR/Ptr.h"
+#include "MuR/RefCounted.h"
+#include "MuT/AST.h"
+
+
+namespace mu
+{
+struct FProgram;
+
+
+//---------------------------------------------------------------------------------------------
+//! A constant ExtensionData
+//---------------------------------------------------------------------------------------------
+class ASTOpConstantExtensionData : public ASTOp
+{
+public:
+
+	//!
+	ExtensionDataPtrConst Value;
+
+public:
+
+	// ASTOp interface
+	OP_TYPE GetOpType() const override { return OP_TYPE::ED_CONSTANT; }
+	void ForEachChild(const TFunctionRef<void(ASTChild&)>) override {}
+	bool IsEqual(const ASTOp& OtherUntyped) const override;
+	Ptr<ASTOp> Clone(MapChildFuncRef MapChild) const override;
+	uint64 Hash() const override;
+	void Link(FProgram& Program, const FLinkerOptions*) override;
+};
+
+
+}
+
diff --git a/Engine/Plugins/Experimental/Mutable/Source/MutableTools/Private/MuT/CodeGenerator.cpp b/Engine/Plugins/Experimental/Mutable/Source/MutableTools/Private/MuT/CodeGenerator.cpp
--- a/Engine/Plugins/Experimental/Mutable/Source/MutableTools/Private/MuT/CodeGenerator.cpp
+++ b/Engine/Plugins/Experimental/Mutable/Source/MutableTools/Private/MuT/CodeGenerator.cpp
@@ -14,6 +14,7 @@
 #include "MuR/ParametersPrivate.h"
 #include "MuR/Platform.h"
 #include "MuT/ASTOpAddLOD.h"
+#include "MuT/ASTOpAddExtensionData.h"
 #include "MuT/ASTOpConditional.h"
 #include "MuT/ASTOpConstantBool.h"
 #include "MuT/ASTOpConstantResource.h"
@@ -1955,6 +1956,62 @@ namespace mu
         }
         Ptr<ASTOp> rootOp = lodsOp;
 
+		// Add an ASTOpAddExtensionData for each connected ExtensionData node
+		for (const NodeObjectNew::Private::NamedExtensionDataNode& NamedNode : node.m_extensionDataNodes)
+		{
+			if (!NamedNode.Node.get())
+			{
+				// No node connected
+				continue;
+			}
+
+			// Name must be valid
+			check(NamedNode.Name.length() > 0);
+
+			FExtensionDataGenerationResult Result;
+			GenerateExtensionData(Result, NamedNode.Node);
+
+			if (!Result.Op.get())
+			{
+				// Failed to generate anything for this node
+				continue;
+			}
+
+			FConditionalExtensionDataOp& SavedOp = m_conditionalExtensionDataOps.AddDefaulted_GetRef();
+			if (m_currentObject.Num() > 0)
+			{
+				SavedOp.Condition = m_currentObject.Last().m_condition;
+			}
+			SavedOp.ExtensionDataOp = Result.Op;
+			SavedOp.ExtensionDataName = NamedNode.Name;
+		}
+
+		if (m_currentObject.Num() == 0)
+		{
+			for (const FConditionalExtensionDataOp& SavedOp : m_conditionalExtensionDataOps)
+			{
+				Ptr<ASTOpAddExtensionData> ExtensionPinOp = new ASTOpAddExtensionData();
+				ExtensionPinOp->Instance = ASTChild(ExtensionPinOp, rootOp);
+				ExtensionPinOp->ExtensionData = ASTChild(ExtensionPinOp, SavedOp.ExtensionDataOp);
+				ExtensionPinOp->ExtensionDataName = SavedOp.ExtensionDataName;
+
+				if (SavedOp.Condition.get())
+				{
+					Ptr<ASTOpConditional> ConditionOp = new ASTOpConditional();
+					ConditionOp->type = OP_TYPE::IN_CONDITIONAL;
+					ConditionOp->no = rootOp;
+					ConditionOp->yes = ExtensionPinOp;
+					ConditionOp->condition = ASTChild(ConditionOp, SavedOp.Condition);
+					
+					rootOp = ConditionOp;
+				}
+				else
+				{
+					rootOp = ExtensionPinOp;
+				}
+			}
+		}
+
         m_currentParents.Pop();
 
         return rootOp;
diff --git a/Engine/Plugins/Experimental/Mutable/Source/MutableTools/Private/MuT/CodeGenerator.h b/Engine/Plugins/Experimental/Mutable/Source/MutableTools/Private/MuT/CodeGenerator.h
--- a/Engine/Plugins/Experimental/Mutable/Source/MutableTools/Private/MuT/CodeGenerator.h
+++ b/Engine/Plugins/Experimental/Mutable/Source/MutableTools/Private/MuT/CodeGenerator.h
@@ -24,6 +24,7 @@
 #include "MuT/NodeColourSampleImage.h"
 #include "MuT/NodeComponentEdit.h"
 #include "MuT/NodeComponentNew.h"
+#include "MuT/NodeExtensionData.h"
 #include "MuT/NodeImageProject.h"
 #include "MuT/NodeLOD.h"
 #include "MuT/NodeLayout.h"
@@ -306,6 +307,15 @@ namespace mu
         //! Variables added for every node
         map< Ptr<const Node>, Ptr<ASTOpParameter> > m_nodeVariables;
 
+		struct FConditionalExtensionDataOp
+		{
+			Ptr<ASTOp> Condition;
+			Ptr<ASTOp> ExtensionDataOp;
+			string ExtensionDataName;
+		};
+
+		TArray<FConditionalExtensionDataOp> m_conditionalExtensionDataOps;
+
 		//-----------------------------------------------------------------------------------------
 
 		// Get the modifiers that have to be applied to elements with a specific tag.
@@ -546,6 +556,18 @@ namespace mu
 		//!
 		Ptr<const Layout> AddLayout(Ptr<const Layout> SourceLayout);
 
+		struct FExtensionDataGenerationResult
+		{
+			Ptr<ASTOp> Op;
+		};
+
+		typedef const NodeExtensionData* FGeneratedExtensionDataCacheKey;
+		typedef TMap<FGeneratedExtensionDataCacheKey, FExtensionDataGenerationResult> GeneratedExtensionDataMap;
+		GeneratedExtensionDataMap m_generatedExtensionData;
+
+		void GenerateExtensionData(FExtensionDataGenerationResult& OutResult, const NodeExtensionDataPtrConst& InUntypedNode);
+		void GenerateExtensionData_Constant(FExtensionDataGenerationResult& OutResult, const class NodeExtensionDataConstant* Constant);
+
         //-----------------------------------------------------------------------------------------
         // Projectors
         struct FProjectorGenerationResult
diff --git a/Engine/Plugins/Experimental/Mutable/Source/MutableTools/Private/MuT/CodeGenerator_ExtensionData.cpp b/Engine/Plugins/Experimental/Mutable/Source/MutableTools/Private/MuT/CodeGenerator_ExtensionData.cpp
new file mode 100644
--- /dev/null
+++ b/Engine/Plugins/Experimental/Mutable/Source/MutableTools/Private/MuT/CodeGenerator_ExtensionData.cpp
@@ -0,0 +1,84 @@
+// Copyright Epic Games, Inc. All Rights Reserved.
+
+#include "Containers/Array.h"
+#include "Misc/AssertionMacros.h"
+#include "MuR/Layout.h"
+#include "MuR/MutableMath.h"
+#include "MuR/Operations.h"
+#include "MuR/ParametersPrivate.h"
+#include "MuR/Ptr.h"
+#include "MuR/RefCounted.h"
+#include "MuR/System.h"
+#include "MuT/AST.h"
+#include "MuT/ASTOpConditional.h"
+#include "MuT/ASTOpConstantExtensionData.h"
+#include "MuT/CodeGenerator.h"
+#include "MuT/CompilerPrivate.h"
+#include "MuT/ErrorLog.h"
+#include "MuT/ErrorLogPrivate.h"
+#include "MuT/NodeExtensionData.h"
+#include "MuT/NodeExtensionDataConstant.h"
+#include "MuT/NodeExtensionDataConstantPrivate.h"
+
+
+namespace mu
+{
+class Node;
+
+	//---------------------------------------------------------------------------------------------
+	void CodeGenerator::GenerateExtensionData(FExtensionDataGenerationResult& OutResult, const NodeExtensionDataPtrConst& InUntypedNode)
+	{
+		if (!InUntypedNode)
+		{
+			OutResult = FExtensionDataGenerationResult();
+			return;
+		}
+
+		// Clear bottom-up state
+		m_currentBottomUpState.m_address = nullptr;
+
+		// See if it was already generated
+		const FGeneratedExtensionDataCacheKey Key = InUntypedNode.get();
+		GeneratedExtensionDataMap::ValueType* CachedResult = m_generatedExtensionData.Find(Key);
+		if (CachedResult)
+		{
+			OutResult = *CachedResult;
+			return;
+		}
+
+		const NodeExtensionData* Node = InUntypedNode.get();
+
+		// Generate for each different type of node
+		switch (Node->GetExtensionDataNodeType())
+		{
+			case NodeExtensionData::EType::Constant: GenerateExtensionData_Constant(OutResult, static_cast<const NodeExtensionDataConstant*>(Node)); break;
+			case NodeExtensionData::EType::None: check(false);
+		}
+
+		// Cache the result
+		m_generatedExtensionData.Add(Key, OutResult);
+	}
+
+	//---------------------------------------------------------------------------------------------
+	void CodeGenerator::GenerateExtensionData_Constant(FExtensionDataGenerationResult& OutResult, const NodeExtensionDataConstant* Constant)
+	{
+		NodeExtensionDataConstant::Private& Node = *Constant->GetPrivate();
+
+		Ptr<ASTOpConstantExtensionData> Op = new ASTOpConstantExtensionData();
+		OutResult.Op = Op;
+
+		ExtensionDataPtrConst Data = Node.Value;
+		if (!Data)
+		{
+			// Data can't be null, so make an empty one
+			Data = new ExtensionData();
+			
+			// Log an error message
+			m_pErrorLog->GetPrivate()->Add("Constant extension data not set", ELMT_WARNING, Node.m_errorContext);
+		}
+
+		Op->Value = Data;
+	}
+
+}
+
diff --git a/Engine/Plugins/Experimental/Mutable/Source/MutableTools/Private/MuT/NodeExtensionData.cpp b/Engine/Plugins/Experimental/Mutable/Source/MutableTools/Private/MuT/NodeExtensionData.cpp
new file mode 100644
--- /dev/null
+++ b/Engine/Plugins/Experimental/Mutable/Source/MutableTools/Private/MuT/NodeExtensionData.cpp
@@ -0,0 +1,68 @@
+// Copyright Epic Games, Inc. All Rights Reserved.
+
+#include "MuT/NodeExtensionData.h"
+
+#include "MuR/Serialisation.h"
+#include "MuR/SerialisationPrivate.h"
+#include "MuT/NodeExtensionDataConstant.h"
+
+namespace mu
+{
+
+
+	//---------------------------------------------------------------------------------------------
+	// Static initialisation
+	//---------------------------------------------------------------------------------------------
+	static NODE_TYPE s_type = NODE_TYPE("NodeExtensionData", Node::GetStaticType());
+
+	//---------------------------------------------------------------------------------------------
+	//!
+	//---------------------------------------------------------------------------------------------
+
+	//---------------------------------------------------------------------------------------------
+	void NodeExtensionData::Serialise(const NodeExtensionData* Node, OutputArchive& Archive)
+	{
+		uint32_t Ver = 0;
+		Archive << Ver;
+
+		Archive << uint32_t(Node->Type);
+		Node->SerialiseWrapper(Archive);
+	}
+
+
+	//---------------------------------------------------------------------------------------------
+	NodeExtensionDataPtr NodeExtensionData::StaticUnserialise(InputArchive& Archive)
+	{
+		uint32_t Ver;
+		Archive >> Ver;
+		check(Ver == 0);
+
+		uint32_t Id;
+		Archive >> Id;
+
+		switch (EType(Id))
+		{
+			case EType::Constant:		return NodeExtensionDataConstant::StaticUnserialise(Archive); break;
+			
+			default: check(false);
+		}
+
+		return 0;
+	}
+
+
+	//---------------------------------------------------------------------------------------------
+	const NODE_TYPE* NodeExtensionData::GetType() const
+	{
+		return GetStaticType();
+	}
+
+
+	//---------------------------------------------------------------------------------------------
+	const NODE_TYPE* NodeExtensionData::GetStaticType()
+	{
+		return &s_type;
+	}
+}
+
+
diff --git a/Engine/Plugins/Experimental/Mutable/Source/MutableTools/Private/MuT/NodeExtensionDataConstant.cpp b/Engine/Plugins/Experimental/Mutable/Source/MutableTools/Private/MuT/NodeExtensionDataConstant.cpp
new file mode 100644
--- /dev/null
+++ b/Engine/Plugins/Experimental/Mutable/Source/MutableTools/Private/MuT/NodeExtensionDataConstant.cpp
@@ -0,0 +1,66 @@
+// Copyright Epic Games, Inc. All Rights Reserved.
+
+
+#include "MuT/NodeExtensionDataConstant.h"
+
+#include "Misc/AssertionMacros.h"
+#include "MuT/NodeExtensionDataConstantPrivate.h"
+#include "MuT/NodePrivate.h"
+
+
+namespace mu
+{
+
+	//---------------------------------------------------------------------------------------------
+	// Static initialisation
+	//---------------------------------------------------------------------------------------------
+	NODE_TYPE NodeExtensionDataConstant::Private::s_type =
+			NODE_TYPE("ExtensionDataConstant", NodeExtensionData::GetStaticType());
+
+
+	//---------------------------------------------------------------------------------------------
+	//!
+	//---------------------------------------------------------------------------------------------
+
+	MUTABLE_IMPLEMENT_NODE(NodeExtensionDataConstant, EType::Constant, Node, Node::EType::ExtensionData);
+
+
+	//---------------------------------------------------------------------------------------------
+	// Node Interface
+	//---------------------------------------------------------------------------------------------
+	int NodeExtensionDataConstant::GetInputCount() const
+	{
+		return 0;
+	}
+
+
+	//---------------------------------------------------------------------------------------------
+	Node* NodeExtensionDataConstant::GetInputNode(int i) const
+	{
+		check(false);
+		return nullptr;
+	}
+
+
+	//---------------------------------------------------------------------------------------------
+	void NodeExtensionDataConstant::SetInputNode(int i, NodePtr pNode)
+	{
+		check(false);
+	}
+
+
+	//---------------------------------------------------------------------------------------------
+	// Own Interface
+	//---------------------------------------------------------------------------------------------
+	ExtensionDataPtrConst NodeExtensionDataConstant::GetValue() const
+	{
+		return m_pD->Value;
+	}
+
+
+	//---------------------------------------------------------------------------------------------
+	void NodeExtensionDataConstant::SetValue(ExtensionDataPtrConst Value)
+	{
+		m_pD->Value = Value;
+    }
+}
diff --git a/Engine/Plugins/Experimental/Mutable/Source/MutableTools/Private/MuT/NodeExtensionDataConstantPrivate.h b/Engine/Plugins/Experimental/Mutable/Source/MutableTools/Private/MuT/NodeExtensionDataConstantPrivate.h
new file mode 100644
--- /dev/null
+++ b/Engine/Plugins/Experimental/Mutable/Source/MutableTools/Private/MuT/NodeExtensionDataConstantPrivate.h
@@ -0,0 +1,47 @@
+// Copyright Epic Games, Inc. All Rights Reserved.
+
+#pragma once
+
+#include "MuR/ExtensionData.h"
+#include "MuT/NodeExtensionDataConstant.h"
+#include "MuT/NodeLayout.h"
+#include "MuT/NodePrivate.h"
+#include "MuT/AST.h"
+
+
+namespace mu
+{
+
+class NodeExtensionDataConstant::Private : public Node::Private
+{
+public:
+
+	MUTABLE_DEFINE_CONST_VISITABLE()
+
+public:
+
+	static NODE_TYPE s_type;
+
+	ExtensionDataPtrConst Value;
+	
+	//!
+	void Serialise(OutputArchive& Archive) const
+	{
+		uint32_t Version = 0;
+		Archive << Version;
+
+		Archive << Value;
+	}
+
+	//!
+	void Unserialise(InputArchive& Archive)
+	{
+		uint32_t Version;
+		Archive >> Version;
+		check(Version == 0);
+
+		Archive >> Value;
+	}
+};
+
+}
diff --git a/Engine/Plugins/Experimental/Mutable/Source/MutableTools/Private/MuT/NodeObjectNew.cpp b/Engine/Plugins/Experimental/Mutable/Source/MutableTools/Private/MuT/NodeObjectNew.cpp
--- a/Engine/Plugins/Experimental/Mutable/Source/MutableTools/Private/MuT/NodeObjectNew.cpp
+++ b/Engine/Plugins/Experimental/Mutable/Source/MutableTools/Private/MuT/NodeObjectNew.cpp
@@ -332,6 +332,13 @@ namespace mu
 		Data.NumExtraLODsToBuildAfterFirstLOD = NumExtraLODsToBuildAfterFirstLOD;
     }
 
+    //---------------------------------------------------------------------------------------------
+	void NodeObjectNew::AddExtensionDataNode(NodeExtensionDataPtr Node, const char* Name)
+	{
+		NodeObjectNew::Private::NamedExtensionDataNode& Entry = m_pD->m_extensionDataNodes.AddDefaulted_GetRef();
+		Entry.Node = Node;
+		Entry.Name = Name;
+	}
 }
 
 
diff --git a/Engine/Plugins/Experimental/Mutable/Source/MutableTools/Private/MuT/NodeObjectNewPrivate.h b/Engine/Plugins/Experimental/Mutable/Source/MutableTools/Private/MuT/NodeObjectNewPrivate.h
--- a/Engine/Plugins/Experimental/Mutable/Source/MutableTools/Private/MuT/NodeObjectNewPrivate.h
+++ b/Engine/Plugins/Experimental/Mutable/Source/MutableTools/Private/MuT/NodeObjectNewPrivate.h
@@ -6,6 +6,7 @@
 #include "MuT/NodeLODPrivate.h"
 #include "MuT/CompilerPrivate.h"
 
+#include "MuT/NodeExtensionData.h"
 #include "MuT/NodeObjectNew.h"
 #include "MuT/NodeLayout.h"
 
@@ -31,6 +32,26 @@ namespace mu
 
 		TArray<NodeObjectPtr> m_children;
 
+		struct NamedExtensionDataNode
+		{
+			NodeExtensionDataPtr Node;
+			string Name;
+
+			void Serialise(OutputArchive& arch) const
+			{
+				arch << Node;
+				arch << Name;
+			}
+
+			void Unserialise(InputArchive& arch)
+			{
+				arch >> Node;
+				arch >> Name;
+			}
+		};
+
+		TArray<NamedExtensionDataNode> m_extensionDataNodes;
+
 
 		//! List of states
         TArray<FObjectState> m_states;
@@ -39,7 +60,7 @@ namespace mu
 		//!
 		void Serialise( OutputArchive& arch ) const
 		{
-            uint32_t ver = 2;
+            uint32_t ver = 3;
 			arch << ver;
 
 			arch << m_name;
@@ -47,6 +68,7 @@ namespace mu
 			arch << m_lods;
 			arch << m_children;
 			arch << m_states;
+			arch << m_extensionDataNodes;
 		}
 
 		//!
@@ -54,13 +76,18 @@ namespace mu
 		{
             uint32_t ver;
 			arch >> ver;
-            check( ver==2 );
+            check(ver >= 2);
 
 			arch >> m_name;
             arch >> m_uid;
 			arch >> m_lods;
 			arch >> m_children;
 			arch >> m_states;
+
+			if (ver >= 3)
+			{
+				arch >> m_extensionDataNodes;
+			}
 		}
 
         //! Return true if the given component is set in any lod of this object.
diff --git a/Engine/Plugins/Experimental/Mutable/Source/MutableTools/Public/MuT/Node.h b/Engine/Plugins/Experimental/Mutable/Source/MutableTools/Public/MuT/Node.h
--- a/Engine/Plugins/Experimental/Mutable/Source/MutableTools/Public/MuT/Node.h
+++ b/Engine/Plugins/Experimental/Mutable/Source/MutableTools/Public/MuT/Node.h
@@ -77,6 +77,7 @@ namespace mu
 			Range = 19,
 			String = 20,
 			Bool = 21,
+			ExtensionData = 22,
 
 			None
 		};
diff --git a/Engine/Plugins/Experimental/Mutable/Source/MutableTools/Public/MuT/NodeExtensionData.h b/Engine/Plugins/Experimental/Mutable/Source/MutableTools/Public/MuT/NodeExtensionData.h
new file mode 100644
--- /dev/null
+++ b/Engine/Plugins/Experimental/Mutable/Source/MutableTools/Public/MuT/NodeExtensionData.h
@@ -0,0 +1,67 @@
+// Copyright Epic Games, Inc. All Rights Reserved.
+
+#pragma once
+
+#include "MuR/Ptr.h"
+#include "MuR/RefCounted.h"
+#include "MuT/Node.h"
+#include "MuT/NodeObject.h"
+
+
+namespace mu
+{
+
+// Forward definitions
+class NodeExtensionData;
+typedef Ptr<NodeExtensionData> NodeExtensionDataPtr;
+typedef Ptr<const NodeExtensionData> NodeExtensionDataPtrConst;
+
+
+//! Node that evaluates to an ExtensionData
+//! \ingroup model
+class MUTABLETOOLS_API NodeExtensionData : public Node
+{
+public:
+	// Possible subclasses
+	enum class EType : uint8
+	{
+		Constant = 0,
+
+		None
+	};
+
+	//-----------------------------------------------------------------------------------------
+	// Life cycle
+	//-----------------------------------------------------------------------------------------
+
+	static void Serialise(const NodeExtensionData* pNode, OutputArchive& arch);
+	static NodeExtensionDataPtr StaticUnserialise(InputArchive& arch);
+
+	//-----------------------------------------------------------------------------------------
+	// Node Interface
+	//-----------------------------------------------------------------------------------------
+	const NODE_TYPE* GetType() const override;
+	static const NODE_TYPE* GetStaticType();
+
+	inline EType GetExtensionDataNodeType() const { return Type; }
+
+	//-----------------------------------------------------------------------------------------
+	// Own Interface
+	//-----------------------------------------------------------------------------------------
+
+	//-----------------------------------------------------------------------------------------
+	// Interface pattern
+	//-----------------------------------------------------------------------------------------
+	class Private;
+
+protected:
+
+	//! Forbidden. Manage with the Ptr<> template.
+	inline ~NodeExtensionData() {}
+
+	EType Type = EType::None;
+};
+
+
+
+}
diff --git a/Engine/Plugins/Experimental/Mutable/Source/MutableTools/Public/MuT/NodeExtensionDataConstant.h b/Engine/Plugins/Experimental/Mutable/Source/MutableTools/Public/MuT/NodeExtensionDataConstant.h
new file mode 100644
--- /dev/null
+++ b/Engine/Plugins/Experimental/Mutable/Source/MutableTools/Public/MuT/NodeExtensionDataConstant.h
@@ -0,0 +1,72 @@
+// Copyright Epic Games, Inc. All Rights Reserved.
+
+#pragma once
+
+#include "MuR/ExtensionData.h"
+#include "MuR/Ptr.h"
+#include "MuR/RefCounted.h"
+#include "MuT/NodeExtensionData.h"
+
+
+namespace mu
+{
+
+// Forward definitions
+class NodeExtensionDataConstant;
+typedef Ptr<NodeExtensionDataConstant> NodeExtensionDataConstantPtr;
+typedef Ptr<const NodeExtensionDataConstant> NodeExtensionDataConstantPtrConst;
+
+//! Node that outputs a constant ExtensionData
+//! \ingroup model
+class MUTABLETOOLS_API NodeExtensionDataConstant : public NodeExtensionData
+{
+public:
+
+	//-----------------------------------------------------------------------------------------
+	// Life cycle
+	//-----------------------------------------------------------------------------------------
+	NodeExtensionDataConstant();
+
+	void SerialiseWrapper(OutputArchive& arch) const override;
+	static void Serialise(const NodeExtensionDataConstant* pNode, OutputArchive& arch);
+	static NodeExtensionDataConstantPtr StaticUnserialise(InputArchive& arch);
+
+
+	//-----------------------------------------------------------------------------------------
+	// Node Interface
+	//-----------------------------------------------------------------------------------------
+	const NODE_TYPE* GetType() const override;
+	static const NODE_TYPE* GetStaticType();
+
+	int GetInputCount() const override;
+	Node* GetInputNode(int i) const override;
+	void SetInputNode(int i, NodePtr pNode) override;
+
+	//-----------------------------------------------------------------------------------------
+	// Own Interface
+	//-----------------------------------------------------------------------------------------
+
+	//! Get the constant ExtensionData that will be returned.
+	ExtensionDataPtrConst GetValue() const;
+
+	//! Set the constant ExtensionData that will be returned.
+	void SetValue(ExtensionDataPtrConst Value);
+
+	//-----------------------------------------------------------------------------------------
+	// Interface pattern
+	//-----------------------------------------------------------------------------------------
+	class Private;
+	Private* GetPrivate() const;
+	Node::Private* GetBasePrivate() const override;
+
+protected:
+	//! Forbidden. Manage with the Ptr<> template.
+	~NodeExtensionDataConstant();
+
+private:
+	Private* m_pD;
+
+};
+
+
+}
diff --git a/Engine/Plugins/Experimental/Mutable/Source/MutableTools/Public/MuT/NodeObjectNew.h b/Engine/Plugins/Experimental/Mutable/Source/MutableTools/Public/MuT/NodeObjectNew.h
--- a/Engine/Plugins/Experimental/Mutable/Source/MutableTools/Public/MuT/NodeObjectNew.h
+++ b/Engine/Plugins/Experimental/Mutable/Source/MutableTools/Public/MuT/NodeObjectNew.h
@@ -3,9 +3,10 @@
 #pragma once
 #include "MuR/Ptr.h"
 #include "MuR/RefCounted.h"
+#include "MuR/System.h"
 #include "MuT/Node.h"
 #include "MuT/NodeObject.h"
-#include "MuR/System.h"
+#include "MuT/NodeExtensionData.h"
 
 
 namespace mu
@@ -100,6 +101,9 @@ namespace mu
 			uint8 FirstLOD,
 			uint8 NumExtraLODsToBuildAfterFirstLOD);
 
+		//! Connect a node that produces ExtensionData to be added to the final Instance, and provide a name to associate with the data
+		void AddExtensionDataNode(NodeExtensionDataPtr Node, const char* Name);
+
 		//-----------------------------------------------------------------------------------------
 		// Interface pattern
 		//-----------------------------------------------------------------------------------------
