Fix for non-unity compile error in NullExtensionDataStreamer

diff --git a/Engine/Plugins/Experimental/Mutable/Source/MutableRuntime/Private/MuR/NullExtensionDataStreamer.cpp b/Engine/Plugins/Experimental/Mutable/Source/MutableRuntime/Private/MuR/NullExtensionDataStreamer.cpp
--- a/Engine/Plugins/Experimental/Mutable/Source/MutableRuntime/Private/MuR/NullExtensionDataStreamer.cpp
+++ b/Engine/Plugins/Experimental/Mutable/Source/MutableRuntime/Private/MuR/NullExtensionDataStreamer.cpp
@@ -3,6 +3,7 @@
 #include "MuR/NullExtensionDataStreamer.h"
 
 #include "MuR/ExtensionData.h"
+#include "Templates/SharedPointer.h"
 
 namespace mu
 {
diff --git a/Engine/Plugins/Experimental/Mutable/Source/MutableRuntime/Public/MuR/ExtensionDataStreamer.h b/Engine/Plugins/Experimental/Mutable/Source/MutableRuntime/Public/MuR/ExtensionDataStreamer.h
--- a/Engine/Plugins/Experimental/Mutable/Source/MutableRuntime/Public/MuR/ExtensionDataStreamer.h
+++ b/Engine/Plugins/Experimental/Mutable/Source/MutableRuntime/Public/MuR/ExtensionDataStreamer.h
@@ -7,6 +7,7 @@
 #include "MuR/RefCounted.h"
 #include "MuR/Serialisation.h"
 #include "MuR/Types.h"
+#include "Templates/SharedPointerFwd.h"
 
 namespace mu
 {
