Added more missing includes for non-unity build

diff --git a/Engine/Plugins/Experimental/Mutable/Source/CustomizableObject/Public/MuCO/UnrealExtensionDataStreamer.h b/Engine/Plugins/Experimental/Mutable/Source/CustomizableObject/Public/MuCO/UnrealExtensionDataStreamer.h
--- a/Engine/Plugins/Experimental/Mutable/Source/CustomizableObject/Public/MuCO/UnrealExtensionDataStreamer.h
+++ b/Engine/Plugins/Experimental/Mutable/Source/CustomizableObject/Public/MuCO/UnrealExtensionDataStreamer.h
@@ -2,6 +2,7 @@
 
 #pragma once
 
+#include "Async/Future.h"
 #include "Engine/StreamableManager.h"
 #include "MuR/ExtensionData.h"
 #include "MuR/ExtensionDataStreamer.h"
diff --git a/Engine/Plugins/Experimental/Mutable/Source/CustomizableObjectEditor/Private/MuCOE/CustomizableObjectCookPackageSplitter.cpp b/Engine/Plugins/Experimental/Mutable/Source/CustomizableObjectEditor/Private/MuCOE/CustomizableObjectCookPackageSplitter.cpp
--- a/Engine/Plugins/Experimental/Mutable/Source/CustomizableObjectEditor/Private/MuCOE/CustomizableObjectCookPackageSplitter.cpp
+++ b/Engine/Plugins/Experimental/Mutable/Source/CustomizableObjectEditor/Private/MuCOE/CustomizableObjectCookPackageSplitter.cpp
@@ -5,6 +5,7 @@
 #include "Algo/Find.h"
 #include "MuCO/CustomizableObject.h"
 #include "UObject/NameTypes.h"
+#include "UObject/Package.h"
 
 REGISTER_COOKPACKAGE_SPLITTER(FCustomizableObjectCookPackageSplitter, UCustomizableObject);
 
