[Mutable] Added new UpdateSkeletalMesh functoin with result callback.

Allows the user to know when an update has finished.

- Created a new UpdateSkeletalMeshAsync delegate which allways calls the Callback, even if there is an error.
- Added a common en point to the update process (EndUpdate).
- Minor simplifications of the update flow.
- Fixed the update result. No longer depends on the SkeletalMeshStatus flag.
- Removed unused functions.
- Removed unused InstanceUpdated.

#rnx

diff --git a/Engine/Plugins/Experimental/Mutable/Source/CustomizableObject/Classes/MuCO/CustomizableObject.h b/Engine/Plugins/Experimental/Mutable/Source/CustomizableObject/Classes/MuCO/CustomizableObject.h
--- a/Engine/Plugins/Experimental/Mutable/Source/CustomizableObject/Classes/MuCO/CustomizableObject.h
+++ b/Engine/Plugins/Experimental/Mutable/Source/CustomizableObject/Classes/MuCO/CustomizableObject.h
@@ -1628,7 +1628,7 @@ public:
 	// Return a pointer to the BulkData subobject, only valid in packaged builds
 	const UCustomizableObjectBulk* GetStreamableBulkData() const { return BulkData; }
 
-	FCustomizableObjectPrivateData* GetPrivate() const { return PrivateData.Get(); }
+	FCustomizableObjectPrivateData* GetPrivate() const;
 
 	// This will always return true in a packaged game
 	UFUNCTION(BlueprintCallable, Category = CustomizableObject)
diff --git a/Engine/Plugins/Experimental/Mutable/Source/CustomizableObject/Classes/MuCO/CustomizableObjectInstance.h b/Engine/Plugins/Experimental/Mutable/Source/CustomizableObject/Classes/MuCO/CustomizableObjectInstance.h
--- a/Engine/Plugins/Experimental/Mutable/Source/CustomizableObject/Classes/MuCO/CustomizableObjectInstance.h
+++ b/Engine/Plugins/Experimental/Mutable/Source/CustomizableObject/Classes/MuCO/CustomizableObjectInstance.h
@@ -51,14 +51,42 @@ namespace ESkeletalMeshState
 };
 
 
+/** Result of all the checks just before beginning an update. */
+enum class EUpdateRequired : uint8
+{
+	NoUpdate, // No work required.
+	Update, // Normal update.
+	Discard // Discard instead of update.
+};
+
+
 /** Instance Update Result. */
+UENUM()
 enum class EUpdateResult : uint8
 {
-	Success,
-	Error
+	Success, // There only exist one Success case. Any other new cases have to be errors.
+	Error, // Generic error.
+	ErrorOptimized, // The update was skipped since its result would have been the same as the current customization.
+	ErrorReplaced, // The update was replaced by a newer update request.
+	ErrorDiscarded // The update was not finished since due to the LOD management discarding the data.
+};
+
+
+/** Instance Update Context.
+ * Used to avoid changing the delegate signature in the future.  */
+USTRUCT()
+struct FUpdateContext
+{	
+	GENERATED_BODY()
+
+	UPROPERTY()
+	EUpdateResult UpdateResult;
 };
 
 
+DECLARE_DYNAMIC_DELEGATE_OneParam(FInstanceUpdateDelegate, const FUpdateContext&, Result);
+
+
 /* When creating new delegates use the following conventions:
  *
  * - All delegates must be multicast.
@@ -150,7 +178,6 @@ public:
 	virtual void PreEditChange(FProperty* PropertyAboutToChange) override;
 	virtual void PostEditChangeProperty(FPropertyChangedEvent& PropertyChangedEvent) override;
 	virtual bool CanEditChange( const FProperty* InProperty ) const override;
-	bool InstanceUpdated; // Flag for the editor, to know when the instance's skeletal mesh has been updated
 #endif //WITH_EDITOR
 
 	virtual void Serialize(FArchive& Ar) override;
@@ -233,10 +260,27 @@ public:
 	void SetSelectedParameterProfileDirty();
 	bool IsSelectedParameterProfileDirty() const;
 
-	//
+	/** Update Skeletal Mesh asynchronously. */
 	UFUNCTION(BlueprintCallable, Category = CustomizableObjectInstance)
 	void UpdateSkeletalMeshAsync(bool bIgnoreCloseDist = false, bool bForceHighPriority = false);
+		
+	/** Update Skeletal Mesh asynchronously. Callback will be called once the update finishes, even if it fails. */
+	UFUNCTION(BlueprintCallable, Category = CustomizableObjectInstance)
+	void UpdateSkeletalMeshAsyncResult(FInstanceUpdateDelegate Callback, bool bIgnoreCloseDist = false, bool bForceHighPriority = false);
+
+private:
+	/** Perform all the checks required to see if the update should begin. */
+	EUpdateRequired IsUpdateRequired(bool bIsCloseDistTick, bool bOnlyUpdateIfNotGenerated, bool bIgnoreCloseDist) const;
+	
+	/** Update Skeletal Mesh asynchronously. Immersive function. Not exposed to the API.
+	 * Once the update reaches this function, the update has been considered started and must complete all the update flow.
+	/* Starting at this function, all Update code paths must end up in FinishUpdateGlobal!
+	 *
+	 * @param bIsCloseDistTick true iff called from the tick.
+	 */
+	void DoUpdateSkeletalMesh(bool bIsCloseDistTick, bool bOnlyUpdateIfNotGenerated, bool bIgnoreCloseDist, bool bForceHighPriority, const EUpdateRequired* OptionalUpdateRequired, FInstanceUpdateDelegate* UpdateCallback);
 
+public:	
 	// Clones the instance creating a new identical transient instance.
 	UFUNCTION(BlueprintCallable, Category = CustomizableObjectInstance)
 	UCustomizableObjectInstance* Clone();
@@ -594,8 +638,8 @@ public:
 
 	const TArray<uint16>& GetRequestedLODsPerComponent() const;
 
-	/** Instance updated. */
-	void Updated(EUpdateResult Result, const FDescriptorRuntimeHash& UpdatedHash);
+	/** Common end point of all updates. Even those which failed. */
+	void FinishUpdate(EUpdateResult UpdateResult, const FDescriptorRuntimeHash& UpdatedHash);
 	
 	/** Return the UCustomizableObjectInstance::Descriptor hash on the last update request. */
 	FDescriptorRuntimeHash GetDescriptorRuntimeHash() const;
@@ -608,10 +652,7 @@ public:
 	/** Flag to know if a property of this instance changed in the editor */
 	bool bEditorPropertyChanged = false;
 
-	UCustomizableInstancePrivateData* GetPrivate() const { 
-		check(nullptr != PrivateData); // Currently this is initialized in the constructor so we expect it always to exist.
-		return PrivateData; 
-	}
+	UCustomizableInstancePrivateData* GetPrivate() const;
 
 	bool ProjectorUpdatedInViewport = false;
 
diff --git a/Engine/Plugins/Experimental/Mutable/Source/CustomizableObject/Classes/MuCO/CustomizableObjectSystem.h b/Engine/Plugins/Experimental/Mutable/Source/CustomizableObject/Classes/MuCO/CustomizableObjectSystem.h
--- a/Engine/Plugins/Experimental/Mutable/Source/CustomizableObject/Classes/MuCO/CustomizableObjectSystem.h
+++ b/Engine/Plugins/Experimental/Mutable/Source/CustomizableObject/Classes/MuCO/CustomizableObjectSystem.h
@@ -5,6 +5,8 @@
 #include "HAL/IConsoleManager.h"
 #include "Engine/StreamableManager.h"
 #include "AssetRegistry/AssetData.h"
+#include "MuCO/CustomizableObjectInstance.h"
+#include "MuCO/CustomizableObjectInstanceDescriptor.h"
 
 #if WITH_EDITOR
 #include "Framework/Notifications/NotificationManager.h"
@@ -14,6 +16,7 @@
 
 class IConsoleVariable;
 class FCustomizableObjectCompilerBase;
+class UCustomizableInstanceLODManagementBase;
 class ITargetPlatform;
 class SNotificationItem;
 class UCustomizableObject;
@@ -181,7 +184,8 @@ public:
 };
 
 
-class UCustomizableInstanceLODManagementBase;
+/** End a Customizable Object Instance Update. All code paths of an update have to end here. */
+void FinishUpdateGlobal(UCustomizableObjectInstance* Instance, EUpdateResult UpdateResult, FInstanceUpdateDelegate* UpdateCallback, const FDescriptorRuntimeHash InUpdatedHash = FDescriptorRuntimeHash());
 
 
 UCLASS(Blueprintable, BlueprintType)
diff --git a/Engine/Plugins/Experimental/Mutable/Source/CustomizableObject/Classes/MuCO/CustomizableSkeletalComponent.h b/Engine/Plugins/Experimental/Mutable/Source/CustomizableObject/Classes/MuCO/CustomizableSkeletalComponent.h
--- a/Engine/Plugins/Experimental/Mutable/Source/CustomizableObject/Classes/MuCO/CustomizableSkeletalComponent.h
+++ b/Engine/Plugins/Experimental/Mutable/Source/CustomizableObject/Classes/MuCO/CustomizableSkeletalComponent.h
@@ -3,6 +3,7 @@
 #pragma once
 
 #include "Components/SceneComponent.h"
+#include "MuCO/CustomizableObjectInstance.h"
 
 #include "CustomizableSkeletalComponent.generated.h"
 
@@ -12,7 +13,6 @@ class UCustomizableSkeletalComponent;
 class UObject;
 class USkeletalMesh;
 struct FFrame;
-class UCustomizableObjectInstance;
 enum class EUpdateResult : uint8;
 
 DECLARE_DELEGATE_TwoParams(FCustomizableSkeletalComponentPreUpdateDelegate, UCustomizableSkeletalComponent* Component, USkeletalMesh* NextMesh);
@@ -43,15 +43,21 @@ public:
 
 	FCustomizableSkeletalComponentPreUpdateDelegate PreUpdateDelegate;
 	FCustomizableSkeletalComponentUpdatedDelegate UpdatedDelegate;
-	
-	void Updated(EUpdateResult Result) const;
+
+	/** Common end point of all updates. Even those which failed. */
+	void Callbacks() const;
 	
 	USkeletalMesh* GetSkeletalMesh() const;
 	USkeletalMesh* GetAttachedSkeletalMesh() const;
 
+	/** Update Skeletal Mesh asynchronously. */
 	UFUNCTION(BlueprintCallable, Category = CustomizableObject)
 	void UpdateSkeletalMeshAsync(bool bNeverSkipUpdate = false);
 
+	/** Update Skeletal Mesh asynchronously. Callback will be called once the update finishes, even if it fails. */
+	UFUNCTION(BlueprintCallable, Category = CustomizableObjectInstance)
+	void UpdateSkeletalMeshAsyncResult(FInstanceUpdateDelegate Callback, bool bIgnoreCloseDist = false, bool bForceHighPriority = false);
+
 	void SetSkeletalMesh(USkeletalMesh* SkeletalMesh, bool bReinitPose=true, bool bForceClothReset=true);
 	void SetPhysicsAsset(class UPhysicsAsset* PhysicsAsset);
 	void UpdateDistFromComponentToPlayer(const AActor* const Pawn, bool bForceEvenIfNotBegunPlay = false);
diff --git a/Engine/Plugins/Experimental/Mutable/Source/CustomizableObject/Private/MuCO/CustomizableInstancePrivateData.h b/Engine/Plugins/Experimental/Mutable/Source/CustomizableObject/Private/MuCO/CustomizableInstancePrivateData.h
--- a/Engine/Plugins/Experimental/Mutable/Source/CustomizableObject/Private/MuCO/CustomizableInstancePrivateData.h
+++ b/Engine/Plugins/Experimental/Mutable/Source/CustomizableObject/Private/MuCO/CustomizableInstancePrivateData.h
@@ -25,7 +25,7 @@ enum ECOInstanceFlags
 	// Update process
 	Updating						= 1 << 0,	//
 	CreatingSkeletalMesh			= 1 << 1,	//
-	Generated						= 1 << 2,	//
+	Generated						= 1 << 2,	// Generated or update in progress.
 	ReuseTextures					= 1 << 3, 	// 
 	ReplacePhysicsAssets			= 1 << 4,	// Merge active PhysicsAssets and replace the base physics asset
 
@@ -46,6 +46,8 @@ enum ECOInstanceFlags
 	ForceGenerateAllLODs			= 1 << 12,	// If set, Requested LOD Levels will be ignored and all LODs in between the current min/max lod will be generated
 };
 
+ENUM_CLASS_FLAGS(ECOInstanceFlags);
+
 
 USTRUCT()
 struct FGeneratedTexture
@@ -247,10 +249,6 @@ public:
 	// \TODO: CustomizableObject shouldn't be here
 	static void ProcessTextureCoverageQueries(const TSharedPtr<FMutableOperationData>& OperationData, UCustomizableObject* CustomizableObject, const FString& ImageKeyName, FTexturePlatformData *PlatformData, UMaterialInterface* Material);
 
-	// Prepare the data for the unreal reference skeleton
-	// it runs in the mutable thread
-	static void PrepareSkeletonData(const TSharedPtr<FMutableOperationData>& OperationData);
-
 	// Copy data generated in the mutable thread over to the instance and initializes additional data required during the update
 	void PrepareForUpdate(const TSharedPtr<FMutableOperationData>& OperationData);
 
@@ -259,8 +257,6 @@ public:
 	// The following method is basically copied from PostEditChangeProperty and/or SkeletalMesh.cpp to be able to replicate PostEditChangeProperty without the editor
 	void PostEditChangePropertyWithoutEditor(USkeletalMesh* SkeletalMesh);
 	
-	void DoUpdateSkeletalMesh(UCustomizableObjectInstance& Instance, bool bIsCloseDistTick=false, bool bOnlyUpdateIfNotGenerated=false, bool bIgnoreCloseDist=false, bool bForceHighPriority = false);
-
 	void DiscardResourcesAndSetReferenceSkeletalMesh(UCustomizableObjectInstance* Public);
 
 	const TArray<FAnimInstanceOverridePhysicsAsset>* GetGeneratedPhysicsAssetsForAnimInstance(TSubclassOf<UAnimInstance> AnimInstance) const;
diff --git a/Engine/Plugins/Experimental/Mutable/Source/CustomizableObject/Private/MuCO/CustomizableObject.cpp b/Engine/Plugins/Experimental/Mutable/Source/CustomizableObject/Private/MuCO/CustomizableObject.cpp
--- a/Engine/Plugins/Experimental/Mutable/Source/CustomizableObject/Private/MuCO/CustomizableObject.cpp
+++ b/Engine/Plugins/Experimental/Mutable/Source/CustomizableObject/Private/MuCO/CustomizableObject.cpp
@@ -908,6 +908,14 @@ void UCustomizableObject::LoadEmbeddedData(FArchive& Ar)
 	}
 }
 
+
+FCustomizableObjectPrivateData* UCustomizableObject::GetPrivate() const
+{
+	check(PrivateData)
+	return PrivateData.Get();
+}
+
+
 bool UCustomizableObject::IsCompiled() const
 {
 	bool IsCompiled = PrivateData->GetModel() != nullptr;
diff --git a/Engine/Plugins/Experimental/Mutable/Source/CustomizableObject/Private/MuCO/CustomizableObjectInstance.cpp b/Engine/Plugins/Experimental/Mutable/Source/CustomizableObject/Private/MuCO/CustomizableObjectInstance.cpp
--- a/Engine/Plugins/Experimental/Mutable/Source/CustomizableObject/Private/MuCO/CustomizableObjectInstance.cpp
+++ b/Engine/Plugins/Experimental/Mutable/Source/CustomizableObject/Private/MuCO/CustomizableObjectInstance.cpp
@@ -749,144 +749,184 @@ void UCustomizableInstancePrivateData::PostEditChangePropertyWithoutEditor(USkel
 
 void UCustomizableObjectInstance::UpdateSkeletalMeshAsync(bool bIgnoreCloseDist, bool bForceHighPriority)
 {
-	PrivateData->DoUpdateSkeletalMesh(*this, false, false, bIgnoreCloseDist, bForceHighPriority);
+	DoUpdateSkeletalMesh( false, false, bIgnoreCloseDist, bForceHighPriority, nullptr, nullptr);
 }
 
 
-void UCustomizableInstancePrivateData::DoUpdateSkeletalMesh(UCustomizableObjectInstance& CustomizableInstance, bool bIsCloseDistTick, bool bOnlyUpdateIfNotGenerated, bool bIgnoreCloseDist, bool bForceHighPriority)
+void UCustomizableObjectInstance::UpdateSkeletalMeshAsyncResult(FInstanceUpdateDelegate Callback, bool bIgnoreCloseDist, bool bForceHighPriority)
 {
-	MUTABLE_CPUPROFILER_SCOPE(UCustomizableInstancePrivateData::DoUpdateSkeletalMesh);
-	check(IsInGameThread());
+	DoUpdateSkeletalMesh(false, false, bIgnoreCloseDist, bForceHighPriority, nullptr, &Callback);
+}
 
-	if (!CustomizableInstance.IsValidLowLevel())
-	{
-		return;
-	}
-	
+
+EUpdateRequired UCustomizableObjectInstance::IsUpdateRequired(bool bIsCloseDistTick, bool bOnlyUpdateIfNotGenerated, bool bIgnoreCloseDist) const
+{
 	UCustomizableObjectSystem* System = UCustomizableObjectSystem::GetInstance();
+	const UCustomizableInstancePrivateData* const Private = GetPrivate();
 	
-	UCustomizableObject* CustomizableObject = CustomizableInstance.GetCustomizableObject();
-
-	if (CustomizableObject && !CustomizableObject->IsCompiled())
+	UCustomizableObject* CustomizableObject = GetCustomizableObject();
+	if (!CustomizableObject || CustomizableObject->IsLocked())
 	{
+		return EUpdateRequired::NoUpdate;
+	}
+	
 #if WITH_EDITOR
-		if(!CustomizableObject->ConditionalAutoCompile())
-		{
-			return;
-		}
+	if (!CustomizableObject->IsCompiled() && !CustomizableObject->ConditionalAutoCompile())
+	{
+		return EUpdateRequired::NoUpdate;
+	}
 #else
-		System->AddUncompiledCOWarning(CustomizableObject);
-		return;
-#endif
-
+	if (!CustomizableObject->IsCompiled())
+	{
+		return EUpdateRequired::NoUpdate;		
 	}
+#endif
 
-	const bool bIsGenerated = HasCOInstanceFlags(Generated);
-
+	const bool bIsGenerated = Private->HasCOInstanceFlags(Generated);
 	const int32 NumGeneratedInstancesLimit = System->GetInstanceLODManagement()->GetNumGeneratedInstancesLimitFullLODs();
 	const int32 NumGeneratedInstancesLimitLOD1 = System->GetInstanceLODManagement()->GetNumGeneratedInstancesLimitLOD1();
 	const int32 NumGeneratedInstancesLimitLOD2 = System->GetInstanceLODManagement()->GetNumGeneratedInstancesLimitLOD2();
 
-	if (NumGeneratedInstancesLimit > 0 && System->GetPrivate()->GetCountAllocatedSkeletalMesh() > NumGeneratedInstancesLimit + NumGeneratedInstancesLimitLOD1 + NumGeneratedInstancesLimitLOD2)
+	if (!bIsGenerated && // Prevent generating more instances than the limit, but let updates to existing instances run normally
+		NumGeneratedInstancesLimit > 0 &&
+		System->GetPrivate()->GetCountAllocatedSkeletalMesh() > NumGeneratedInstancesLimit + NumGeneratedInstancesLimitLOD1 + NumGeneratedInstancesLimitLOD2)
 	{
-		if (!bIsGenerated) // Prevent generating more instances than the limit, but let updates to existing instances run normally
-		{
-			return;
-		}
+		return EUpdateRequired::NoUpdate;
 	}
 
-	const bool bShouldUpdateLODs = HasCOInstanceFlags(PendingLODsUpdate);
+	const bool bShouldUpdateLODs = Private->HasCOInstanceFlags(PendingLODsUpdate);
 
-	if (
-		!HasCOInstanceFlags(DiscardedByNumInstancesLimit)
-		&&
-		CustomizableObject && !CustomizableObject->IsLocked()
-		&&
-		(
-			!System->GetInstanceLODManagement()->IsOnlyUpdateCloseCustomizableObjectsEnabled() 
-			|| LastMinSquareDistFromComponentToPlayer <= FMath::Square(System->GetInstanceLODManagement()->GetOnlyUpdateCloseCustomizableObjectsDist())
-			|| bIgnoreCloseDist 
-			|| bShouldUpdateLODs
-			)
-		)
+	const bool bDiscardByDistance = Private->LastMinSquareDistFromComponentToPlayer > FMath::Square(System->GetInstanceLODManagement()->GetOnlyUpdateCloseCustomizableObjectsDist());
+	const bool bLODManagementDiscard = System->GetInstanceLODManagement()->IsOnlyUpdateCloseCustomizableObjectsEnabled() &&
+			bDiscardByDistance &&
+			!bIgnoreCloseDist;
+	
+	if (Private->HasCOInstanceFlags(DiscardedByNumInstancesLimit) ||
+		bLODManagementDiscard)
 	{
-		if (bShouldUpdateLODs)
+		if (bIsGenerated && !Private->HasCOInstanceFlags(Updating))
 		{
-			UE_LOG(LogMutable, Verbose, TEXT("LOD change: %d, %d -> %d, %d"), CustomizableInstance.GetCurrentMinLOD(), CustomizableInstance.GetCurrentMaxLOD(), CustomizableInstance.GetMinLODToLoad(), CustomizableInstance.GetMinLODToLoad());
+			return EUpdateRequired::Discard;		
 		}
-
-		if ((!(bIsGenerated && bOnlyUpdateIfNotGenerated) && !(bIsCloseDistTick && bIsGenerated)) || bShouldUpdateLODs)
+		else
 		{
-			FMutableQueueElem::EQueuePriorityType Priority = FMutableQueueElem::EQueuePriorityType::Low;
+			return EUpdateRequired::NoUpdate;
+		}
+	}
 
-			const bool bIsDowngradeLODUpdate = HasCOInstanceFlags(PendingLODsDowngrade);
-			const bool bIsPlayerOrNearIt = HasCOInstanceFlags(UsedByPlayerOrNearIt);
+	if (bIsGenerated &&
+		!bShouldUpdateLODs &&
+		(bOnlyUpdateIfNotGenerated || bIsCloseDistTick))
+	{
+		return EUpdateRequired::NoUpdate;
+	}
 
-			if (bForceHighPriority)
-			{
-				Priority = FMutableQueueElem::EQueuePriorityType::High;
-			}
-			else if (!bIsGenerated || !CustomizableInstance.HasAnySkeletalMesh())
-			{
-				Priority = FMutableQueueElem::EQueuePriorityType::Med;
-			}
-			else if (bShouldUpdateLODs && bIsDowngradeLODUpdate)
-			{
-				Priority = FMutableQueueElem::EQueuePriorityType::Med_Low;
-			}
-			else if (bIsPlayerOrNearIt && bShouldUpdateLODs && !bIsDowngradeLODUpdate)
-			{
-				Priority = FMutableQueueElem::EQueuePriorityType::High;
-			}
-			else if (bShouldUpdateLODs && !bIsDowngradeLODUpdate)
-			{
-				Priority = FMutableQueueElem::EQueuePriorityType::Med;
-			}
-			else if (bIsPlayerOrNearIt)
-			{
-				Priority = FMutableQueueElem::EQueuePriorityType::High;
-			}
+	return EUpdateRequired::Update;
+}
 
-			{
-				const uint32 InstanceId = CustomizableInstance.GetUniqueID();
-				const float Distance = FMath::Sqrt(CustomizableInstance.GetPrivate()->LastMinSquareDistFromComponentToPlayer);
-				UE_LOG(LogMutable, Verbose, TEXT("Started UpdateSkeletalMesh Async. of Instance %d with priority %d at dist %f bIsPlayerOrNearIt=%d, frame=%d"), InstanceId, static_cast<int32>(Priority), Distance, bIsPlayerOrNearIt, GFrameNumber);				
-			}
 
-			CustomizableInstance.UpdateDescriptorRuntimeHash = FDescriptorRuntimeHash(CustomizableInstance.Descriptor);
+void UCustomizableObjectInstance::DoUpdateSkeletalMesh(bool bIsCloseDistTick, bool bOnlyUpdateIfNotGenerated, bool bIgnoreCloseDist, bool bForceHighPriority, const EUpdateRequired* OptionalUpdateRequired, FInstanceUpdateDelegate* UpdateCallback)
+{
+	MUTABLE_CPUPROFILER_SCOPE(UCustomizableInstancePrivateData::DoUpdateSkeletalMesh);
+	check(IsInGameThread());
+	check(!OptionalUpdateRequired || *OptionalUpdateRequired != EUpdateRequired::NoUpdate) // If no update is required this functions must not be called.
+
+	UCustomizableObjectSystem* System = UCustomizableObjectSystem::GetInstance();
+
+	const EUpdateRequired UpdateRequired = OptionalUpdateRequired ? *OptionalUpdateRequired : IsUpdateRequired(bIsCloseDistTick, bOnlyUpdateIfNotGenerated, bIgnoreCloseDist);
+	switch (UpdateRequired)
+	{
+	case EUpdateRequired::NoUpdate:
+	{	
+		FinishUpdateGlobal(this, EUpdateResult::ErrorDiscarded, UpdateCallback);
+		break;
+	}		
+	case EUpdateRequired::Update:
+	{
+		const bool bIsGenerated = GetPrivate()->HasCOInstanceFlags(Generated);
+		const bool bShouldUpdateLODs = GetPrivate()->HasCOInstanceFlags(PendingLODsUpdate);	
+		const bool bIsDowngradeLODUpdate = GetPrivate()->HasCOInstanceFlags(PendingLODsDowngrade);
+		const bool bIsPlayerOrNearIt = GetPrivate()->HasCOInstanceFlags(UsedByPlayerOrNearIt);
+
+		FMutableQueueElem::EQueuePriorityType Priority = FMutableQueueElem::EQueuePriorityType::Low;
+		if (bForceHighPriority)
+		{
+			Priority = FMutableQueueElem::EQueuePriorityType::High;
+		}
+		else if (!bIsGenerated || !HasAnySkeletalMesh())
+		{
+			Priority = FMutableQueueElem::EQueuePriorityType::Med;
+		}
+		else if (bShouldUpdateLODs && bIsDowngradeLODUpdate)
+		{
+			Priority = FMutableQueueElem::EQueuePriorityType::Med_Low;
+		}
+		else if (bIsPlayerOrNearIt && bShouldUpdateLODs && !bIsDowngradeLODUpdate)
+		{
+			Priority = FMutableQueueElem::EQueuePriorityType::High;
+		}
+		else if (bShouldUpdateLODs && !bIsDowngradeLODUpdate)
+		{
+			Priority = FMutableQueueElem::EQueuePriorityType::Med;
+		}
+		else if (bIsPlayerOrNearIt)
+		{
+			Priority = FMutableQueueElem::EQueuePriorityType::High;
+		}
+
+		{
+			const uint32 InstanceId = GetUniqueID();
+			const float Distance = FMath::Sqrt(GetPrivate()->LastMinSquareDistFromComponentToPlayer);
+			UE_LOG(LogMutable, Verbose, TEXT("Started UpdateSkeletalMesh Async. of Instance %d with priority %d at dist %f bIsPlayerOrNearIt=%d, frame=%d"), InstanceId, static_cast<int32>(Priority), Distance, bIsPlayerOrNearIt, GFrameNumber);				
+		}
+
+		UpdateDescriptorRuntimeHash = FDescriptorRuntimeHash(Descriptor);
 
 #if WITH_EDITOR
-			if (!UCustomizableObjectSystem::GetInstance()->IsCompilationDisabled())
+		if (!UCustomizableObjectSystem::GetInstance()->IsCompilationDisabled())
+		{
+			if (SkeletalMeshStatus != ESkeletalMeshState::AsyncUpdatePending)
 			{
-				if (CustomizableInstance.SkeletalMeshStatus != ESkeletalMeshState::AsyncUpdatePending)
-				{
-					CustomizableInstance.PreUpdateSkeletalMeshStatus = CustomizableInstance.SkeletalMeshStatus;
-				}
+				PreUpdateSkeletalMeshStatus = SkeletalMeshStatus;
 			}
+		}
 #endif
-
-			SetCOInstanceFlags(Generated); // Will be done in UpdateSkeletalMesh_PostBeginUpdate
-			
-			// Do not do work after calling InitUpdateSkeletalMesh. This function can optimize an update and fully complete it before even exiting its scope.
-			System->GetPrivate()->InitUpdateSkeletalMesh(CustomizableInstance, Priority);
+		
+		if (bShouldUpdateLODs)
+		{
+			UE_LOG(LogMutable, Verbose, TEXT("LOD change: %d, %d -> %d, %d"), GetCurrentMinLOD(), GetCurrentMaxLOD(), GetMinLODToLoad(), GetMinLODToLoad());
 		}
+		
+		GetPrivate()->SetCOInstanceFlags(Generated); // Will be done in UpdateSkeletalMesh_PostBeginUpdate
 
-		ClearCOInstanceFlags(PendingLODsUpdate);
+		// Do not do work after calling InitUpdateSkeletalMesh. This function can optimize an update and fully complete it before even exiting its scope.
+		System->GetPrivate()->InitUpdateSkeletalMesh(*this, Priority, UpdateCallback);
+		break;
 	}
-	else
+
+	case EUpdateRequired::Discard:
 	{
-		if (HasCOInstanceFlags(Generated) && !HasCOInstanceFlags(Updating))
-		{
-			System->GetPrivate()->InitDiscardResourcesSkeletalMesh(&CustomizableInstance);
-			ClearCOInstanceFlags(PendingLODsUpdate);
-		}
+		System->GetPrivate()->InitDiscardResourcesSkeletalMesh(this);
+		
+		FinishUpdateGlobal(this, EUpdateResult::ErrorDiscarded, UpdateCallback);
+		break;
+	}
+
+	default:
+		check(false); // Case not implemented.
 	}
 }
 
 
-void UCustomizableInstancePrivateData::TickUpdateCloseCustomizableObjects( UCustomizableObjectInstance& Public )
+void UCustomizableInstancePrivateData::TickUpdateCloseCustomizableObjects(UCustomizableObjectInstance& Public)
 {
-	DoUpdateSkeletalMesh(Public, true, true, false);
+	const EUpdateRequired UpdateRequired = Public.IsUpdateRequired(true, true, false);
+	if (UpdateRequired != EUpdateRequired::NoUpdate) // Since this is done in the tick, avoid starting an update that we know for sure that would not be performed. Once started it has some performance implications that we want to avoid.
+	{
+		Public.DoUpdateSkeletalMesh(true, true, false, false, &UpdateRequired, nullptr);
+	}
+
+	ClearCOInstanceFlags(PendingLODsUpdate | PendingLODsDowngrade);
 }
 
 
@@ -894,7 +934,7 @@ void UCustomizableInstancePrivateData::UpdateInstanceIfNotGenerated(UCustomizabl
 {
 	if (!HasCOInstanceFlags(Generated))
 	{
-		DoUpdateSkeletalMesh(Public, false, true);
+		Public.DoUpdateSkeletalMesh(false, true, false, false, nullptr, nullptr);
 	}
 }
 
@@ -2926,30 +2966,19 @@ void SetTexturePropertiesFromMutableImageProps(UTexture2D* Texture, const FMutab
 }
 
 
-void UCustomizableObjectInstance::Updated(EUpdateResult Result, const FDescriptorRuntimeHash& InUpdatedHash)
+void UCustomizableObjectInstance::FinishUpdate(EUpdateResult UpdateResult, const FDescriptorRuntimeHash& InUpdatedHash)
 {
-	if (Result == EUpdateResult::Success)
+	if (UpdateResult == EUpdateResult::Success)
 	{
 		DescriptorRuntimeHash = InUpdatedHash;
 	}
 
-	// Call Customizable Skeletal Components updated callbacks.
-	for (TObjectIterator<UCustomizableSkeletalComponent> It; It; ++It)
-	{
-		if (const UCustomizableSkeletalComponent* CustomizableSkeletalComponent = *It;
-			CustomizableSkeletalComponent &&
-			CustomizableSkeletalComponent->CustomizableObjectInstance == this)
-		{
-			CustomizableSkeletalComponent->Updated(Result);
-		}
-	}
-
-	if (Result == EUpdateResult::Success)
+	if (UpdateResult == EUpdateResult::Success)
 	{
 		// Call Instance updated (this) callbacks.
 		UpdatedDelegate.Broadcast(this);
 		UpdatedNativeDelegate.Broadcast(this);
-	}	
+	}
 }
 
 
@@ -2965,6 +2994,13 @@ FDescriptorRuntimeHash UCustomizableObjectInstance::GetUpdateDescriptorRuntimeHa
 }
 
 
+UCustomizableInstancePrivateData* UCustomizableObjectInstance::GetPrivate() const
+{ 
+	check(PrivateData); // Currently this is initialized in the constructor so we expect it always to exist.
+	return PrivateData; 
+}
+
+
 void UCustomizableObjectInstance::CommitMinMaxLOD()
 {
 	CurrentMinLOD = Descriptor.MinLOD;
diff --git a/Engine/Plugins/Experimental/Mutable/Source/CustomizableObject/Private/MuCO/CustomizableObjectSystem.cpp b/Engine/Plugins/Experimental/Mutable/Source/CustomizableObject/Private/MuCO/CustomizableObjectSystem.cpp
--- a/Engine/Plugins/Experimental/Mutable/Source/CustomizableObject/Private/MuCO/CustomizableObjectSystem.cpp
+++ b/Engine/Plugins/Experimental/Mutable/Source/CustomizableObject/Private/MuCO/CustomizableObjectSystem.cpp
@@ -125,6 +125,10 @@ void FMutableQueue::Enqueue(const FMutableQueueElem& TaskToEnqueue)
 		if ((QueueElem.Operation->Type == FMutableOperation::EOperationType::Update || QueueElem.Operation->Type == FMutableOperation::EOperationType::Discard) &&
 			QueueElem.Operation->CustomizableObjectInstance == TaskToEnqueue.Operation->CustomizableObjectInstance)
 		{
+			UCustomizableObjectInstance* Instance = QueueElem.Operation->CustomizableObjectInstance.Get();
+			FInstanceUpdateDelegate* Callback = &QueueElem.Operation->UpdateCallback;
+			FinishUpdateGlobal(Instance, EUpdateResult::ErrorReplaced, Callback);
+		
 			QueueElem.Operation = TaskToEnqueue.Operation;
 			QueueElem.PriorityType = FMath::Min(QueueElem.PriorityType, TaskToEnqueue.PriorityType);
 			QueueElem.Priority = FMath::Min(QueueElem.Priority, TaskToEnqueue.Priority);
@@ -745,92 +749,105 @@ FAutoConsoleVariableRef CVarMaxTextureSizeToGenerate(
 		"If a texture doesn't have small enough mips, mip 0 will be the last mip available."));
 
 
-/** Update the given Instance Skeletal Meshes and call its callbacks. */
-void UpdateSkeletalMesh(UCustomizableObjectInstance* CustomizableObjectInstance, const FDescriptorRuntimeHash& UpdatedDescriptorRuntimeHash)
+void FinishUpdateGlobal(UCustomizableObjectInstance* Instance, EUpdateResult UpdateResult, FInstanceUpdateDelegate* UpdateCallback, const FDescriptorRuntimeHash InUpdatedHash)
 {
-	// This used to be CustomizableObjectInstance::UpdateSkeletalMesh_PostBeginUpdate3
+	// Callbacks. Must be done at the end.
+	if (Instance)
 	{
-		MUTABLE_CPUPROFILER_SCOPE(UpdateSkeletalMesh_PostBeginUpdate3);
-		check(CustomizableObjectInstance != nullptr);
+		Instance->FinishUpdate(UpdateResult, InUpdatedHash);			
+	}
 
-#if !WITH_EDITOR
-		for (int32 ComponentIndex = 0; ComponentIndex < CustomizableObjectInstance->SkeletalMeshes.Num(); ++ComponentIndex)
+	if (UpdateResult == EUpdateResult::Success)
+	{
+		// Call Customizable Skeletal Components updated callbacks.
+		for (TObjectIterator<UCustomizableSkeletalComponent> It; It; ++It) // Since iterating objects is expensive, for now CustomizableSkeletalComponent does not have a FinishUpdate function.
 		{
-			if (CustomizableObjectInstance->SkeletalMeshes[ComponentIndex])
+			if (const UCustomizableSkeletalComponent* CustomizableSkeletalComponent = *It;
+				CustomizableSkeletalComponent &&
+				CustomizableSkeletalComponent->CustomizableObjectInstance == Instance)
 			{
-				CustomizableObjectInstance->SkeletalMeshes[ComponentIndex]->RebuildSocketMap();
+				CustomizableSkeletalComponent->Callbacks();
 			}
 		}
-#endif
+	}
 
-		UCustomizableInstancePrivateData* CustomizableObjectInstancePrivateData = CustomizableObjectInstance->GetPrivate();
-		check(CustomizableObjectInstancePrivateData != nullptr);
-		for (TObjectIterator<UCustomizableSkeletalComponent> It; It; ++It)
-		{
-			UCustomizableSkeletalComponent* CustomizableSkeletalComponent = *It;
+	if (UpdateCallback)
+	{
+		FUpdateContext Context;
+		Context.UpdateResult = UpdateResult;
+		
+		UpdateCallback->ExecuteIfBound(Context);
+	}
+}
 
-			if (CustomizableSkeletalComponent &&
-				(CustomizableSkeletalComponent->CustomizableObjectInstance == CustomizableObjectInstance) &&
-				CustomizableObjectInstance->SkeletalMeshes.IsValidIndex(CustomizableSkeletalComponent->ComponentIndex)
-			   )
-			{
-				MUTABLE_CPUPROFILER_SCOPE(UpdateSkeletalMesh_SetSkeletalMesh);
 
-				const bool bIsCreatingSkeletalMesh = CustomizableObjectInstancePrivateData->HasCOInstanceFlags(CreatingSkeletalMesh); //TODO MTBL-391: Review
-				CustomizableSkeletalComponent->SetSkeletalMesh(CustomizableObjectInstance->SkeletalMeshes[CustomizableSkeletalComponent->ComponentIndex], false, bIsCreatingSkeletalMesh);
+/** Update the given Instance Skeletal Meshes and call its callbacks. */
+void UpdateSkeletalMesh(UCustomizableObjectInstance& CustomizableObjectInstance, const FDescriptorRuntimeHash& UpdatedDescriptorRuntimeHash, EUpdateResult UpdateResult, FInstanceUpdateDelegate* UpdateCallback)
+{
+	MUTABLE_CPUPROFILER_SCOPE(UpdateSkeletalMesh);
 
-				if (CustomizableObjectInstancePrivateData->HasCOInstanceFlags(ReplacePhysicsAssets))
-				{
-					CustomizableSkeletalComponent->SetPhysicsAsset(
-						CustomizableObjectInstance->SkeletalMeshes[CustomizableSkeletalComponent->ComponentIndex] ? 
-						CustomizableObjectInstance->SkeletalMeshes[CustomizableSkeletalComponent->ComponentIndex]->GetPhysicsAsset() : nullptr);
-				}
-			}
+	for (int32 ComponentIndex = 0; ComponentIndex < CustomizableObjectInstance.SkeletalMeshes.Num(); ++ComponentIndex)
+	{
+		if (TObjectPtr<USkeletalMesh> SkeletalMesh = CustomizableObjectInstance.SkeletalMeshes[ComponentIndex])
+		{
+#if WITH_EDITOR
+			FUnrealBakeHelpers::BakeHelper_RegenerateImportedModel(SkeletalMesh);
+#else
+			SkeletalMesh->RebuildSocketMap();
+#endif
 		}
+	}
 
-		CustomizableObjectInstancePrivateData->SetCOInstanceFlags(Generated);
-		CustomizableObjectInstancePrivateData->ClearCOInstanceFlags(CreatingSkeletalMesh);
-
-		CustomizableObjectInstance->bEditorPropertyChanged = false;
+	UCustomizableInstancePrivateData* CustomizableObjectInstancePrivateData = CustomizableObjectInstance.GetPrivate();
+	check(CustomizableObjectInstancePrivateData != nullptr);
+	for (TObjectIterator<UCustomizableSkeletalComponent> It; It; ++It)
+	{
+		UCustomizableSkeletalComponent* CustomizableSkeletalComponent = *It;
 
+		if (CustomizableSkeletalComponent &&
+			(CustomizableSkeletalComponent->CustomizableObjectInstance == &CustomizableObjectInstance) &&
+			CustomizableObjectInstance.SkeletalMeshes.IsValidIndex(CustomizableSkeletalComponent->ComponentIndex)
+		   )
 		{
-			MUTABLE_CPUPROFILER_SCOPE(UpdateSkeletalMesh_UpdatedDelegate);
-
-			EUpdateResult State = CustomizableObjectInstance->SkeletalMeshStatus == ESkeletalMeshState::Correct ? EUpdateResult::Success : EUpdateResult::Error;
-			CustomizableObjectInstance->Updated(State, UpdatedDescriptorRuntimeHash);
+			MUTABLE_CPUPROFILER_SCOPE(UpdateSkeletalMesh_SetSkeletalMesh);
 
+			const bool bIsCreatingSkeletalMesh = CustomizableObjectInstancePrivateData->HasCOInstanceFlags(CreatingSkeletalMesh); //TODO MTBL-391: Review
+			CustomizableSkeletalComponent->SetSkeletalMesh(CustomizableObjectInstance.SkeletalMeshes[CustomizableSkeletalComponent->ComponentIndex], false, bIsCreatingSkeletalMesh);
 
-#if WITH_EDITOR
-			// \TODO: Review if we can avoid all this editor code here.
-			for (int32 MeshIndex = 0; MeshIndex < CustomizableObjectInstance->SkeletalMeshes.Num(); ++MeshIndex)
+			if (CustomizableObjectInstancePrivateData->HasCOInstanceFlags(ReplacePhysicsAssets))
 			{
-				if (USkeletalMesh* SkeletalMesh = CustomizableObjectInstance->GetSkeletalMesh(MeshIndex))
-				{
-					FUnrealBakeHelpers::BakeHelper_RegenerateImportedModel(SkeletalMesh);
-				}
+				CustomizableSkeletalComponent->SetPhysicsAsset(
+					CustomizableObjectInstance.SkeletalMeshes[CustomizableSkeletalComponent->ComponentIndex] ? 
+					CustomizableObjectInstance.SkeletalMeshes[CustomizableSkeletalComponent->ComponentIndex]->GetPhysicsAsset() : nullptr);
 			}
-#endif //WITH_EDITOR
 		}
+	}
 
-		check(UCustomizableObjectSystem::GetInstance() != nullptr);
-		FCustomizableObjectSystemPrivate* CustomizableObjectSystem = UCustomizableObjectSystem::GetInstance()->GetPrivate();
-		check(CustomizableObjectSystem != nullptr);
-		if (CustomizableObjectSystem->bReleaseTexturesImmediately)
-		{
-			FMutableResourceCache& Cache = CustomizableObjectSystem->GetObjectCache(CustomizableObjectInstance->GetCustomizableObject());
+	CustomizableObjectInstancePrivateData->SetCOInstanceFlags(Generated);
+	CustomizableObjectInstancePrivateData->ClearCOInstanceFlags(CreatingSkeletalMesh);
 
-			for (TPair<uint32, FGeneratedTexture>& Item : CustomizableObjectInstancePrivateData->TexturesToRelease)
-			{
-				UCustomizableInstancePrivateData::ReleaseMutableTexture(Item.Value.Id, Item.Value.Texture, Cache);
-			}
+	CustomizableObjectInstance.bEditorPropertyChanged = false;
+
+	check(UCustomizableObjectSystem::GetInstance() != nullptr);
+	FCustomizableObjectSystemPrivate* CustomizableObjectSystem = UCustomizableObjectSystem::GetInstance()->GetPrivate();
+	check(CustomizableObjectSystem != nullptr);
+	if (CustomizableObjectSystem->bReleaseTexturesImmediately)
+	{
+		FMutableResourceCache& Cache = CustomizableObjectSystem->GetObjectCache(CustomizableObjectInstance.GetCustomizableObject());
 
-			CustomizableObjectInstancePrivateData->TexturesToRelease.Empty();
+		for (TPair<uint32, FGeneratedTexture>& Item : CustomizableObjectInstancePrivateData->TexturesToRelease)
+		{
+			UCustomizableInstancePrivateData::ReleaseMutableTexture(Item.Value.Id, Item.Value.Texture, Cache);
 		}
+
+		CustomizableObjectInstancePrivateData->TexturesToRelease.Empty();
 	}
+
+	FinishUpdateGlobal(&CustomizableObjectInstance, UpdateResult, UpdateCallback, UpdatedDescriptorRuntimeHash);
 }
 
 
-void FCustomizableObjectSystemPrivate::InitUpdateSkeletalMesh(UCustomizableObjectInstance& Instance, FMutableQueueElem::EQueuePriorityType Priority)
+void FCustomizableObjectSystemPrivate::InitUpdateSkeletalMesh(UCustomizableObjectInstance& Instance, FMutableQueueElem::EQueuePriorityType Priority, FInstanceUpdateDelegate* UpdateCallback)
 {
 	MUTABLE_CPUPROFILER_SCOPE(FCustomizableObjectSystemPrivate::InitUpdateSkeletalMesh);
 
@@ -858,8 +875,6 @@ void FCustomizableObjectSystemPrivate::InitUpdateSkeletalMesh(UCustomizableObjec
 		MipsToSkip = 255; // This means skip all possible mips until only UTexture::GetStaticMinTextureResidentMipCount() are left
 	}
 	
-	const TSharedPtr<FMutableOperation> Operation = MakeShared<FMutableOperation>(FMutableOperation::CreateInstanceUpdate(&Instance, bNeverStream, MipsToSkip));
-
 	const FDescriptorRuntimeHash UpdateDescriptorHash = Instance.GetUpdateDescriptorRuntimeHash();
 
 	if (const FMutableQueueElem* QueueElem = MutableOperationQueue.Get(&Instance))
@@ -869,6 +884,7 @@ void FCustomizableObjectSystemPrivate::InitUpdateSkeletalMesh(UCustomizableObjec
 			QueuedOperation->Type == FMutableOperation::EOperationType::Update &&
 			UpdateDescriptorHash.IsSubset(QueuedOperation->InstanceDescriptorRuntimeHash))
 		{
+			FinishUpdateGlobal(&Instance, EUpdateResult::ErrorOptimized, UpdateCallback);			
 			return; // The the requested update is equal to the last enqueued update.
 		}
 	}
@@ -877,6 +893,7 @@ void FCustomizableObjectSystemPrivate::InitUpdateSkeletalMesh(UCustomizableObjec
 		CurrentMutableOperation->Type == FMutableOperation::EOperationType::Update &&
 		UpdateDescriptorHash.IsSubset(CurrentMutableOperation->InstanceDescriptorRuntimeHash))
 	{
+		FinishUpdateGlobal(&Instance, EUpdateResult::ErrorOptimized, UpdateCallback);			
 		return; // The requested update is equal to the running update.
 	}
 	
@@ -887,11 +904,13 @@ void FCustomizableObjectSystemPrivate::InitUpdateSkeletalMesh(UCustomizableObjec
 	if (UpdateDescriptorHash.IsSubset(Instance.GetDescriptorRuntimeHash()))
 	{
 		Instance.SkeletalMeshStatus = ESkeletalMeshState::Correct; // TODO GMT MTBL-1033 should not be here. Move to UCustomizableObjectInstance::Updated
-		UpdateSkeletalMesh(&Instance, Instance.GetDescriptorRuntimeHash());
+		UpdateSkeletalMesh(Instance, Instance.GetDescriptorRuntimeHash(), EUpdateResult::Success, UpdateCallback);
 	}
 	else
 	{
 		Instance.SkeletalMeshStatus = ESkeletalMeshState::AsyncUpdatePending;
+
+		const TSharedPtr<FMutableOperation> Operation = MakeShared<FMutableOperation>(FMutableOperation::CreateInstanceUpdate(&Instance, bNeverStream, MipsToSkip, UpdateCallback));
 		MutableOperationQueue.Enqueue(FMutableQueueElem::Create(Operation, Priority, Instance.GetPrivate()->MinSquareDistFromComponentToPlayer));
 	}
 }
@@ -953,7 +972,6 @@ void UCustomizableObjectSystem::ClearResourceCacheProtected()
 bool UCustomizableObjectSystem::LockObject(const class UCustomizableObject* InObject)
 {
 	check(InObject != nullptr);
-	check(InObject->GetPrivate() != nullptr);
 	check(!InObject->GetPrivate()->bLocked);
 	check(IsInGameThread() && !IsInParallelGameThread());
 
@@ -1015,7 +1033,6 @@ bool UCustomizableObjectSystem::LockObject(const class UCustomizableObject* InOb
 void UCustomizableObjectSystem::UnlockObject(const class UCustomizableObject* Obj)
 {
 	check(Obj != nullptr);
-	check(Obj->GetPrivate() != nullptr);
 	check(Obj->GetPrivate()->bLocked);
 	check(IsInGameThread() && !IsInParallelGameThread());
 	Obj->GetPrivate()->bLocked = false;
@@ -1807,6 +1824,7 @@ namespace impl
 		UCustomizableObjectSystem* System = UCustomizableObjectSystem::GetInstance();
 		if (!System || !System->IsValidLowLevel() || System->HasAnyFlags(RF_BeginDestroyed))
 		{
+			FinishUpdateGlobal(CustomizableObjectInstancePtr.Get(), EUpdateResult::Error, &OperationData->UpdateCallback);
 			return;
 		}
 
@@ -1816,6 +1834,7 @@ namespace impl
 		if (!CustomizableObjectInstance || !CustomizableObjectInstance->IsValidLowLevel() )
 		{
 			System->ClearCurrentMutableOperation();
+			FinishUpdateGlobal(CustomizableObjectInstance, EUpdateResult::Error, &OperationData->UpdateCallback);
 			return;
 		}
 
@@ -1824,7 +1843,7 @@ namespace impl
 
 		// Actual work
 		// TODO MTBL-391: Review This hotfix
-		UpdateSkeletalMesh(CustomizableObjectInstance, CustomizableObjectSystemPrivateData->CurrentMutableOperation->InstanceDescriptorRuntimeHash);
+		UpdateSkeletalMesh(*CustomizableObjectInstance, CustomizableObjectSystemPrivateData->CurrentMutableOperation->InstanceDescriptorRuntimeHash, OperationData->UpdateResult, &OperationData->UpdateCallback);
 
 		// TODO: T2927
 		if (LogBenchmarkUtil::isLoggingActive())
@@ -1856,6 +1875,7 @@ namespace impl
 		UCustomizableObjectSystem* System = UCustomizableObjectSystem::GetInstance();
 		if (!System || !System->IsValidLowLevel() || System->HasAnyFlags(RF_BeginDestroyed))
 		{
+			FinishUpdateGlobal(CustomizableObjectInstancePtr.Get(), EUpdateResult::Error, &OperationData->UpdateCallback);
 			return;
 		}
 
@@ -1863,20 +1883,12 @@ namespace impl
 
 		UCustomizableObjectInstance* CustomizableObjectInstance = CustomizableObjectInstancePtr.Get();
 
-		// TODO: Review checks.
-		bool bCancel = false;
-		if (!CustomizableObjectInstance || !CustomizableObjectInstance->IsValidLowLevel())
-		{
-			bCancel = true;
-		}
-
 		// Actual work
-
-
-		if (!bCancel)
+		// TODO: Review checks.
+		const bool bInstanceInvalid = !CustomizableObjectInstance || !CustomizableObjectInstance->IsValidLowLevel();
+		if (!bInstanceInvalid)
 		{
 			UCustomizableInstancePrivateData* CustomizableInstancePrivateData = CustomizableObjectInstance->GetPrivate();
-			check(CustomizableInstancePrivateData != nullptr);
 
 			// Process the pending texture coverage queries
 			{
@@ -1942,7 +1954,7 @@ namespace impl
 					}
 				}
 			} // END - Process texture coverage queries
-		} // if (!bCancel)
+		} // if (!bInstanceValid)
 
 		FCustomizableObjectSystemPrivate* CustomizableObjectSystemPrivateData = System->GetPrivate();
 		check(CustomizableObjectSystemPrivateData != nullptr);
@@ -1960,7 +1972,7 @@ namespace impl
 
 		// Next Task: Release Platform Data
 		//-------------------------------------------------------------
-		if (!bCancel)
+		if (!bInstanceInvalid)
 		{
 			TSharedPtr<FMutableReleasePlatformOperationData> ReleaseOperationData = MakeShared<FMutableReleasePlatformOperationData>();
 			check(ReleaseOperationData);
@@ -1993,6 +2005,10 @@ namespace impl
 					{}
 				});
 		}
+		else
+		{
+			FinishUpdateGlobal(CustomizableObjectInstance, EUpdateResult::Error, &OperationData->UpdateCallback);
+		}
 	}
 
 
@@ -2002,6 +2018,7 @@ namespace impl
 		MUTABLE_CPUPROFILER_SCOPE(Task_Game_ConvertResources)
 
 		check(IsInGameThread());
+		check(OperationData);
 
 		UCustomizableObjectSystem* System = UCustomizableObjectSystem::GetInstance();
 		if (!System)
@@ -2014,6 +2031,7 @@ namespace impl
 		if (!ObjectInstance)
 		{
 			System->ClearCurrentMutableOperation();
+			FinishUpdateGlobal(ObjectInstance, EUpdateResult::Error, &OperationData->UpdateCallback);
 			return;
 		}
 
@@ -2035,6 +2053,7 @@ namespace impl
 		if (!CustomizableObject)
 		{
 			System->ClearCurrentMutableOperation();
+			FinishUpdateGlobal(CustomizableObjectInstancePtr.Get(), EUpdateResult::Error, &OperationData->UpdateCallback);
 			return;
 		}
 
@@ -2099,7 +2118,6 @@ namespace impl
 #endif
 		{
 			// Task inputs
-			check(CustomizableObject->GetPrivate() != nullptr);
 			TSharedPtr<mu::Model, ESPMode::ThreadSafe> Model = CustomizableObject->GetPrivate()->GetModel();
 			int32 State = ObjectInstance->GetState();
 
@@ -2205,21 +2223,21 @@ namespace impl
 		if (!Operation->CustomizableObjectInstance.IsValid() || !Operation->CustomizableObjectInstance->IsValidLowLevel()) // Only start if it hasn't been already destroyed (i.e. GC after finish PIE)
 		{
 			System->ClearCurrentMutableOperation();
+			FinishUpdateGlobal(nullptr, EUpdateResult::Error, &Operation->UpdateCallback);
 			return;
 		}
 
 		TObjectPtr<UCustomizableObjectInstance> CandidateInstance = Operation->CustomizableObjectInstance.Get();
 		
 		UCustomizableInstancePrivateData* CandidateInstancePrivateData = CandidateInstance->GetPrivate();
-		check(CandidateInstancePrivateData != nullptr);
-
 		if (!CandidateInstancePrivateData)
 		{
 			System->ClearCurrentMutableOperation();
+			FinishUpdateGlobal(nullptr, EUpdateResult::Error, &Operation->UpdateCallback);
 			return;
 		}
-
-		if (CandidateInstancePrivateData && CandidateInstancePrivateData->HasCOInstanceFlags(PendingLODsUpdate))
+		
+		if (CandidateInstancePrivateData->HasCOInstanceFlags(PendingLODsUpdate))
 		{
 			CandidateInstancePrivateData->ClearCOInstanceFlags(PendingLODsUpdate);
 			// TODO: Is anything needed for this now?
@@ -2236,7 +2254,6 @@ namespace impl
 		}
 		else
 		{
-			check(CustomizableObject->GetPrivate());
 			if (CustomizableObject->GetPrivate()->bLocked)
 			{
 				bCancel = true;
@@ -2258,7 +2275,7 @@ namespace impl
 		if (Operation->InstanceDescriptorRuntimeHash.IsSubset(CandidateInstance->GetDescriptorRuntimeHash()))
 		{
 			CandidateInstance->SkeletalMeshStatus = ESkeletalMeshState::Correct;
-			UpdateSkeletalMesh(CandidateInstance, CandidateInstance->GetDescriptorRuntimeHash());
+			UpdateSkeletalMesh(*CandidateInstance, CandidateInstance->GetDescriptorRuntimeHash(), EUpdateResult::Success, &Operation->UpdateCallback);
 			bCancel = true;
 		}
 
@@ -2273,6 +2290,8 @@ namespace impl
 			CandidateInstancePrivateData->ClearCOInstanceFlags(Updating);
 
 			System->ClearCurrentMutableOperation();
+
+			FinishUpdateGlobal(CandidateInstance, EUpdateResult::Error, &Operation->UpdateCallback);
 			return;
 		}
 
@@ -2350,6 +2369,8 @@ namespace impl
 		CurrentOperationData->MipsToSkip = Operation->MipsToSkip;
 		CurrentOperationData->MutableParameters = Parameters;
 		CurrentOperationData->State = CandidateInstance->GetState();
+		CurrentOperationData->UpdateResult = EUpdateResult::Success;
+		CurrentOperationData->UpdateCallback = Operation->UpdateCallback;
 		CustomizableObject->GetLowPriorityTextureNames(CurrentOperationData->LowPriorityTextures);
 
 		if (System->IsOnlyGenerateRequestedLODsEnabled() && System->CurrentInstanceLODManagement->IsOnlyGenerateRequestedLODLevelsEnabled() && !CandidateInstancePrivateData->HasCOInstanceFlags(ForceGenerateAllLODs))
@@ -2369,7 +2390,6 @@ namespace impl
 			TSharedPtr<FMutableOperation> CurrentMutableOperation = SystemPrivateData->CurrentMutableOperation;
 			check(CurrentMutableOperation);
 			bool bBuildParameterDecorations = CurrentMutableOperation->IsBuildParameterDecorations();
-			check(CustomizableObject->GetPrivate() != nullptr);
 			TSharedPtr<mu::Model, ESPMode::ThreadSafe> Model = CustomizableObject->GetPrivate()->GetModel();
 			int32 State = CustomizableObjectInstancePtr->GetState();
 
@@ -2404,7 +2424,6 @@ namespace impl
 				});
 		}
 	}
-
 } // namespace impl
 
 
@@ -2570,7 +2589,6 @@ bool UCustomizableObjectSystem::Tick(float DeltaTime)
 			if (IsValidChecked(*CustomizableObjectInstance) && CustomizableObjectInstance->GetPrivate())
 			{
 				UCustomizableInstancePrivateData* ObjectInstancePrivateData = CustomizableObjectInstance->GetPrivate();
-				check(ObjectInstancePrivateData != nullptr);
 
 				if (ObjectInstancePrivateData->HasCOInstanceFlags(UsedByComponentInPlay))
 				{
@@ -2702,7 +2720,7 @@ UDefaultImageProvider& UCustomizableObjectSystem::GetOrCreateDefaultImageProvide
 }
 
 
-FMutableOperation FMutableOperation::CreateInstanceUpdate(UCustomizableObjectInstance* InCustomizableObjectInstance, bool bInNeverStream, int32 InMipsToSkip)
+FMutableOperation FMutableOperation::CreateInstanceUpdate(UCustomizableObjectInstance* InCustomizableObjectInstance, bool bInNeverStream, int32 InMipsToSkip, const FInstanceUpdateDelegate* UpdateCallback)
 {
 	check(InCustomizableObjectInstance != nullptr);
 	check(InCustomizableObjectInstance->GetPrivate() != nullptr);
@@ -2717,7 +2735,11 @@ FMutableOperation FMutableOperation::CreateInstanceUpdate(UCustomizableObjectIns
 	Op.bStarted = false;
 	Op.bBuildParameterDecorations = InCustomizableObjectInstance->GetBuildParameterDecorations();
 	Op.Parameters = InCustomizableObjectInstance->GetPrivate()->GetParameters(InCustomizableObjectInstance);
-
+	if (UpdateCallback)
+	{
+		Op.UpdateCallback = *UpdateCallback;		
+	}
+	
 	InCustomizableObjectInstance->GetCustomizableObject()->ApplyStateForcedValuesToParameters(InCustomizableObjectInstance->GetState(), Op.Parameters.get());
 
 	if (!Op.Parameters)
@@ -2755,35 +2777,6 @@ FMutableOperation FMutableOperation::CreateInstanceIDRelease(mu::Instance::ID ID
 }
 
 
-void FMutableOperation::MutableIsDisabledCase()
-{
-	check(CustomizableObjectInstance != nullptr);
-	check(CustomizableObjectInstance->GetPrivate() != nullptr);
-	check(CustomizableObjectInstance->GetCustomizableObject() != nullptr);
-
-	for (TObjectIterator<UCustomizableSkeletalComponent> It; It; ++It)
-	{
-		UCustomizableSkeletalComponent* CustomizableSkeletalComponent = *It;
-
-		if (CustomizableSkeletalComponent && CustomizableSkeletalComponent->CustomizableObjectInstance == CustomizableObjectInstance.Get())
-		{
-			CustomizableSkeletalComponent->SetSkeletalMesh(CustomizableObjectInstance->GetCustomizableObject()->GetRefSkeletalMesh(CustomizableSkeletalComponent->ComponentIndex), false);
-		}
-	}
-
-	CustomizableObjectInstance->GetPrivate()->SetCOInstanceFlags(Generated);
-	CustomizableObjectInstance->GetPrivate()->ClearCOInstanceFlags(CreatingSkeletalMesh);
-
-	if (CustomizableObjectInstance->bEditorPropertyChanged)
-	{
-		CustomizableObjectInstance->bEditorPropertyChanged = false;
-	}
-
-	// We must invalidate the current DescriptorRuntimeHash, since we're discarding everything.
-	CustomizableObjectInstance->Updated(EUpdateResult::Success, FDescriptorRuntimeHash());
-}
-
-
 int32 UCustomizableObjectSystem::GetNumInstances() const
 {
 	check(Private != nullptr);
diff --git a/Engine/Plugins/Experimental/Mutable/Source/CustomizableObject/Private/MuCO/CustomizableObjectSystemPrivate.h b/Engine/Plugins/Experimental/Mutable/Source/CustomizableObject/Private/MuCO/CustomizableObjectSystemPrivate.h
--- a/Engine/Plugins/Experimental/Mutable/Source/CustomizableObject/Private/MuCO/CustomizableObjectSystemPrivate.h
+++ b/Engine/Plugins/Experimental/Mutable/Source/CustomizableObject/Private/MuCO/CustomizableObjectSystemPrivate.h
@@ -42,7 +42,7 @@ class FMutableOperation
 public:
 
 
-	static FMutableOperation CreateInstanceUpdate(UCustomizableObjectInstance* COInstance, bool bInNeverStream, int32 MipsToSkip);
+	static FMutableOperation CreateInstanceUpdate(UCustomizableObjectInstance* COInstance, bool bInNeverStream, int32 MipsToSkip, const FInstanceUpdateDelegate* UpdateCallback);
 	static FMutableOperation CreateInstanceDiscard(UCustomizableObjectInstance* COInstance);
 	static FMutableOperation CreateInstanceIDRelease(mu::Instance::ID);
 
@@ -91,16 +91,13 @@ public:
 	/** Only used in the IDRelease operation type */
 	mu::Instance::ID IDToRelease;
 
-	//!
+	FInstanceUpdateDelegate UpdateCallback;
+	
 	bool IsBuildParameterDecorations() const
 	{
 		return bBuildParameterDecorations;
 	}
 
-	// \TODO: This should be intercepted earlier.
-	// In case Mutable Compilation has been disabled, do all the corresponding steps to assign the reference mesh of the customizable object
-	void MutableIsDisabledCase();
-
 	/** Read-only access to the mutable instance parameters for this operation. */
 	mu::ParametersPtrConst GetParameters() const
 	{
@@ -520,6 +517,9 @@ struct FMutableOperationData
 	mu::Ptr<const mu::Parameters> MutableParameters;
 	int32 State = 0;
 
+	EUpdateResult UpdateResult;
+	FInstanceUpdateDelegate UpdateCallback;
+
 #if WITH_EDITOR
 	/** Used for profiling in the editor. */
 	uint32 MutableRuntimeCycles = 0;
@@ -789,7 +789,7 @@ public:
 
 
 	// Init the async Skeletal Mesh creation/update
-	void InitUpdateSkeletalMesh(UCustomizableObjectInstance& Public, FMutableQueueElem::EQueuePriorityType Priority);
+	void InitUpdateSkeletalMesh(UCustomizableObjectInstance& Public, FMutableQueueElem::EQueuePriorityType Priority, FInstanceUpdateDelegate* UpdateCallback = nullptr);
 		
 	// Init an async and safe release of the UE and Mutable resources used by the instance without actually destroying the instance, for example if it's very far away
 	void InitDiscardResourcesSkeletalMesh(UCustomizableObjectInstance* InCustomizableObjectInstance);
diff --git a/Engine/Plugins/Experimental/Mutable/Source/CustomizableObject/Private/MuCO/CustomizableSkeletalComponent.cpp b/Engine/Plugins/Experimental/Mutable/Source/CustomizableObject/Private/MuCO/CustomizableSkeletalComponent.cpp
--- a/Engine/Plugins/Experimental/Mutable/Source/CustomizableObject/Private/MuCO/CustomizableSkeletalComponent.cpp
+++ b/Engine/Plugins/Experimental/Mutable/Source/CustomizableObject/Private/MuCO/CustomizableSkeletalComponent.cpp
@@ -24,12 +24,9 @@ UCustomizableSkeletalComponent::UCustomizableSkeletalComponent()
 }
 
 
-void UCustomizableSkeletalComponent::Updated(EUpdateResult Result) const
+void UCustomizableSkeletalComponent::Callbacks() const
 {
-	if (Result == EUpdateResult::Success)
-	{
-		UpdatedDelegate.ExecuteIfBound();
-	}	
+	UpdatedDelegate.ExecuteIfBound();
 }
 
 
@@ -130,6 +127,15 @@ void UCustomizableSkeletalComponent::UpdateSkeletalMeshAsync(bool bNeverSkipUpda
 }
 
 
+void UCustomizableSkeletalComponent::UpdateSkeletalMeshAsyncResult(FInstanceUpdateDelegate Callback, bool bIgnoreCloseDist, bool bForceHighPriority)
+{
+	if (CustomizableObjectInstance)
+	{
+		CustomizableObjectInstance->UpdateSkeletalMeshAsyncResult(Callback, false, false);
+	}
+}
+
+
 void UCustomizableSkeletalComponent::BeginDestroy()
 {
 	Super::BeginDestroy();
diff --git a/Engine/Plugins/Experimental/Mutable/Source/CustomizableObjectEditor/Private/MuCOE/SCustomizableInstanceProperties.cpp b/Engine/Plugins/Experimental/Mutable/Source/CustomizableObjectEditor/Private/MuCOE/SCustomizableInstanceProperties.cpp
--- a/Engine/Plugins/Experimental/Mutable/Source/CustomizableObjectEditor/Private/MuCOE/SCustomizableInstanceProperties.cpp
+++ b/Engine/Plugins/Experimental/Mutable/Source/CustomizableObjectEditor/Private/MuCOE/SCustomizableInstanceProperties.cpp
@@ -1345,8 +1345,6 @@ void SCustomizableInstanceProperties::OnStateComboBoxSelectionChanged(TSharedPtr
 
 	// Non-continuous change: collect garbage.
 	GEngine->ForceGarbageCollection();
-
-	CustomInstance->InstanceUpdated = false;
 }
 
 
