[Mutable] Fix minor compilation error due to an invalid comment.

diff --git a/Engine/Plugins/Experimental/Mutable/Source/CustomizableObject/Classes/MuCO/CustomizableObjectInstance.h b/Engine/Plugins/Experimental/Mutable/Source/CustomizableObject/Classes/MuCO/CustomizableObjectInstance.h
--- a/Engine/Plugins/Experimental/Mutable/Source/CustomizableObject/Classes/MuCO/CustomizableObjectInstance.h
+++ b/Engine/Plugins/Experimental/Mutable/Source/CustomizableObject/Classes/MuCO/CustomizableObjectInstance.h
@@ -274,7 +274,7 @@ private:
 	
 	/** Update Skeletal Mesh asynchronously. Immersive function. Not exposed to the API.
 	 * Once the update reaches this function, the update has been considered started and must complete all the update flow.
-	/* Starting at this function, all Update code paths must end up in FinishUpdateGlobal!
+	 * Starting at this function, all Update code paths must end up in FinishUpdateGlobal!
 	 *
 	 * @param bIsCloseDistTick true iff called from the tick.
 	 */
