[Mutable] Fix NonUnity build errors.

#rnx

diff --git a/Engine/Plugins/Experimental/Mutable/Source/CustomizableObject/Public/MuCO/CustomizableObjectInstanceDescriptor.h b/Engine/Plugins/Experimental/Mutable/Source/CustomizableObject/Public/MuCO/CustomizableObjectInstanceDescriptor.h
--- a/Engine/Plugins/Experimental/Mutable/Source/CustomizableObject/Public/MuCO/CustomizableObjectInstanceDescriptor.h
+++ b/Engine/Plugins/Experimental/Mutable/Source/CustomizableObject/Public/MuCO/CustomizableObjectInstanceDescriptor.h
@@ -3,18 +3,12 @@
 #pragma once
 
 #include "MuCO/MultilayerProjector.h"
+#include "MuCO/CustomizableObjectParameterTypeDefinitions.h"
 
 #include "CustomizableObjectInstanceDescriptor.generated.h"
 
 class UTexture2D;
 enum class ECustomizableObjectProjectorType : uint8;
-struct FCustomizableObjectBoolParameterValue;
-struct FCustomizableObjectFloatParameterValue;
-struct FCustomizableObjectIntParameterValue;
-struct FCustomizableObjectProjector;
-struct FCustomizableObjectProjectorParameterValue;
-struct FCustomizableObjectTextureParameterValue;
-struct FCustomizableObjectVectorParameterValue;
 
 class FArchive;
 class UCustomizableInstancePrivateData;
