[Mutable] Fix Working Set not working unless the CustomizableObject was loaded.

#rnx

diff --git a/Engine/Plugins/Experimental/Mutable/Source/CustomizableObjectEditor/Private/MuCOE/CustomizableObjectCompiler.cpp b/Engine/Plugins/Experimental/Mutable/Source/CustomizableObjectEditor/Private/MuCOE/CustomizableObjectCompiler.cpp
--- a/Engine/Plugins/Experimental/Mutable/Source/CustomizableObjectEditor/Private/MuCOE/CustomizableObjectCompiler.cpp
+++ b/Engine/Plugins/Experimental/Mutable/Source/CustomizableObjectEditor/Private/MuCOE/CustomizableObjectCompiler.cpp
@@ -569,11 +569,11 @@ mu::NodeObjectPtr FCustomizableObjectCompiler::GenerateMutableRoot(
 			const int32 MaxIndex = Object->WorkingSet.Num();
 			for (int32 i = 0; i < MaxIndex; ++i)
 			{
-				if (Object->WorkingSet[i] != nullptr)
+				if (UCustomizableObject* WorkingSetObject = Object->WorkingSet[i].LoadSynchronous())
 				{
 					ArrayCustomizableObject.Reset();
 
-					if (!GetParentsUntilRoot(Object->WorkingSet[i].LoadSynchronous(), ArrayNodeObject, ArrayCustomizableObject))
+					if (!GetParentsUntilRoot(WorkingSetObject, ArrayNodeObject, ArrayCustomizableObject))
 					{
 						CompilerLog(LOCTEXT("NoReferenceMesh", "Error! Cycle detected in the Customizable Object hierarchy."), Root);
 						return nullptr;
