[mutable] Fixed crash when setting parameter texture to null in a COInstance Editor.

#rnx

diff --git a/Engine/Plugins/Experimental/Mutable/Source/CustomizableObject/Private/MuCO/CustomizableObjectInstance.cpp b/Engine/Plugins/Experimental/Mutable/Source/CustomizableObject/Private/MuCO/CustomizableObjectInstance.cpp
--- a/Engine/Plugins/Experimental/Mutable/Source/CustomizableObject/Private/MuCO/CustomizableObjectInstance.cpp
+++ b/Engine/Plugins/Experimental/Mutable/Source/CustomizableObject/Private/MuCO/CustomizableObjectInstance.cpp
@@ -245,7 +245,11 @@ void UCustomizableObjectInstance::PreEditChange(FProperty* PropertyAboutToChange
 			if (Texture)
 			{
 				const int32 TextureId = DefaultImageProvider.Get(Texture);
-				DefaultImageProvider.Keep(TextureId, false);
+
+				if (TextureId != INDEX_NONE)
+				{
+					DefaultImageProvider.Keep(TextureId, false);
+				}
 			}
 		}
 	}
