[Mutable] Removed CO dependency to MuR private.

- Created model default parameter functions.
- Exposed Model Rom Id funcitons.
- Exposed ImagePixelFormat and GetImageFormatData and ImageSwizzle functoins.
- Split serialization macros into declaration and definitions.
- Moved serialization macros bodies to their private source files.

#rnx

diff --git a/Engine/Plugins/Experimental/Mutable/Source/CustomizableObject/Classes/MuCO/CustomizableObject.h b/Engine/Plugins/Experimental/Mutable/Source/CustomizableObject/Classes/MuCO/CustomizableObject.h
--- a/Engine/Plugins/Experimental/Mutable/Source/CustomizableObject/Classes/MuCO/CustomizableObject.h
+++ b/Engine/Plugins/Experimental/Mutable/Source/CustomizableObject/Classes/MuCO/CustomizableObject.h
@@ -1335,18 +1335,6 @@ public:
 	UFUNCTION(BlueprintCallable, Category = CustomizableObject)
 	void UnloadMaskOutCache();
 
-private:
-	
-	/** Returns true or false if the parameter with name can be located and it has the type the caller is looking for. It will also
-	 * check if the model has been set to ensure access to it can take place at the calculated parameter index.
-	 * @param InParameterName The name of the parameter to look for.
-	 * @param InParameterType The type the parameter we are looking for we know has. If the name does not match this type this check will fail and return false.
-	 * @param OutParameterIndex The index of the parameter.
-	 * @return True if the parameter can be accessed for it's default values, false if it can not be accessed.
-	 */
-	bool CanDefaultParameterBeAccessed(const FString& InParameterName,const EMutableParameterType& InParameterType, int32& OutParameterIndex) const;
-
-public:
 	/** Get the default value of a parameter of type Float.
 	 * @param InParameterName The name of the Float parameter to get the default value of.
 	 * @return The default value of the provided parameter name.
diff --git a/Engine/Plugins/Experimental/Mutable/Source/CustomizableObject/CustomizableObject.Build.cs b/Engine/Plugins/Experimental/Mutable/Source/CustomizableObject/CustomizableObject.Build.cs
--- a/Engine/Plugins/Experimental/Mutable/Source/CustomizableObject/CustomizableObject.Build.cs
+++ b/Engine/Plugins/Experimental/Mutable/Source/CustomizableObject/CustomizableObject.Build.cs
@@ -51,10 +51,6 @@ public class CustomizableObject : ModuleRules
                 "TargetPlatform"
         });
 
-		PrivateIncludePaths.AddRange(new string[] {
-				"MutableRuntime/Private",
-			});
-
 		if (TargetRules.bBuildEditor == true)
         {
             PublicDependencyModuleNames.Add("UnrealEd");	// @todo api: Only public because of WITH_EDITOR
diff --git a/Engine/Plugins/Experimental/Mutable/Source/CustomizableObject/Private/MuCO/CustomizableObject.cpp b/Engine/Plugins/Experimental/Mutable/Source/CustomizableObject/Private/MuCO/CustomizableObject.cpp
--- a/Engine/Plugins/Experimental/Mutable/Source/CustomizableObject/Private/MuCO/CustomizableObject.cpp
+++ b/Engine/Plugins/Experimental/Mutable/Source/CustomizableObject/Private/MuCO/CustomizableObject.cpp
@@ -23,8 +23,7 @@
 #include "MuCO/MutableProjectorTypeUtils.h"
 #include "MuCO/UnrealMutableModelDiskStreamer.h"
 #include "MuCO/UnrealPortabilityHelpers.h"
-#include "MuR/ModelPrivate.h"
-#include "MuR/ParametersPrivate.h"
+#include "MuR/Model.h"
 #include "PhysicsEngine/PhysicsAsset.h"
 #include "Serialization/MemoryReader.h"
 #include "Serialization/MemoryWriter.h"
@@ -35,6 +34,8 @@
 #endif
 
 
+#include "MuR/Model.h"
+
 #include UE_INLINE_GENERATED_CPP_BY_NAME(CustomizableObject)
 
 class UMaterialInterface;
@@ -348,19 +349,17 @@ void UCustomizableObject::UpdateCompiledDataFromModel()
 	TSharedPtr<mu::Model, ESPMode::ThreadSafe> Model = PrivateData->GetModel();
 
 	// Generate a map that using the resource id tells the offset and size of the resource inside the bulk data
-	if(Model)
+	if (Model)
 	{
-		uint64_t Offset = 0;
-		uint32_t ResourceId = 0;
-		uint32_t ResourceSize = 0;
+		uint64 Offset = 0;
 
-		int32 NumStreamingFiles = Model->GetPrivate()->m_program.m_roms.Num();
+		const int32 NumStreamingFiles = Model->GetRomCount();
 		HashToStreamableBlock.Empty(NumStreamingFiles);
 
 		for (size_t FileIndex = 0; FileIndex < NumStreamingFiles; ++FileIndex)
 		{
-			ResourceId = Model->GetPrivate()->m_program.m_roms[FileIndex].Id;
-			ResourceSize = Model->GetPrivate()->m_program.m_roms[FileIndex].Size;
+			const uint32 ResourceId = Model->GetRomId(FileIndex);
+			const uint32 ResourceSize = Model->GetRomSize(FileIndex);
 
 			HashToStreamableBlock.Add(ResourceId, FMutableStreamableBlock({0, Offset, ResourceSize }));
 			Offset += ResourceSize;
@@ -1437,99 +1436,85 @@ void UCustomizableObject::UnloadMaskOutCache()
 }
 
 
-bool UCustomizableObject::CanDefaultParameterBeAccessed(const FString& InParameterName,
-	const EMutableParameterType& InParameterType, int32& OutParameterIndex) const
+float UCustomizableObject::GetFloatParameterDefaultValue(const FString& InParameterName) const
 {
-	OutParameterIndex = FindParameter(InParameterName);
-	if (OutParameterIndex == INDEX_NONE)
-	{
-		UE_LOG(LogMutable,Error,TEXT("The parameter with name [%s] does not exist in the CO with name [%s]."), *InParameterName,*GetName());
-		return false;
-	}
-
-	const TSharedPtr<const mu::Model> MutableModel = PrivateData->GetModel();
-	if (!MutableModel)
-	{
-		UE_LOG(LogMutable,Error,TEXT("Mutable model not set in [%s] CO. Compile this Customizable Object in order to be able to acces the default values of the parameters."),*GetName());
-		return false;
-	}
-	
-	const int32 ParameterCount = MutableModel->GetPrivate()->m_program.m_parameters.Num();
-	check (OutParameterIndex < ParameterCount);		// this failing would mean the data between CO and core is not synced 
-	
-	if ( GetParameterType(OutParameterIndex) != InParameterType)
+	const int32 ParameterIndex = FindParameter(InParameterName);
+	if (ParameterIndex == INDEX_NONE)
 	{
-		UE_LOG(LogMutable,Error,TEXT("The parameter with name [%s] is not a [%s] parameter of the [%s] CO."), *InParameterName,*UEnum::GetValueAsString(InParameterType),*GetName());
-		return false;
+		checkNoEntry();
+		return FCustomizableObjectFloatParameterValue::DEFAULT_PARAMETER_VALUE;
 	}
-	
-	return true;
-}
 
-float UCustomizableObject::GetFloatParameterDefaultValue(const FString& InParameterName) const
-{
-	float Value = -1.0f;
-	
-	int32 ParameterIndex;
-	if (!CanDefaultParameterBeAccessed(InParameterName,EMutableParameterType::Float,ParameterIndex))
+	const TSharedPtr<mu::Model> Model = GetModel();
+	if (!Model)
 	{
-		return Value;
+		checkNoEntry();
+		return FCustomizableObjectFloatParameterValue::DEFAULT_PARAMETER_VALUE;
 	}
 
-	const TSharedPtr<const mu::Model> MutableModel = PrivateData->GetModel();
-	Value = MutableModel->GetPrivate()->m_program.m_parameters[ParameterIndex].m_defaultValue.Get<mu::ParamFloatType>();
-	return Value;
+	return Model->GetFloatDefaultValue(ParameterIndex);
 }
 
 
 int32 UCustomizableObject::GetIntParameterDefaultValue(const FString& InParameterName) const
 {
-	int32 Value = -1;
-	
-	int32 ParameterIndex;
-	if (!CanDefaultParameterBeAccessed(InParameterName,EMutableParameterType::Int,ParameterIndex))
+	const int32 ParameterIndex = FindParameter(InParameterName);
+	if (ParameterIndex == INDEX_NONE)
 	{
-		return Value;
+		checkNoEntry();
+		return FCustomizableObjectIntParameterValue::DEFAULT_PARAMETER_VALUE;
 	}
 
-	const TSharedPtr<const mu::Model> MutableModel = PrivateData->GetModel();
-	Value = MutableModel->GetPrivate()->m_program.m_parameters[ParameterIndex].m_defaultValue.Get<mu::ParamIntType>();
-	return Value;
+	const TSharedPtr<mu::Model> Model = GetModel();
+	if (!Model)
+	{
+		checkNoEntry();
+		return FCustomizableObjectIntParameterValue::DEFAULT_PARAMETER_VALUE;
+	}
+	
+	return Model->GetIntDefaultValue(ParameterIndex);
 }
 
 
 bool UCustomizableObject::GetBoolParameterDefaultValue(const FString& InParameterName) const
 {
-	bool Value = false;
-	
-	int32 ParameterIndex;
-	if (!CanDefaultParameterBeAccessed(InParameterName,EMutableParameterType::Bool,ParameterIndex))
+	const int32 ParameterIndex = FindParameter(InParameterName);
+	if (ParameterIndex == INDEX_NONE)
 	{
-		return Value;
+		checkNoEntry();
+		return FCustomizableObjectBoolParameterValue::DEFAULT_PARAMETER_VALUE;
+	}
+
+	const TSharedPtr<mu::Model> Model = GetModel();
+	if (!Model)
+	{
+		checkNoEntry();
+		return FCustomizableObjectBoolParameterValue::DEFAULT_PARAMETER_VALUE;;
 	}
 	
-	const TSharedPtr<const mu::Model> MutableModel = PrivateData->GetModel();
-	Value = MutableModel->GetPrivate()->m_program.m_parameters[ParameterIndex].m_defaultValue.Get<mu::ParamBoolType>();
-	return Value;
+	return Model->GetBoolDefaultValue(ParameterIndex);
 }
 
 
 FLinearColor UCustomizableObject::GetColorParameterDefaultValue(const FString& InParameterName) const
 {
-	FLinearColor Value = FLinearColor::Black;
-	
-	int32 ParameterIndex;
-	if (!CanDefaultParameterBeAccessed(InParameterName,EMutableParameterType::Color,ParameterIndex))
+	const int32 ParameterIndex = FindParameter(InParameterName);
+	if (ParameterIndex == INDEX_NONE)
+	{
+		checkNoEntry();
+		return FCustomizableObjectVectorParameterValue::DEFAULT_PARAMETER_VALUE;;
+	}
+
+	const TSharedPtr<mu::Model> Model = GetModel();
+	if (!Model)
 	{
-		return Value;
+		checkNoEntry();
+		return FCustomizableObjectVectorParameterValue::DEFAULT_PARAMETER_VALUE;
 	}
 	
-	const mu::Model::Private* ModelPrivate = PrivateData->GetModel()->GetPrivate();
-	Value.R = ModelPrivate->m_program.m_parameters[ParameterIndex].m_defaultValue.Get<mu::ParamColorType>()[0];
-	Value.G = ModelPrivate->m_program.m_parameters[ParameterIndex].m_defaultValue.Get<mu::ParamColorType>()[1];
-	Value.B = ModelPrivate->m_program.m_parameters[ParameterIndex].m_defaultValue.Get<mu::ParamColorType>()[2];
-	Value.A = 1.0f;
-	
+	FLinearColor Value;
+	Model->GetColourDefaultValue(ParameterIndex, &Value.R, &Value.G, &Value.B);
+
 	return Value;
 }
 
@@ -1538,75 +1523,75 @@ void UCustomizableObject::GetProjectorParameterDefaultValue(const FString& InPar
 	FVector3f& OutDirection, FVector3f& OutUp, FVector3f& OutScale, float& OutAngle,
 	ECustomizableObjectProjectorType& OutType) const
 {
-	int32 ParameterIndex;
-	if (!CanDefaultParameterBeAccessed(InParameterName,EMutableParameterType::Projector,ParameterIndex))
-	{
-		return;
-	}
-	
-	const TSharedPtr<const mu::Model> MutableModel = PrivateData->GetModel();
-	
-	mu::FProjector MutableProjector = MutableModel->GetPrivate()->m_program.m_parameters[ParameterIndex].m_defaultValue.Get<mu::ParamProjectorType>();
-	OutType = ProjectorUtils::GetEquivalentProjectorType(MutableProjector.type);
-	OutPos = MutableProjector.position;
-	OutDirection = MutableProjector.direction;
-	OutUp = MutableProjector.up;
-	OutScale = MutableProjector.scale;
-	OutAngle = MutableProjector.projectionAngle;
+	const FCustomizableObjectProjector Projector = GetProjectorParameterDefaultValue(InParameterName);
+		
+	OutType = Projector.ProjectionType;
+	OutPos = Projector.Position;
+	OutDirection = Projector.Direction;
+	OutUp = Projector.Up;
+	OutScale = Projector.Scale;
+	OutAngle = Projector.Angle;
 }
 
 
 FCustomizableObjectProjector UCustomizableObject::GetProjectorParameterDefaultValue(const FString& InParameterName) const 
 {
-	FCustomizableObjectProjector Value;
-	
-	int32 ParameterIndex;
-	if (!CanDefaultParameterBeAccessed(InParameterName,EMutableParameterType::Projector,ParameterIndex))
+	const int32 ParameterIndex = FindParameter(InParameterName);
+	if (ParameterIndex == INDEX_NONE)
 	{
-		return Value;
+		checkNoEntry();
+		return FCustomizableObjectProjectorParameterValue::DEFAULT_PARAMETER_VALUE;
 	}
-	
-	const TSharedPtr<const mu::Model> MutableModel = PrivateData->GetModel();
 
-	mu::FProjector MutableProjector = MutableModel->GetPrivate()->m_program.m_parameters[ParameterIndex].m_defaultValue.Get<mu::ParamProjectorType>();
-	Value.Angle = MutableProjector.projectionAngle;
-	Value.Direction = MutableProjector.direction;
-	Value.Position = MutableProjector.position;
-	Value.Scale = MutableProjector.scale;
-	Value.Up = MutableProjector.up;
-	Value.ProjectionType = ProjectorUtils::GetEquivalentProjectorType(MutableProjector.type);
+	const TSharedPtr<mu::Model> Model = GetModel();
+	if (!Model)
+	{
+		checkNoEntry();
+		return FCustomizableObjectProjectorParameterValue::DEFAULT_PARAMETER_VALUE;
+	}
 
+	FCustomizableObjectProjector Value;
+	mu::PROJECTOR_TYPE Type;
+	Model->GetProjectorDefaultValue(ParameterIndex, &Type, &Value.Position, &Value.Direction, &Value.Up, &Value.Scale, &Value.Angle);
+	Value.ProjectionType = ProjectorUtils::GetEquivalentProjectorType(Type);
+	
 	return Value;
 }
 
 
 uint64 UCustomizableObject::GetTextureParameterDefaultValue(const FString& InParameterName) const
 {
-	uint64 Value = false;
-	
-	int32 ParameterIndex;
-	if (!CanDefaultParameterBeAccessed(InParameterName,EMutableParameterType::Texture,ParameterIndex))
+	const int32 ParameterIndex = FindParameter(InParameterName);
+	if (ParameterIndex == INDEX_NONE)
 	{
-		return Value;
+		checkNoEntry();
+		return FCustomizableObjectTextureParameterValue::DEFAULT_PARAMETER_VALUE;
+	}
+
+	const TSharedPtr<mu::Model> Model = GetModel();
+	if (!Model)
+	{
+		checkNoEntry();
+		return FCustomizableObjectTextureParameterValue::DEFAULT_PARAMETER_VALUE;
 	}
 	
-	const TSharedPtr<const mu::Model> MutableModel = PrivateData->GetModel();
-	Value = MutableModel->GetPrivate()->m_program.m_parameters[ParameterIndex].m_defaultValue.Get<mu::ParamImageType>();
-	return Value;
+	return Model->GetImageDefaultValue(ParameterIndex);
 }
 
+
 bool UCustomizableObject::IsParameterMultidimensional(const FString& InParameterName) const
 {
 	const int32 ParameterIndex = FindParameter(InParameterName);
 	if (ParameterIndex == INDEX_NONE)
 	{
-		UE_LOG(LogMutable,Error,TEXT("The parameter with name %s does not exist in the CO with name %s."), *InParameterName,*GetName());
 		checkNoEntry();
+		return {};
 	}
 
 	return IsParameterMultidimensional(ParameterIndex);
 }
 
+
 bool UCustomizableObject::IsParameterMultidimensional(const int32& InParamIndex) const
 {
 	check(InParamIndex != INDEX_NONE);
@@ -1972,7 +1957,7 @@ void UCustomizableObjectBulk::PrepareBulkData(UCustomizableObject* InOuter, cons
 	BulkDataFileNames.Empty();
 	
 	// Split the Streamable data into several separate files and fix up FileIndex and Offset of each StreamableBlock
-	if(TSharedPtr<const mu::Model, ESPMode::ThreadSafe> Model = CustomizableObject->GetModel())
+	if (TSharedPtr<const mu::Model, ESPMode::ThreadSafe> Model = CustomizableObject->GetModel())
 	{
 		const uint64 MaxChunkSize = UCustomizableObjectSystem::GetInstance()->GetMaxChunkSizeForPlatform(TargetPlatform);
 
@@ -1981,10 +1966,10 @@ void UCustomizableObjectBulk::PrepareBulkData(UCustomizableObject* InOuter, cons
 		uint16 CurrentFileIndex = 0;
 		uint64 CurrentChunkSize = 0;
 
-		const int32 NumStreamingFiles = Model->GetPrivate()->m_program.m_roms.Num();
+		const int32 NumStreamingFiles = Model->GetRomCount();
 		for (size_t FileIndex = 0; FileIndex < NumStreamingFiles; ++FileIndex)
 		{
-			const uint32 ResourceId = Model->GetPrivate()->m_program.m_roms[FileIndex].Id;
+			const uint32 ResourceId = Model->GetRomId(FileIndex);
 
 			FMutableStreamableBlock& StreamableBlock = CustomizableObject->HashToStreamableBlock[ResourceId];
 
diff --git a/Engine/Plugins/Experimental/Mutable/Source/CustomizableObject/Private/MuCO/CustomizableObjectInstance.cpp b/Engine/Plugins/Experimental/Mutable/Source/CustomizableObject/Private/MuCO/CustomizableObjectInstance.cpp
--- a/Engine/Plugins/Experimental/Mutable/Source/CustomizableObject/Private/MuCO/CustomizableObjectInstance.cpp
+++ b/Engine/Plugins/Experimental/Mutable/Source/CustomizableObject/Private/MuCO/CustomizableObjectInstance.cpp
@@ -26,10 +26,7 @@
 #include "MuCO/ICustomizableObjectModule.h"
 #include "MuCO/UnrealConversionUtils.h"
 #include "MuCO/UnrealPortabilityHelpers.h"
-#include "MuR/OpImagePixelFormat.h"
 #include "MuR/Model.h"
-#include "MuR/OpImagePixelFormat.h"
-#include "MuR/OpImageSwizzle.h"
 #include "PhysicsEngine/PhysicsAsset.h"
 #include "Rendering/Texture2DResource.h"
 #include "RenderingThread.h"
@@ -3307,6 +3304,7 @@ FTexturePlatformData* UCustomizableInstancePrivateData::MutableCreateImagePlatfo
 }
 
 
+
 void UCustomizableInstancePrivateData::ConvertImage(UTexture2D* Texture, mu::ImagePtrConst MutableImage, const FMutableModelImageProperties& Props, int OnlyLOD, int32 ExtractChannel )
 {
 	MUTABLE_CPUPROFILER_SCOPE(UCustomizableInstancePrivateData::ConvertImage);
@@ -6140,6 +6138,7 @@ void UCustomizableInstancePrivateData::ProcessTextureCoverageQueries(const TShar
 }
 
 
+
 void UCustomizableInstancePrivateData::UpdateParameterDecorationsEngineResources(const TSharedPtr<FMutableOperationData>& OperationData)
 {
 	// This should run in the unreal thread.
diff --git a/Engine/Plugins/Experimental/Mutable/Source/CustomizableObject/Private/MuCO/CustomizableObjectInstanceDescriptor.cpp b/Engine/Plugins/Experimental/Mutable/Source/CustomizableObject/Private/MuCO/CustomizableObjectInstanceDescriptor.cpp
--- a/Engine/Plugins/Experimental/Mutable/Source/CustomizableObject/Private/MuCO/CustomizableObjectInstanceDescriptor.cpp
+++ b/Engine/Plugins/Experimental/Mutable/Source/CustomizableObject/Private/MuCO/CustomizableObjectInstanceDescriptor.cpp
@@ -14,9 +14,6 @@
 #include UE_INLINE_GENERATED_CPP_BY_NAME(CustomizableObjectInstanceDescriptor)
 
 
-const FString EmptyString;
-
-
 FString GetAvailableOptionsString(const UCustomizableObject& CustomizableObject, const int32 ParameterIndexInObject)
 {
 	FString OptionsString;
@@ -1155,7 +1152,7 @@ const FString& FCustomizableObjectInstanceDescriptor::GetIntParameterSelectedOpt
 		}
 	}
 
-	return EmptyString;
+	return FCustomizableObjectBoolParameterValue::DEFAULT_PARAMETER_VALUE_NAME;
 }
 
 
@@ -1240,7 +1237,7 @@ float FCustomizableObjectInstanceDescriptor::GetFloatParameterSelectedOption(con
 		}
 	}
 
-	return -1.0f; 
+	return 	FCustomizableObjectFloatParameterValue::DEFAULT_PARAMETER_VALUE; 
 }
 
 
@@ -1303,7 +1300,7 @@ uint64 FCustomizableObjectInstanceDescriptor::GetTextureParameterSelectedOption(
 		}
 	}
 
-	return -1; 
+	return FCustomizableObjectTextureParameterValue::DEFAULT_PARAMETER_VALUE; 
 }
 
 
@@ -1371,7 +1368,7 @@ FLinearColor FCustomizableObjectInstanceDescriptor::GetColorParameterSelectedOpt
 		return VectorParameters[ColorParamIndex].ParameterValue;
 	}
 
-	return FLinearColor::Black;
+	return FCustomizableObjectVectorParameterValue::DEFAULT_PARAMETER_VALUE;
 }
 
 
@@ -1394,7 +1391,7 @@ bool FCustomizableObjectInstanceDescriptor::GetBoolParameterSelectedOption(const
 		return BoolParameters[BoolParamIndex].ParameterValue;
 	}
 
-	return false;
+	return FCustomizableObjectBoolParameterValue::DEFAULT_PARAMETER_VALUE;
 }
 
 
@@ -1752,7 +1749,7 @@ FCustomizableObjectProjector FCustomizableObjectInstanceDescriptor::GetProjector
 		}
 	}
 
-	return FCustomizableObjectProjector();
+	return FCustomizableObjectProjectorParameterValue::DEFAULT_PARAMETER_VALUE;
 }
 
 
@@ -1883,6 +1880,19 @@ int32 FCustomizableObjectInstanceDescriptor::AddValueToFloatRange(const FString&
 }
 
 
+int32 FCustomizableObjectInstanceDescriptor::AddValueToTextureRange(const FString& ParamName)
+{
+	const int32 TextureParameterIndex = FindTextureParameterNameIndex(ParamName);
+	if (TextureParameterIndex != -1)
+	{
+		FCustomizableObjectTextureParameterValue& TextureParameter = TextureParameters[TextureParameterIndex];
+		return TextureParameter.ParameterRangeValues.Add(FCustomizableObjectTextureParameterValue::DEFAULT_PARAMETER_VALUE);		
+	}
+	
+	return -1;
+}
+
+
 int32 FCustomizableObjectInstanceDescriptor::AddValueToProjectorRange(const FString& ParamName)
 {
 	check(CustomizableObject);
diff --git a/Engine/Plugins/Experimental/Mutable/Source/CustomizableObject/Private/MuCO/CustomizableObjectSystemPrivate.h b/Engine/Plugins/Experimental/Mutable/Source/CustomizableObject/Private/MuCO/CustomizableObjectSystemPrivate.h
--- a/Engine/Plugins/Experimental/Mutable/Source/CustomizableObject/Private/MuCO/CustomizableObjectSystemPrivate.h
+++ b/Engine/Plugins/Experimental/Mutable/Source/CustomizableObject/Private/MuCO/CustomizableObjectSystemPrivate.h
@@ -872,8 +872,3 @@ private:
 
 };
 
-
-void ConvertImage(UTexture2D* Texture, mu::ImagePtrConst MutableImage, const struct FMutableModelImageProperties& Props);
-
-
-
diff --git a/Engine/Plugins/Experimental/Mutable/Source/CustomizableObject/Public/MuCO/CustomizableObjectParameterTypeDefinitions.h b/Engine/Plugins/Experimental/Mutable/Source/CustomizableObject/Public/MuCO/CustomizableObjectParameterTypeDefinitions.h
--- a/Engine/Plugins/Experimental/Mutable/Source/CustomizableObject/Public/MuCO/CustomizableObjectParameterTypeDefinitions.h
+++ b/Engine/Plugins/Experimental/Mutable/Source/CustomizableObject/Public/MuCO/CustomizableObjectParameterTypeDefinitions.h
@@ -46,6 +46,10 @@ struct FCustomizableObjectBoolParameterValue
 {
 	GENERATED_USTRUCT_BODY()
 
+	static constexpr bool DEFAULT_PARAMETER_VALUE = false;
+
+	inline static const FString DEFAULT_PARAMETER_VALUE_NAME;
+	
 	UPROPERTY(Category = CustomizableObjectBoolParameterValue, VisibleAnywhere)
 	FString ParameterName;
 
@@ -71,6 +75,8 @@ struct FCustomizableObjectIntParameterValue
 {
 	GENERATED_USTRUCT_BODY()
 
+	static constexpr int32 DEFAULT_PARAMETER_VALUE = INDEX_NONE;
+	
 	UPROPERTY(Category = CustomizableObjectIntParameterValue, VisibleAnywhere)
 	FString ParameterName;
 
@@ -116,6 +122,8 @@ struct FCustomizableObjectFloatParameterValue
 {
 	GENERATED_USTRUCT_BODY()
 
+	static constexpr float DEFAULT_PARAMETER_VALUE = 0;
+
 	UPROPERTY(Category = CustomizableObjectFloatParameterValue, VisibleAnywhere)
 	FString ParameterName;
 
@@ -147,10 +155,10 @@ inline uint32 GetTypeHash(const FCustomizableObjectFloatParameterValue& Key)
 USTRUCT(BlueprintType)
 struct FCustomizableObjectTextureParameterValue
 {
-	static constexpr uint64 DEFAULT_PARAMETER_VALUE = 0;
-	
 	GENERATED_USTRUCT_BODY()
 
+	static constexpr uint64 DEFAULT_PARAMETER_VALUE = 0;
+	
 	UPROPERTY(Category = CustomizableObjectTextureParameterValue, VisibleAnywhere)
 	FString ParameterName;
 
@@ -184,6 +192,8 @@ struct FCustomizableObjectVectorParameterValue
 {
 	GENERATED_USTRUCT_BODY()
 
+	inline static const FLinearColor DEFAULT_PARAMETER_VALUE = FLinearColor::Black;
+	
 	UPROPERTY(Category = CustomizableObjectVectorParameterValue, VisibleAnywhere)
 	FString ParameterName;
 
@@ -269,6 +279,8 @@ struct FCustomizableObjectProjectorParameterValue
 {
 	GENERATED_USTRUCT_BODY()
 
+	static constexpr FCustomizableObjectProjector DEFAULT_PARAMETER_VALUE;
+
 	UPROPERTY(Category = CustomizableObjectVectorParameterValue, VisibleAnywhere)
 	FString ParameterName;
 
diff --git a/Engine/Plugins/Experimental/Mutable/Source/MutableRuntime/Private/MuR/ImagePrivate.h b/Engine/Plugins/Experimental/Mutable/Source/MutableRuntime/Private/MuR/ImagePrivate.h
--- a/Engine/Plugins/Experimental/Mutable/Source/MutableRuntime/Private/MuR/ImagePrivate.h
+++ b/Engine/Plugins/Experimental/Mutable/Source/MutableRuntime/Private/MuR/ImagePrivate.h
@@ -10,61 +10,12 @@
 namespace mu
 {
 
-	MUTABLE_DEFINE_ENUM_SERIALISABLE( EBlendType );
-	MUTABLE_DEFINE_ENUM_SERIALISABLE( EMipmapFilterType );
-	MUTABLE_DEFINE_ENUM_SERIALISABLE( ECompositeImageMode );
-	MUTABLE_DEFINE_ENUM_SERIALISABLE( ESamplingMethod );
-	MUTABLE_DEFINE_ENUM_SERIALISABLE( EMinFilterMethod );
-	MUTABLE_DEFINE_ENUM_SERIALISABLE( EImageFormat );
-
-    struct FImageFormatData
-	{
-		static constexpr SIZE_T MAX_BYTES_PER_BLOCK = 16;
-
-		FImageFormatData
-			(
-				unsigned pixelsPerBlockX = 0,
-				unsigned pixelsPerBlockY = 0,
-				unsigned bytesPerBlock = 0,
-				unsigned channels = 0
-			)
-		{
-			m_pixelsPerBlockX = (uint8)pixelsPerBlockX;
-			m_pixelsPerBlockY = (uint8)pixelsPerBlockY;
-			m_bytesPerBlock = (uint16)bytesPerBlock;
-			m_channels = (uint16)channels;	
-		}
-
-		FImageFormatData
-			(
-				unsigned pixelsPerBlockX,
-				unsigned pixelsPerBlockY,
-				unsigned bytesPerBlock,
-				unsigned channels,
-				std::initializer_list<uint8> BlackBlockInit
-			)
-			: FImageFormatData(pixelsPerBlockX, pixelsPerBlockY, bytesPerBlock, channels)
-		{
-			check(MAX_BYTES_PER_BLOCK >= BlackBlockInit.size());
-
-			const SIZE_T SanitizedBlockSize = FMath::Min<SIZE_T>(MAX_BYTES_PER_BLOCK, BlackBlockInit.size());
-			FMemory::Memcpy(BlackBlock, BlackBlockInit.begin(), SanitizedBlockSize);
-		}
-
-		//! For block based formats, size of the block size. For uncompressed formats it will
-		//! always be 1,1. For non-block-based compressed formats, it will be 0,0.
-        uint8 m_pixelsPerBlockX, m_pixelsPerBlockY;
-
-		//! Number of bytes used by every pixel block, if uncompressed or block-compressed format.
-		//! For non-block-compressed formats, it returns 0.
-        uint16 m_bytesPerBlock;
-
-		//! Channels in every pixel of the image.
-        uint16 m_channels;
-
-		//! Representation of a black block of the image.
-		uint8 BlackBlock[MAX_BYTES_PER_BLOCK] = { 0 };
-	};
+	MUTABLE_IMPLEMENT_ENUM_SERIALISABLE( EBlendType );
+	MUTABLE_IMPLEMENT_ENUM_SERIALISABLE( EMipmapFilterType );
+	MUTABLE_IMPLEMENT_ENUM_SERIALISABLE( ECompositeImageMode );
+	MUTABLE_IMPLEMENT_ENUM_SERIALISABLE( ESamplingMethod );
+	MUTABLE_IMPLEMENT_ENUM_SERIALISABLE( EMinFilterMethod );
+	MUTABLE_IMPLEMENT_ENUM_SERIALISABLE( EImageFormat );
 
 
 	struct FMipmapGenerationSettings
@@ -96,13 +47,6 @@ namespace mu
 		}
 	};
 
-
-	//---------------------------------------------------------------------------------------------
-	//!
-	//---------------------------------------------------------------------------------------------
-	MUTABLERUNTIME_API const FImageFormatData& GetImageFormatData(EImageFormat format );
-
-
 	//---------------------------------------------------------------------------------------------
 	//!
 	//---------------------------------------------------------------------------------------------
diff --git a/Engine/Plugins/Experimental/Mutable/Source/MutableRuntime/Private/MuR/Layout.cpp b/Engine/Plugins/Experimental/Mutable/Source/MutableRuntime/Private/MuR/Layout.cpp
--- a/Engine/Plugins/Experimental/Mutable/Source/MutableRuntime/Private/MuR/Layout.cpp
+++ b/Engine/Plugins/Experimental/Mutable/Source/MutableRuntime/Private/MuR/Layout.cpp
@@ -6,6 +6,7 @@
 #include "HAL/LowLevelMemTracker.h"
 #include "Math/IntPoint.h"
 #include "MuR/MutableMath.h"
+#include "MuR/SerialisationPrivate.h"
 
 
 namespace mu {
@@ -185,6 +186,66 @@ namespace mu {
 	}
 
 
+	//---------------------------------------------------------------------------------------------
+	void Layout::Serialise(OutputArchive& arch) const
+	{
+		uint32 ver = 5;
+		arch << ver;
+
+		arch << m_size;
+		arch << m_blocks;
+
+		arch << m_maxsize;
+		arch << uint32(m_strategy);
+		arch << FirstLODToIgnoreWarnings;
+		arch << uint32(ReductionMethod);
+	}
+
+	
+	//---------------------------------------------------------------------------------------------
+	void Layout::Unserialise(InputArchive& arch)
+	{
+		uint32 ver;
+		arch >> ver;
+		check(ver <= 5);
+
+		arch >> m_size;
+
+		if (ver < 5)
+		{
+			uint32_t Size = 0;
+			arch >> Size;
+			m_blocks.SetNum(Size);
+
+			for (uint32_t BlockIndex = 0; BlockIndex < Size; ++BlockIndex)
+			{
+				m_blocks[BlockIndex].UnserialiseOldVersion(arch, ver);
+			}
+		}
+		else
+		{
+			arch >> m_blocks;
+		}
+
+		arch >> m_maxsize;
+
+		uint32 temp;
+		arch >> temp;
+		m_strategy = EPackStrategy(temp);
+
+		if (ver >= 4)
+		{
+			arch >> FirstLODToIgnoreWarnings;
+		}
+
+		if (ver >= 5)
+		{
+			arch >> temp;
+			ReductionMethod = EReductionMethod(temp);
+		}
+	}
+
+
 	//---------------------------------------------------------------------------------------------
 	bool Layout::IsSimilar(const Layout& o) const
 	{
@@ -258,5 +319,38 @@ namespace mu {
 	{
 		return ReductionMethod;
 	}
+
+	
+	//---------------------------------------------------------------------------------------------
+	void Layout::FBlock::Serialise(OutputArchive& arch) const
+	{
+		arch << m_min;
+		arch << m_size;
+		arch << m_id;
+		arch << m_priority;
+		arch << bUseSymmetry;
+	}
+
+
+	//---------------------------------------------------------------------------------------------
+	void Layout::FBlock::Unserialise(InputArchive& arch)
+	{
+		arch >> m_min;
+		arch >> m_size;
+		arch >> m_id;
+		arch >> m_priority;
+		arch >> bUseSymmetry;
+	}
+
+	
+	//---------------------------------------------------------------------------------------------
+	void Layout::FBlock::UnserialiseOldVersion(InputArchive& Archive, const int32 Version)
+	{
+		// Use the version if the Layout version changes to 6
+		Archive >> m_min;
+		Archive >> m_size;
+		Archive >> m_id;
+		Archive >> m_priority;
+	}
 }
 
diff --git a/Engine/Plugins/Experimental/Mutable/Source/MutableRuntime/Private/MuR/Mesh.cpp b/Engine/Plugins/Experimental/Mutable/Source/MutableRuntime/Private/MuR/Mesh.cpp
--- a/Engine/Plugins/Experimental/Mutable/Source/MutableRuntime/Private/MuR/Mesh.cpp
+++ b/Engine/Plugins/Experimental/Mutable/Source/MutableRuntime/Private/MuR/Mesh.cpp
@@ -987,6 +987,153 @@ void Mesh::ResetBufferIndices()
 }
 
 
+//-------------------------------------------------------------------------------------------------
+void Mesh::FACE_GROUP::Serialise(OutputArchive& arch) const
+{
+	const int32 ver = 0;
+	arch << ver;
+
+	arch << m_name;
+	arch << m_faces;
+}
+
+
+//-------------------------------------------------------------------------------------------------
+void Mesh::FACE_GROUP::Unserialise(InputArchive& arch)
+{
+	int32 ver = 0;
+	arch >> ver;
+	check(ver == 0);
+
+	arch >> m_name;
+	arch >> m_faces;
+}
+
+
+//-------------------------------------------------------------------------------------------------
+void Mesh::FBonePose::Serialise(OutputArchive& arch) const
+{
+	const int32 ver = 1;
+	arch << ver;
+
+	arch << BoneName;
+	arch << BoneUsageFlags;
+	arch << BoneTransform;
+}
+
+
+//-------------------------------------------------------------------------------------------------
+void Mesh::FBonePose::Unserialise(InputArchive& arch)
+{
+	int32 ver = 0;
+	arch >> ver;
+	check(ver <= 1);
+
+	arch >> BoneName;
+
+	if (ver == 0)
+	{
+		uint8 Skinned = 0;
+		arch >> Skinned;
+		BoneUsageFlags = Skinned ? EBoneUsageFlags::Skinning : EBoneUsageFlags::None;
+	}
+	else
+	{
+		arch >> BoneUsageFlags;
+	}
+
+	arch >> BoneTransform;
+}
+
+
+//-------------------------------------------------------------------------------------------------
+void Mesh::Serialise(OutputArchive& arch) const
+{
+	uint32 ver = 15;
+	arch << ver;
+
+	arch << m_IndexBuffers;
+	arch << m_VertexBuffers;
+	arch << m_FaceBuffers;
+	arch << m_AdditionalBuffers;
+	arch << m_layouts;
+
+	arch << SkeletonIDs;
+
+	arch << m_pSkeleton;
+	arch << m_pPhysicsBody;
+
+	arch << m_staticFormatFlags;
+	arch << m_surfaces;
+	arch << m_faceGroups;
+
+	arch << m_tags;
+
+	arch << BonePoses;
+
+	arch << AdditionalPhysicsBodies;
+}
+
+
+//-------------------------------------------------------------------------------------------------
+void Mesh::Unserialise(InputArchive& arch)
+{
+	uint32 ver;
+	arch >> ver;
+	check(ver <= 15);
+
+	arch >> m_IndexBuffers;
+	arch >> m_VertexBuffers;
+	arch >> m_FaceBuffers;
+	arch >> m_AdditionalBuffers;
+	arch >> m_layouts;
+
+	if (ver >= 14)
+	{
+		arch >> SkeletonIDs;
+	}
+
+	arch >> m_pSkeleton;
+	if (ver >= 12)
+	{ 
+		arch >> m_pPhysicsBody;
+	}
+	else
+	{
+		m_pPhysicsBody = nullptr;
+	}
+
+	arch >> m_staticFormatFlags;
+	arch >> m_surfaces;
+	arch >> m_faceGroups;
+
+	arch >> m_tags;
+
+	if (ver >= 13)
+	{
+		arch >> BonePoses;
+	}
+	else if (m_pSkeleton)
+	{
+		const int32 NumBones = m_pSkeleton->GetBoneCount();
+		BonePoses.SetNum(NumBones);
+		check(m_pSkeleton->m_boneTransforms_DEPRECATED.Num() == NumBones);
+
+		for (int32 BoneIndex = 0; BoneIndex < NumBones; ++BoneIndex)
+		{
+			BonePoses[BoneIndex].BoneName = m_pSkeleton->GetBoneName(BoneIndex);
+			BonePoses[BoneIndex].BoneUsageFlags = EBoneUsageFlags::Skinning;
+			BonePoses[BoneIndex].BoneTransform = m_pSkeleton->m_boneTransforms_DEPRECATED[BoneIndex];
+		}
+	}
+
+	if (ver >= 15)
+	{
+		arch >> AdditionalPhysicsBodies;
+	}
+}
+
+
 //---------------------------------------------------------------------------------------------
 bool Mesh::IsSimilar(const Mesh& o, bool bCompareLayouts) const
 {
@@ -1458,4 +1605,10 @@ void Mesh::Log( FString& out, int32 BufferElementLimit)
     LogBuffer( out, m_FaceBuffers, BufferElementLimit);
 }
 
+
+MUTABLE_IMPLEMENT_POD_VECTOR_SERIALISABLE(MESH_SURFACE)
+MUTABLE_IMPLEMENT_ENUM_SERIALISABLE(EBoneUsageFlags)
+MUTABLE_IMPLEMENT_ENUM_SERIALISABLE(EMeshBufferType)
+MUTABLE_IMPLEMENT_ENUM_SERIALISABLE(EShapeBindingMethod)
+	
 }
diff --git a/Engine/Plugins/Experimental/Mutable/Source/MutableRuntime/Private/MuR/MeshBufferSet.cpp b/Engine/Plugins/Experimental/Mutable/Source/MutableRuntime/Private/MuR/MeshBufferSet.cpp
--- a/Engine/Plugins/Experimental/Mutable/Source/MutableRuntime/Private/MuR/MeshBufferSet.cpp
+++ b/Engine/Plugins/Experimental/Mutable/Source/MutableRuntime/Private/MuR/MeshBufferSet.cpp
@@ -8,6 +8,7 @@
 #include "Math/UnrealMathSSE.h"
 #include "Misc/AssertionMacros.h"
 #include "MuR/MutableTrace.h"
+#include "MuR/SerialisationPrivate.h"
 
 
 namespace mu
@@ -52,6 +53,22 @@ namespace mu
 	}
 
 
+	//-----------------------------------------------------------------------------------------
+	void FMeshBufferSet::Serialise(OutputArchive& arch) const
+	{
+		arch << m_elementCount;
+		arch << m_buffers;
+	}
+
+
+	//-----------------------------------------------------------------------------------------
+	void FMeshBufferSet::Unserialise(InputArchive& arch)
+	{
+		arch >> m_elementCount;
+		arch >> m_buffers;
+	}
+
+
 	//---------------------------------------------------------------------------------------------
 	//---------------------------------------------------------------------------------------------
 	//---------------------------------------------------------------------------------------------
@@ -605,4 +622,28 @@ namespace mu
 			m_buffers[b].m_elementSize = offset;
 	}
 
-}
\ No newline at end of file
+	
+	//-----------------------------------------------------------------------------------------
+	void MESH_BUFFER::Serialise(OutputArchive& arch) const
+	{
+		arch << m_channels;
+		arch << m_data;
+		arch << m_elementSize;
+	}
+
+	
+	//-----------------------------------------------------------------------------------------
+	void MESH_BUFFER::Unserialise(InputArchive& arch)
+	{
+		arch >> m_channels;
+		arch >> m_data;
+		arch >> m_elementSize;
+	}
+
+	
+	MUTABLE_IMPLEMENT_POD_SERIALISABLE(MESH_BUFFER_CHANNEL);
+	MUTABLE_IMPLEMENT_POD_VECTOR_SERIALISABLE(MESH_BUFFER_CHANNEL);
+	MUTABLE_IMPLEMENT_ENUM_SERIALISABLE(MESH_BUFFER_FORMAT);
+	MUTABLE_IMPLEMENT_ENUM_SERIALISABLE(MESH_BUFFER_SEMANTIC);
+}
+
diff --git a/Engine/Plugins/Experimental/Mutable/Source/MutableRuntime/Private/MuR/Model.cpp b/Engine/Plugins/Experimental/Mutable/Source/MutableRuntime/Private/MuR/Model.cpp
--- a/Engine/Plugins/Experimental/Mutable/Source/MutableRuntime/Private/MuR/Model.cpp
+++ b/Engine/Plugins/Experimental/Mutable/Source/MutableRuntime/Private/MuR/Model.cpp
@@ -318,7 +318,6 @@ namespace mu
         }
     }
 
-
     //---------------------------------------------------------------------------------------------
     Model::Private* Model::GetPrivate() const
     {
@@ -338,6 +337,137 @@ namespace mu
     }
 
 
+    //---------------------------------------------------------------------------------------------
+    bool Model::GetBoolDefaultValue(int32 Index) const
+    {
+    	check(m_pD->m_program.m_parameters.IsValidIndex(Index));
+		check(m_pD->m_program.m_parameters[Index].m_type == PARAMETER_TYPE::T_BOOL);
+
+        // Early out in case of invalid parameters
+        if (m_pD->m_program.m_parameters.IsValidIndex(Index) ||
+            m_pD->m_program.m_parameters[Index].m_type == PARAMETER_TYPE::T_BOOL)
+        {
+            return false;
+        }
+		
+        return m_pD->m_program.m_parameters[Index].m_defaultValue.Get<ParamBoolType>();
+    }
+
+
+    //---------------------------------------------------------------------------------------------
+	int32 Model::GetIntDefaultValue(int Index) const
+	{
+		check(m_pD->m_program.m_parameters.IsValidIndex(Index));
+		check(m_pD->m_program.m_parameters[Index].m_type == PARAMETER_TYPE::T_INT);
+
+        // Early out in case of invalid parameters
+        if (m_pD->m_program.m_parameters.IsValidIndex(Index) ||
+            m_pD->m_program.m_parameters[Index].m_type == PARAMETER_TYPE::T_INT)
+        {
+            return 0;
+        }
+		
+        return m_pD->m_program.m_parameters[Index].m_defaultValue.Get<ParamIntType>();
+	}
+
+
+    //---------------------------------------------------------------------------------------------
+	float Model::GetFloatDefaultValue(int32 Index) const
+	{
+    	check(m_pD->m_program.m_parameters.IsValidIndex(Index));
+		check(m_pD->m_program.m_parameters[Index].m_type == PARAMETER_TYPE::T_FLOAT);
+
+        // Early out in case of invalid parameters
+        if (m_pD->m_program.m_parameters.IsValidIndex(Index) ||
+            m_pD->m_program.m_parameters[Index].m_type == PARAMETER_TYPE::T_FLOAT)
+        {
+            return 0.0f;
+        }
+		
+        return m_pD->m_program.m_parameters[Index].m_defaultValue.Get<ParamFloatType>();
+	}
+
+
+	void Model::GetColourDefaultValue(int32 Index, float* R, float* G, float* B) const
+    {
+    	check(m_pD->m_program.m_parameters.IsValidIndex(Index));
+		check(m_pD->m_program.m_parameters[Index].m_type == PARAMETER_TYPE::T_COLOUR);
+
+        // Early out in case of invalid parameters
+        if (m_pD->m_program.m_parameters.IsValidIndex(Index) ||
+            m_pD->m_program.m_parameters[Index].m_type == PARAMETER_TYPE::T_COLOUR)
+        {
+            return;
+        }
+
+        ParamColorType& Color = m_pD->m_program.m_parameters[Index].m_defaultValue.Get<ParamColorType>();
+        if (R) *R = Color[0];
+    	if (G) *G = Color[1];
+    	if (B) *B = Color[2];
+    }
+
+
+	void Model::GetProjectorDefaultValue(int32 Index, PROJECTOR_TYPE* OutProjectionType, FVector3f* OutPos,
+		FVector3f* OutDir, FVector3f* OutUp, FVector3f* OutScale, float* OutProjectionAngle) const
+	{
+    	check(m_pD->m_program.m_parameters.IsValidIndex(Index));
+		check(m_pD->m_program.m_parameters[Index].m_type == PARAMETER_TYPE::T_PROJECTOR);
+
+        // Early out in case of invalid parameters
+        if (m_pD->m_program.m_parameters.IsValidIndex(Index) ||
+            m_pD->m_program.m_parameters[Index].m_type == PARAMETER_TYPE::T_PROJECTOR)
+        {
+            return;
+        }
+
+        const ParamProjectorType& Projector = m_pD->m_program.m_parameters[Index].m_defaultValue.Get<ParamProjectorType>();
+        if (OutProjectionType) *OutProjectionType = Projector.type;
+    	if (OutPos) *OutPos = Projector.position;
+		if (OutDir) *OutDir = Projector.direction;
+    	if (OutUp) *OutUp = Projector.up;
+    	if (OutScale) *OutScale = Projector.scale;
+    	if (OutProjectionAngle) *OutProjectionAngle = Projector.projectionAngle;
+	}
+
+
+	//---------------------------------------------------------------------------------------------
+    EXTERNAL_IMAGE_ID Model::GetImageDefaultValue(int32 Index) const
+    {
+	    check(m_pD->m_program.m_parameters.IsValidIndex(Index));
+		check(m_pD->m_program.m_parameters[Index].m_type == PARAMETER_TYPE::T_IMAGE);
+
+        // Early out in case of invalid parameters
+        if (m_pD->m_program.m_parameters.IsValidIndex(Index) ||
+            m_pD->m_program.m_parameters[Index].m_type == PARAMETER_TYPE::T_IMAGE)
+        {
+            return 0;
+        }
+		
+        return m_pD->m_program.m_parameters[Index].m_defaultValue.Get<ParamImageType>();
+    }
+
+
+    //---------------------------------------------------------------------------------------------
+    int32 Model::GetRomCount() const
+    {
+    	return m_pD->m_program.m_roms.Num();
+    }
+
+
+    //---------------------------------------------------------------------------------------------
+    uint32 Model::GetRomId(int32 Index) const
+    {
+    	return m_pD->m_program.m_roms[Index].Id;
+    }
+
+
+    //---------------------------------------------------------------------------------------------
+    uint32 Model::GetRomSize(int32 Index) const
+    {
+    	return m_pD->m_program.m_roms[Index].Size;
+    }
+
+
     //---------------------------------------------------------------------------------------------
     ParametersPtr Model::NewParameters(TSharedPtr<const Model> Model, const Parameters* pOld )
     {
diff --git a/Engine/Plugins/Experimental/Mutable/Source/MutableRuntime/Private/MuR/ModelPrivate.h b/Engine/Plugins/Experimental/Mutable/Source/MutableRuntime/Private/MuR/ModelPrivate.h
--- a/Engine/Plugins/Experimental/Mutable/Source/MutableRuntime/Private/MuR/ModelPrivate.h
+++ b/Engine/Plugins/Experimental/Mutable/Source/MutableRuntime/Private/MuR/ModelPrivate.h
@@ -44,7 +44,7 @@ namespace mu
 		uint32 ResourceType;        
     };
 
-    MUTABLE_DEFINE_POD_SERIALISABLE(FRomData);
+    MUTABLE_IMPLEMENT_POD_SERIALISABLE(FRomData);
 
     //!
     template<typename DATA>
@@ -66,7 +66,7 @@ namespace mu
 		uint8 LODCount = 0;
 		EImageFormat ImageFormat = EImageFormat::IF_NONE;
 	};
-	MUTABLE_DEFINE_POD_SERIALISABLE(FImageLODRange);
+	MUTABLE_IMPLEMENT_POD_SERIALISABLE(FImageLODRange);
 
 	struct FExtensionDataConstant
 	{
diff --git a/Engine/Plugins/Experimental/Mutable/Source/MutableRuntime/Private/MuR/OpImagePixelFormat.h b/Engine/Plugins/Experimental/Mutable/Source/MutableRuntime/Private/MuR/OpImagePixelFormat.h
--- a/Engine/Plugins/Experimental/Mutable/Source/MutableRuntime/Private/MuR/OpImagePixelFormat.h
+++ b/Engine/Plugins/Experimental/Mutable/Source/MutableRuntime/Private/MuR/OpImagePixelFormat.h
@@ -15,13 +15,6 @@ namespace mu
     //! \return false if the conversion failed, usually because not enough memory was allocated in
     //!     the result. This is only checked for RLE compression.
     //---------------------------------------------------------------------------------------------
-	MUTABLERUNTIME_API extern bool ImagePixelFormatInPlace( int quality, Image* pResult, const Image* pBase,
+	bool ImagePixelFormatInPlace( int quality, Image* pResult, const Image* pBase,
                                          int onlyLOD = -1 );
-
-    //---------------------------------------------------------------------------------------------
-    //! Wrapper for the above method that allocates the destination image.
-    //---------------------------------------------------------------------------------------------
-	MUTABLERUNTIME_API extern ImagePtr ImagePixelFormat( int quality, const Image* pBase, EImageFormat targetFormat,
-                                      int onlyLOD = -1 );
-
 }
diff --git a/Engine/Plugins/Experimental/Mutable/Source/MutableRuntime/Private/MuR/OpImageSwizzle.cpp b/Engine/Plugins/Experimental/Mutable/Source/MutableRuntime/Private/MuR/OpImageSwizzle.cpp
new file mode 100644
--- /dev/null
+++ b/Engine/Plugins/Experimental/Mutable/Source/MutableRuntime/Private/MuR/OpImageSwizzle.cpp
@@ -0,0 +1,245 @@
+// Copyright Epic Games, Inc. All Rights Reserved.
+
+#include "MuR/OpImageSwizzle.h"
+
+#include "MuR/ImagePrivate.h"
+#include "Async/ParallelFor.h"
+
+namespace mu
+{
+	
+	Ptr<Image> ImageSwizzle
+		(
+			EImageFormat format,
+			const Ptr<const Image> pSources[],
+            const uint8 channels[]
+		)
+	{
+		MUTABLE_CPUPROFILER_SCOPE(ImageSwizzle);
+
+		if (!pSources[0])
+		{
+			return nullptr;
+		}
+
+        ImagePtr pDest = new Image( pSources[0]->GetSizeX(), pSources[0]->GetSizeY(),
+                                    pSources[0]->GetLODCount(),
+                                    format );
+
+		// Very slow generic implementations
+        int32 PixelCount = pDest->CalculatePixelCount();
+
+        // Pixelcount should already match, but due to bugs it may not be the case. Try to detect it,
+        // but avoid crashing below:
+        uint16 numChannels = GetImageFormatData(format).m_channels;
+        for (uint16 c=0;c<numChannels; ++c)
+        {
+            if (pSources[c])
+            {
+                size_t sourcePixelCount = pSources[c]->CalculatePixelCount();
+                if (PixelCount>sourcePixelCount)
+                {
+                    check(false);
+
+                    // Something went wrong
+					PixelCount = sourcePixelCount;
+                }
+            }
+        }
+
+		int NumDestChannels = 0;
+
+		switch ( format )
+		{
+		case EImageFormat::IF_L_UBYTE:
+			NumDestChannels = 1;
+			break;
+
+		case EImageFormat::IF_RGB_UBYTE:
+			NumDestChannels = 3;
+			break;
+
+        case EImageFormat::IF_RGBA_UBYTE:
+        case EImageFormat::IF_BGRA_UBYTE:
+			NumDestChannels = 4;
+			break;
+
+        default:
+			check(false);
+		}
+
+		for (int i = 0; i < NumDestChannels; ++i)
+		{
+			uint8* pDestBuf = pDest->GetData() + i;
+
+			if (format == EImageFormat::IF_BGRA_UBYTE)
+			{
+				if (i == 0)
+				{
+					pDestBuf = pDest->GetData() + 2;
+				}
+				else if (i == 2)
+				{
+					pDestBuf = pDest->GetData() + 0;
+				}
+			}
+
+			bool filled = false;
+
+			constexpr int32 NumBatchElems = 4096*2;
+			const int32 NumBatches = FMath::DivideAndRoundUp(PixelCount, NumBatchElems);
+
+			if (pSources[i])
+			{
+				const uint8* pSourceBuf = pSources[i]->GetData() + channels[i];
+
+				switch (pSources[i]->GetFormat())
+				{
+				case EImageFormat::IF_L_UBYTE:
+					if (channels[i] < 1)
+					{
+						auto ProcessBatch = [pDestBuf, pSourceBuf, NumDestChannels, PixelCount, NumBatchElems ](int32 BatchId)
+						{
+							const int32 BatchBegin = BatchId * NumBatchElems;
+							const int32 BatchEnd = FMath::Min(BatchBegin + NumBatchElems, PixelCount);
+
+							for (int32 p = BatchBegin; p < BatchEnd; ++p)
+							{
+								pDestBuf[p * NumDestChannels] = pSourceBuf[p];
+							}
+						};
+
+						if (NumBatches == 1)
+						{
+							ProcessBatch(0);
+						}
+						else if (NumBatches > 1)
+						{
+							ParallelFor(NumBatches, ProcessBatch);
+						}
+
+						filled = true;
+					}
+					break;
+
+				case EImageFormat::IF_RGB_UBYTE:
+					if (channels[i] < 3)
+					{
+						auto ProcessBatch = [pDestBuf, pSourceBuf, NumDestChannels, PixelCount, NumBatchElems](int32 BatchId)
+						{
+							const int32 BatchBegin = BatchId * NumBatchElems;
+							const int32 BatchEnd = FMath::Min(BatchBegin + NumBatchElems, PixelCount);
+
+							for (int32 p = BatchBegin; p < BatchEnd; ++p)
+							{
+								pDestBuf[p * NumDestChannels] = pSourceBuf[p * 3];
+							}
+						};
+
+						if (NumBatches == 1)
+						{
+							ProcessBatch(0);
+						}
+						else if (NumBatches > 1)
+						{
+							ParallelFor(NumBatches, ProcessBatch);
+						}
+
+						filled = true;
+					}
+					break;
+
+				case EImageFormat::IF_RGBA_UBYTE:
+					if (channels[i] < 4)
+					{	
+						auto ProcessBatch = [pDestBuf, pSourceBuf, NumDestChannels, PixelCount, NumBatchElems](int32 BatchId)
+						{
+							const int32 BatchBegin = BatchId * NumBatchElems;
+							const int32 BatchEnd = FMath::Min(BatchBegin + NumBatchElems, PixelCount);
+
+							for (int32 p = BatchBegin; p < BatchEnd; ++p)
+							{
+								pDestBuf[p * NumDestChannels] = pSourceBuf[p * 4];
+							}
+						};
+
+						if (NumBatches == 1)
+						{
+							ProcessBatch(0);
+						}
+						else if (NumBatches > 1)
+						{
+							ParallelFor(NumBatches, ProcessBatch);
+						}
+
+						filled = true;
+					}
+					break;
+
+				case EImageFormat::IF_BGRA_UBYTE:
+					if (channels[i] == 0)
+					{
+						pSourceBuf = pSources[i]->GetData() + 2;
+					}
+					else if (channels[i] == 2)
+					{
+						pSourceBuf = pSources[i]->GetData() + 0;
+					}
+					if (channels[i] < 4)
+					{
+						auto ProcessBatch = [pDestBuf, pSourceBuf, NumDestChannels, PixelCount, NumBatchElems](int32 BatchId)
+						{
+							const int32 BatchBegin = BatchId * NumBatchElems;
+							const int32 BatchEnd = FMath::Min(BatchBegin + NumBatchElems, PixelCount);
+							for (int32 p = BatchBegin; p < BatchEnd; ++p)
+							{
+								pDestBuf[p * NumDestChannels] = pSourceBuf[p * 4];
+							}
+						};
+
+						if (NumBatches == 1)
+						{
+							ProcessBatch(0);
+						}
+						else if (NumBatches > 1)
+						{
+							ParallelFor(NumBatches, ProcessBatch);
+						}
+
+						filled = true;
+					}
+					break;
+
+				default:
+					check(false);
+				}
+			}
+
+			if (!filled)
+			{
+				// Source not set. Clear to 0
+				auto ProcessBatch = [pDestBuf, NumDestChannels, PixelCount, NumBatchElems](int32 BatchId)
+				{
+					const int32 BatchBegin = BatchId * NumBatchElems;
+					const int32 BatchEnd = FMath::Min(BatchBegin + NumBatchElems, PixelCount);
+					for (int32 p = BatchBegin; p < BatchEnd; ++p)
+					{
+						pDestBuf[p * NumDestChannels] = 0;
+					}
+				};
+
+				if (NumBatches == 1)
+				{
+					ProcessBatch(0);
+				}
+				else if (NumBatches > 1)
+				{
+					ParallelFor(NumBatches, ProcessBatch);
+				}
+			}
+		}
+
+		return pDest;
+	}
+	
+}
diff --git a/Engine/Plugins/Experimental/Mutable/Source/MutableRuntime/Private/MuR/OpImageSwizzle.h b/Engine/Plugins/Experimental/Mutable/Source/MutableRuntime/Private/MuR/OpImageSwizzle.h
--- a/Engine/Plugins/Experimental/Mutable/Source/MutableRuntime/Private/MuR/OpImageSwizzle.h
+++ b/Engine/Plugins/Experimental/Mutable/Source/MutableRuntime/Private/MuR/OpImageSwizzle.h
@@ -2,245 +2,9 @@
 
 #pragma once
 
-#include "MuR/ImagePrivate.h"
-#include "Async/ParallelFor.h"
-#include "MuR/Platform.h"
+#include "MuR/Image.h"
+
 
 namespace mu
 {
-
-	inline Ptr<Image> ImageSwizzle
-		(
-			EImageFormat format,
-			const Ptr<const Image> pSources[],
-            const uint8 channels[]
-		)
-	{
-		MUTABLE_CPUPROFILER_SCOPE(ImageSwizzle);
-
-		if (!pSources[0])
-		{
-			return nullptr;
-		}
-
-        ImagePtr pDest = new Image( pSources[0]->GetSizeX(), pSources[0]->GetSizeY(),
-                                    pSources[0]->GetLODCount(),
-                                    format );
-
-		// Very slow generic implementations
-        int32 PixelCount = pDest->CalculatePixelCount();
-
-        // Pixelcount should already match, but due to bugs it may not be the case. Try to detect it,
-        // but avoid crashing below:
-        uint16 numChannels = GetImageFormatData(format).m_channels;
-        for (uint16 c=0;c<numChannels; ++c)
-        {
-            if (pSources[c])
-            {
-                size_t sourcePixelCount = pSources[c]->CalculatePixelCount();
-                if (PixelCount>sourcePixelCount)
-                {
-                    check(false);
-
-                    // Something went wrong
-					PixelCount = sourcePixelCount;
-                }
-            }
-        }
-
-		int NumDestChannels = 0;
-
-		switch ( format )
-		{
-		case EImageFormat::IF_L_UBYTE:
-			NumDestChannels = 1;
-			break;
-
-		case EImageFormat::IF_RGB_UBYTE:
-			NumDestChannels = 3;
-			break;
-
-        case EImageFormat::IF_RGBA_UBYTE:
-        case EImageFormat::IF_BGRA_UBYTE:
-			NumDestChannels = 4;
-			break;
-
-        default:
-			check(false);
-		}
-
-		for (int i = 0; i < NumDestChannels; ++i)
-		{
-			uint8* pDestBuf = pDest->GetData() + i;
-
-			if (format == EImageFormat::IF_BGRA_UBYTE)
-			{
-				if (i == 0)
-				{
-					pDestBuf = pDest->GetData() + 2;
-				}
-				else if (i == 2)
-				{
-					pDestBuf = pDest->GetData() + 0;
-				}
-			}
-
-			bool filled = false;
-
-			constexpr int32 NumBatchElems = 4096*2;
-			const int32 NumBatches = FMath::DivideAndRoundUp(PixelCount, NumBatchElems);
-
-			if (pSources[i])
-			{
-				const uint8* pSourceBuf = pSources[i]->GetData() + channels[i];
-
-				switch (pSources[i]->GetFormat())
-				{
-				case EImageFormat::IF_L_UBYTE:
-					if (channels[i] < 1)
-					{
-						auto ProcessBatch = [pDestBuf, pSourceBuf, NumDestChannels, PixelCount, NumBatchElems ](int32 BatchId)
-						{
-							const int32 BatchBegin = BatchId * NumBatchElems;
-							const int32 BatchEnd = FMath::Min(BatchBegin + NumBatchElems, PixelCount);
-
-							for (int32 p = BatchBegin; p < BatchEnd; ++p)
-							{
-								pDestBuf[p * NumDestChannels] = pSourceBuf[p];
-							}
-						};
-
-						if (NumBatches == 1)
-						{
-							ProcessBatch(0);
-						}
-						else if (NumBatches > 1)
-						{
-							ParallelFor(NumBatches, ProcessBatch);
-						}
-
-						filled = true;
-					}
-					break;
-
-				case EImageFormat::IF_RGB_UBYTE:
-					if (channels[i] < 3)
-					{
-						auto ProcessBatch = [pDestBuf, pSourceBuf, NumDestChannels, PixelCount, NumBatchElems](int32 BatchId)
-						{
-							const int32 BatchBegin = BatchId * NumBatchElems;
-							const int32 BatchEnd = FMath::Min(BatchBegin + NumBatchElems, PixelCount);
-
-							for (int32 p = BatchBegin; p < BatchEnd; ++p)
-							{
-								pDestBuf[p * NumDestChannels] = pSourceBuf[p * 3];
-							}
-						};
-
-						if (NumBatches == 1)
-						{
-							ProcessBatch(0);
-						}
-						else if (NumBatches > 1)
-						{
-							ParallelFor(NumBatches, ProcessBatch);
-						}
-
-						filled = true;
-					}
-					break;
-
-				case EImageFormat::IF_RGBA_UBYTE:
-					if (channels[i] < 4)
-					{	
-						auto ProcessBatch = [pDestBuf, pSourceBuf, NumDestChannels, PixelCount, NumBatchElems](int32 BatchId)
-						{
-							const int32 BatchBegin = BatchId * NumBatchElems;
-							const int32 BatchEnd = FMath::Min(BatchBegin + NumBatchElems, PixelCount);
-
-							for (int32 p = BatchBegin; p < BatchEnd; ++p)
-							{
-								pDestBuf[p * NumDestChannels] = pSourceBuf[p * 4];
-							}
-						};
-
-						if (NumBatches == 1)
-						{
-							ProcessBatch(0);
-						}
-						else if (NumBatches > 1)
-						{
-							ParallelFor(NumBatches, ProcessBatch);
-						}
-
-						filled = true;
-					}
-					break;
-
-				case EImageFormat::IF_BGRA_UBYTE:
-					if (channels[i] == 0)
-					{
-						pSourceBuf = pSources[i]->GetData() + 2;
-					}
-					else if (channels[i] == 2)
-					{
-						pSourceBuf = pSources[i]->GetData() + 0;
-					}
-					if (channels[i] < 4)
-					{
-						auto ProcessBatch = [pDestBuf, pSourceBuf, NumDestChannels, PixelCount, NumBatchElems](int32 BatchId)
-						{
-							const int32 BatchBegin = BatchId * NumBatchElems;
-							const int32 BatchEnd = FMath::Min(BatchBegin + NumBatchElems, PixelCount);
-							for (int32 p = BatchBegin; p < BatchEnd; ++p)
-							{
-								pDestBuf[p * NumDestChannels] = pSourceBuf[p * 4];
-							}
-						};
-
-						if (NumBatches == 1)
-						{
-							ProcessBatch(0);
-						}
-						else if (NumBatches > 1)
-						{
-							ParallelFor(NumBatches, ProcessBatch);
-						}
-
-						filled = true;
-					}
-					break;
-
-				default:
-					check(false);
-				}
-			}
-
-			if (!filled)
-			{
-				// Source not set. Clear to 0
-				auto ProcessBatch = [pDestBuf, NumDestChannels, PixelCount, NumBatchElems](int32 BatchId)
-				{
-					const int32 BatchBegin = BatchId * NumBatchElems;
-					const int32 BatchEnd = FMath::Min(BatchBegin + NumBatchElems, PixelCount);
-					for (int32 p = BatchBegin; p < BatchEnd; ++p)
-					{
-						pDestBuf[p * NumDestChannels] = 0;
-					}
-				};
-
-				if (NumBatches == 1)
-				{
-					ProcessBatch(0);
-				}
-				else if (NumBatches > 1)
-				{
-					ParallelFor(NumBatches, ProcessBatch);
-				}
-			}
-		}
-
-		return pDest;
-	}
-
 }
diff --git a/Engine/Plugins/Experimental/Mutable/Source/MutableRuntime/Private/MuR/Operations.h b/Engine/Plugins/Experimental/Mutable/Source/MutableRuntime/Private/MuR/Operations.h
--- a/Engine/Plugins/Experimental/Mutable/Source/MutableRuntime/Private/MuR/Operations.h
+++ b/Engine/Plugins/Experimental/Mutable/Source/MutableRuntime/Private/MuR/Operations.h
@@ -1170,8 +1170,8 @@ namespace mu
     typedef t_OP<uint32> OP;
 
     // OP is not serialisable anymore, we link code now.
-    MUTABLE_DEFINE_POD_SERIALISABLE( OP );
-    MUTABLE_DEFINE_POD_VECTOR_SERIALISABLE( OP );
+    MUTABLE_IMPLEMENT_POD_SERIALISABLE( OP );
+    MUTABLE_IMPLEMENT_POD_VECTOR_SERIALISABLE( OP );
     
     // Check that we didn't go out of control with the operation size
     static_assert( sizeof(OP::FArgs)==28, "Argument union has an unexpected size." );
diff --git a/Engine/Plugins/Experimental/Mutable/Source/MutableRuntime/Private/MuR/Parameters.cpp b/Engine/Plugins/Experimental/Mutable/Source/MutableRuntime/Private/MuR/Parameters.cpp
--- a/Engine/Plugins/Experimental/Mutable/Source/MutableRuntime/Private/MuR/Parameters.cpp
+++ b/Engine/Plugins/Experimental/Mutable/Source/MutableRuntime/Private/MuR/Parameters.cpp
@@ -580,7 +580,7 @@ namespace mu
              ||
              GetType(index) != PARAMETER_TYPE::T_FLOAT )
         {
-            return false;
+            return 0.0f;
         }
 
         // Single value case
@@ -768,6 +768,16 @@ namespace mu
         check( index >= 0 && index < (int)m_pD->m_values.Num() );
         check( GetType( index ) == PARAMETER_TYPE::T_IMAGE );
 
+		// Early out in case of invalid parameters
+        if ( index < 0
+             ||
+             index >= (int)m_pD->m_values.Num()
+             ||
+             GetType(index) != PARAMETER_TYPE::T_IMAGE )
+        {
+            return 0;
+        }
+		
 		// Single value case
 		if (!pos)
 		{
diff --git a/Engine/Plugins/Experimental/Mutable/Source/MutableRuntime/Private/MuR/ParametersPrivate.h b/Engine/Plugins/Experimental/Mutable/Source/MutableRuntime/Private/MuR/ParametersPrivate.h
--- a/Engine/Plugins/Experimental/Mutable/Source/MutableRuntime/Private/MuR/ParametersPrivate.h
+++ b/Engine/Plugins/Experimental/Mutable/Source/MutableRuntime/Private/MuR/ParametersPrivate.h
@@ -16,9 +16,9 @@
 namespace mu
 {
 
-    MUTABLE_DEFINE_ENUM_SERIALISABLE( PARAMETER_TYPE )
-    MUTABLE_DEFINE_ENUM_SERIALISABLE( PARAMETER_DETAILED_TYPE )
-    MUTABLE_DEFINE_ENUM_SERIALISABLE( PROJECTOR_TYPE )
+    MUTABLE_IMPLEMENT_ENUM_SERIALISABLE( PARAMETER_TYPE )
+    MUTABLE_IMPLEMENT_ENUM_SERIALISABLE( PARAMETER_DETAILED_TYPE )
+    MUTABLE_IMPLEMENT_ENUM_SERIALISABLE( PROJECTOR_TYPE )
 
 
     //---------------------------------------------------------------------------------------------
diff --git a/Engine/Plugins/Experimental/Mutable/Source/MutableRuntime/Private/MuR/PhysicsBody.cpp b/Engine/Plugins/Experimental/Mutable/Source/MutableRuntime/Private/MuR/PhysicsBody.cpp
--- a/Engine/Plugins/Experimental/Mutable/Source/MutableRuntime/Private/MuR/PhysicsBody.cpp
+++ b/Engine/Plugins/Experimental/Mutable/Source/MutableRuntime/Private/MuR/PhysicsBody.cpp
@@ -15,6 +15,7 @@ namespace mu
 	}
 
 
+	//-------------------------------------------------------------------------------------------------
 	mu::Ptr<PhysicsBody> PhysicsBody::StaticUnserialise(InputArchive& arch)
 	{
 		mu::Ptr<PhysicsBody> pResult = new PhysicsBody();
@@ -22,6 +23,206 @@ namespace mu
 		return pResult;
 	}
 
+
+	//-------------------------------------------------------------------------------------------------
+	void FBodyShape::Serialise(OutputArchive& arch) const
+	{
+		const uint32 ver = 0;
+		arch << ver;
+
+		arch << Name;
+		arch << Flags;
+	}
+
+
+	//-------------------------------------------------------------------------------------------------
+	void FBodyShape::Unserialise(InputArchive& arch)
+	{
+		uint32 ver;
+		arch >> ver;
+		check(ver == 0);
+
+		arch >> Name;
+		arch >> Flags;
+	}
+
+
+	//-------------------------------------------------------------------------------------------------
+	void FSphereBody::Serialise(OutputArchive& arch) const
+	{
+		FBodyShape::Serialise(arch);
+
+		const uint32 ver = 0;
+		arch << ver;
+
+		arch << Position;
+		arch << Radius;
+	}
+
+
+	//-------------------------------------------------------------------------------------------------
+	void FSphereBody::Unserialise(InputArchive& arch)
+	{
+		FBodyShape::Unserialise(arch);
+
+		uint32 ver;
+		arch >> ver;
+		check(ver == 0);
+
+		arch >> Position;
+		arch >> Radius;
+	}
+
+
+	//-------------------------------------------------------------------------------------------------
+	void FBoxBody::Serialise(OutputArchive& arch) const
+	{
+		FBodyShape::Serialise( arch );
+
+		const uint32 ver = 0;
+		arch << ver;
+
+		arch << Position;
+		arch << Orientation;
+		arch << Size;
+	}
+
+
+	//-------------------------------------------------------------------------------------------------
+	void FBoxBody::Unserialise(InputArchive& arch)
+	{
+		FBodyShape::Unserialise( arch );
+
+		uint32 ver;
+		arch >> ver;
+		check(ver == 0);
+
+		arch >> Position;
+		arch >> Orientation;
+		arch >> Size;
+	}
+
+	
+	//-------------------------------------------------------------------------------------------------
+	void FSphylBody::Serialise(OutputArchive& arch) const
+	{
+		FBodyShape::Serialise( arch );
+
+		const uint32 ver = 0;
+		arch << ver;
+
+		arch << Position;
+		arch << Orientation;
+		arch << Radius;
+		arch << Length;
+	}
+
+
+	//-------------------------------------------------------------------------------------------------
+	void FSphylBody::Unserialise(InputArchive& arch)
+	{
+		FBodyShape::Unserialise( arch );
+
+		uint32 ver;
+		arch >> ver;
+		check(ver == 0);
+
+		arch >> Position;
+		arch >> Orientation;
+		arch >> Radius;
+		arch >> Length;
+	}
+
+
+	//-------------------------------------------------------------------------------------------------
+	void FTaperedCapsuleBody::Serialise(OutputArchive& arch) const
+	{
+		FBodyShape::Serialise( arch );
+
+		const uint32 ver = 0;
+		arch << ver;
+
+		arch << Position;
+		arch << Orientation;
+		arch << Radius0;
+		arch << Radius1;
+		arch << Length;
+	}
+
+
+	//-------------------------------------------------------------------------------------------------
+	void FTaperedCapsuleBody::Unserialise(InputArchive& arch)
+	{
+		FBodyShape::Unserialise( arch );
+			
+		uint32 ver;
+		arch >> ver;
+		check(ver == 0);
+
+		arch >> Position;
+		arch >> Orientation;
+		arch >> Radius0;
+		arch >> Radius1;
+		arch >> Length;
+	}
+
+
+	//-------------------------------------------------------------------------------------------------
+	void FConvexBody::Serialise(OutputArchive& arch) const
+	{
+		FBodyShape::Serialise( arch );
+	
+		uint32 ver = 0;
+		arch << ver;
+
+		arch << Vertices;
+		arch << Indices;
+		arch << Transform;
+	}
+
+
+	//-------------------------------------------------------------------------------------------------
+	void FConvexBody::Unserialise(InputArchive& arch)
+	{
+		uint32 ver;
+		arch >> ver;
+		check(ver == 0);
+
+		arch >> Vertices;
+		arch >> Indices;
+		arch >> Transform;
+	}
+
+
+	//-------------------------------------------------------------------------------------------------
+	void FPhysicsBodyAggregate::Serialise(OutputArchive& arch) const
+	{
+		uint32 ver = 0;
+		arch << ver;
+		
+		arch << Spheres;
+		arch << Boxes;
+		arch << Convex;
+		arch << Sphyls;
+		arch << TaperedCapsules;
+	}
+
+
+	//-------------------------------------------------------------------------------------------------
+	void FPhysicsBodyAggregate::Unserialise(InputArchive& arch)
+	{
+		uint32 ver;
+		arch >> ver;
+		check(ver == 0);
+		
+		arch >> Spheres;
+		arch >> Boxes;
+		arch >> Convex;
+		arch >> Sphyls;
+		arch >> TaperedCapsules;
+	}
+
+
 	//-------------------------------------------------------------------------------------------------
 	mu::Ptr<PhysicsBody> PhysicsBody::Clone() const
 	{
@@ -504,4 +705,40 @@ namespace mu
 		return Bodies[B].TaperedCapsules[I].Name.c_str();
 	}
 
+
+	//-------------------------------------------------------------------------------------------------
+	void PhysicsBody::Serialise(OutputArchive& arch) const
+	{
+		uint32 ver = 2;
+		arch << ver;
+        	
+		arch << CustomId;
+		arch << Bodies;
+		arch << Bones;
+		arch << BodiesCustomIds;
+		arch << bBodiesModified;
+	}
+
+
+	//-------------------------------------------------------------------------------------------------
+	void PhysicsBody::Unserialise(InputArchive& arch)
+	{
+		uint32 ver;
+		arch >> ver;
+		check(ver <= 2);
+        	
+		if (ver >= 2)
+		{
+			arch >> CustomId;
+		}
+        	
+		arch >> Bodies;
+		arch >> Bones;
+		arch >> BodiesCustomIds;
+
+		if (ver >= 1)
+		{
+			arch >> bBodiesModified;
+		}
+	}
 }
diff --git a/Engine/Plugins/Experimental/Mutable/Source/MutableRuntime/Private/MuR/SerialisationPrivate.h b/Engine/Plugins/Experimental/Mutable/Source/MutableRuntime/Private/MuR/SerialisationPrivate.h
--- a/Engine/Plugins/Experimental/Mutable/Source/MutableRuntime/Private/MuR/SerialisationPrivate.h
+++ b/Engine/Plugins/Experimental/Mutable/Source/MutableRuntime/Private/MuR/SerialisationPrivate.h
@@ -142,45 +142,36 @@ namespace mu
 
 
 	//---------------------------------------------------------------------------------------------
-	template< typename T >
-	inline void operator<< ( OutputArchive& arch, const T& t )
-	{
-        t.Serialise( arch );
-	}
-
-	template< typename T >
-	inline void operator>> ( InputArchive& arch, T& t )
-	{
-        t.Unserialise( arch );
-	}
-
-
-	//---------------------------------------------------------------------------------------------
-#define MUTABLE_DEFINE_POD_SERIALISABLE(T)								\
+#define MUTABLE_IMPLEMENT_POD_SERIALISABLE(T)							\
 		template<>														\
 		inline void operator<< <T>( OutputArchive& arch, const T& t )	\
 		{																\
 			arch.GetPrivate()->m_pStream->Write( &t, sizeof(T) );		\
 		}																\
 																		\
+		template void operator<< ( OutputArchive& arch, const T& t );	\
+																		\
 		template<>														\
 		inline void operator>> <T>( InputArchive& arch, T& t )			\
 		{																\
 			arch.GetPrivate()->m_pStream->Read( &t, sizeof(T) );		\
-		}
+		}																\
+																		\
+		template void operator>> ( InputArchive& arch, T& t );			\
 
-	MUTABLE_DEFINE_POD_SERIALISABLE(float);
-	MUTABLE_DEFINE_POD_SERIALISABLE(double);
 
-    MUTABLE_DEFINE_POD_SERIALISABLE( int8 );
-    MUTABLE_DEFINE_POD_SERIALISABLE( int16 );
-    MUTABLE_DEFINE_POD_SERIALISABLE( int32 );
-    MUTABLE_DEFINE_POD_SERIALISABLE( int64 );
+	MUTABLE_IMPLEMENT_POD_SERIALISABLE(float);
+	MUTABLE_IMPLEMENT_POD_SERIALISABLE(double);
 
-    MUTABLE_DEFINE_POD_SERIALISABLE( uint8 );
-    MUTABLE_DEFINE_POD_SERIALISABLE( uint16 );
-    MUTABLE_DEFINE_POD_SERIALISABLE( uint32 );
-    MUTABLE_DEFINE_POD_SERIALISABLE( uint64 );
+    MUTABLE_IMPLEMENT_POD_SERIALISABLE( int8 );
+    MUTABLE_IMPLEMENT_POD_SERIALISABLE( int16 );
+    MUTABLE_IMPLEMENT_POD_SERIALISABLE( int32 );
+    MUTABLE_IMPLEMENT_POD_SERIALISABLE( int64 );
+
+    MUTABLE_IMPLEMENT_POD_SERIALISABLE( uint8 );
+    MUTABLE_IMPLEMENT_POD_SERIALISABLE( uint16 );
+    MUTABLE_IMPLEMENT_POD_SERIALISABLE( uint32 );
+    MUTABLE_IMPLEMENT_POD_SERIALISABLE( uint64 );
 
 	//---------------------------------------------------------------------------------------------
 	// TVariant custom serialize. Based on the default serialization.
@@ -257,21 +248,26 @@ namespace mu
 
 
 	//---------------------------------------------------------------------------------------------
-#define MUTABLE_DEFINE_ENUM_SERIALISABLE(T)								\
+#define MUTABLE_IMPLEMENT_ENUM_SERIALISABLE(T)							\
 		template<>														\
         inline void operator<< <T>( OutputArchive& arch, const T& t )	\
 		{																\
-            auto v = (uint32)t;                                       \
-            arch.GetPrivate()->m_pStream->Write( &v, sizeof(uint32) );\
+            auto v = (uint32)t;                                         \
+            arch.GetPrivate()->m_pStream->Write( &v, sizeof(uint32) );  \
 		}																\
 																		\
+		template void operator<< ( OutputArchive& arch, const T& t );	\
+																		\
 		template<>														\
-        inline void operator>> <T>( InputArchive& arch, T& t )          \
+        inline void operator>> <T>( InputArchive& arch, T& t )   		\
 		{																\
             uint32 v;													\
             arch.GetPrivate()->m_pStream->Read( &v, sizeof(uint32) );	\
 			t = (T)v;													\
-		}
+		}																\
+																		\
+		template void operator>> ( InputArchive& arch, T& t ); 			\
+
 
 
     //---------------------------------------------------------------------------------------------
@@ -289,31 +285,6 @@ namespace mu
         arch >> v.second;
     }
 
-
-	
-	//---------------------------------------------------------------------------------------------
-	template<typename T> inline void operator<<(OutputArchive& arch, const TArray<T>& v)
-	{
-		// TODO: Optimise for vectors of PODs
-		arch << (uint32)v.Num();
-		for (std::size_t i = 0; i < v.Num(); ++i)
-		{
-			arch << v[i];
-		}
-	}
-
-	template<typename T> inline void operator>>(InputArchive& arch, TArray<T>& v)
-	{
-		// TODO: Optimise for vectors of PODs
-		uint32 size;
-		arch >> size;
-		v.SetNum(size);
-		for (std::size_t i = 0; i < size; ++i)
-		{
-			arch >> v[i];
-		}
-	}
-
 	
 	//---------------------------------------------------------------------------------------------
 	template<typename T, uint32 Size> void operator<<(OutputArchive& arch, const TStaticArray<T, Size>& v)
@@ -363,8 +334,9 @@ namespace mu
 
 
 	//---------------------------------------------------------------------------------------------
-#define MUTABLE_DEFINE_POD_VECTOR_SERIALISABLE(T)										\
-	template<> inline void operator<<<T>(OutputArchive& arch, const TArray<T>& v)		\
+#define MUTABLE_IMPLEMENT_POD_VECTOR_SERIALISABLE(T)									\
+	template<>																			\
+	inline void operator<< <T>(OutputArchive& arch, const TArray<T>& v)					\
 	{                                                                                   \
 		uint32 Num = uint32(v.Num());													\
 		arch << Num;																	\
@@ -374,7 +346,10 @@ namespace mu
 		}                                                                               \
 	}                                                                                   \
 																						\
-	template<> inline void operator>><T>(InputArchive& arch, TArray<T>& v)				\
+	template void operator<< (OutputArchive& arch, const TArray<T>& v);					\
+																						\
+	template<>																			\
+	inline void operator>> <T>(InputArchive& arch, TArray<T>& v)						\
 	{                                                                                   \
 		uint32 Num;																		\
 		arch >> Num;																	\
@@ -383,37 +358,39 @@ namespace mu
 		{                                                                               \
 			arch.GetPrivate()->m_pStream->Read(&v[0], Num * sizeof(T));					\
 		}                                                                               \
-	}
-
-	MUTABLE_DEFINE_POD_VECTOR_SERIALISABLE(float);
-	MUTABLE_DEFINE_POD_VECTOR_SERIALISABLE(double);
-	MUTABLE_DEFINE_POD_VECTOR_SERIALISABLE(uint8);
-	MUTABLE_DEFINE_POD_VECTOR_SERIALISABLE(uint16);
-	MUTABLE_DEFINE_POD_VECTOR_SERIALISABLE(uint32);
-	MUTABLE_DEFINE_POD_VECTOR_SERIALISABLE(uint64);
-	MUTABLE_DEFINE_POD_VECTOR_SERIALISABLE(int8);
-	MUTABLE_DEFINE_POD_VECTOR_SERIALISABLE(int16);
-	MUTABLE_DEFINE_POD_VECTOR_SERIALISABLE(int32);
-	MUTABLE_DEFINE_POD_VECTOR_SERIALISABLE(int64);
-
-	MUTABLE_DEFINE_POD_SERIALISABLE(vec2f);
-	MUTABLE_DEFINE_POD_SERIALISABLE(vec3f);
-	MUTABLE_DEFINE_POD_SERIALISABLE(mat3f);
-	MUTABLE_DEFINE_POD_SERIALISABLE(mat4f);
-	MUTABLE_DEFINE_POD_SERIALISABLE(vec2<int>);
-	MUTABLE_DEFINE_POD_VECTOR_SERIALISABLE(vec2f);
-	MUTABLE_DEFINE_POD_VECTOR_SERIALISABLE(vec3f);
-	MUTABLE_DEFINE_POD_VECTOR_SERIALISABLE(mat3f);
-	MUTABLE_DEFINE_POD_VECTOR_SERIALISABLE(mat4f);
-	MUTABLE_DEFINE_POD_VECTOR_SERIALISABLE(vec2<int>);
+	}																					\
+																						\
+	template void operator>> (InputArchive& arch, TArray<T>& v);						\
+
+	MUTABLE_IMPLEMENT_POD_VECTOR_SERIALISABLE(float);
+	MUTABLE_IMPLEMENT_POD_VECTOR_SERIALISABLE(double);
+	MUTABLE_IMPLEMENT_POD_VECTOR_SERIALISABLE(uint8);
+	MUTABLE_IMPLEMENT_POD_VECTOR_SERIALISABLE(uint16);
+	MUTABLE_IMPLEMENT_POD_VECTOR_SERIALISABLE(uint32);
+	MUTABLE_IMPLEMENT_POD_VECTOR_SERIALISABLE(uint64);
+	MUTABLE_IMPLEMENT_POD_VECTOR_SERIALISABLE(int8);
+	MUTABLE_IMPLEMENT_POD_VECTOR_SERIALISABLE(int16);
+	MUTABLE_IMPLEMENT_POD_VECTOR_SERIALISABLE(int32);
+	MUTABLE_IMPLEMENT_POD_VECTOR_SERIALISABLE(int64);
+
+	MUTABLE_IMPLEMENT_POD_SERIALISABLE(vec2f);
+	MUTABLE_IMPLEMENT_POD_SERIALISABLE(vec3f);
+	MUTABLE_IMPLEMENT_POD_SERIALISABLE(mat3f);
+	MUTABLE_IMPLEMENT_POD_SERIALISABLE(mat4f);
+	MUTABLE_IMPLEMENT_POD_SERIALISABLE(vec2<int>);
+	MUTABLE_IMPLEMENT_POD_VECTOR_SERIALISABLE(vec2f);
+	MUTABLE_IMPLEMENT_POD_VECTOR_SERIALISABLE(vec3f);
+	MUTABLE_IMPLEMENT_POD_VECTOR_SERIALISABLE(mat3f);
+	MUTABLE_IMPLEMENT_POD_VECTOR_SERIALISABLE(mat4f);
+	MUTABLE_IMPLEMENT_POD_VECTOR_SERIALISABLE(vec2<int>);
 
 	// Unreal POD Serializables
-	MUTABLE_DEFINE_POD_SERIALISABLE(FUintVector2);
-	MUTABLE_DEFINE_POD_SERIALISABLE(UE::Math::TIntVector2<uint16>);
-	MUTABLE_DEFINE_POD_SERIALISABLE(UE::Math::TIntVector2<int16>);
-	MUTABLE_DEFINE_POD_VECTOR_SERIALISABLE(FUintVector2);
-	MUTABLE_DEFINE_POD_VECTOR_SERIALISABLE(UE::Math::TIntVector2<uint16>);
-	MUTABLE_DEFINE_POD_VECTOR_SERIALISABLE(UE::Math::TIntVector2<int16>);
+	MUTABLE_IMPLEMENT_POD_SERIALISABLE(FUintVector2);
+	MUTABLE_IMPLEMENT_POD_SERIALISABLE(UE::Math::TIntVector2<uint16>);
+	MUTABLE_IMPLEMENT_POD_SERIALISABLE(UE::Math::TIntVector2<int16>);
+	MUTABLE_IMPLEMENT_POD_VECTOR_SERIALISABLE(FUintVector2);
+	MUTABLE_IMPLEMENT_POD_VECTOR_SERIALISABLE(UE::Math::TIntVector2<uint16>);
+	MUTABLE_IMPLEMENT_POD_VECTOR_SERIALISABLE(UE::Math::TIntVector2<int16>);
 
 
 	//---------------------------------------------------------------------------------------------
diff --git a/Engine/Plugins/Experimental/Mutable/Source/MutableRuntime/Private/MuR/Skeleton.cpp b/Engine/Plugins/Experimental/Mutable/Source/MutableRuntime/Private/MuR/Skeleton.cpp
--- a/Engine/Plugins/Experimental/Mutable/Source/MutableRuntime/Private/MuR/Skeleton.cpp
+++ b/Engine/Plugins/Experimental/Mutable/Source/MutableRuntime/Private/MuR/Skeleton.cpp
@@ -101,4 +101,43 @@ namespace mu
 	}
 
 
-}
\ No newline at end of file
+	//-------------------------------------------------------------------------------------------------
+	void Skeleton::Serialise(OutputArchive& arch) const
+	{
+		uint32 ver = 5;
+		arch << ver;
+
+		arch << m_bones;
+		arch << m_boneParents;
+	}
+
+
+	//-------------------------------------------------------------------------------------------------
+	void Skeleton::Unserialise(InputArchive& arch)
+	{
+		uint32 ver;
+		arch >> ver;
+		check(ver >= 3);
+
+		arch >> m_bones;
+
+		if (ver == 3)
+		{
+			arch >> m_boneTransforms_DEPRECATED;
+		}
+
+		arch >> m_boneParents;
+
+		if (ver <= 4)
+		{
+			TArray<int32> boneIds;
+			arch >> boneIds;
+		}
+
+		if (ver == 3)
+		{
+			bool bBoneTransformModified;
+			arch >> bBoneTransformModified;
+		}
+	}
+}
diff --git a/Engine/Plugins/Experimental/Mutable/Source/MutableRuntime/Public/MuR/Image.h b/Engine/Plugins/Experimental/Mutable/Source/MutableRuntime/Public/MuR/Image.h
--- a/Engine/Plugins/Experimental/Mutable/Source/MutableRuntime/Public/MuR/Image.h
+++ b/Engine/Plugins/Experimental/Mutable/Source/MutableRuntime/Public/MuR/Image.h
@@ -176,7 +176,71 @@ namespace mu
 	};
 	static_assert(uint32(EMinFilterMethod::MaxValue) <= (1 << 3), "EMinFilterMethod enum cannot hold more than 8 values");
 
+	struct FImageFormatData
+	{
+		static constexpr SIZE_T MAX_BYTES_PER_BLOCK = 16;
+
+		FImageFormatData
+			(
+				unsigned pixelsPerBlockX = 0,
+				unsigned pixelsPerBlockY = 0,
+				unsigned bytesPerBlock = 0,
+				unsigned channels = 0
+			)
+		{
+			m_pixelsPerBlockX = (uint8)pixelsPerBlockX;
+			m_pixelsPerBlockY = (uint8)pixelsPerBlockY;
+			m_bytesPerBlock = (uint16)bytesPerBlock;
+			m_channels = (uint16)channels;	
+		}
+
+		FImageFormatData
+			(
+				unsigned pixelsPerBlockX,
+				unsigned pixelsPerBlockY,
+				unsigned bytesPerBlock,
+				unsigned channels,
+				std::initializer_list<uint8> BlackBlockInit
+			)
+			: FImageFormatData(pixelsPerBlockX, pixelsPerBlockY, bytesPerBlock, channels)
+		{
+			check(MAX_BYTES_PER_BLOCK >= BlackBlockInit.size());
+
+			const SIZE_T SanitizedBlockSize = FMath::Min<SIZE_T>(MAX_BYTES_PER_BLOCK, BlackBlockInit.size());
+			FMemory::Memcpy(BlackBlock, BlackBlockInit.begin(), SanitizedBlockSize);
+		}
 
+		//! For block based formats, size of the block size. For uncompressed formats it will
+		//! always be 1,1. For non-block-based compressed formats, it will be 0,0.
+        uint8 m_pixelsPerBlockX, m_pixelsPerBlockY;
+
+		//! Number of bytes used by every pixel block, if uncompressed or block-compressed format.
+		//! For non-block-compressed formats, it returns 0.
+        uint16 m_bytesPerBlock;
+
+		//! Channels in every pixel of the image.
+        uint16 m_channels;
+
+		//! Representation of a black block of the image.
+		uint8 BlackBlock[MAX_BYTES_PER_BLOCK] = { 0 };
+	};
+		
+	MUTABLERUNTIME_API const FImageFormatData& GetImageFormatData(EImageFormat format );
+
+	//---------------------------------------------------------------------------------------------
+	//! Convert an image to another pixel format.
+	//! Allocates the destination image.
+	//! \warning Not all format conversions are implemented.
+	//! \param onlyLOD If different than -1, only the specified lod level will be converted in the
+	//! returned image.
+	//! \return false if the conversion failed, usually because not enough memory was allocated in
+	//!     the result. This is only checked for RLE compression.
+	//---------------------------------------------------------------------------------------------
+	MUTABLERUNTIME_API ImagePtr ImagePixelFormat( int quality, const Image* pBase, EImageFormat targetFormat,
+	                                      int onlyLOD = -1 );
+
+	MUTABLERUNTIME_API Ptr<Image> ImageSwizzle(EImageFormat format, const Ptr<const Image> pSources[], const uint8 channels[]);
+	
     //! \brief 2D image resource with mipmaps.
 	//! \ingroup runtime
     class MUTABLERUNTIME_API Image : public RefCounted
@@ -387,7 +451,6 @@ namespace mu
 
 		/** Fill the image with a plain colour. */
 		void FillColour(FVector4f c);
-
 	};
 }
 
diff --git a/Engine/Plugins/Experimental/Mutable/Source/MutableRuntime/Public/MuR/Layout.h b/Engine/Plugins/Experimental/Mutable/Source/MutableRuntime/Public/MuR/Layout.h
--- a/Engine/Plugins/Experimental/Mutable/Source/MutableRuntime/Public/MuR/Layout.h
+++ b/Engine/Plugins/Experimental/Mutable/Source/MutableRuntime/Public/MuR/Layout.h
@@ -5,7 +5,6 @@
 #include "MuR/Ptr.h"
 #include "MuR/RefCounted.h"
 #include "MuR/Serialisation.h"
-#include "MuR/SerialisationPrivate.h"
 #include "Containers/Array.h"
 #include "Misc/AssertionMacros.h"
 #include "Math/IntVector.h"
@@ -178,24 +177,10 @@ namespace mu
 
 
 			//!
-			void Serialise(OutputArchive& arch) const
-			{
-				arch << m_min;
-				arch << m_size;
-				arch << m_id;
-				arch << m_priority;
-				arch << bUseSymmetry;
-			}
+			void Serialise(OutputArchive& arch) const;
 
 			//!
-			void Unserialise(InputArchive& arch)
-			{
-				arch >> m_min;
-				arch >> m_size;
-				arch >> m_id;
-				arch >> m_priority;
-				arch >> bUseSymmetry;
-			}
+			void Unserialise(InputArchive& arch);
 
 			//!
 			inline bool operator==(const FBlock& o) const
@@ -217,14 +202,7 @@ namespace mu
 			}
 
 			//!
-			void UnserialiseOldVersion(InputArchive& Archive, const int32 Version)
-			{
-				// Use the version if the Layout version changes to 6
-				Archive >> m_min;
-				Archive >> m_size;
-				Archive >> m_id;
-				Archive >> m_priority;
-			}
+			void UnserialiseOldVersion(InputArchive& Archive, const int32 Version);
 		};
 
 
@@ -245,62 +223,10 @@ namespace mu
 
 
 		//!
-		void Serialise(OutputArchive& arch) const
-		{
-			uint32 ver = 5;
-			arch << ver;
-
-			arch << m_size;
-			arch << m_blocks;
-
-			arch << m_maxsize;
-			arch << uint32(m_strategy);
-			arch << FirstLODToIgnoreWarnings;
-			arch << uint32(ReductionMethod);
-		}
+		void Serialise(OutputArchive& arch) const;
 
 		//!
-		void Unserialise(InputArchive& arch)
-		{
-			uint32 ver;
-			arch >> ver;
-			check(ver <= 5);
-
-			arch >> m_size;
-
-			if (ver < 5)
-			{
-				uint32_t Size = 0;
-				arch >> Size;
-				m_blocks.SetNum(Size);
-
-				for (uint32_t BlockIndex = 0; BlockIndex < Size; ++BlockIndex)
-				{
-					m_blocks[BlockIndex].UnserialiseOldVersion(arch, ver);
-				}
-			}
-			else
-			{
-				arch >> m_blocks;
-			}
-
-			arch >> m_maxsize;
-
-			uint32 temp;
-			arch >> temp;
-			m_strategy = EPackStrategy(temp);
-
-			if (ver >= 4)
-			{
-				arch >> FirstLODToIgnoreWarnings;
-			}
-
-			if (ver >= 5)
-			{
-				arch >> temp;
-				ReductionMethod = EReductionMethod(temp);
-			}
-		}
+		void Unserialise(InputArchive& arch);
 
 		//!
 		bool IsSimilar(const Layout& o) const;
diff --git a/Engine/Plugins/Experimental/Mutable/Source/MutableRuntime/Public/MuR/Mesh.h b/Engine/Plugins/Experimental/Mutable/Source/MutableRuntime/Public/MuR/Mesh.h
--- a/Engine/Plugins/Experimental/Mutable/Source/MutableRuntime/Public/MuR/Mesh.h
+++ b/Engine/Plugins/Experimental/Mutable/Source/MutableRuntime/Public/MuR/Mesh.h
@@ -15,7 +15,7 @@
 #include "MuR/Ptr.h"
 #include "MuR/RefCounted.h"
 #include "MuR/Serialisation.h"
-#include "MuR/SerialisationPrivate.h"
+#include "MuR/Serialisation.h"
 #include "MuR/Skeleton.h"
 #include "Templates/Tuple.h"
 
@@ -71,7 +71,6 @@ namespace mu
 
 	};
 
-	MUTABLE_DEFINE_POD_VECTOR_SERIALISABLE(MESH_SURFACE);
 
 	enum class EBoneUsageFlags : uint32
 	{
@@ -87,7 +86,6 @@ namespace mu
 	};
 
 	ENUM_CLASS_FLAGS(EBoneUsageFlags);
-	MUTABLE_DEFINE_ENUM_SERIALISABLE(EBoneUsageFlags)
 
 	//!
 	enum class EMeshBufferType
@@ -98,7 +96,6 @@ namespace mu
 		PhysicsBodyDeformSelection,
 		PhysicsBodyDeformOffsets
 	};
-	MUTABLE_DEFINE_ENUM_SERIALISABLE(EMeshBufferType)
 
 	//!
 	enum class EShapeBindingMethod : uint32
@@ -108,7 +105,6 @@ namespace mu
 		ClipDeformClosestToSurface = 2,
 		ClipDeformNormalProject = 3	
 	};
-	MUTABLE_DEFINE_ENUM_SERIALISABLE(EShapeBindingMethod)
 
 	enum class EMeshCloneFlags : uint32
 	{
@@ -365,24 +361,10 @@ namespace mu
 			string m_name;
 			TArray<int32> m_faces;
 
-			inline void Serialise(OutputArchive& arch) const
-			{
-				const int32 ver = 0;
-				arch << ver;
+			inline void Serialise(OutputArchive& arch) const;
 
-				arch << m_name;
-				arch << m_faces;
-			}
 
-			inline void Unserialise(InputArchive& arch)
-			{
-				int32 ver = 0;
-				arch >> ver;
-				check(ver == 0);
-
-				arch >> m_name;
-				arch >> m_faces;
-			}
+			inline void Unserialise(InputArchive& arch);
 
 			//!
 			inline bool operator==(const FACE_GROUP& o) const
@@ -403,37 +385,10 @@ namespace mu
 			EBoneUsageFlags BoneUsageFlags = EBoneUsageFlags::None;
 			FTransform3f BoneTransform;
 
-			inline void Serialise(OutputArchive& arch) const
-			{
-				const int32 ver = 1;
-				arch << ver;
-
-				arch << BoneName;
-				arch << BoneUsageFlags;
-				arch << BoneTransform;
-			}
-
-			inline void Unserialise(InputArchive& arch)
-			{
-				int32 ver = 0;
-				arch >> ver;
-				check(ver <= 1);
-
-				arch >> BoneName;
+			inline void Serialise(OutputArchive& arch) const;
 
-				if (ver == 0)
-				{
-					uint8 Skinned = 0;
-					arch >> Skinned;
-					BoneUsageFlags = Skinned ? EBoneUsageFlags::Skinning : EBoneUsageFlags::None;
-				}
-				else
-				{
-					arch >> BoneUsageFlags;
-				}
 
-				arch >> BoneTransform;
-			}
+			inline void Unserialise(InputArchive& arch);
 
 
 			//!
@@ -448,93 +403,13 @@ namespace mu
 		TArray<FBonePose> BonePoses;
 
 		//!
-		inline void Serialise(OutputArchive& arch) const
-		{
-			uint32 ver = 15;
-			arch << ver;
-
-			arch << m_IndexBuffers;
-			arch << m_VertexBuffers;
-			arch << m_FaceBuffers;
-			arch << m_AdditionalBuffers;
-			arch << m_layouts;
-
-			arch << SkeletonIDs;
-
-			arch << m_pSkeleton;
-			arch << m_pPhysicsBody;
-
-			arch << m_staticFormatFlags;
-			arch << m_surfaces;
-			arch << m_faceGroups;
-
-			arch << m_tags;
-
-			arch << BonePoses;
-
-			arch << AdditionalPhysicsBodies;
-		}
-
-		//!
-		inline void Unserialise(InputArchive& arch)
-		{
-			uint32 ver;
-			arch >> ver;
-			check(ver <= 15);
-
-			arch >> m_IndexBuffers;
-			arch >> m_VertexBuffers;
-			arch >> m_FaceBuffers;
-			arch >> m_AdditionalBuffers;
-			arch >> m_layouts;
-
-			if (ver >= 14)
-			{
-				arch >> SkeletonIDs;
-			}
-
-			arch >> m_pSkeleton;
-			if (ver >= 12)
-			{ 
-				arch >> m_pPhysicsBody;
-			}
-			else
-			{
-				m_pPhysicsBody = nullptr;
-			}
-
-			arch >> m_staticFormatFlags;
-			arch >> m_surfaces;
-			arch >> m_faceGroups;
-
-			arch >> m_tags;
-
-			if (ver >= 13)
-			{
-				arch >> BonePoses;
-			}
-			else if (m_pSkeleton)
-			{
-				const int32 NumBones = m_pSkeleton->GetBoneCount();
-				BonePoses.SetNum(NumBones);
-				check(m_pSkeleton->m_boneTransforms_DEPRECATED.Num() == NumBones);
-
-				for (int32 BoneIndex = 0; BoneIndex < NumBones; ++BoneIndex)
-				{
-					BonePoses[BoneIndex].BoneName = m_pSkeleton->GetBoneName(BoneIndex);
-					BonePoses[BoneIndex].BoneUsageFlags = EBoneUsageFlags::Skinning;
-					BonePoses[BoneIndex].BoneTransform = m_pSkeleton->m_boneTransforms_DEPRECATED[BoneIndex];
-				}
-			}
+		inline void Serialise(OutputArchive& arch) const;
 
-			if (ver >= 15)
-			{
-				arch >> AdditionalPhysicsBodies;
-			}
-		}
+        //!
+		inline void Unserialise(InputArchive& arch);
 
 
-		//!
+        //!
 		inline bool operator==(const Mesh& o) const
 		{
 			bool equal = (m_IndexBuffers == o.m_IndexBuffers);
@@ -645,5 +520,9 @@ namespace mu
     };
 
 
+	MUTABLE_DEFINE_POD_VECTOR_SERIALISABLE(MESH_SURFACE)
+	MUTABLE_DEFINE_ENUM_SERIALISABLE(EBoneUsageFlags)
+	MUTABLE_DEFINE_ENUM_SERIALISABLE(EMeshBufferType)
+	MUTABLE_DEFINE_ENUM_SERIALISABLE(EShapeBindingMethod)
 }
 
diff --git a/Engine/Plugins/Experimental/Mutable/Source/MutableRuntime/Public/MuR/MeshBufferSet.h b/Engine/Plugins/Experimental/Mutable/Source/MutableRuntime/Public/MuR/MeshBufferSet.h
--- a/Engine/Plugins/Experimental/Mutable/Source/MutableRuntime/Public/MuR/MeshBufferSet.h
+++ b/Engine/Plugins/Experimental/Mutable/Source/MutableRuntime/Public/MuR/MeshBufferSet.h
@@ -2,16 +2,14 @@
 
 #pragma once
 
-#include "MuR/SerialisationPrivate.h"
+#include "MuR/Serialisation.h"
 #include "Containers/Array.h"
 #include "HAL/PlatformMath.h"
 #include "MuR/MutableMemory.h"
-#include "MuR/Serialisation.h"
 
 
 namespace mu
 {
-
 	//! Supported formats for the elements in mesh buffers.
 	//! \ingroup runtime
 	typedef enum
@@ -163,11 +161,6 @@ namespace mu
 
 	};
 
-	MUTABLE_DEFINE_POD_SERIALISABLE(MESH_BUFFER_CHANNEL);
-	MUTABLE_DEFINE_POD_VECTOR_SERIALISABLE(MESH_BUFFER_CHANNEL);
-	MUTABLE_DEFINE_ENUM_SERIALISABLE(MESH_BUFFER_FORMAT);
-	MUTABLE_DEFINE_ENUM_SERIALISABLE(MESH_BUFFER_SEMANTIC);
-
 
 	struct MESH_BUFFER
 	{
@@ -178,7 +171,7 @@ namespace mu
 		}
 
 		//!
-		TArray<MESH_BUFFER_CHANNEL> m_channels;
+		TArray<mu::MESH_BUFFER_CHANNEL> m_channels;
 
 		//!
 		TArray<uint8> m_data;
@@ -187,20 +180,10 @@ namespace mu
 		uint32 m_elementSize;
 
 		//!
-		inline void Serialise(OutputArchive& arch) const
-		{
-			arch << m_channels;
-			arch << m_data;
-			arch << m_elementSize;
-		}
+		void Serialise(mu::OutputArchive& arch) const;
 
 		//!
-		inline void Unserialise(InputArchive& arch)
-		{
-			arch >> m_channels;
-			arch >> m_data;
-			arch >> m_elementSize;
-		}
+		inline void Unserialise(mu::InputArchive& arch);
 
 		//!
 		inline bool operator==(const MESH_BUFFER& o) const
@@ -226,18 +209,10 @@ namespace mu
 		TArray<MESH_BUFFER> m_buffers;
 
 		//!
-		inline void Serialise(OutputArchive& arch) const
-		{
-			arch << m_elementCount;
-			arch << m_buffers;
-		}
+		inline void Serialise(OutputArchive& arch) const;
 
 		//!
-		inline void Unserialise(InputArchive& arch)
-		{
-			arch >> m_elementCount;
-			arch >> m_buffers;
-		}
+		inline void Unserialise(InputArchive& arch);
 
 		//!
 		inline bool operator==(const FMeshBufferSet& o) const
@@ -401,7 +376,12 @@ namespace mu
 
 		//!
 		void UpdateOffsets(int32 b);
-	};
+	};	
 
+	
+	MUTABLE_DEFINE_POD_SERIALISABLE(MESH_BUFFER_CHANNEL);
+	MUTABLE_DEFINE_POD_VECTOR_SERIALISABLE(MESH_BUFFER_CHANNEL);
+	MUTABLE_DEFINE_ENUM_SERIALISABLE(MESH_BUFFER_FORMAT);
+	MUTABLE_DEFINE_ENUM_SERIALISABLE(MESH_BUFFER_SEMANTIC);
 }
 
diff --git a/Engine/Plugins/Experimental/Mutable/Source/MutableRuntime/Public/MuR/Model.h b/Engine/Plugins/Experimental/Mutable/Source/MutableRuntime/Public/MuR/Model.h
--- a/Engine/Plugins/Experimental/Mutable/Source/MutableRuntime/Public/MuR/Model.h
+++ b/Engine/Plugins/Experimental/Mutable/Source/MutableRuntime/Public/MuR/Model.h
@@ -104,7 +104,53 @@ namespace mu
         //! that keep models loaded. This could be called when a game finishes, or a change of
         //! level.
         void ClearCaches();
-
+    	
+		//! Return the default value of a boolean parameter.
+		//! \pre The parameter specified by index is a T_BOOL.
+		//! \param Index Index of the parameter from 0 to GetCount()-1
+    	bool GetBoolDefaultValue(int32 Index) const;
+
+   		//! Return the default value of a integer parameter.
+		//! \pre The parameter specified by index is a T_INT.
+		//! \param Index Index of the parameter from 0 to GetCount()-1
+    	int32 GetIntDefaultValue(int32 Index) const;
+
+		//! Return the default value of a float parameter.
+		//! \pre The parameter specified by index is a T_FLOAT.
+		//! \param Index Index of the parameter from 0 to GetCount()-1
+        float GetFloatDefaultValue(int32 Index) const;
+
+		//! Return the default value of a colour parameter.
+		//! \pre The parameter specified by index is a T_FLOAT.
+        //! \param Index Index of the parameter from 0 to GetCount()-1
+        //! \param R,G,B Pointers to values where every resulting colour channel will be stored
+    	void GetColourDefaultValue(int32 Index, float* R, float* G, float* B) const;
+
+    	//! Return the default value of a projector parameter, as a 4x4 matrix. The matrix is supposed to be
+		//! a linear transform in column-major.
+		//! \pre The parameter specified by index is a T_PROJECTOR.
+        //! \param Index Index of the parameter from 0 to GetCount()-1
+        //! \param OutPos Pointer to where the object-space position coordinates of the projector will be stored.
+        //! \param OutDir Pointer to where the object-space direction vector of the projector will be stored.
+        //! \param OutUp Pointer to where the object-space vertically up direction vector
+        //!         of the projector will be stored. This controls the "roll" angle of the
+        //!         projector.
+        //! \param OutScale Pointer to the projector-space scaling of the projector.
+    	void GetProjectorDefaultValue(int32 Index, PROJECTOR_TYPE* OutProjectionType, FVector3f* OutPos,
+    	 	FVector3f* OutDir, FVector3f* OutUp, FVector3f* OutScale, float* OutProjectionAngle) const;
+
+        //! Return the default value of an image parameter.
+        //! \pre The parameter specified by index is a T_IMAGE.
+        //! \param Index Index of the parameter from 0 to GetCount()-1
+		//! \return The externalId specified when setting the image value (\see SetImageValue)
+        EXTERNAL_IMAGE_ID GetImageDefaultValue(int32 Index) const;
+
+    	int32 GetRomCount() const;
+
+		uint32 GetRomId(int32 Index) const;
+
+		uint32 GetRomSize(int32 Index) const;
+    	
 		//-----------------------------------------------------------------------------------------
 		// Interface pattern
 		//-----------------------------------------------------------------------------------------
diff --git a/Engine/Plugins/Experimental/Mutable/Source/MutableRuntime/Public/MuR/PhysicsBody.h b/Engine/Plugins/Experimental/Mutable/Source/MutableRuntime/Public/MuR/PhysicsBody.h
--- a/Engine/Plugins/Experimental/Mutable/Source/MutableRuntime/Public/MuR/PhysicsBody.h
+++ b/Engine/Plugins/Experimental/Mutable/Source/MutableRuntime/Public/MuR/PhysicsBody.h
@@ -2,7 +2,9 @@
 
 #pragma once
 
-#include "MuR/SerialisationPrivate.h"
+#include <string>
+
+#include "MuR/Serialisation.h"
 #include "MuR/Ptr.h"
 #include "MuR/RefCounted.h"
 #include "MuR/Serialisation.h"
@@ -17,6 +19,8 @@
 
 namespace mu
 {
+	typedef std::string string;
+
 	struct FBodyShape
 	{
 		string Name;
@@ -28,25 +32,10 @@ namespace mu
 		}
 	
 		//!
-		inline void Serialise( OutputArchive& arch ) const
-		{
-			const uint32 ver = 0;
-			arch << ver;
-
-			arch << Name;
-			arch << Flags;
-		}
+		inline void Serialise( OutputArchive& arch ) const;
 
 		//!
-		inline void Unserialise( InputArchive& arch )
-		{
-			uint32 ver;
-			arch >> ver;
-			check(ver == 0);
-
-			arch >> Name;
-			arch >> Flags;
-		}
+		inline void Unserialise( InputArchive& arch );
 	};
 
 	struct FSphereBody : FBodyShape
@@ -60,29 +49,10 @@ namespace mu
 		}
 		 
 		//!
-		inline void Serialise( OutputArchive& arch ) const
-		{
-			FBodyShape::Serialise(arch);
-
-			const uint32 ver = 0;
-			arch << ver;
-
-			arch << Position;
-			arch << Radius;
-		}
+		inline void Serialise( OutputArchive& arch ) const;
 
 		//!
-		inline void Unserialise( InputArchive& arch )
-		{
-			FBodyShape::Unserialise(arch);
-
-			uint32 ver;
-			arch >> ver;
-			check(ver == 0);
-
-			arch >> Position;
-			arch >> Radius;
-		} 
+		inline void Unserialise( InputArchive& arch );
 	};
 	
 	struct FBoxBody : FBodyShape
@@ -97,31 +67,10 @@ namespace mu
 		}
 
 		//!
-		inline void Serialise( OutputArchive& arch ) const
-		{
-			FBodyShape::Serialise( arch );
+		inline void Serialise( OutputArchive& arch ) const;
 
-			const uint32 ver = 0;
-			arch << ver;
-
-			arch << Position;
-			arch << Orientation;
-			arch << Size;
-		}
-
-		  //!
-		inline void Unserialise( InputArchive& arch )
-		{
-			FBodyShape::Unserialise( arch );
-
-			uint32 ver;
-			arch >> ver;
-			check(ver == 0);
-
-			arch >> Position;
-			arch >> Orientation;
-			arch >> Size;
-		} 
+		//!
+		inline void Unserialise( InputArchive& arch );
 	};
 
 	struct FSphylBody : FBodyShape
@@ -141,33 +90,10 @@ namespace mu
 		}
 		 
 		//!
-		inline void Serialise( OutputArchive& arch ) const
-		{
-			FBodyShape::Serialise( arch );
-
-			const uint32 ver = 0;
-			arch << ver;
-
-			arch << Position;
-			arch << Orientation;
-			arch << Radius;
-			arch << Length;
-		}
+		inline void Serialise( OutputArchive& arch ) const;
 
 		//!
-		inline void Unserialise( InputArchive& arch )
-		{
-			FBodyShape::Unserialise( arch );
-
-			uint32 ver;
-			arch >> ver;
-			check(ver == 0);
-
-			arch >> Position;
-			arch >> Orientation;
-			arch >> Radius;
-			arch >> Length;
-		}
+		inline void Unserialise( InputArchive& arch );
 	};
 	
 	struct FTaperedCapsuleBody : FBodyShape
@@ -189,35 +115,10 @@ namespace mu
 		}
 		 
 		//!
-		inline void Serialise( OutputArchive& arch ) const
-		{
-			FBodyShape::Serialise( arch );
+		inline void Serialise( OutputArchive& arch ) const;
 
-			const uint32 ver = 0;
-			arch << ver;
-
-			arch << Position;
-			arch << Orientation;
-			arch << Radius0;
-			arch << Radius1;
-			arch << Length;
-		}
-
-		 //!
-		inline void Unserialise( InputArchive& arch )
-		{
-			FBodyShape::Unserialise( arch );
-			
-			uint32 ver;
-			arch >> ver;
-			check(ver == 0);
-
-			arch >> Position;
-			arch >> Orientation;
-			arch >> Radius0;
-			arch >> Radius1;
-			arch >> Length;
-		 }
+		//!
+		inline void Unserialise( InputArchive& arch );
 	};
 	
 	struct FConvexBody : FBodyShape
@@ -236,29 +137,10 @@ namespace mu
 		}
 
 		//!
-		inline void Serialise( OutputArchive& arch ) const
-		{
-			FBodyShape::Serialise( arch );
-	
-		 	uint32 ver = 0;
-		 	arch << ver;
+		inline void Serialise( OutputArchive& arch ) const;
 
-		 	arch << Vertices;
-		 	arch << Indices;
-		 	arch << Transform;
-		}
-
-		  //!
-		inline void Unserialise( InputArchive& arch )
-		{
-		 	uint32 ver;
-		 	arch >> ver;
-		 	check(ver == 0);
-
-		 	arch >> Vertices;
-		 	arch >> Indices;
-		 	arch >> Transform;
-		}
+		//!
+		inline void Unserialise( InputArchive& arch );
 	};
 	
 	struct FPhysicsBodyAggregate
@@ -279,31 +161,10 @@ namespace mu
 		}
 
 		 //!
-		inline void Serialise( OutputArchive& arch ) const
-		{
-		 	uint32 ver = 0;
-		 	arch << ver;
-		
-		 	arch << Spheres;
-		 	arch << Boxes;
-		 	arch << Convex;
-		 	arch << Sphyls;
-		 	arch << TaperedCapsules;
-		}
+		inline void Serialise( OutputArchive& arch ) const;
 
-		 //!
-		inline void Unserialise( InputArchive& arch )
-		{
-		 	uint32 ver;
-		 	arch >> ver;
-		 	check(ver == 0);
-		
-		 	arch >> Spheres;
-		 	arch >> Boxes;
-		 	arch >> Convex;
-		 	arch >> Sphyls;
-		 	arch >> TaperedCapsules;
-		}
+		//!
+		inline void Unserialise( InputArchive& arch );
 	};
 
 	class MUTABLERUNTIME_API PhysicsBody : public RefCounted
@@ -444,41 +305,12 @@ namespace mu
 		bool bBodiesModified = false;
 
         //!
-        inline void Serialise( OutputArchive& arch ) const
-		{
-            uint32 ver = 2;
-			arch << ver;
-        	
-			arch << CustomId;
-        	arch << Bodies;
-        	arch << Bones;
-        	arch << BodiesCustomIds;
-			arch << bBodiesModified;
-        }
+        inline void Serialise( OutputArchive& arch ) const;
 
-        //!
-        inline void Unserialise( InputArchive& arch )
-		{
-            uint32 ver;
-			arch >> ver;
-            check(ver <= 2);
-        	
-			if (ver >= 2)
-			{
-				arch >> CustomId;
-			}
-        	
-        	arch >> Bodies;
-        	arch >> Bones;
-        	arch >> BodiesCustomIds;
-
-			if (ver >= 1)
-			{
-				arch >> bBodiesModified;
-			}
-		}
-		
-        //!
+		//!
+        inline void Unserialise( InputArchive& arch );
+
+		//!
 		inline bool operator==(const PhysicsBody& Other) const
         {
         	return CustomId        == Other.CustomId        &&
diff --git a/Engine/Plugins/Experimental/Mutable/Source/MutableRuntime/Public/MuR/Serialisation.h b/Engine/Plugins/Experimental/Mutable/Source/MutableRuntime/Public/MuR/Serialisation.h
--- a/Engine/Plugins/Experimental/Mutable/Source/MutableRuntime/Public/MuR/Serialisation.h
+++ b/Engine/Plugins/Experimental/Mutable/Source/MutableRuntime/Public/MuR/Serialisation.h
@@ -13,6 +13,31 @@ namespace mu
     class Image;
 	class Model;
 
+
+#define MUTABLE_DEFINE_POD_SERIALISABLE(T)										\
+	template<>																	\
+	MUTABLERUNTIME_API void operator<< <T>(OutputArchive& arch, const T& t);	\
+																				\
+	template<>																	\
+	MUTABLERUNTIME_API void operator>> <T>(InputArchive& arch, T& t);			\
+
+
+#define MUTABLE_DEFINE_POD_VECTOR_SERIALISABLE(T)										\
+	template<>																			\
+	MUTABLERUNTIME_API void operator<< <T>(OutputArchive& arch, const TArray<T>& v);	\
+																						\
+	template<>																			\
+	MUTABLERUNTIME_API void operator>> <T>(InputArchive& arch, TArray<T>& v);			\
+
+
+#define MUTABLE_DEFINE_ENUM_SERIALISABLE(T)										\
+	template<>																	\
+    MUTABLERUNTIME_API void operator<< <T>(OutputArchive& arch, const T& t);	\
+																				\
+	template<>																	\
+    MUTABLERUNTIME_API void operator>> <T>(InputArchive& arch, T& t);			\
+	
+	
     //! \brief
     //! \ingroup model
     template<class R>
@@ -286,5 +311,42 @@ namespace mu
 
     };
 
+
+	template< typename T >
+	void operator<< ( OutputArchive& arch, const T& t )
+	{
+        t.Serialise( arch );
+	}
+
+	template< typename T >
+	void operator>> ( InputArchive& arch, T& t )
+	{
+        t.Unserialise( arch );
+	}
+
+	
+	//---------------------------------------------------------------------------------------------
+	template<typename T> void operator<<(OutputArchive& arch, const TArray<T>& v)
+	{
+		// TODO: Optimise for vectors of PODs
+		arch << (uint32)v.Num();
+		for (std::size_t i = 0; i < v.Num(); ++i)
+		{
+			arch << v[i];
+		}
+	}
+
+	template<typename T> void operator>>(InputArchive& arch, TArray<T>& v)
+	{
+		// TODO: Optimise for vectors of PODs
+		uint32 size;
+		arch >> size;
+		v.SetNum(size);
+		for (std::size_t i = 0; i < size; ++i)
+		{
+			arch >> v[i];
+		}
+	}
+
 }
 
diff --git a/Engine/Plugins/Experimental/Mutable/Source/MutableRuntime/Public/MuR/Skeleton.h b/Engine/Plugins/Experimental/Mutable/Source/MutableRuntime/Public/MuR/Skeleton.h
--- a/Engine/Plugins/Experimental/Mutable/Source/MutableRuntime/Public/MuR/Skeleton.h
+++ b/Engine/Plugins/Experimental/Mutable/Source/MutableRuntime/Public/MuR/Skeleton.h
@@ -2,7 +2,7 @@
 
 #pragma once
 
-#include "MuR/SerialisationPrivate.h"
+#include "MuR/Serialisation.h"
 #include "Containers/Array.h"
 #include "HAL/PlatformMath.h"
 #include "Math/Transform.h"
@@ -10,12 +10,14 @@
 #include "Misc/AssertionMacros.h"
 #include "MuR/Ptr.h"
 #include "MuR/RefCounted.h"
-#include "MuR/Serialisation.h"
+
+#include <string>
 
 
 namespace mu
 {
-
+	typedef std::string string;
+	
 	// Forward references
     class Skeleton;
 
@@ -77,43 +79,10 @@ namespace mu
 
 
 		//!
-		inline void Serialise(OutputArchive& arch) const
-		{
-			uint32 ver = 5;
-			arch << ver;
-
-			arch << m_bones;
-			arch << m_boneParents;
-		}
+		inline void Serialise(OutputArchive& arch) const;
 
 		//!
-		inline void Unserialise(InputArchive& arch)
-		{
-			uint32 ver;
-			arch >> ver;
-			check(ver >= 3);
-
-			arch >> m_bones;
-
-			if (ver == 3)
-			{
-				arch >> m_boneTransforms_DEPRECATED;
-			}
-
-			arch >> m_boneParents;
-
-			if (ver <= 4)
-			{
-				TArray<int32> boneIds;
-				arch >> boneIds;
-			}
-
-			if (ver == 3)
-			{
-				bool bBoneTransformModified;
-				arch >> bBoneTransformModified;
-			}
-		}
+		inline void Unserialise(InputArchive& arch);
 
 
 		//!
diff --git a/Engine/Plugins/Experimental/Mutable/Source/MutableTools/Private/MuT/CompilerPrivate.h b/Engine/Plugins/Experimental/Mutable/Source/MutableTools/Private/MuT/CompilerPrivate.h
--- a/Engine/Plugins/Experimental/Mutable/Source/MutableTools/Private/MuT/CompilerPrivate.h
+++ b/Engine/Plugins/Experimental/Mutable/Source/MutableTools/Private/MuT/CompilerPrivate.h
@@ -13,7 +13,7 @@
 
 namespace mu
 {
-	MUTABLE_DEFINE_ENUM_SERIALISABLE(ETextureCompressionStrategy);
+	MUTABLE_IMPLEMENT_ENUM_SERIALISABLE(ETextureCompressionStrategy);
 
     //!
     struct FStateOptimizationOptions
diff --git a/Engine/Plugins/Experimental/Mutable/Source/MutableTools/Private/MuT/NodeColourArithmeticOperationPrivate.h b/Engine/Plugins/Experimental/Mutable/Source/MutableTools/Private/MuT/NodeColourArithmeticOperationPrivate.h
--- a/Engine/Plugins/Experimental/Mutable/Source/MutableTools/Private/MuT/NodeColourArithmeticOperationPrivate.h
+++ b/Engine/Plugins/Experimental/Mutable/Source/MutableTools/Private/MuT/NodeColourArithmeticOperationPrivate.h
@@ -9,7 +9,7 @@
 namespace mu
 {
 
-	MUTABLE_DEFINE_ENUM_SERIALISABLE(NodeColourArithmeticOperation::OPERATION)
+	MUTABLE_IMPLEMENT_ENUM_SERIALISABLE(NodeColourArithmeticOperation::OPERATION)
 
 	class NodeColourArithmeticOperation::Private : public NodeColour::Private
 	{
diff --git a/Engine/Plugins/Experimental/Mutable/Source/MutableTools/Private/MuT/NodeObjectGroupPrivate.h b/Engine/Plugins/Experimental/Mutable/Source/MutableTools/Private/MuT/NodeObjectGroupPrivate.h
--- a/Engine/Plugins/Experimental/Mutable/Source/MutableTools/Private/MuT/NodeObjectGroupPrivate.h
+++ b/Engine/Plugins/Experimental/Mutable/Source/MutableTools/Private/MuT/NodeObjectGroupPrivate.h
@@ -11,7 +11,7 @@
 namespace mu
 {
 
-	MUTABLE_DEFINE_ENUM_SERIALISABLE( NodeObjectGroup::CHILD_SELECTION )
+	MUTABLE_IMPLEMENT_ENUM_SERIALISABLE( NodeObjectGroup::CHILD_SELECTION )
 
 
 	class NodeObjectGroup::Private : public NodeObject::Private
diff --git a/Engine/Plugins/Experimental/Mutable/Source/MutableTools/Private/MuT/NodeScalarArithmeticOperationPrivate.h b/Engine/Plugins/Experimental/Mutable/Source/MutableTools/Private/MuT/NodeScalarArithmeticOperationPrivate.h
--- a/Engine/Plugins/Experimental/Mutable/Source/MutableTools/Private/MuT/NodeScalarArithmeticOperationPrivate.h
+++ b/Engine/Plugins/Experimental/Mutable/Source/MutableTools/Private/MuT/NodeScalarArithmeticOperationPrivate.h
@@ -9,7 +9,7 @@
 namespace mu
 {
 
-    MUTABLE_DEFINE_ENUM_SERIALISABLE(NodeScalarArithmeticOperation::OPERATION)
+    MUTABLE_IMPLEMENT_ENUM_SERIALISABLE(NodeScalarArithmeticOperation::OPERATION)
 
     class NodeScalarArithmeticOperation::Private : public NodeScalar::Private
 	{
diff --git a/Engine/Plugins/Experimental/Mutable/Source/MutableTools/Private/MuT/TablePrivate.h b/Engine/Plugins/Experimental/Mutable/Source/MutableTools/Private/MuT/TablePrivate.h
--- a/Engine/Plugins/Experimental/Mutable/Source/MutableTools/Private/MuT/TablePrivate.h
+++ b/Engine/Plugins/Experimental/Mutable/Source/MutableTools/Private/MuT/TablePrivate.h
@@ -14,7 +14,7 @@
 namespace mu
 {
 
-	MUTABLE_DEFINE_ENUM_SERIALISABLE( TABLE_COLUMN_TYPE )
+	MUTABLE_IMPLEMENT_ENUM_SERIALISABLE( TABLE_COLUMN_TYPE )
 
 
 	struct TABLE_COLUMN
