[Mutable] Fix Clang compialtion errors.

#rnx

diff --git a/Engine/Plugins/Experimental/Mutable/Source/CustomizableObject/Public/MuCO/CustomizableObjectParameterTypeDefinitions.h b/Engine/Plugins/Experimental/Mutable/Source/CustomizableObject/Public/MuCO/CustomizableObjectParameterTypeDefinitions.h
--- a/Engine/Plugins/Experimental/Mutable/Source/CustomizableObject/Public/MuCO/CustomizableObjectParameterTypeDefinitions.h
+++ b/Engine/Plugins/Experimental/Mutable/Source/CustomizableObject/Public/MuCO/CustomizableObjectParameterTypeDefinitions.h
@@ -279,7 +279,7 @@ struct FCustomizableObjectProjectorParameterValue
 {
 	GENERATED_USTRUCT_BODY()
 
-	static constexpr FCustomizableObjectProjector DEFAULT_PARAMETER_VALUE;
+	inline static const FCustomizableObjectProjector DEFAULT_PARAMETER_VALUE = {};
 
 	UPROPERTY(Category = CustomizableObjectVectorParameterValue, VisibleAnywhere)
 	FString ParameterName;
