[Mutable] Fix MipDataProvider synchronization issues.

- Fixed synchronization issue where the PollMips could have been called before completing the task, completely skipping the copy of the MIP. Even though the probability of this issue was low, it was plausible and could be manifested with MIP corruption.
- Added CPU scopes.

#rnx

diff --git a/Engine/Plugins/Experimental/Mutable/Source/CustomizableObject/Private/MuCO/CustomizableObjectMipDataProvider.cpp b/Engine/Plugins/Experimental/Mutable/Source/CustomizableObject/Private/MuCO/CustomizableObjectMipDataProvider.cpp
--- a/Engine/Plugins/Experimental/Mutable/Source/CustomizableObject/Private/MuCO/CustomizableObjectMipDataProvider.cpp
+++ b/Engine/Plugins/Experimental/Mutable/Source/CustomizableObject/Private/MuCO/CustomizableObjectMipDataProvider.cpp
@@ -165,6 +165,8 @@ namespace impl
 
 int32 FMutableTextureMipDataProvider::GetMips(const FTextureUpdateContext& Context, int32 StartingMipIndex, const FTextureMipInfoArray& MipInfos, const FTextureUpdateSyncOptions& SyncOptions)
 {
+	MUTABLE_CPUPROFILER_SCOPE(FMutableTextureMipDataProvider::GetMips)
+
 #if WITH_EDITOR
 	check(Context.Texture->HasPendingInitOrStreaming());
 	check(CustomizableObjectInstance->GetCustomizableObject());
@@ -261,6 +263,11 @@ int32 FMutableTextureMipDataProvider::GetMips(const FTextureUpdateContext& Conte
 
 bool FMutableTextureMipDataProvider::PollMips(const FTextureUpdateSyncOptions& SyncOptions)
 {
+	MUTABLE_CPUPROFILER_SCOPE(FMutableTextureMipDataProvider::PollMips)
+
+	// Once this point is reached, even if the task has not been completed, we know that all the work we need from it has been completed.
+	// Furthermore, checking if the task is completed is incorrect since PollMips could have been called by RescheduleCallback (before completing the task).
+	
 #if WITH_EDITOR
 	check(CustomizableObjectInstance->GetCustomizableObject());
 	if (CustomizableObjectInstance->GetCustomizableObject()->IsLocked())
@@ -271,41 +278,32 @@ bool FMutableTextureMipDataProvider::PollMips(const FTextureUpdateSyncOptions& S
 	}
 #endif
 
-	if (!bRequestAborted 
-		&& 
-#ifdef MUTABLE_USE_NEW_TASKGRAPH
-		UpdateImageMutableTaskEvent.IsCompleted()
-#else
-		UpdateImageMutableTaskEvent
-#endif
-		)
-	{
-		if (OperationData && OperationData->Result && OperationData->Levels.Num())
-		{
-			// The counter must be zero meaning the Mutable image operation has finished
-			check(SyncOptions.Counter->GetValue() == 0);
-
-			mu::ImagePtrConst Mip = OperationData->Result;
-			int32 MipIndex = 0;
-			check(Mip->GetSizeX() == OperationData->Levels[0].SizeX);
-			check(Mip->GetSizeY() == OperationData->Levels[0].SizeY);
-
-			for (FMutableMipUpdateLevel& Level : OperationData->Levels)
-			{
-				// Check Mip DataSize for consistency, but skip if 0 because it's optional and might be zero in cooked mips
-				check(Level.DataSize == 0 || Mip->GetLODDataSize(MipIndex) == Level.DataSize);
-				void* Dest = Level.Dest;
-				FMemory::Memcpy(Dest, Mip->GetMipData(MipIndex), Mip->GetLODDataSize(MipIndex));
-				++MipIndex;
-			}
-		}
-	}
-	else if (bRequestAborted)
+	if (bRequestAborted)
 	{
 		OperationData = nullptr;
 		AdvanceTo(ETickState::CleanUp, ETickThread::Async);
 		return false;
 	}
+	
+	if (OperationData && OperationData->Result && OperationData->Levels.Num())
+	{
+		// The counter must be zero meaning the Mutable image operation has finished
+		check(SyncOptions.Counter->GetValue() == 0);
+
+		mu::ImagePtrConst Mip = OperationData->Result;
+		int32 MipIndex = 0;
+		check(Mip->GetSizeX() == OperationData->Levels[0].SizeX);
+		check(Mip->GetSizeY() == OperationData->Levels[0].SizeY);
+
+		for (FMutableMipUpdateLevel& Level : OperationData->Levels)
+		{
+			// Check Mip DataSize for consistency, but skip if 0 because it's optional and might be zero in cooked mips
+			check(Level.DataSize == 0 || Mip->GetLODDataSize(MipIndex) == Level.DataSize);
+			void* Dest = Level.Dest;
+			FMemory::Memcpy(Dest, Mip->GetMipData(MipIndex), Mip->GetLODDataSize(MipIndex));
+			++MipIndex;
+		}
+	}
 
 	OperationData = nullptr;
 	AdvanceTo(ETickState::Done, ETickThread::None);
