[mutable] Added a Clip Deform Mesh node column to Tag Explorer.
+Improved Tag Explorer Usability by removing the "Find Tags" button. Now the method is called automatically when the combo box content is generated.

#rnx

diff --git a/Engine/Plugins/Experimental/Mutable/Source/CustomizableObjectEditor/Private/MuCOE/SCustomizableObjectEditorTagExplorer.cpp b/Engine/Plugins/Experimental/Mutable/Source/CustomizableObjectEditor/Private/MuCOE/SCustomizableObjectEditorTagExplorer.cpp
--- a/Engine/Plugins/Experimental/Mutable/Source/CustomizableObjectEditor/Private/MuCOE/SCustomizableObjectEditorTagExplorer.cpp
+++ b/Engine/Plugins/Experimental/Mutable/Source/CustomizableObjectEditor/Private/MuCOE/SCustomizableObjectEditorTagExplorer.cpp
@@ -2,6 +2,7 @@
 
 #include "MuCOE/SCustomizableObjectEditorTagExplorer.h"
 
+#include "DetailLayoutBuilder.h"
 #include "Framework/Views/TableViewMetadata.h"
 #include "HAL/PlatformApplicationMisc.h"
 #include "MuCOE/CustomizableObjectEditor.h"
@@ -10,10 +11,12 @@
 #include "MuCOE/Nodes/CustomizableObjectNodeExtendMaterial.h"
 #include "MuCOE/Nodes/CustomizableObjectNodeMaterial.h"
 #include "MuCOE/Nodes/CustomizableObjectNodeMaterialVariation.h"
+#include "MuCOE/Nodes/CustomizableObjectNodeMeshClipDeform.h"
 #include "MuCOE/Nodes/CustomizableObjectNodeMeshClipMorph.h"
 #include "MuCOE/Nodes/CustomizableObjectNodeMeshClipWithMesh.h"
 #include "Widgets/Input/SButton.h"
-#include "Widgets/Input/SComboBox.h"
+#include "Widgets/Input/SComboButton.h"
+#include "Widgets/Views/SListView.h"
 
 class ITableRow;
 class STableViewBase;
@@ -30,20 +33,6 @@ void SCustomizableObjectEditorTagExplorer::Construct(const FArguments & InArgs)
 	[
 		SNew(SVerticalBox)
 
-		+ SVerticalBox::Slot()
-		.AutoHeight()
-		.Padding(5.0f)
-		[
-			SNew(SHorizontalBox)
-
-			+ SHorizontalBox::Slot()
-			.AutoWidth()
-			[
-				SNew(SButton).Text(LOCTEXT("FindTags", "Find Tags")).OnClicked(this, &SCustomizableObjectEditorTagExplorer::GetCustomizableObjectTags)
-				.ToolTipText(LOCTEXT("FindTagToolTip", "Find all tags related to this Customizable Object."))
-			]
-		]
-
 		+ SVerticalBox::Slot()
 		.AutoHeight()
 		.Padding(5.0f)
@@ -59,13 +48,13 @@ void SCustomizableObjectEditorTagExplorer::Construct(const FArguments & InArgs)
 			+ SHorizontalBox::Slot()
 			.Padding(5.0f, 0, 0, 0)
 			[
-				SAssignNew(TagComboBox, SComboBox<TSharedPtr<FString>>)
-				.OptionsSource(&ComboboxTags)
-				.OnGenerateWidget(this, &SCustomizableObjectEditorTagExplorer::MakeComboButtonItemWidget)
-				.OnSelectionChanged(this, &SCustomizableObjectEditorTagExplorer::OnComboBoxSelectionChanged)
-				.InitiallySelectedItem(TagComboBoxItem)
+				SAssignNew(TagComboBox, SComboButton)
+				.OnGetMenuContent(this, &SCustomizableObjectEditorTagExplorer::OnGetTagsMenuContent)
+				.VAlign(VAlign_Center)
+				.ButtonContent()
 				[
 					SNew(STextBlock)
+					.Font(IDetailLayoutBuilder::GetDetailFont())
 					.Text(this, &SCustomizableObjectEditorTagExplorer::GetCurrentItemLabel)
 				]
 			]
@@ -102,7 +91,7 @@ void SCustomizableObjectEditorTagExplorer::Construct(const FArguments & InArgs)
 				(
 					SNew(SHeaderRow)
 					+ SHeaderRow::Column("MaterialNode")
-					.DefaultLabel(LOCTEXT("MaterialNode", "Material Node"))
+					.DefaultLabel(LOCTEXT("MaterialNode_ColumnName", "Material Node"))
 					.FillWidth(0.3f)
 					.HAlignHeader(EHorizontalAlignment::HAlign_Center)
 					.HAlignCell(EHorizontalAlignment::HAlign_Center)
@@ -119,7 +108,7 @@ void SCustomizableObjectEditorTagExplorer::Construct(const FArguments & InArgs)
 				(
 					SNew(SHeaderRow)
 					+ SHeaderRow::Column("ClipMeshNode")
-					.DefaultLabel(LOCTEXT("ClipMeshNode", "Clip Mesh Node"))
+					.DefaultLabel(LOCTEXT("ClipMeshNode_ColumnName", "Clip Mesh Node"))
 					.FillWidth(0.3f)
 					.HAlignHeader(EHorizontalAlignment::HAlign_Center)
 					.HAlignCell(EHorizontalAlignment::HAlign_Center)
@@ -136,7 +125,7 @@ void SCustomizableObjectEditorTagExplorer::Construct(const FArguments & InArgs)
 				(
 					SNew(SHeaderRow)
 					+ SHeaderRow::Column("ClipMorphNode")
-					.DefaultLabel(LOCTEXT("ClipMorphNode", "Clip Morph Node"))
+					.DefaultLabel(LOCTEXT("ClipMorphNode_ColumnName", "Clip Morph Node"))
 					.FillWidth(0.3f)
 					.HAlignHeader(EHorizontalAlignment::HAlign_Center)
 					.HAlignCell(EHorizontalAlignment::HAlign_Center)
@@ -153,7 +142,7 @@ void SCustomizableObjectEditorTagExplorer::Construct(const FArguments & InArgs)
 				(
 					SNew(SHeaderRow)
 					+ SHeaderRow::Column("VariationNode")
-					.DefaultLabel(LOCTEXT("VariationNode", "Variation Node"))
+					.DefaultLabel(LOCTEXT("VariationNode_ColumnName", "Variation Node"))
 					.FillWidth(0.3f)
 					.HAlignHeader(EHorizontalAlignment::HAlign_Center)
 					.HAlignCell(EHorizontalAlignment::HAlign_Center)
@@ -161,26 +150,43 @@ void SCustomizableObjectEditorTagExplorer::Construct(const FArguments & InArgs)
 			]
 
 			+ SHorizontalBox::Slot()
-				[
-					SAssignNew(ColumnExtend, SListView<UCustomizableObjectNode*>)
-					.ListItemsSource(&ExtendNodes)
-					.OnGenerateRow(this, &SCustomizableObjectEditorTagExplorer::OnGenerateTableRow)
-					.OnSelectionChanged(this, &SCustomizableObjectEditorTagExplorer::OnTagTableSelectionChanged)
-					.HeaderRow
-					(
-						SNew(SHeaderRow)
-						+ SHeaderRow::Column("ExtendNode")
-						.DefaultLabel(LOCTEXT("ExtendNode", "Extend Node"))
-						.FillWidth(0.3f)
-						.HAlignHeader(EHorizontalAlignment::HAlign_Center)
-						.HAlignCell(EHorizontalAlignment::HAlign_Center)
-					)
-				]
+			[
+				SAssignNew(ColumnExtend, SListView<UCustomizableObjectNode*>)
+				.ListItemsSource(&ExtendNodes)
+				.OnGenerateRow(this, &SCustomizableObjectEditorTagExplorer::OnGenerateTableRow)
+				.OnSelectionChanged(this, &SCustomizableObjectEditorTagExplorer::OnTagTableSelectionChanged)
+				.HeaderRow
+				(
+					SNew(SHeaderRow)
+					+ SHeaderRow::Column("ExtendNode")
+					.DefaultLabel(LOCTEXT("ExtendNode_ColumnName", "Extend Node"))
+					.FillWidth(0.3f)
+					.HAlignHeader(EHorizontalAlignment::HAlign_Center)
+					.HAlignCell(EHorizontalAlignment::HAlign_Center)
+				)
+			]
+
+			+ SHorizontalBox::Slot()
+			[
+				SAssignNew(ColumnClipDeform, SListView<UCustomizableObjectNode*>)
+				.ListItemsSource(&ClipDeformNodes)
+				.OnGenerateRow(this, &SCustomizableObjectEditorTagExplorer::OnGenerateTableRow)
+				.OnSelectionChanged(this, &SCustomizableObjectEditorTagExplorer::OnTagTableSelectionChanged)
+				.HeaderRow
+				(
+					SNew(SHeaderRow)
+					+ SHeaderRow::Column("ClipDeform")
+					.DefaultLabel(LOCTEXT("ClipDeform_ColumnName", "Clip Deform Node"))
+					.FillWidth(0.3f)
+					.HAlignHeader(EHorizontalAlignment::HAlign_Center)
+					.HAlignCell(EHorizontalAlignment::HAlign_Center)
+				)
+			]
 		]
 	];
 }
 
-FReply SCustomizableObjectEditorTagExplorer::GetCustomizableObjectTags()
+TSharedRef<SWidget> SCustomizableObjectEditorTagExplorer::OnGetTagsMenuContent()
 {
 	bool bMultipleBaseObjectsFound = false;
 	NodeTags.Empty();
@@ -200,22 +206,11 @@ FReply SCustomizableObjectEditorTagExplorer::GetCustomizableObjectTags()
 			CustomizableObjectFamily.Add(AbsoluteCO);
 			// Stores external customizable objects
 			CustomizableObjectEditorPtr->GetExternalChildObjects(AbsoluteCO, CustomizableObjectFamily, true);
-			ComboboxTags.Empty();
 
 			for (UCustomizableObject* CustObject : CustomizableObjectFamily)
 			{
 				FillTagInformation(CustObject, Tags);
 			}
-
-			for (FString tag : Tags)
-			{
-				ComboboxTags.Add(MakeShareable(new FString(tag)));
-			}
-
-			if (TagComboBox.IsValid())
-			{
-				TagComboBox.ToSharedRef()->RefreshOptions();
-			}
 		}
 	}
 	else
@@ -231,9 +226,21 @@ FReply SCustomizableObjectEditorTagExplorer::GetCustomizableObjectTags()
 		}
 	}
 
-	TagComboBox.ToSharedRef()->RefreshOptions();
+	if (Tags.Num())
+	{
+		FMenuBuilder MenuBuilder(true, NULL);
 
-	return FReply::Handled();
+		for (int32 TagIndex = 0; TagIndex < Tags.Num(); ++TagIndex)
+		{
+			FText TagText = FText::FromString(Tags[TagIndex]);
+			FUIAction Action(FExecuteAction::CreateSP(this, &SCustomizableObjectEditorTagExplorer::OnComboBoxSelectionChanged, Tags[TagIndex]));
+			MenuBuilder.AddMenuEntry(TagText, FText::GetEmpty(), FSlateIcon(), Action);
+		}
+
+		return MenuBuilder.MakeWidget();
+	}
+
+	return SNullWidget::NullWidget;
 }
 
 void SCustomizableObjectEditorTagExplorer::FillTagInformation(UCustomizableObject * Object, TArray<FString>& Tags)
@@ -301,6 +308,18 @@ void SCustomizableObjectEditorTagExplorer::FillTagInformation(UCustomizableObjec
 					}
 				}
 			}
+
+			if (UCustomizableObjectNodeMeshClipDeform* TypedNodeClipDeform = Cast<UCustomizableObjectNodeMeshClipDeform>(Node))
+			{
+				for (int i = 0; i < TypedNodeClipDeform->Tags.Num(); ++i)
+				{
+					NodeTags.Add(TypedNodeClipDeform->Tags[i], TypedNodeClipDeform);
+					if (Tags.Find(TypedNodeClipDeform->Tags[i]) == INDEX_NONE)
+					{
+						Tags.Add(TypedNodeClipDeform->Tags[i]);
+					}
+				}
+			}
 		}
 	}
 }
@@ -312,9 +331,9 @@ TSharedRef<SWidget> SCustomizableObjectEditorTagExplorer::MakeComboButtonItemWid
 
 FText SCustomizableObjectEditorTagExplorer::GetCurrentItemLabel() const
 {
-	if (TagComboBoxItem.IsValid())
+	if (!SelectedTag.IsEmpty())
 	{
-		return FText::FromString(*TagComboBoxItem);
+		return FText::FromString(*SelectedTag);
 	}
 
 	return LOCTEXT("InvalidComboEntryText", "None");
@@ -322,69 +341,72 @@ FText SCustomizableObjectEditorTagExplorer::GetCurrentItemLabel() const
 
 FReply SCustomizableObjectEditorTagExplorer::CopyTagToClipboard()
 {
-	if (TagComboBoxItem.IsValid())
+	if (!SelectedTag.IsEmpty())
 	{
-		FString TagName = *TagComboBoxItem.Get();
-		FPlatformApplicationMisc::ClipboardCopy(*TagName);
+		FPlatformApplicationMisc::ClipboardCopy(*SelectedTag);
 	}
 
 	return FReply::Handled();
 }
 
-void SCustomizableObjectEditorTagExplorer::OnComboBoxSelectionChanged(TSharedPtr<FString> NewValue, ESelectInfo::Type)
+void SCustomizableObjectEditorTagExplorer::OnComboBoxSelectionChanged(FString NewValue)
 {
-	TagComboBoxItem = NewValue;
+	SelectedTag = NewValue;
 
-	for (TSharedPtr<FString> tag : ComboboxTags)
+	if (!NewValue.IsEmpty())
 	{
-		if (tag == NewValue)
-		{
-			TArray<UCustomizableObjectNode*> auxNodes;
-			NodeTags.MultiFind(*tag.Get(), auxNodes, false);
+		TArray<UCustomizableObjectNode*> auxNodes;
+		NodeTags.MultiFind(NewValue, auxNodes, false);
 
-			MaterialNodes.Empty();
-			VariationNodes.Empty();
-			ClipMeshNodes.Empty();
-			ClipMorphNodes.Empty();
-			ExtendNodes.Empty();
+		MaterialNodes.Empty();
+		VariationNodes.Empty();
+		ClipMeshNodes.Empty();
+		ClipMorphNodes.Empty();
+		ExtendNodes.Empty();
+		ClipDeformNodes.Empty();
 
-			for (UCustomizableObjectNode* node : auxNodes)
+		for (UCustomizableObjectNode* node : auxNodes)
+		{
+			if (Cast<UCustomizableObjectNodeMaterial>(node) && MaterialNodes.Find(node) == INDEX_NONE)
 			{
-				if (Cast<UCustomizableObjectNodeMaterial>(node) && MaterialNodes.Find(node) == INDEX_NONE)
-				{
-					MaterialNodes.Add(node);
-				}
+				MaterialNodes.Add(node);
+			}
 
-				if (Cast<UCustomizableObjectNodeMaterialVariation>(node) && VariationNodes.Find(node) == INDEX_NONE)
-				{
-					VariationNodes.Add(node);
-				}
+			if (Cast<UCustomizableObjectNodeMaterialVariation>(node) && VariationNodes.Find(node) == INDEX_NONE)
+			{
+				VariationNodes.Add(node);
+			}
 
-				if (Cast<UCustomizableObjectNodeMeshClipMorph>(node) && ClipMorphNodes.Find(node) == INDEX_NONE)
-				{
-					ClipMorphNodes.Add(node);
-				}
+			if (Cast<UCustomizableObjectNodeMeshClipMorph>(node) && ClipMorphNodes.Find(node) == INDEX_NONE)
+			{
+				ClipMorphNodes.Add(node);
+			}
 
-				if (Cast<UCustomizableObjectNodeMeshClipWithMesh>(node) && ClipMeshNodes.Find(node) == INDEX_NONE)
-				{
-					ClipMeshNodes.Add(node);
-				}
-				
-				if (Cast<UCustomizableObjectNodeExtendMaterial>(node) && ExtendNodes.Find(node) == INDEX_NONE)
-				{
-					ExtendNodes.Add(node);
-				}
+			if (Cast<UCustomizableObjectNodeMeshClipWithMesh>(node) && ClipMeshNodes.Find(node) == INDEX_NONE)
+			{
+				ClipMeshNodes.Add(node);
+			}
+			
+			if (Cast<UCustomizableObjectNodeExtendMaterial>(node) && ExtendNodes.Find(node) == INDEX_NONE)
+			{
+				ExtendNodes.Add(node);
+			}
+
+			if (Cast<UCustomizableObjectNodeMeshClipDeform>(node) && ClipDeformNodes.Find(node) == INDEX_NONE)
+			{
+				ClipDeformNodes.Add(node);
 			}
 		}
 	}
 
-	if (ColumnMat.IsValid() && ColumnClipMesh.IsValid() && ColumnClipMorph.IsValid() && ColumnVar.IsValid() && ColumnExtend.IsValid())
+	if (ColumnMat.IsValid() && ColumnClipMesh.IsValid() && ColumnClipMorph.IsValid() && ColumnVar.IsValid() && ColumnExtend.IsValid() && ColumnClipDeform.IsValid())
 	{
 		ColumnMat.ToSharedRef()->RequestListRefresh();
 		ColumnClipMesh.ToSharedRef()->RequestListRefresh();
 		ColumnClipMorph.ToSharedRef()->RequestListRefresh();
 		ColumnVar.ToSharedRef()->RequestListRefresh();
 		ColumnExtend.ToSharedRef()->RequestListRefresh();
+		ColumnClipDeform.ToSharedRef()->RequestListRefresh();
 	}
 }
 
@@ -401,8 +423,6 @@ TSharedRef<ITableRow> SCustomizableObjectEditorTagExplorer::OnGenerateTableRow(U
 		ObjectNode.RemoveAt(Indx, ObjectNode.Len() - Indx);
 	}
 
-
-
 	return SNew(STableRow<UCustomizableObjectNode*>, OwnerTable)
 		[
 			SNew(SHorizontalBox)
@@ -428,13 +448,14 @@ void SCustomizableObjectEditorTagExplorer::OnTagTableSelectionChanged(UCustomiza
 		TSharedPtr<ICustomizableObjectEditor> Editor = Entry->GetGraphEditor();
 		Editor->SelectNode(Entry);
 
-		if (ColumnMat.IsValid() && ColumnClipMesh.IsValid() && ColumnClipMorph.IsValid() && ColumnVar.IsValid() && ColumnExtend.IsValid())
+		if (ColumnMat.IsValid() && ColumnClipMesh.IsValid() && ColumnClipMorph.IsValid() && ColumnVar.IsValid() && ColumnExtend.IsValid() && ColumnClipDeform.IsValid())
 		{
 			ColumnMat.ToSharedRef()->ClearSelection();
 			ColumnClipMesh.ToSharedRef()->ClearSelection();
 			ColumnClipMorph.ToSharedRef()->ClearSelection();
 			ColumnVar.ToSharedRef()->ClearSelection();
 			ColumnExtend.ToSharedRef()->ClearSelection();
+			ColumnClipDeform.ToSharedRef()->RequestListRefresh();
 		}
 	}
 }
diff --git a/Engine/Plugins/Experimental/Mutable/Source/CustomizableObjectEditor/Private/MuCOE/SCustomizableObjectEditorTagExplorer.h b/Engine/Plugins/Experimental/Mutable/Source/CustomizableObjectEditor/Private/MuCOE/SCustomizableObjectEditorTagExplorer.h
--- a/Engine/Plugins/Experimental/Mutable/Source/CustomizableObjectEditor/Private/MuCOE/SCustomizableObjectEditorTagExplorer.h
+++ b/Engine/Plugins/Experimental/Mutable/Source/CustomizableObjectEditor/Private/MuCOE/SCustomizableObjectEditorTagExplorer.h
@@ -4,16 +4,16 @@
 
 #include "Widgets/SCompoundWidget.h"
 
-class UCustomizableObjectNode;
 namespace ESelectInfo { enum Type : int; }
 template <typename ItemType> class SListView;
-template <typename OptionType> class SComboBox;
 
 class FCustomizableObjectEditor;
 class ITableRow;
 class STableViewBase;
+class SComboButton;
 class SWidget;
 class UCustomizableObject;
+class UCustomizableObjectNode;
 
 class SCustomizableObjectEditorTagExplorer : public SCompoundWidget
 {
@@ -27,8 +27,8 @@ public:
 
 private:
 
-	/** Button Callback to fill the tag list */
-	FReply GetCustomizableObjectTags();
+	/** Callback to fill the combobox options */
+	TSharedRef<SWidget> OnGetTagsMenuContent();
 
 	/** Fills a list with all the tags found in the nodes of a graph */
 	void FillTagInformation(UCustomizableObject* Object, TArray<FString>& Tags);
@@ -43,7 +43,7 @@ private:
 	FReply CopyTagToClipboard();
 
 	/** OnSelectionChanged callback of the tags combobox */
-	void OnComboBoxSelectionChanged(TSharedPtr<FString> NewValue, ESelectInfo::Type);
+	void OnComboBoxSelectionChanged(FString NewValue);
 
 	/** Tags table callbacks*/
 	TSharedRef<ITableRow> OnGenerateTableRow(UCustomizableObjectNode* Node, const TSharedRef<STableViewBase>& OwnerTable);
@@ -55,13 +55,10 @@ private:
 	FCustomizableObjectEditor* CustomizableObjectEditorPtr;
 
 	/** Combobox for the Customizable Object Tags */
-	TSharedPtr<SComboBox<TSharedPtr<FString>>> TagComboBox;
-
-	/** Combobox string container */
-	TArray<TSharedPtr<FString>> ComboboxTags;
+	TSharedPtr<SComboButton> TagComboBox;
 
 	/** Combobox Selection */
-	TSharedPtr<FString> TagComboBoxItem;
+	FString SelectedTag;
 
 	/** List views for the nodes that contain the same tag */
 	TSharedPtr<SListView<UCustomizableObjectNode*>> ColumnMat;
@@ -69,6 +66,7 @@ private:
 	TSharedPtr<SListView<UCustomizableObjectNode*>> ColumnClipMesh;
 	TSharedPtr<SListView<UCustomizableObjectNode*>> ColumnClipMorph;
 	TSharedPtr<SListView<UCustomizableObjectNode*>> ColumnExtend;
+	TSharedPtr<SListView<UCustomizableObjectNode*>> ColumnClipDeform;
 
 	/** Arrays for each type of node that contains a tag */
 	TArray<UCustomizableObjectNode*> MaterialNodes;
@@ -76,6 +74,7 @@ private:
 	TArray<UCustomizableObjectNode*> ClipMeshNodes;
 	TArray<UCustomizableObjectNode*> ClipMorphNodes;
 	TArray<UCustomizableObjectNode*> ExtendNodes;
+	TArray<UCustomizableObjectNode*> ClipDeformNodes;
 
 	/** Multimap to relatre the nodes with the tags */
 	TMultiMap<FString, UCustomizableObjectNode*> NodeTags;
