[Mutable] Make it optional to load the reference SkeletalMesh alongside the CO.
+ Allow users to disable certain functionalities by adding/removing specific flags to the CO.
+ Add missing PostCompile callback when compiling COs synchronously.
+ Fix CustomnizableSkeletalComponent::SetPhysicsAsset

#rnx

diff --git a/Engine/Plugins/Experimental/Mutable/Source/CustomizableObject/Classes/MuCO/CustomizableObject.h b/Engine/Plugins/Experimental/Mutable/Source/CustomizableObject/Classes/MuCO/CustomizableObject.h
--- a/Engine/Plugins/Experimental/Mutable/Source/CustomizableObject/Classes/MuCO/CustomizableObject.h
+++ b/Engine/Plugins/Experimental/Mutable/Source/CustomizableObject/Classes/MuCO/CustomizableObject.h
@@ -34,6 +34,7 @@ class UPhysicsAsset;
 class USkeletalMesh;
 class USkeleton;
 struct FFrame;
+struct FStreamableHandle;
 template <typename FuncType> class TFunctionRef;
 
 DECLARE_MULTICAST_DELEGATE(FPostCompileDelegate)
@@ -899,7 +900,7 @@ struct FMutableRefAssetUserData
 
 	friend FArchive& operator<<(FArchive& Ar, FMutableRefAssetUserData& Data);
 
-	void InitResources(UObject* InOuter);
+	void InitResources(UCustomizableObject* InOuter);
 #endif
 
 };
@@ -909,7 +910,15 @@ USTRUCT()
 struct FMutableRefSkeletalMeshData
 {
 	GENERATED_BODY()
-	
+
+	// Reference Skeletal Mesh
+	UPROPERTY()
+	TObjectPtr<USkeletalMesh> SkeletalMesh;
+
+	// Path to load the ReferenceSkeletalMesh
+	UPROPERTY()
+	FSoftObjectPath SkeletalMeshAssetPath;
+
 	// LOD info
 	UPROPERTY()
 	TArray<FMutableRefLODData> LODData;
@@ -947,60 +956,9 @@ struct FMutableRefSkeletalMeshData
 	TArray<FMutableRefAssetUserData> AssetUserData;
 
 #if WITH_EDITORONLY_DATA
-	friend FArchive& operator<<(FArchive& Ar, FMutableRefSkeletalMeshData& Data)
-	{
-		Ar << Data.LODData;
-		Ar << Data.Sockets;
-		Ar << Data.Bounds;
-		Ar << Data.Settings;
-
-		if (Ar.IsSaving())
-		{
-			FString AssetPath = Data.Skeleton.ToSoftObjectPath().ToString();
-			Ar << AssetPath;
-
-			AssetPath = Data.PhysicsAsset.ToSoftObjectPath().ToString();
-			Ar << AssetPath;
-
-			AssetPath = Data.PostProcessAnimInst.ToSoftObjectPath().ToString();
-			Ar << AssetPath;
-
-			AssetPath = Data.ShadowPhysicsAsset.ToSoftObjectPath().ToString();
-			Ar << AssetPath;
-		
-		}
-		else
-		{
-			FString SkeletonAssetPath;
-			Ar << SkeletonAssetPath;
-			Data.Skeleton = TSoftObjectPtr<USkeleton>(FSoftObjectPath(SkeletonAssetPath));
-
-			FString PhysicsAssetPath;
-			Ar << PhysicsAssetPath;
-			Data.PhysicsAsset = TSoftObjectPtr<UPhysicsAsset>(FSoftObjectPath(PhysicsAssetPath));
-			
-			FString PostProcessAnimInstAssetPath;
-			Ar << PostProcessAnimInstAssetPath;
-			Data.PostProcessAnimInst = TSoftClassPtr<UAnimInstance>(FSoftObjectPath(PostProcessAnimInstAssetPath));
-			
-			FString ShadowPhysicsAssetPath;
-			Ar << ShadowPhysicsAssetPath;
-			Data.ShadowPhysicsAsset = TSoftObjectPtr<UPhysicsAsset>(FSoftObjectPath(ShadowPhysicsAssetPath));
-		}
-
-		Ar << Data.AssetUserData;
-
-		return Ar;
-	}
+	friend FArchive& operator<<(FArchive& Ar, FMutableRefSkeletalMeshData& Data);
 
-	void InitResources(UObject* InOuter)
-	{
-		check(InOuter);
-		for (FMutableRefAssetUserData& Data : AssetUserData)
-		{
-			Data.InitResources(InOuter);
-		}
-	}
+	void InitResources(UCustomizableObject* InOuter);
 #endif
 
 };
@@ -1107,6 +1065,7 @@ public:
 
 	UCustomizableObject();
 
+#if WITH_EDITORONLY_DATA
 	/** All the SkeletalMeshes generated for this CustomizableObject instances will use the Reference Skeletal Mesh 
 	* properties for everything that Mutable doesn't create or modify. This includes data like LOD distances, Physics
 	* properties, Bounding Volumes, Skeleton, etc.
@@ -1134,6 +1093,7 @@ public:
 	*/
 	UPROPERTY(EditAnywhere, Category=CustomizableObject)
 	TArray< TObjectPtr<class USkeletalMesh> > ReferenceSkeletalMeshes;
+#endif
 
 	/** All the SkeletalMeshes generated for this CustomizableObject instances will use the Reference Skeletal Mesh
 	 * properties for everything that Mutable doesn't create or modify. This struct stores the information used from
@@ -1143,11 +1103,6 @@ public:
 	UPROPERTY()
 	TArray<FMutableRefSkeletalMeshData> ReferenceSkeletalMeshesData;
 
-	// Hide for now, since it is not supported yet
-	//UPROPERTY(EditAnywhere, Category = CustomizableObject)
-	UPROPERTY()
-	TObjectPtr<class UStaticMesh> ReferenceStaticMesh;
-
 	/** List of Materials referenced by this or any child customizable object. */
 	UPROPERTY()
 	TArray<TSoftObjectPtr<UMaterialInterface>> ReferencedMaterials;
@@ -1196,6 +1151,13 @@ public:
 
 #if WITH_EDITORONLY_DATA
 
+	/** Use the SkeletalMesh of reference as a placeholder until the custom mesh is ready to use.
+	  * 
+	  * Note: If disabled, a null mesh will be used to replace the discarded mesh due to 'ReplaceDiscardedWithReferenceMesh' being enabled.
+	  */
+	UPROPERTY(EditAnywhere, Category = CustomizableObject)
+	bool bEnableUseRefSkeletalMeshAsPlaceholder = true;
+
 	// Hide this property because it is not used yet.
 	//UPROPERTY(EditAnywhere, Category = CustomizableObject)
 	UPROPERTY()
@@ -1335,7 +1297,23 @@ public:
 	UFUNCTION(BlueprintCallable, Category = CustomizableObject)
 	void UnloadMaskOutCache();
 
+	// Called to load the reference SkeletalMesh if it needs to be used as a placeholder and it's not loaded.
+	UFUNCTION(BlueprintCallable, Category = CustomizableObject)
+	void LoadReferenceSkeletalMeshesAsync();
+
+	UFUNCTION(BlueprintCallable, Category = CustomizableObject)
+	void UnloadReferenceSkeletalMeshes();
+
+	// Callback of LoadReferenceSkeletalMeshesAsync
+	void OnReferenceSkeletalMeshesAsyncLoaded();
+
 private:
+
+#if !WITH_EDITORONLY_DATA
+	// Handle used to store a streaming request operation.
+	TSharedPtr<FStreamableHandle> RefSkeletalMeshStreamingHandle;
+#endif
+
 	
 	/** Returns true or false if the parameter with name can be located and it has the type the caller is looking for. It will also
 	 * check if the model has been set to ensure access to it can take place at the calculated parameter index.
@@ -1435,7 +1413,7 @@ private:
 	// This is a manual version number for the binary blobs in this asset.
 	// Increasing it invalidates all the previously compiled models.
 	// Warning: If while merging code both versions have changed, take the highest+1.
-	static const int32 CurrentSupportedVersion = 380;
+	static const int32 CurrentSupportedVersion = 381;
 
 public:
 
@@ -1517,6 +1495,7 @@ public:
 	void LoadEmbeddedData(FArchive& Ar);
 
 	void PostLoad() override;
+	void BeginDestroy() override;
 
 	void Serialize(FArchive& Ar) override;
 
diff --git a/Engine/Plugins/Experimental/Mutable/Source/CustomizableObject/Private/MuCO/CustomizableObject.cpp b/Engine/Plugins/Experimental/Mutable/Source/CustomizableObject/Private/MuCO/CustomizableObject.cpp
--- a/Engine/Plugins/Experimental/Mutable/Source/CustomizableObject/Private/MuCO/CustomizableObject.cpp
+++ b/Engine/Plugins/Experimental/Mutable/Source/CustomizableObject/Private/MuCO/CustomizableObject.cpp
@@ -169,6 +169,21 @@ void UCustomizableObject::PostLoad()
 }
 
 
+void UCustomizableObject::BeginDestroy()
+{
+#if !WITH_EDITORONLY_DATA
+	if (RefSkeletalMeshStreamingHandle.IsValid() && RefSkeletalMeshStreamingHandle->IsActive())
+	{
+		RefSkeletalMeshStreamingHandle->CancelHandle();
+	}
+	
+	RefSkeletalMeshStreamingHandle = nullptr;
+#endif
+
+	Super::BeginDestroy();
+}
+
+
 bool UCustomizableObject::IsLocked() const
 {
 	if (PrivateData.IsValid())
@@ -464,7 +479,7 @@ void UCustomizableObject::LoadCompiledData(FArchive& MemoryReader, bool bSkipEdi
 
 		MemoryReader << ReferenceSkeletalMeshesData;
 
-		// Some resources require an outer to initialize themselves. 
+		// Initialize resources. 
 		for(FMutableRefSkeletalMeshData& ReferenceSkeletalMeshData : ReferenceSkeletalMeshesData)
 		{
 			ReferenceSkeletalMeshData.InitResources(this);
@@ -926,11 +941,18 @@ bool UCustomizableObject::IsCompiled() const
 
 USkeletalMesh* UCustomizableObject::GetRefSkeletalMesh(int32 ComponentIndex)
 {
+#if WITH_EDITORONLY_DATA
 	if (ReferenceSkeletalMeshes.IsValidIndex(ComponentIndex))
 	{
 		return ReferenceSkeletalMeshes[ComponentIndex];
 	}
-
+#else
+	if (ReferenceSkeletalMeshesData.IsValidIndex(ComponentIndex))
+	{
+		// Can be nullptr if RefSkeletalMeshes are not loaded yet.
+		return ReferenceSkeletalMeshesData[ComponentIndex].SkeletalMesh;
+	}
+#endif
 	return nullptr;
 }
 
@@ -978,6 +1000,60 @@ TSoftObjectPtr<USkeleton> UCustomizableObject::GetReferencedSkeletonAssetPtr( ui
 }
 
 
+void UCustomizableObject::LoadReferenceSkeletalMeshesAsync()
+{
+#if !WITH_EDITORONLY_DATA
+	if (!RefSkeletalMeshStreamingHandle)
+	{
+		TArray<FSoftObjectPath> MeshesToStream;
+
+		for (const FMutableRefSkeletalMeshData& RefSkeletalMeshData : ReferenceSkeletalMeshesData)
+		{
+			MeshesToStream.Add(RefSkeletalMeshData.SkeletalMeshAssetPath);
+		}
+
+		if (!MeshesToStream.IsEmpty())
+		{
+			TWeakObjectPtr<UCustomizableObject> WeakCO(this);
+
+			FStreamableManager& StreamableManager = UCustomizableObjectSystem::GetInstance()->GetStreamableManager();
+			RefSkeletalMeshStreamingHandle = StreamableManager.RequestAsyncLoad(MeshesToStream, FStreamableDelegate::CreateUObject(this, &UCustomizableObject::OnReferenceSkeletalMeshesAsyncLoaded),
+				FStreamableManager::AsyncLoadHighPriority);
+		}
+
+	}
+#endif
+}
+
+
+void UCustomizableObject::UnloadReferenceSkeletalMeshes()
+{
+#if !WITH_EDITORONLY_DATA
+	for (FMutableRefSkeletalMeshData& Data : ReferenceSkeletalMeshesData)
+	{
+		Data.SkeletalMesh = nullptr;
+	}
+#endif
+}
+
+
+void UCustomizableObject::OnReferenceSkeletalMeshesAsyncLoaded()
+{
+#if !WITH_EDITORONLY_DATA
+	if (RefSkeletalMeshStreamingHandle)
+	{
+		for (FMutableRefSkeletalMeshData& Data : ReferenceSkeletalMeshesData)
+		{
+			Data.SkeletalMesh = TSoftObjectPtr<USkeletalMesh>(Data.SkeletalMeshAssetPath).Get();
+			check(Data.SkeletalMesh);
+		}
+
+		RefSkeletalMeshStreamingHandle.Reset();
+	}
+#endif
+}
+
+
 int32 UCustomizableObject::FindState( const FString& Name ) const
 {
 	int32 Result = -1;
@@ -2022,7 +2098,7 @@ void UCustomizableObjectBulk::PrepareBulkData(UCustomizableObject* InOuter, cons
 //-------------------------------------------------------------------------------------------------
 //-------------------------------------------------------------------------------------------------
 //-------------------------------------------------------------------------------------------------
-void FMutableRefAssetUserData::InitResources(UObject* InOuter)
+void FMutableRefAssetUserData::InitResources(UCustomizableObject* InOuter)
 {
 	UClass* AssetUserDataClass = FindObject<UClass>(nullptr, *ClassPath);
 	if (AssetUserDataClass)
@@ -2067,6 +2143,77 @@ FArchive& operator<<(FArchive& Ar, FMutableRefAssetUserData& Data)
 	return Ar;
 }
 
+
+FArchive& operator<<(FArchive& Ar, FMutableRefSkeletalMeshData& Data)
+{
+	Ar << Data.LODData;
+	Ar << Data.Sockets;
+	Ar << Data.Bounds;
+	Ar << Data.Settings;
+
+	if (Ar.IsSaving())
+	{
+		FString AssetPath = Data.SkeletalMeshAssetPath.ToString();
+		Ar << AssetPath;
+
+		AssetPath = Data.Skeleton.ToSoftObjectPath().ToString();
+		Ar << AssetPath;
+
+		AssetPath = Data.PhysicsAsset.ToSoftObjectPath().ToString();
+		Ar << AssetPath;
+
+		AssetPath = Data.PostProcessAnimInst.ToSoftObjectPath().ToString();
+		Ar << AssetPath;
+
+		AssetPath = Data.ShadowPhysicsAsset.ToSoftObjectPath().ToString();
+		Ar << AssetPath;
+
+	}
+	else
+	{
+		FString SkeletalMeshAssetPath;
+		Ar << SkeletalMeshAssetPath;
+		Data.SkeletalMeshAssetPath = SkeletalMeshAssetPath;
+
+		FString SkeletonAssetPath;
+		Ar << SkeletonAssetPath;
+		Data.Skeleton = TSoftObjectPtr<USkeleton>(FSoftObjectPath(SkeletonAssetPath));
+
+		FString PhysicsAssetPath;
+		Ar << PhysicsAssetPath;
+		Data.PhysicsAsset = TSoftObjectPtr<UPhysicsAsset>(FSoftObjectPath(PhysicsAssetPath));
+
+		FString PostProcessAnimInstAssetPath;
+		Ar << PostProcessAnimInstAssetPath;
+		Data.PostProcessAnimInst = TSoftClassPtr<UAnimInstance>(FSoftObjectPath(PostProcessAnimInstAssetPath));
+
+		FString ShadowPhysicsAssetPath;
+		Ar << ShadowPhysicsAssetPath;
+		Data.ShadowPhysicsAsset = TSoftObjectPtr<UPhysicsAsset>(FSoftObjectPath(ShadowPhysicsAssetPath));
+	}
+
+	Ar << Data.AssetUserData;
+
+	return Ar;
+}
+
+
+void FMutableRefSkeletalMeshData::InitResources(UCustomizableObject* InOuter)
+{
+	check(InOuter);
+	if (InOuter->bEnableUseRefSkeletalMeshAsPlaceholder)
+	{
+		SkeletalMesh = TSoftObjectPtr<USkeletalMesh>(SkeletalMeshAssetPath).LoadSynchronous();
+	}
+
+	// Initialize AssetUserData
+	for (FMutableRefAssetUserData& Data : AssetUserData)
+	{
+		Data.InitResources(InOuter);
+	}
+}
+
+
 #endif
 
 #endif
diff --git a/Engine/Plugins/Experimental/Mutable/Source/CustomizableObject/Private/MuCO/CustomizableObjectInstance.cpp b/Engine/Plugins/Experimental/Mutable/Source/CustomizableObject/Private/MuCO/CustomizableObjectInstance.cpp
--- a/Engine/Plugins/Experimental/Mutable/Source/CustomizableObject/Private/MuCO/CustomizableObjectInstance.cpp
+++ b/Engine/Plugins/Experimental/Mutable/Source/CustomizableObject/Private/MuCO/CustomizableObjectInstance.cpp
@@ -322,9 +322,12 @@ void UCustomizableObjectInstance::BeginDestroy()
 	// Release the Live Instance ID if there it hadn't been released before
 	DestroyLiveUpdateInstance();
 
-	if (PrivateData && PrivateData->StreamingHandle)
+	if (PrivateData)
 	{
-		PrivateData->StreamingHandle->CancelHandle();
+		if (PrivateData->StreamingHandle.IsValid() && PrivateData->StreamingHandle->IsActive())
+		{
+			PrivateData->StreamingHandle->CancelHandle();
+		}
 		PrivateData->StreamingHandle = nullptr;
 	}
 	
@@ -1923,7 +1926,6 @@ bool UCustomizableInstancePrivateData::DoComponentsNeedUpdate(UCustomizableObjec
 }
 
 
-
 bool UCustomizableInstancePrivateData::UpdateSkeletalMesh_PostBeginUpdate0(UCustomizableObjectInstance* Public, const TSharedPtr<FMutableOperationData>& OperationData)
 {
 	MUTABLE_CPUPROFILER_SCOPE(UCustomizableInstancePrivateData::UpdateSkeletalMesh_PostBeginUpdate0)
@@ -5020,7 +5022,8 @@ bool UCustomizableInstancePrivateData::BuildOrCopyRenderData(const TSharedPtr<FM
 
 			const mu::FMeshBufferSet& MutableMeshVertexBuffers = Component.Mesh->GetVertexBuffers();
 
-			for (int32 ProfileIndex = 0; ProfileIndex < SkinWeightProfilesInfo.Num(); ++ProfileIndex)
+			const int32 SkinWeightProfilesCount = CustomizableObject->SkinWeightProfilesInfo.Num();
+			for (int32 ProfileIndex = 0; ProfileIndex < SkinWeightProfilesCount; ++ProfileIndex)
 			{
 				const int32 ProfileSemanticsIndex = ProfileIndex + 10;
 				int32 BoneIndicesBufferIndex, BoneIndicesBufferChannelIndex;
@@ -5040,7 +5043,7 @@ bool UCustomizableInstancePrivateData::BuildOrCopyRenderData(const TSharedPtr<FM
 					bHasSkinWeightProfiles = true;
 				}
 
-				const FMutableSkinWeightProfileInfo& Profile = SkinWeightProfilesInfo[ProfileIndex];
+				const FMutableSkinWeightProfileInfo& Profile = CustomizableObject->SkinWeightProfilesInfo[ProfileIndex];
 
 				const FSkinWeightProfileInfo* ExistingProfile = SkeletalMesh->GetSkinWeightProfiles().FindByPredicate(
 					[&Profile](const FSkinWeightProfileInfo& P) { return P.Name == Profile.Name; });
diff --git a/Engine/Plugins/Experimental/Mutable/Source/CustomizableObject/Private/MuCO/CustomizableSkeletalComponent.cpp b/Engine/Plugins/Experimental/Mutable/Source/CustomizableObject/Private/MuCO/CustomizableSkeletalComponent.cpp
--- a/Engine/Plugins/Experimental/Mutable/Source/CustomizableObject/Private/MuCO/CustomizableSkeletalComponent.cpp
+++ b/Engine/Plugins/Experimental/Mutable/Source/CustomizableObject/Private/MuCO/CustomizableSkeletalComponent.cpp
@@ -2,6 +2,9 @@
 
 #include "MuCO/CustomizableSkeletalComponent.h"
 
+#include "MuCO/CustomizableObjectSystem.h"
+#include "MuCO/CustomizableInstancePrivateData.h"
+
 #include "AnimationRuntime.h"
 #include "Components/SkeletalMeshComponent.h"
 #include "Engine/SkeletalMesh.h"
@@ -9,7 +12,6 @@
 #include "Engine/World.h"
 #include "GameFramework/Pawn.h"
 #include "Materials/MaterialInstanceDynamic.h"
-#include "MuCO/CustomizableInstancePrivateData.h"
 #include "Rendering/SkeletalMeshRenderData.h"
 #include "UObject/ObjectSaveContext.h"
 #include "MuCO/UnrealPortabilityHelpers.h"
@@ -98,7 +100,7 @@ void UCustomizableSkeletalComponent::SetPhysicsAsset(UPhysicsAsset* PhysicsAsset
 {
 	USkeletalMeshComponent* Parent = Cast<USkeletalMeshComponent>(GetAttachParent());
 
-	if (Parent && GetWorld())
+	if (Parent && Parent->GetWorld())
 	{
 		Parent->SetPhysicsAsset(PhysicsAsset, true);
 	}
@@ -303,24 +305,32 @@ void UCustomizableSkeletalComponent::TickComponent(float DeltaTime, ELevelTick T
 {
 	Super::TickComponent(DeltaTime, TickType, ThisTickFunction);
 
-	if (!bPendingSetSkeletalMesh)
+	if (!bPendingSetSkeletalMesh || !CustomizableObjectInstance || !CustomizableObjectInstance->GetCustomizableObject())
 	{
-		USkeletalMeshComponent* Parent = Cast<USkeletalMeshComponent>(GetAttachParent());
+		return;
+	}
 
-		if (Parent && CustomizableObjectInstance)
-		{
-			UCustomizableObject* CustomizableObject = CustomizableObjectInstance->GetCustomizableObject(); 
+	if (USkeletalMeshComponent* Parent = Cast<USkeletalMeshComponent>(GetAttachParent()))
+	{
+		UCustomizableObject* CustomizableObject = CustomizableObjectInstance->GetCustomizableObject();
 
-			bPendingSetSkeletalMesh = true;
+		// Hacky. Replace once we know if the instance has been generated
+		const bool bInstanceGenerated = CustomizableObjectInstance->HasAnySkeletalMesh();
 
-			if (CustomizableObjectInstance->GetSkeletalMesh(ComponentIndex))
-			{
-				Parent->SetSkeletalMesh(CustomizableObjectInstance->GetSkeletalMesh(ComponentIndex));
-			}
-			else if (!bSkipSetReferenceSkeletalMesh && CustomizableObject)
-			{
-				Parent->SetSkeletalMesh(CustomizableObject->GetRefSkeletalMesh(ComponentIndex));
-			}
+		// Generated SkeletalMesh to set, can be null if the component is empty
+		USkeletalMesh* SkeletalMesh = CustomizableObjectInstance->GetSkeletalMesh(ComponentIndex);
+
+		// If not generated yet, conditionally set the SkeletalMesh of reference
+		if (!bInstanceGenerated && !bSkipSetReferenceSkeletalMesh)
+		{
+			// Can be nullptr
+			SkeletalMesh = CustomizableObject->GetRefSkeletalMesh(ComponentIndex);
+		}
+
+		// Set SkeletalMesh
+		if (bInstanceGenerated || SkeletalMesh)
+		{
+			Parent->SetSkeletalMesh(SkeletalMesh);
 
 			if (Parent->OverrideMaterials.Num() > 0)
 			{
@@ -334,6 +344,7 @@ void UCustomizableSkeletalComponent::TickComponent(float DeltaTime, ELevelTick T
 				Parent->EmptyOverrideMaterials();
 			}
 
+			bPendingSetSkeletalMesh = false;
 		}
 	}
 }
@@ -347,7 +358,7 @@ void UCustomizableSkeletalComponent::OnAttachmentChanged()
 
 	if (Parent)
 	{
-		bPendingSetSkeletalMesh = false;
+		bPendingSetSkeletalMesh = true;
 	}
 	else if(!GetAttachParent())
 	{
diff --git a/Engine/Plugins/Experimental/Mutable/Source/CustomizableObjectEditor/Private/MuCOE/CustomizableObjectCompiler.cpp b/Engine/Plugins/Experimental/Mutable/Source/CustomizableObjectEditor/Private/MuCOE/CustomizableObjectCompiler.cpp
--- a/Engine/Plugins/Experimental/Mutable/Source/CustomizableObjectEditor/Private/MuCOE/CustomizableObjectCompiler.cpp
+++ b/Engine/Plugins/Experimental/Mutable/Source/CustomizableObjectEditor/Private/MuCOE/CustomizableObjectCompiler.cpp
@@ -1270,6 +1270,8 @@ void FCustomizableObjectCompiler::CompileInternal(UCustomizableObject* Object, c
 			CleanCachedReferencers();
 			UpdateArrayGCProtect();
 
+			CurrentObject->PostCompile(); 
+
 			State = ECustomizableObjectCompilationState::Completed;
 		}
 		else
diff --git a/Engine/Plugins/Experimental/Mutable/Source/CustomizableObjectEditor/Private/MuCOE/GenerateMutableSource/GenerateMutableSource.cpp b/Engine/Plugins/Experimental/Mutable/Source/CustomizableObjectEditor/Private/MuCOE/GenerateMutableSource/GenerateMutableSource.cpp
--- a/Engine/Plugins/Experimental/Mutable/Source/CustomizableObjectEditor/Private/MuCOE/GenerateMutableSource/GenerateMutableSource.cpp
+++ b/Engine/Plugins/Experimental/Mutable/Source/CustomizableObjectEditor/Private/MuCOE/GenerateMutableSource/GenerateMutableSource.cpp
@@ -1282,11 +1282,16 @@ void PopulateReferenceSkeletalMeshesData(FMutableGraphGenerationContext& Generat
 	GenerationContext.ReferenceSkeletalMeshesData.AddDefaulted(ComponentCount);
 	for(uint32 ComponentIndex = 0; ComponentIndex < ComponentCount; ++ComponentIndex)
 	{
-		const USkeletalMesh* RefSkeletalMesh = GenerationContext.ComponentInfos[ComponentIndex].RefSkeletalMesh;
+		USkeletalMesh* RefSkeletalMesh = GenerationContext.ComponentInfos[ComponentIndex].RefSkeletalMesh;
 		check(RefSkeletalMesh);
 
-		// Gather LODData, this may include per LOD settings such as render data config or LODDataInfoArray
 		FMutableRefSkeletalMeshData& Data = GenerationContext.ReferenceSkeletalMeshesData[ComponentIndex];
+
+		// Set the RefSkeletalMesh
+		Data.SkeletalMesh = TObjectPtr<USkeletalMesh>(RefSkeletalMesh);
+		Data.SkeletalMeshAssetPath = FSoftObjectPath(RefSkeletalMesh);
+
+		// Gather LODData, this may include per LOD settings such as render data config or LODDataInfoArray
 		Data.LODData.AddDefaulted(GenerationContext.NumLODsInRoot);
 		
 		const uint32 RefSkeletalMeshLODCount = RefSkeletalMesh->GetLODNum();
