#rnx

diff --git a/Engine/Plugins/Experimental/Mutable/Source/CustomizableObjectEditor/Private/MuCOE/Nodes/CustomizableObjectNodeMaterial.cpp b/Engine/Plugins/Experimental/Mutable/Source/CustomizableObjectEditor/Private/MuCOE/Nodes/CustomizableObjectNodeMaterial.cpp
--- a/Engine/Plugins/Experimental/Mutable/Source/CustomizableObjectEditor/Private/MuCOE/Nodes/CustomizableObjectNodeMaterial.cpp
+++ b/Engine/Plugins/Experimental/Mutable/Source/CustomizableObjectEditor/Private/MuCOE/Nodes/CustomizableObjectNodeMaterial.cpp
@@ -753,7 +753,6 @@ int32 UCustomizableObjectNodeMaterial::GetImageUVLayout(const int32 ImageIndex)
 
 int32 UCustomizableObjectNodeMaterial::GetNumParameters(const EMaterialParameterType Type) const
 {
-	ensure(Material);
 	if (Material)
 	{
 #if ENGINE_MAJOR_VERSION==5 && ENGINE_MINOR_VERSION>=1
