[Mutable] Add LODReductionSettings to the CO NodeObject to remove bones per component and LOD. Mutable will remove bones specified in the mesh settings & bones set in the component settings.
+ Remove unnecessary StaticDuplicate of meshes when converting meshes to Mutable
+ Remove cache of meshes with bone removal applied
+ Fix BoneIndices remapping after bone removal
+ Fix CustomizableObjectNodeObject backward compatibility

#rnx

diff --git a/Engine/Plugins/Experimental/Mutable/Source/CustomizableObjectEditor/Classes/MuCOE/Nodes/CustomizableObjectNodeObject.h b/Engine/Plugins/Experimental/Mutable/Source/CustomizableObjectEditor/Classes/MuCOE/Nodes/CustomizableObjectNodeObject.h
--- a/Engine/Plugins/Experimental/Mutable/Source/CustomizableObjectEditor/Classes/MuCOE/Nodes/CustomizableObjectNodeObject.h
+++ b/Engine/Plugins/Experimental/Mutable/Source/CustomizableObjectEditor/Classes/MuCOE/Nodes/CustomizableObjectNodeObject.h
@@ -22,7 +22,7 @@ struct FBoneToRemove
 	GENERATED_USTRUCT_BODY()
 
 	UPROPERTY(EditAnywhere, Category = CustomizableObject)
-	bool bIncludeBone = false;
+	bool bOnlyRemoveChildren = false;
 
 	UPROPERTY(EditAnywhere, Category = CustomizableObject)
 	FName BoneName;
diff --git a/Engine/Plugins/Experimental/Mutable/Source/CustomizableObjectEditor/Private/MuCOE/GenerateMutableSource/GenerateMutableSource.cpp b/Engine/Plugins/Experimental/Mutable/Source/CustomizableObjectEditor/Private/MuCOE/GenerateMutableSource/GenerateMutableSource.cpp
--- a/Engine/Plugins/Experimental/Mutable/Source/CustomizableObjectEditor/Private/MuCOE/GenerateMutableSource/GenerateMutableSource.cpp
+++ b/Engine/Plugins/Experimental/Mutable/Source/CustomizableObjectEditor/Private/MuCOE/GenerateMutableSource/GenerateMutableSource.cpp
@@ -631,6 +631,11 @@ mu::NodeObjectPtr GenerateMutableSource(const UEdGraphPin * Pin, FMutableGraphGe
 				{
 					MaxRefMeshLODs = RefSkeletalMesh->GetLODNum();
 				}
+
+				if (TypedNodeObj->ComponentSettings.IsValidIndex(MeshIndex))
+				{
+					GenerationContext.ComponentInfos[MeshIndex].AccumulateBonesToRemovePerLOD(TypedNodeObj->ComponentSettings[MeshIndex], TypedNodeObj->NumLODs);
+				}
 			}
 
 			if (MaxRefMeshLODs < NumLODs)
@@ -1515,4 +1520,73 @@ void AddSocketTagsToMesh(const USkeletalMesh* SourceMesh, mu::MeshPtr MutableMes
 	}
 }
 
+
+FMutableComponentInfo::FMutableComponentInfo(USkeletalMesh* InRefSkeletalMesh)
+{
+	if (!InRefSkeletalMesh || !InRefSkeletalMesh->GetSkeleton())
+	{
+		return;
+	}
+
+	RefSkeletalMesh = InRefSkeletalMesh;
+	RefSkeleton = RefSkeletalMesh->GetSkeleton();
+
+	const int32 NumBones = RefSkeleton->GetReferenceSkeleton().GetRawBoneNum();
+	BoneNamesToPathHash.Reserve(NumBones);
+
+	const TArray<FMeshBoneInfo>& Bones = RefSkeleton->GetReferenceSkeleton().GetRawRefBoneInfo();
+
+	for (int32 BoneIndex = 0; BoneIndex < NumBones; ++BoneIndex)
+	{
+		const FMeshBoneInfo& Bone = Bones[BoneIndex];
+
+		// Retrieve parent bone name and respective hash, root-bone is assumed to have a parent hash of 0
+		const FName ParentName = Bone.ParentIndex != INDEX_NONE ? Bones[Bone.ParentIndex].Name : NAME_None;
+		const uint32 ParentHash = Bone.ParentIndex != INDEX_NONE ? GetTypeHash(ParentName) : 0;
+
+		// Look-up the path-hash from root to the parent bone
+		const uint32* ParentPath = BoneNamesToPathHash.Find(ParentName);
+		const uint32 ParentPathHash = ParentPath ? *ParentPath : 0;
+
+		// Append parent hash to path to give full path hash to current bone
+		const uint32 BonePathHash = HashCombine(ParentPathHash, ParentHash);
+
+		// Add path hash to current bone
+		BoneNamesToPathHash.Add(Bone.Name, BonePathHash);
+	}
+}
+
+
+void FMutableComponentInfo::AccumulateBonesToRemovePerLOD(const FComponentSettings& ComponentSettings, int32 NumLODs)
+{
+	BonesToRemovePerLOD.SetNum(NumLODs);
+
+	TMap<FName, bool> BonesToRemove;
+
+	const int32 ComponentSettingsLODCount = ComponentSettings.LODReductionSettings.Num();
+	for (int32 LODIndex = 0; LODIndex < NumLODs; ++LODIndex)
+	{
+		if (LODIndex < ComponentSettingsLODCount)
+		{
+			const FLODReductionSettings& LODReductionSettings = ComponentSettings.LODReductionSettings[LODIndex];
+
+			for (const FBoneToRemove& Bone : LODReductionSettings.BonesToRemove)
+			{
+				if (bool* bOnlyRemoveChildren = BonesToRemove.Find(Bone.BoneName))
+				{
+					// Removed by a previous LOD
+					*bOnlyRemoveChildren = (*bOnlyRemoveChildren) && Bone.bOnlyRemoveChildren;
+				}
+				else
+				{
+					BonesToRemove.Add(Bone.BoneName, Bone.bOnlyRemoveChildren);
+				}
+			}
+		}
+		
+		BonesToRemovePerLOD[LODIndex] = BonesToRemove;
+	}
+}
+
+
 #undef LOCTEXT_NAMESPACE
diff --git a/Engine/Plugins/Experimental/Mutable/Source/CustomizableObjectEditor/Private/MuCOE/GenerateMutableSource/GenerateMutableSource.h b/Engine/Plugins/Experimental/Mutable/Source/CustomizableObjectEditor/Private/MuCOE/GenerateMutableSource/GenerateMutableSource.h
--- a/Engine/Plugins/Experimental/Mutable/Source/CustomizableObjectEditor/Private/MuCOE/GenerateMutableSource/GenerateMutableSource.h
+++ b/Engine/Plugins/Experimental/Mutable/Source/CustomizableObjectEditor/Private/MuCOE/GenerateMutableSource/GenerateMutableSource.h
@@ -367,40 +367,9 @@ inline uint32 GetTypeHash(const FMeshData& Key)
 /** Struct used to store info specific to each component during compilation */
 struct FMutableComponentInfo
 {
-	FMutableComponentInfo(USkeletalMesh* InRefSkeletalMesh)
-	{
-		if(!InRefSkeletalMesh || !InRefSkeletalMesh->GetSkeleton())
-		{
-			return;
-		}
-		
-		RefSkeletalMesh = InRefSkeletalMesh;
-		RefSkeleton = RefSkeletalMesh->GetSkeleton();
-			
-		const int32 NumBones = RefSkeleton->GetReferenceSkeleton().GetRawBoneNum();
-		BoneNamesToPathHash.Reserve(NumBones);
-
-		const TArray<FMeshBoneInfo>& Bones = RefSkeleton->GetReferenceSkeleton().GetRawRefBoneInfo();
-
-		for (int32 BoneIndex = 0; BoneIndex < NumBones; ++BoneIndex)
-		{
-			const FMeshBoneInfo& Bone = Bones[BoneIndex];
-
-			// Retrieve parent bone name and respective hash, root-bone is assumed to have a parent hash of 0
-			const FName ParentName = Bone.ParentIndex != INDEX_NONE ? Bones[Bone.ParentIndex].Name : NAME_None;
-			const uint32 ParentHash = Bone.ParentIndex != INDEX_NONE ? GetTypeHash(ParentName) : 0;
-
-			// Look-up the path-hash from root to the parent bone
-			const uint32* ParentPath = BoneNamesToPathHash.Find(ParentName);
-			const uint32 ParentPathHash = ParentPath ? *ParentPath : 0;
+	FMutableComponentInfo(USkeletalMesh* InRefSkeletalMesh);
 
-			// Append parent hash to path to give full path hash to current bone
-			const uint32 BonePathHash = HashCombine(ParentPathHash, ParentHash);
-
-			// Add path hash to current bone
-			BoneNamesToPathHash.Add(Bone.Name, BonePathHash);
-		}
-	}
+	void AccumulateBonesToRemovePerLOD(const FComponentSettings& ComponentSettings, int32 NumLODs);
 
 	// Each component must have a reference SkeletalMesh with a valid Skeleton
 	USkeletalMesh* RefSkeletalMesh = nullptr;
@@ -412,6 +381,9 @@ struct FMutableComponentInfo
 	// Hierarchy hash from parent-bone to root bone, used to check if additional skeletons are compatible with
 	// the RefSkeleton
 	TMap<FName, uint32> BoneNamesToPathHash;
+
+	// Bones to remove on each LOD, include bones on previous LODs. FName (BoneToRemove) - bool (bOnlyRemoveChildren)
+	TArray<TMap<FName, bool>> BonesToRemovePerLOD;
 };
 
 
@@ -741,19 +713,6 @@ struct FMutableGraphGenerationContext
 	// Stores the parameters generated in the node tables
 	TMap<const class UCustomizableObjectNodeTable*, TArray<FGuid>> GeneratedParametersInTables;
 
-	struct FMeshWithBoneRemovalApplied
-	{
-		TObjectPtr<class USkeletalMesh> Mesh;
-		
-		// Key is LOD index
-		//
-		// This is a cache. Entries are added on demand. A processed LOD with no bones removed will have an empty entry.
-		TMap<int32, TMap<int32, int32>> RemovedBonesActiveParentIndicesPerLOD;
-
-		bool bHasBonesToRemove = false;
-	};
-
-	TMap<TObjectPtr<class USkeletalMesh>, FMeshWithBoneRemovalApplied> MeshesWithBoneRemovalApplied;
 
 	/** Extension Data constants are collected here */
 	FExtensionDataCompilerInterface ExtensionDataCompilerInterface;
diff --git a/Engine/Plugins/Experimental/Mutable/Source/CustomizableObjectEditor/Private/MuCOE/GenerateMutableSource/GenerateMutableSourceMesh.cpp b/Engine/Plugins/Experimental/Mutable/Source/CustomizableObjectEditor/Private/MuCOE/GenerateMutableSource/GenerateMutableSourceMesh.cpp
--- a/Engine/Plugins/Experimental/Mutable/Source/CustomizableObjectEditor/Private/MuCOE/GenerateMutableSource/GenerateMutableSourceMesh.cpp
+++ b/Engine/Plugins/Experimental/Mutable/Source/CustomizableObjectEditor/Private/MuCOE/GenerateMutableSource/GenerateMutableSourceMesh.cpp
@@ -48,175 +48,96 @@
 
 #define LOCTEXT_NAMESPACE "CustomizableObjectEditor"
 
-TObjectPtr<const USkeletalMesh> GetMeshWithBoneRemovalApplied(TObjectPtr<USkeletalMesh> InSkeletalMesh, int32 InLODIndex, FMutableGraphGenerationContext& GenerationContext, const UCustomizableObjectNode* CurrentNode, TMap<int32, int32>& OutRemovedBonesActiveParentIndices)
+void BuildRemappedBonesArray(const FMutableComponentInfo& InComponentInfo, TObjectPtr<USkeletalMesh> InSkeletalMesh, int32 InLODIndex, const TArray<FBoneIndexType>& InRequiredBones, TArray<FBoneIndexType>& OutRemappedBones)
 {
-	FMutableGraphGenerationContext::FMeshWithBoneRemovalApplied& CacheEntry = GenerationContext.MeshesWithBoneRemovalApplied.FindOrAdd(InSkeletalMesh);
-	if (CacheEntry.Mesh)
+	if (!InSkeletalMesh)
 	{
-		if (CacheEntry.bHasBonesToRemove)
-		{
-			TMap<int32, int32>* RemovedBonesActiveParentIndicesForLOD = CacheEntry.RemovedBonesActiveParentIndicesPerLOD.Find(InLODIndex);
-			if (RemovedBonesActiveParentIndicesForLOD != nullptr)
-			{
-				// Note that this copies the map. This function could provide a pointer into the cache instead, but
-				// it would be hard to track when the cache pointer is invalidated, so this copy is done for safety.
-				OutRemovedBonesActiveParentIndices = *RemovedBonesActiveParentIndicesForLOD;
-				return CacheEntry.Mesh;
-			}
-		}
-		else
-		{
-			OutRemovedBonesActiveParentIndices.Reset();
-			return CacheEntry.Mesh;
-		}
-	}
-
-	CacheEntry.bHasBonesToRemove = false;
-	for (int32 LODIndex = 0; LODIndex < InSkeletalMesh->GetLODNum(); ++LODIndex)
-	{
-		if (InSkeletalMesh->GetLODInfo(LODIndex)->BonesToRemove.Num() > 0)
-		{
-			CacheEntry.bHasBonesToRemove = true;
-			break;
-		}
-	}
-
-	if (!CacheEntry.bHasBonesToRemove)
-	{
-		// No changes needed, just return the original mesh
-		CacheEntry.Mesh = InSkeletalMesh;
-		OutRemovedBonesActiveParentIndices.Reset();
-		return CacheEntry.Mesh;
+		return;
 	}
+	
+	const FReferenceSkeleton& ReferenceSkeleton = InSkeletalMesh->GetRefSkeleton();
+	const int32 NumBones = ReferenceSkeleton.GetNum();
 
-	if (!CacheEntry.Mesh)
-	{
-		UE_LOG(LogMutable, Log, TEXT("GetMeshWithBoneRemovalApplied duplicating mesh %s"), *InSkeletalMesh->GetPathName());
-
-		// Keep the name since it is used at least when gathering info about morph targets. 
-		CacheEntry.Mesh = (USkeletalMesh*)StaticDuplicateObject(InSkeletalMesh, GetTransientPackage(), InSkeletalMesh->GetFName(), EObjectFlags::RF_Transient);
-		check(CacheEntry.Mesh);
-
-		if (InSkeletalMesh->GetLODNum() < CacheEntry.Mesh->GetLODNum())
-		{
-			GenerationContext.Compiler->CompilerLog(LOCTEXT("ReferenceLODMismatch", "The reference mesh has less LODs than the generated mesh."),
-				CurrentNode,
-				EMessageSeverity::Warning);
-		}
-	}
+	// Build RemappedBones array
+	OutRemappedBones.Init(0, NumBones);
 
-	UE_LOG(LogMutable, Log, TEXT("GetMeshWithBoneRemovalApplied calculating bone hierarchy for mesh %s LOD %i"), *InSkeletalMesh->GetPathName(), InLODIndex);
+	const bool bComponentInfoHasBonesToRemove = InComponentInfo.BonesToRemovePerLOD.IsValidIndex(InLODIndex) && !InComponentInfo.BonesToRemovePerLOD[InLODIndex].IsEmpty();
 
-	// If this code is reached, it should mean that the removed bone indices haven't been cached for this LOD yet
-	check(!CacheEntry.RemovedBonesActiveParentIndicesPerLOD.Contains(InLODIndex));
+	const TArray<FMeshBoneInfo>& RefBoneInfos = ReferenceSkeleton.GetRefBoneInfo();
+	const TArray<FSkeletalMeshLODInfo>& LODInfos = InSkeletalMesh->GetLODInfoArray();
+	const int32 NumLODInfos = LODInfos.Num();
 
-	TMap<int32, int32>& CachedRemovedBonesActiveParentIndices = CacheEntry.RemovedBonesActiveParentIndicesPerLOD.Add(InLODIndex);
-	// The new entry should be empty
-	check(CachedRemovedBonesActiveParentIndices.Num() == 0);
+	// Helper to know which bones have been removed
+	TArray<bool> RemovedBones;
+	RemovedBones.SetNumZeroed(NumBones);
 
-	TSet<int32> RemovedBoneIndices;
+	for (const FBoneIndexType& RequiredBoneIndex : InRequiredBones)
 	{
-		TArray<FName> BoneNamesToRemove;
-		BoneNamesToRemove.Reset(InSkeletalMesh->GetLODInfo(InLODIndex)->BonesToRemove.Num());
-		
-		for (const FBoneReference& BoneRef : InSkeletalMesh->GetLODInfo(InLODIndex)->BonesToRemove)
-		{
-			BoneNamesToRemove.Add(BoneRef.BoneName);
-		}
+		const FMeshBoneInfo& BoneInfo = RefBoneInfos[RequiredBoneIndex];
+		FBoneIndexType FinalBoneIndex = RequiredBoneIndex;
 
-		if (BoneNamesToRemove.Num() > 0)
+		// Remove bone if the parent has been removed, Root can't be removed
+		if (BoneInfo.ParentIndex != INDEX_NONE && RemovedBones[BoneInfo.ParentIndex])
 		{
-			UE_LOG(LogMutable, Log, TEXT("GetMeshWithBoneRemovalApplied removing bones from mesh %s LOD %i"), *InSkeletalMesh->GetPathName(), InLODIndex);
-
-			IMeshUtilities& MeshUtilities = FModuleManager::Get().LoadModuleChecked<IMeshUtilities>("MeshUtilities");
-			MeshUtilities.RemoveBonesFromMesh(CacheEntry.Mesh, InLODIndex, &BoneNamesToRemove);
+			RemovedBones[RequiredBoneIndex] = true;
+			FinalBoneIndex = OutRemappedBones[BoneInfo.ParentIndex];
 		}
 
-		for (const FName& BoneName : BoneNamesToRemove)
-		{
-			int32 RefBoneIndex = CacheEntry.Mesh->GetRefSkeleton().FindBoneIndex(BoneName);
-			if (RefBoneIndex != INDEX_NONE)
-			{
-				RemovedBoneIndices.Add(RefBoneIndex);
-			}
-		}
-	}
-
-	// Reconstruct RefSkeleton's bone hierarchy
-	TMap<int32, TArray<int32>> SkeletalTree;
-	FReferenceSkeleton& RefSkeleton = InSkeletalMesh->GetRefSkeleton();
-	for (int32 i = 0; i < RefSkeleton.GetRawBoneNum(); ++i)
-	{
-		int32 ParentBoneIndex = RefSkeleton.GetRawParentIndex(i);
-		if (ParentBoneIndex == INDEX_NONE)
-		{
-			continue;
-		}
-
-		if (TArray<int32>* Children = SkeletalTree.Find(ParentBoneIndex))
-		{
-			Children->Add(i);
-		}
 		else
 		{
-			TArray<int32> NewChildrenArray;
-			NewChildrenArray.Add(i);
-			SkeletalTree.Add(ParentBoneIndex, MoveTemp(NewChildrenArray));
-		}
-	}
-
-	// Search for the first active parent of each removed bone
-	// If a vertex is influenced by a removed bone the influences will be transferred to the parent bone.
-	for (const int32 RemovedBone : RemovedBoneIndices)
-	{
-		int32 ParentIndex = RefSkeleton.GetRawParentIndex(RemovedBone);
+			// Check if it has to be removed
+			bool bBoneRemoved = false;
 
-		if (TArray<int32>* Children = SkeletalTree.Find(RemovedBone))
-		{
-			for (int32 i = 0; i < Children->Num(); ++i)
+			if (bComponentInfoHasBonesToRemove)
 			{
-				int32& ChildIndex = (*Children)[i];
-				CachedRemovedBonesActiveParentIndices.Add(ChildIndex, ParentIndex);
-
-				if (RemovedBoneIndices.Find(ChildIndex))
+				// Remove if found in the BonesToRemove map (ComponentSettings -> LODReductionSettings in the CustomizableObjectNodeObject)
+				if (const bool* bOnlyRemoveChildren = InComponentInfo.BonesToRemovePerLOD[InLODIndex].Find(BoneInfo.Name))
 				{
-					continue;
-				}
+					// Mark bone as removed
+					RemovedBones[RequiredBoneIndex] = true;
 
-				if (TArray<int32>* ChildChildren = SkeletalTree.Find(ChildIndex))
-				{
-					for (int j = 0; j < ChildChildren->Num(); ++j)
-					{
-						Children->Add((*ChildChildren)[j]);
-					}
+					// There's the option of only removing the children of this bone
+					bBoneRemoved = !(*bOnlyRemoveChildren);
 				}
 			}
-		}
 
-		if (ParentIndex >= 0)
-		{
-			CachedRemovedBonesActiveParentIndices.Add(RemovedBone, ParentIndex);
+			// If the bone has not been remove yet, check if it's in the BonesToRemove of the SkeletalMesh.
+			for (int32 LODIndex = 0; !bBoneRemoved && LODIndex <= InLODIndex && LODIndex < NumLODInfos; ++LODIndex)
+			{
+				const FBoneReference* BoneToRemove = LODInfos[LODIndex].BonesToRemove.FindByPredicate(
+					[&BoneInfo](const FBoneReference& BoneReference) { return BoneReference.BoneName == BoneInfo.Name; });
+				
+				bBoneRemoved = BoneToRemove != nullptr;
+				RemovedBones[RequiredBoneIndex] = RemovedBones[RequiredBoneIndex] || bBoneRemoved;
+			}
+
+			// Fix up FinalBoneIndex if it has been removed. Root can't be removed
+			FinalBoneIndex = !bBoneRemoved || BoneInfo.ParentIndex == INDEX_NONE ? RequiredBoneIndex : OutRemappedBones[BoneInfo.ParentIndex];
 		}
+
+		OutRemappedBones[RequiredBoneIndex] = FinalBoneIndex;
 	}
 
-	OutRemovedBonesActiveParentIndices = CachedRemovedBonesActiveParentIndices;
-	return CacheEntry.Mesh;
 }
 
 
-void TransferRemovedBonesInfluences(FBoneIndexType* InfluenceBones, uint16* InfluenceWeights, const int32 InfluenceCount, const TMap<int32, int32>& ReplacementMap)
+void TransferRemovedBonesInfluences(FBoneIndexType* InfluenceBones, uint16* InfluenceWeights, const int32 InfluenceCount, const TArray<FBoneIndexType>& RemappedBoneMapIndices)
 {
+	const int32 BoneMapBoneCount = RemappedBoneMapIndices.Num();
+
 	for (int32 i = 0; i < InfluenceCount; ++i)
 	{
-		if (const int32* IndexPtr = ReplacementMap.Find(InfluenceBones[i]))
+		if (InfluenceBones[i] < BoneMapBoneCount)
 		{
 			bool bParentFound = false;
-			int32 ParentIndex = *IndexPtr;
-			for (int32 j = 0; j < InfluenceCount; ++j)
+			FBoneIndexType ParentIndex = RemappedBoneMapIndices[InfluenceBones[i]];
+			for (int32 j = 0; j < i; ++j)
 			{
 				if (InfluenceBones[j] == ParentIndex)
 				{
 					InfluenceWeights[j] += InfluenceWeights[i];
+
+					InfluenceBones[i] = 0;
 					InfluenceWeights[i] = 0.f;
 					bParentFound = true;
 					break;
@@ -228,6 +149,11 @@ void TransferRemovedBonesInfluences(FBoneIndexType* InfluenceBones, uint16* Infl
 				InfluenceBones[i] = ParentIndex;
 			}
 		}
+		else
+		{
+			InfluenceBones[i] = 0;
+			InfluenceWeights[i] = 0.f;
+		}
 	}
 }
 
@@ -317,6 +243,68 @@ void NormalizeWeights(FBoneIndexType* InfluenceBones, uint16* InfluenceWeights,
 	}
 }
 
+
+bool IsSkeletalMeshCompatibleWithRefSkeleton(FMutableComponentInfo& ComponentInfo, TObjectPtr<USkeletalMesh> InSkeletalMesh, FString& OutErrorMessage)
+{
+	TObjectPtr<USkeleton> Skeleton = InSkeletalMesh->GetSkeleton();
+
+	if (Skeleton == ComponentInfo.RefSkeleton)
+	{
+		return true;
+	}
+
+	if (bool* SkeletonCompatibility = ComponentInfo.SkeletonCompatibility.Find(Skeleton))
+	{
+		return *SkeletonCompatibility;
+	}
+
+
+	// Check if the skeleton is compatible with the reference skeleton
+	const TMap<FName, uint32>& RefMeshBoneNamesToPathHash = ComponentInfo.BoneNamesToPathHash;
+
+	const TArray<FMeshBoneInfo>& Bones = Skeleton->GetReferenceSkeleton().GetRawRefBoneInfo();
+	const int32 NumBones = Bones.Num();
+
+	TMap<FName, uint32> BoneNamesToPathHash;
+	BoneNamesToPathHash.Reserve(NumBones);
+
+	for (int32 BoneIndex = 0; BoneIndex < NumBones; ++BoneIndex)
+	{
+		const FMeshBoneInfo& Bone = Bones[BoneIndex];
+
+		// Retrieve parent bone name and respective hash, root-bone is assumed to have a parent hash of 0
+		const FName ParentName = Bone.ParentIndex != INDEX_NONE ? Bones[Bone.ParentIndex].Name : NAME_None;
+		const uint32 ParentHash = Bone.ParentIndex != INDEX_NONE ? GetTypeHash(ParentName) : 0;
+
+		// Look-up the path-hash from root to the parent bone
+		const uint32* ParentPath = BoneNamesToPathHash.Find(ParentName);
+		const uint32 ParentPathHash = ParentPath ? *ParentPath : 0;
+
+		// Append parent hash to path to give full path hash to current bone
+		const uint32 BonePathHash = HashCombine(ParentPathHash, ParentHash);
+
+		// If the hash differs from the reference one it means skeletons are incompatible
+		if (const uint32* RefSMBonePathHash = RefMeshBoneNamesToPathHash.Find(Bone.Name); RefSMBonePathHash && *RefSMBonePathHash != BonePathHash)
+		{
+			// Different skeletons can't be used if they are incompatible with the reference skeleton.
+			FString Msg = FString::Printf(
+				TEXT("The SkeletalMesh [%s] with Skeleton [%s] is incompatible with the reference mesh [%s] which has [%s]. "
+					"Bone [%s] has a differnt parent on the Skeleton from the reference mesh."),
+				*InSkeletalMesh->GetName(), *Skeleton->GetName(),
+				*ComponentInfo.RefSkeletalMesh->GetName(), *ComponentInfo.RefSkeleton->GetName(),
+				*Bone.ExportName);
+
+
+			return false;
+		}
+
+		// Add path hash to current bone
+		BoneNamesToPathHash.Add(Bone.Name, BonePathHash);
+	}
+
+	return true;
+}
+
 void SetAndPropagatePoseBoneUsage(
 		mu::Mesh& MutableMesh, int32 PoseIndex, mu::EBoneUsageFlags Usage, 
 		const TMap<FName, int32>* BoneNameToSkeletonIndexMap = nullptr,
@@ -561,45 +549,12 @@ mu::Ptr<mu::PhysicsBody> MakePhysicsBodyFromAsset(UPhysicsAsset* Asset, const TA
 
 mu::MeshPtr ConvertSkeletalMeshToMutable(USkeletalMesh* InSkeletalMesh, const TSoftClassPtr<UAnimInstance>& AnimBp, int LOD, int MaterialIndex, FMutableGraphGenerationContext& GenerationContext, const UCustomizableObjectNode* CurrentNode)
 {
-	// Get the mesh generation flags to use
-	const EMutableMeshConversionFlags CurrentFlags = GenerationContext.MeshGenerationFlags.Last();
-	const bool bIgnoreSkeleton = EnumHasAnyFlags(CurrentFlags, EMutableMeshConversionFlags::IgnoreSkinning);
-	const bool bIgnorePhysics = EnumHasAnyFlags(CurrentFlags, EMutableMeshConversionFlags::IgnorePhysics);
-
-	mu::MeshPtr MutableMesh = new mu::Mesh();
-		
 	if(!InSkeletalMesh)
 	{
-		// Trying to convert a nullptr to mutable?
-		return nullptr;
-	}
-	
-	USkeleton* InSkeleton = InSkeletalMesh->GetSkeleton();
-	if (!InSkeleton)
-	{
-		FString Msg = FString::Printf(
-			TEXT("No skeleton provided when converting SkeletalMesh [%s]."),
-			*InSkeletalMesh->GetName()
-		);
-		GenerationContext.Compiler->CompilerLog(FText::FromString(Msg), CurrentNode);
 		return nullptr;
 	}
 
-	TObjectPtr<const USkeletalMesh> SkeletalMesh = InSkeletalMesh;
-
-	// CurrentMeshComponent < 0 implies IgnoreSkeleton flag.  
-	// CurrentMeshComponent < 0 will only happen with modifiers and, for now, any mesh generated from a modifier
-	// should ignore skinning.
-	check(!(GenerationContext.CurrentMeshComponent < 0) || bIgnoreSkeleton);
-
-	// Apply removed bones
-	TMap<int32, int32> RemovedBonesActiveParentIndices;
-	if (!bIgnoreSkeleton)
-	{
-		SkeletalMesh = GetMeshWithBoneRemovalApplied(InSkeletalMesh, LOD, GenerationContext, CurrentNode, RemovedBonesActiveParentIndices);
-	}
-
-	const FSkeletalMeshModel* ImportedModel = SkeletalMesh->GetImportedModel();
+	const FSkeletalMeshModel* ImportedModel = InSkeletalMesh->GetImportedModel();
 
 	// Check in case the data has changed
 	if (!(ImportedModel
@@ -634,13 +589,33 @@ mu::MeshPtr ConvertSkeletalMeshToMutable(USkeletalMesh* InSkeletalMesh, const TS
 		return nullptr;
 	}
 
-	// This map will be used to find and replace vertex influences.
-	TMap<int32, int32> InfluencesToReplaceMap;
+	// Get the mesh generation flags to use
+	const EMutableMeshConversionFlags CurrentFlags = GenerationContext.MeshGenerationFlags.Last();
+	const bool bIgnoreSkeleton = EnumHasAnyFlags(CurrentFlags, EMutableMeshConversionFlags::IgnoreSkinning);
+	const bool bIgnorePhysics = EnumHasAnyFlags(CurrentFlags, EMutableMeshConversionFlags::IgnorePhysics);
+
+	mu::MeshPtr MutableMesh = new mu::Mesh();
 
-	// Check for a matching skeleton to the reference mesh
+	// CurrentMeshComponent < 0 implies IgnoreSkeleton flag.  
+	// CurrentMeshComponent < 0 will only happen with modifiers and, for now, any mesh generated from a modifier
+	// should ignore skinning.
+	check(!(GenerationContext.CurrentMeshComponent < 0) || bIgnoreSkeleton);
+
+	// 
+	bool bBoneMapModified = false;
+	TArray<FBoneIndexType> BoneMap;
+	TArray<FBoneIndexType> RemappedBoneMapIndices;
+
+	// Check if the Skeleton is valid and build the mu::Skeleton
 	if (!bIgnoreSkeleton)
 	{
-		check(GenerationContext.CurrentMeshComponent >= 0);
+		USkeleton* InSkeleton = InSkeletalMesh->GetSkeleton();
+		if (!InSkeleton)
+		{
+			FString Msg = FString::Printf(TEXT("No skeleton provided when converting SkeletalMesh [%s]."), *InSkeletalMesh->GetName());
+			GenerationContext.Compiler->CompilerLog(FText::FromString(Msg), CurrentNode);
+			return nullptr;
+		}
 
 		FMutableComponentInfo& MutComponentInfo = GenerationContext.GetCurrentComponentInfo();
 		USkeletalMesh* ComponentRefSkeletalMesh = MutComponentInfo.RefSkeletalMesh;
@@ -648,74 +623,31 @@ mu::MeshPtr ConvertSkeletalMeshToMutable(USkeletalMesh* InSkeletalMesh, const TS
 		check(ComponentRefSkeletalMesh);
 		check(ComponentRefSkeleton);
 
-		// Add the RefSkeleton ID to the mesh.
-		const int32 RefSkeletonID = GenerationContext.ReferencedSkeletons.AddUnique(ComponentRefSkeleton);
-		MutableMesh->AddSkeletonID(RefSkeletonID);
-
-		if (InSkeleton != ComponentRefSkeleton)
+		// Compatibility check
 		{
-			bool* SkeletonCompatibility = MutComponentInfo.SkeletonCompatibility.Find(InSkeleton);
-			
-			if(!SkeletonCompatibility)
-			{
-				// Check if the skeleton is compatible with the reference skeleton
-				const TMap<FName, uint32>& RefMeshBoneNamesToPathHash = MutComponentInfo.BoneNamesToPathHash;
-
-				const TArray<FMeshBoneInfo>& Bones = InSkeleton->GetReferenceSkeleton().GetRawRefBoneInfo();
-				const int32 NumBones = Bones.Num();
+			FString ErrorMessage;
+			const bool bCompatible = IsSkeletalMeshCompatibleWithRefSkeleton(MutComponentInfo, InSkeletalMesh, ErrorMessage);
+			MutComponentInfo.SkeletonCompatibility.Add(InSkeleton, bCompatible);
 
-				TMap<FName, uint32> BoneNamesToPathHash;
-				BoneNamesToPathHash.Reserve(NumBones);
-
-				for (int32 BoneIndex = 0; BoneIndex < NumBones; ++BoneIndex)
+			if (!bCompatible)
+			{
+				if (ErrorMessage.IsEmpty())
 				{
-					const FMeshBoneInfo& Bone = Bones[BoneIndex];
-
-					// Retrieve parent bone name and respective hash, root-bone is assumed to have a parent hash of 0
-					const FName ParentName = Bone.ParentIndex != INDEX_NONE ? Bones[Bone.ParentIndex].Name : NAME_None;
-					const uint32 ParentHash = Bone.ParentIndex != INDEX_NONE ? GetTypeHash(ParentName) : 0;
-
-					// Look-up the path-hash from root to the parent bone
-					const uint32* ParentPath = BoneNamesToPathHash.Find(ParentName);
-					const uint32 ParentPathHash = ParentPath ? *ParentPath : 0;
-
-					// Append parent hash to path to give full path hash to current bone
-					const uint32 BonePathHash = HashCombine(ParentPathHash, ParentHash);
-
-					// If the hash differs from the reference one it means skeletons are incompatible
-					if (const uint32* RefSMBonePathHash = RefMeshBoneNamesToPathHash.Find(Bone.Name); RefSMBonePathHash && *RefSMBonePathHash != BonePathHash)
-					{
-						// Different skeletons can't be used if they are incompatible with the reference skeleton.
-						FString Msg = FString::Printf(
-							TEXT("The SkeletalMesh [%s] with Skeleton [%s] is incompatible with the reference mesh [%s] which has [%s]. "
-								"Bone [%s] has a differnt parent on the Skeleton from the reference mesh."),
-							*InSkeletalMesh->GetName(),	*InSkeleton->GetName(),
-							*ComponentRefSkeletalMesh->GetName(), *ComponentRefSkeleton->GetName(),
-							*Bone.ExportName);
-
-						GenerationContext.Compiler->CompilerLog(FText::FromString(Msg), CurrentNode, EMessageSeverity::Info);
-
-						MutComponentInfo.SkeletonCompatibility.Add(InSkeleton, false);
-						return nullptr;
-					}
-
-					// Add path hash to current bone
-					BoneNamesToPathHash.Add(Bone.Name, BonePathHash);
+					GenerationContext.Compiler->CompilerLog(FText::FromString(ErrorMessage), CurrentNode, EMessageSeverity::Info);
 				}
-				
-				MutComponentInfo.SkeletonCompatibility.Add(InSkeleton, true);
-			}
-			else if(*SkeletonCompatibility == false)
-			{
-				// Incompatible skeleton 
 				return nullptr;
 			}
 
+			// Add the RefSkeleton ID to the mesh.
+			const int32 RefSkeletonID = GenerationContext.ReferencedSkeletons.AddUnique(ComponentRefSkeleton);
+			MutableMesh->AddSkeletonID(RefSkeletonID);
+
 			// Add the skeleton to the list of referenced skeletons and add its index to the mesh
 			const int32 SkeletonID = GenerationContext.ReferencedSkeletons.AddUnique(InSkeleton);
 			MutableMesh->AddSkeletonID(SkeletonID);
 		}
 
+		// RefSkeleton check
 		{
 			// Ensure the bones used by the Skeletal Mesh exits in the Mesh's Skeleton
 			const TArray<FMeshBoneInfo>& RawRefBoneInfo = InSkeletalMesh->GetRefSkeleton().GetRawRefBoneInfo();
@@ -749,19 +681,78 @@ mu::MeshPtr ConvertSkeletalMeshToMutable(USkeletalMesh* InSkeletalMesh, const TS
 			}
 		}
 
+		const TArray<uint16>& SourceRequiredBones = ImportedModel->LODModels[LOD].RequiredBones;
+
+		// Remove bones and build an array to remap indices of the BoneMap
+		TArray<FBoneIndexType> RemappedBones;
+		BuildRemappedBonesArray(MutComponentInfo, InSkeletalMesh, LOD, SourceRequiredBones, RemappedBones);
+		
+		// Build RequiredBones array
+		TArray<FBoneIndexType> RequiredBones;
+		RequiredBones.Reserve(SourceRequiredBones.Num());
+
+		for (const FBoneIndexType& RequiredBoneIndex : SourceRequiredBones)
 		{
-			// Find removed bone and parent indices in the current section BoneMap.
-			const TArray<uint16>& BoneMap = ImportedModel->LODModels[LOD].Sections[MaterialIndex].BoneMap;
+			RequiredBones.AddUnique(RemappedBones[RequiredBoneIndex]);
+		}
 
-			for (const TPair<int32,int32> it : RemovedBonesActiveParentIndices)
-			{
-				const int32 BoneMapIndex = BoneMap.Find(it.Key);
-				if (BoneMapIndex != INDEX_NONE)
-				{
-					const int32 ParentBoneMapIndex = BoneMap.Find(it.Value);
-					InfluencesToReplaceMap.Add(BoneMapIndex, ParentBoneMapIndex != INDEX_NONE ? ParentBoneMapIndex : 0);
-				}
-			}
+		// Build BoneMap
+		const TArray<uint16>& SourceBoneMap = ImportedModel->LODModels[LOD].Sections[MaterialIndex].BoneMap;
+		const int32 NumBonesInBoneMap = SourceBoneMap.Num();
+		const int32 NumRemappedBones = RemappedBones.Num();
+
+		for (int32 BoneIndex = 0; BoneIndex < NumBonesInBoneMap; ++BoneIndex)
+		{
+			const FBoneIndexType BoneMapBoneIndex = SourceBoneMap[BoneIndex];
+			const FBoneIndexType FinalBoneIndex = BoneMapBoneIndex < NumRemappedBones ? RemappedBones[BoneMapBoneIndex] : 0;
+
+			const int32 BoneMapIndex = BoneMap.AddUnique(FinalBoneIndex);
+			RemappedBoneMapIndices.Add(BoneMapIndex);
+
+			bBoneMapModified = bBoneMapModified || SourceBoneMap[BoneIndex] != FinalBoneIndex;
+		}
+
+		TMap<int32, int32> InverseBoneMap;
+
+		const int32 BoneMapCount = BoneMap.Num();
+		for (int32 BoneIndex = 0; BoneIndex < BoneMapCount; ++BoneIndex)
+		{
+			int32 RefSkeletonIndex = BoneMap[BoneIndex];
+			InverseBoneMap.Add(RefSkeletonIndex, BoneIndex);
+		}
+
+		// Create the skeleton and poses for this mesh
+		mu::SkeletonPtr MutableSkeleton = new mu::Skeleton;
+		MutableMesh->SetSkeleton(MutableSkeleton);
+
+		const int32 NumRequiredBones = RequiredBones.Num();
+		MutableMesh->SetBonePoseCount(NumRequiredBones);
+		MutableSkeleton->SetBoneCount(NumRequiredBones);
+
+		const TArray<FMeshBoneInfo>& RefBoneInfo = InSkeletalMesh->GetRefSkeleton().GetRefBoneInfo();
+		for (int32 RequiredBoneIndex = 0; RequiredBoneIndex < NumRequiredBones; ++RequiredBoneIndex)
+		{
+			const int32 RefSkelIndex = RequiredBones[RequiredBoneIndex];
+
+			const FMeshBoneInfo& BoneInfo = RefBoneInfo[RefSkelIndex];
+
+			const FString BoneName = BoneInfo.Name.ToString();
+			const int32 BoneMapIndex = InverseBoneMap.FindOrAdd(RefSkelIndex, InverseBoneMap.Num());
+			const int32 ParentBoneMapIndex = BoneInfo.ParentIndex != INDEX_NONE ? InverseBoneMap.FindOrAdd(BoneInfo.ParentIndex, InverseBoneMap.Num()) : INDEX_NONE;
+
+			// Set bone hierarchy
+			MutableSkeleton->SetBoneName(BoneMapIndex, StringCast<ANSICHAR>(*BoneName).Get());
+			MutableSkeleton->SetBoneParent(BoneMapIndex, ParentBoneMapIndex);
+
+			// Set bone pose
+			FMatrix44f BaseInvMatrix = InSkeletalMesh->GetRefBasesInvMatrix()[RefSkelIndex];
+			FTransform3f BaseInvTransform;
+			BaseInvTransform.SetFromMatrix(BaseInvMatrix);
+
+			mu::EBoneUsageFlags BoneUsageFlags = mu::EBoneUsageFlags::None;
+			EnumAddFlags(BoneUsageFlags, BoneMapIndex < BoneMapCount ? mu::EBoneUsageFlags::Skinning : mu::EBoneUsageFlags::None);
+			EnumAddFlags(BoneUsageFlags, BoneInfo.ParentIndex == INDEX_NONE ? mu::EBoneUsageFlags::Root : mu::EBoneUsageFlags::None);
+			MutableMesh->SetBonePose(BoneMapIndex, StringCast<ANSICHAR>(*BoneName).Get(), BaseInvTransform.Inverse(), BoneUsageFlags);
 		}
 	}
 
@@ -867,16 +858,16 @@ mu::MeshPtr ConvertSkeletalMeshToMutable(USkeletalMesh* InSkeletalMesh, const TS
 			}
 		}
 
-		const int32 MaxSectionBoneMapIndex = ImportedModel->LODModels[LOD].Sections[MaterialIndex].BoneMap.Num();
+		const int32 MaxSectionBoneMapIndex = BoneMap.Num();
 
 		for (int32 VertexIndex = VertexStart; VertexIndex < VertexStart + VertexCount && VertexIndex < Vertices.Num(); ++VertexIndex)
 		{
 			FSoftSkinVertex& Vertex = Vertices[VertexIndex];
 
 			// Transfer removed bones influences to parent bones
-			if (!InfluencesToReplaceMap.IsEmpty())
+			if (bBoneMapModified)
 			{
-				TransferRemovedBonesInfluences(&Vertex.InfluenceBones[0], &Vertex.InfluenceWeights[0], MaxSectionInfluences, InfluencesToReplaceMap);
+				TransferRemovedBonesInfluences(&Vertex.InfluenceBones[0], &Vertex.InfluenceWeights[0], MaxSectionInfluences, RemappedBoneMapIndices);
 			}
 
 			if (!GenerationContext.Options.bExtraBoneInfluencesEnabled)
@@ -908,7 +899,7 @@ mu::MeshPtr ConvertSkeletalMeshToMutable(USkeletalMesh* InSkeletalMesh, const TS
 
 		// This call involves resolving every TObjectPtr<UMorphTarget> to a UMorphTarget*, so
 		// cache the result here to avoid calling it repeatedly.
-		const TArray<UMorphTarget*>& SkeletalMeshMorphTargets = SkeletalMesh->GetMorphTargets();
+		const TArray<UMorphTarget*>& SkeletalMeshMorphTargets = InSkeletalMesh->GetMorphTargets();
 
 		// Find realtime MorphTargets to be used.
 		TArray<const UMorphTarget*> UsedMorphTargets;
@@ -931,7 +922,7 @@ mu::MeshPtr ConvertSkeletalMeshToMutable(USkeletalMesh* InSkeletalMesh, const TS
                     ? MorphFound 
                     : &RealTimeMorphTargetOverrides.Emplace_GetRef(MorphTarget->GetFName());
 
-            const int32 AddedMeshNameIdx = MorphFound->SkeletalMeshesNames.AddUnique(SkeletalMesh->GetFName());
+            const int32 AddedMeshNameIdx = MorphFound->SkeletalMeshesNames.AddUnique(InSkeletalMesh->GetFName());
 
 			if (AddedMeshNameIdx >= MorphFound->Override.Num())
 			{
@@ -979,7 +970,7 @@ mu::MeshPtr ConvertSkeletalMeshToMutable(USkeletalMesh* InSkeletalMesh, const TS
 					} 
 				
 					const int32 FoundIdx = 
-							MorphTargetOverride.SkeletalMeshesNames.Find(SkeletalMesh->GetFName());
+							MorphTargetOverride.SkeletalMeshesNames.Find(InSkeletalMesh->GetFName());
 
 					if (FoundIdx != INDEX_NONE)
 					{
@@ -1145,7 +1136,7 @@ mu::MeshPtr ConvertSkeletalMeshToMutable(USkeletalMesh* InSkeletalMesh, const TS
 		// clothing assets are shared among all LODs in a section
 		const int32 ClothingAssetIndex = Invoke([&]() -> int32
 		{
-			const UClothingAssetBase* ClothingAssetBase = SkeletalMesh->GetSectionClothingAsset(LOD, MaterialIndex);
+			const UClothingAssetBase* ClothingAssetBase = InSkeletalMesh->GetSectionClothingAsset(LOD, MaterialIndex);
 
 			if (!ClothingAssetBase)
 			{
@@ -1281,7 +1272,7 @@ mu::MeshPtr ConvertSkeletalMeshToMutable(USkeletalMesh* InSkeletalMesh, const TS
 		const int32 MaxSectionInfluences = ImportedModel->LODModels[LOD].Sections[MaterialIndex].GetMaxBoneInfluences();
 		const int32 MaxSectionBoneMapIndex = ImportedModel->LODModels[LOD].Sections[MaterialIndex].BoneMap.Num();
 
-		const TArray<FSkinWeightProfileInfo>& SkinWeightProfilesInfo = SkeletalMesh->GetSkinWeightProfiles();
+		const TArray<FSkinWeightProfileInfo>& SkinWeightProfilesInfo = InSkeletalMesh->GetSkinWeightProfiles();
 		for (const FSkinWeightProfileInfo& Profile : SkinWeightProfilesInfo)
 		{
 			const FImportedSkinWeightProfileData* ImportedProfileData = ImportedModel->LODModels[LOD].SkinWeightProfiles.Find(Profile.Name);
@@ -1300,9 +1291,9 @@ mu::MeshPtr ConvertSkeletalMeshToMutable(USkeletalMesh* InSkeletalMesh, const TS
 			{
 				FRawSkinWeight SkinWeight = ImportedProfileData->SkinWeights[VertexIndex];
 
-				if (!InfluencesToReplaceMap.IsEmpty())
+				if (bBoneMapModified)
 				{
-					TransferRemovedBonesInfluences(&SkinWeight.InfluenceBones[0], &SkinWeight.InfluenceWeights[0], MaxSectionInfluences, InfluencesToReplaceMap);
+					TransferRemovedBonesInfluences(&SkinWeight.InfluenceBones[0], &SkinWeight.InfluenceWeights[0], MaxSectionInfluences, RemappedBoneMapIndices);
 				}
 
 				if (!GenerationContext.Options.bExtraBoneInfluencesEnabled)
@@ -1379,8 +1370,8 @@ mu::MeshPtr ConvertSkeletalMeshToMutable(USkeletalMesh* InSkeletalMesh, const TS
 		// For some reason, the indices in 4.25 (and 4.24) are in different order in the Imported and Rendering data structures. The strange thing 
 		// is actually that the vertices in the imported model seem to match the rendering model indices. Maybe there is some mapping that
 		// we are missing, but for now this will do:
-		const int ElementSize = SkeletalMesh->GetResourceForRendering()->LODRenderData[LOD].MultiSizeIndexContainer.GetDataTypeSize();
-		void* IndexDataPointer = SkeletalMesh->GetResourceForRendering()->LODRenderData[LOD].MultiSizeIndexContainer.GetIndexBuffer()->GetPointerTo(IndexStart);
+		const int ElementSize = InSkeletalMesh->GetResourceForRendering()->LODRenderData[LOD].MultiSizeIndexContainer.GetDataTypeSize();
+		void* IndexDataPointer = InSkeletalMesh->GetResourceForRendering()->LODRenderData[LOD].MultiSizeIndexContainer.GetIndexBuffer()->GetPointerTo(IndexStart);
 		const int FinalElementSize = sizeof(uint32_t);
 		const int ChannelCount = 1;
 		const MESH_BUFFER_SEMANTIC Semantics[ChannelCount] = { MBS_VERTEXINDEX };
@@ -1443,93 +1434,10 @@ mu::MeshPtr ConvertSkeletalMeshToMutable(USkeletalMesh* InSkeletalMesh, const TS
 		}
 	}
 
-	if (!bIgnoreSkeleton)
-	{
-		// Skeleton
-
-		// Add the bones in the order they are used in the mesh
-		TMap<int32, int32> InverseBoneMap; // Inverse of ImportedModel->LODModels[LOD].Chunks[MaterialIndex].BoneMap for the chunk's bone indices
-		TQueue<int32> SkippedBoneMapIndices; // The order and index of the bones in the BoneMap must remain the same even if some bones are removed
-
-		const int32 NumBonesInMesh = ImportedModel->LODModels[LOD].Sections[MaterialIndex].BoneMap.Num();
-		for (int32 BoneMapIndex = 0; BoneMapIndex < NumBonesInMesh; ++BoneMapIndex)
-		{
-			if (InfluencesToReplaceMap.Find(BoneMapIndex))
-			{
-				SkippedBoneMapIndices.Enqueue(BoneMapIndex); // Enqueue the skipped indices, will be used later by the remaining RequiredBones.
-				continue;
-			}
-
-			int32 RefSkeletonIndex = ImportedModel->LODModels[LOD].Sections[MaterialIndex].BoneMap[BoneMapIndex];
-			InverseBoneMap.Add(RefSkeletonIndex, BoneMapIndex);
-
-		}
-
-		// Add the data to mutable
-		int NumRequiredBones = ImportedModel->LODModels[LOD].RequiredBones.Num();
-		mu::SkeletonPtr MutableSkeleton = new mu::Skeleton;
-		MutableMesh->SetSkeleton(MutableSkeleton);
-		MutableMesh->SetBonePoseCount(NumRequiredBones);
-		MutableSkeleton->SetBoneCount(NumRequiredBones);
-		for (int32 RequiredBoneIndex = 0; RequiredBoneIndex < NumRequiredBones; ++RequiredBoneIndex)
-		{
-			int32 RefSkelIndex = ImportedModel->LODModels[LOD].RequiredBones[RequiredBoneIndex];
-
-			bool bSkinned = true;
-			// If the bone was required but not in the BoneMap (parent of a skinning bone) add it now.
-			if (!InverseBoneMap.Contains(RefSkelIndex))
-			{
-				bSkinned = false;
-				int32 NewBoneIndex = 0;
-				if (SkippedBoneMapIndices.Dequeue(NewBoneIndex)) // use skipped index
-				{
-					InverseBoneMap.Add(RefSkelIndex, NewBoneIndex);
-				}
-				else
-				{
-					InverseBoneMap.Add(RefSkelIndex, InverseBoneMap.Num());
-				}
-			}
-
-			int32 BoneMapIndex = InverseBoneMap[RefSkelIndex];
-			const FMeshBoneInfo& BoneInfo = SkeletalMesh->GetRefSkeleton().GetRefBoneInfo()[RefSkelIndex];
-			FString BoneName = BoneInfo.Name.ToString();
-
-			if (BoneMapIndex >= NumRequiredBones)
-			{
-				// This can only happen with invalid data.
-				UE_LOG(LogMutable, Warning, TEXT("In object [%s] ConvertSkeletalMeshToMutable: Source mesh [%s] section uses a bone [%s] that is not in its RequiredBones list."),
-					*GenerationContext.Object->GetName(),
-					*SkeletalMesh->GetName(),
-					*BoneName);
-				MutableSkeleton->SetBoneCount(BoneMapIndex + 1);
-			}
-
-			MutableSkeleton->SetBoneName(BoneMapIndex, StringCast<ANSICHAR>(*BoneName).Get());
-			int32 ParentRefSkeletonIndex = SkeletalMesh->GetRefSkeleton().GetParentIndex(RefSkelIndex);
-			int ParentBoneMapIndex = ParentRefSkeletonIndex >= 0 ? InverseBoneMap[ParentRefSkeletonIndex] : -1;
-			MutableSkeleton->SetBoneParent(BoneMapIndex, ParentBoneMapIndex);
-
-			FMatrix44f BaseInvMatrix = SkeletalMesh->GetRefBasesInvMatrix()[RefSkelIndex];
-			FTransform3f BaseInvTransform;
-			BaseInvTransform.SetFromMatrix(BaseInvMatrix);
-
-			mu::EBoneUsageFlags BoneUsageFlags = bSkinned ? mu::EBoneUsageFlags::Skinning : mu::EBoneUsageFlags::None;
-			
-			if (SkeletalMesh->GetRefSkeleton().GetParentIndex(RefSkelIndex) == INDEX_NONE)
-			{
-				EnumAddFlags(BoneUsageFlags, mu::EBoneUsageFlags::Root);
-			}
-
-			MutableMesh->SetBonePose(BoneMapIndex, StringCast<ANSICHAR>(*BoneName).Get(), BaseInvTransform.Inverse(), BoneUsageFlags);
-		}
-	}
-
-	const bool bPhysicsMergeEnabled = GenerationContext.Options.bPhysicsAssetMergeEnabled;
-	if (!bIgnorePhysics && SkeletalMesh->GetPhysicsAsset() && MutableMesh->GetSkeleton() && bPhysicsMergeEnabled)
+	if (!bIgnorePhysics && InSkeletalMesh->GetPhysicsAsset() && MutableMesh->GetSkeleton())
 	{
 		// Find BodySetups with relevant bones.
-		TArray<TObjectPtr<USkeletalBodySetup>>& SkeletalBodySetups = SkeletalMesh->GetPhysicsAsset()->SkeletalBodySetups;
+		TArray<TObjectPtr<USkeletalBodySetup>>& SkeletalBodySetups = InSkeletalMesh->GetPhysicsAsset()->SkeletalBodySetups;
 		
 		TArray<TObjectPtr<USkeletalBodySetup>> RelevantBodySetups;
 		RelevantBodySetups.Reserve(SkeletalBodySetups.Num());
@@ -1570,7 +1478,7 @@ mu::MeshPtr ConvertSkeletalMeshToMutable(USkeletalMesh* InSkeletalMesh, const TS
 		{
 			FString PhysicsSetupsRemovedMsg = 
 					FString::Printf(TEXT("PhysicsBodySetups in %s attached to bones"), 
-					*(SkeletalMesh->GetPhysicsAsset()->GetName()));
+					*(InSkeletalMesh->GetPhysicsAsset()->GetName()));
 
 			constexpr int32 MaxNumDiscardedShown = 3;
 			
@@ -1590,7 +1498,7 @@ mu::MeshPtr ConvertSkeletalMeshToMutable(USkeletalMesh* InSkeletalMesh, const TS
 			}
 	
 			PhysicsSetupsRemovedMsg += FString::Printf(TEXT("have been discarded because they are not present in the SkeletalMesh [%s] Skeleton."),
-				*SkeletalMesh->GetName());
+				*InSkeletalMesh->GetName());
 					
 			GenerationContext.Compiler->CompilerLog(FText::FromString(PhysicsSetupsRemovedMsg), CurrentNode, EMessageSeverity::Warning);
 		}
@@ -1775,6 +1683,7 @@ mu::MeshPtr ConvertSkeletalMeshToMutable(USkeletalMesh* InSkeletalMesh, const TS
 	return MutableMesh;
 }
 
+
 mu::MeshPtr ConvertStaticMeshToMutable(UStaticMesh* StaticMesh, int LOD, int MaterialIndex, FMutableGraphGenerationContext& GenerationContext, const UCustomizableObjectNode* CurrentNode)
 {
 	if (!StaticMesh->GetRenderData() ||
diff --git a/Engine/Plugins/Experimental/Mutable/Source/CustomizableObjectEditor/Private/MuCOE/Nodes/CustomizableObjectNodeObject.cpp b/Engine/Plugins/Experimental/Mutable/Source/CustomizableObjectEditor/Private/MuCOE/Nodes/CustomizableObjectNodeObject.cpp
--- a/Engine/Plugins/Experimental/Mutable/Source/CustomizableObjectEditor/Private/MuCOE/Nodes/CustomizableObjectNodeObject.cpp
+++ b/Engine/Plugins/Experimental/Mutable/Source/CustomizableObjectEditor/Private/MuCOE/Nodes/CustomizableObjectNodeObject.cpp
@@ -36,12 +36,6 @@ UCustomizableObjectNodeObject::UCustomizableObjectNodeObject()
 		Identifier = FGuid::NewGuid();
 		IdentifierVerification = Identifier;
 	}
-
-	if (ComponentSettings.IsEmpty())
-	{
-		ComponentSettings.Add(FComponentSettings());
-		ComponentSettings.Last().LODReductionSettings.Add(FLODReductionSettings());
-	}
 }
 
 
@@ -273,6 +267,15 @@ void UCustomizableObjectNodeObject::PostBackwardsCompatibleFixup()
 		.Log();
 	}
 
+	// Fix up ComponentSettings
+	if (ComponentSettings.IsEmpty())
+	{
+		FComponentSettings ComponentSettingsTemplate;
+		ComponentSettingsTemplate.LODReductionSettings.SetNum(NumLODs);
+
+		ComponentSettings.Init(ComponentSettingsTemplate, NumMeshComponents);
+	}
+
 	// Reconstruct in case any extension pins have changed
 	ReconstructNode();
 }
diff --git a/Engine/Plugins/Experimental/Mutable/Source/CustomizableObjectEditor/Private/MuCOE/Widgets/CustomizableObjectLODReductionSettings.cpp b/Engine/Plugins/Experimental/Mutable/Source/CustomizableObjectEditor/Private/MuCOE/Widgets/CustomizableObjectLODReductionSettings.cpp
--- a/Engine/Plugins/Experimental/Mutable/Source/CustomizableObjectEditor/Private/MuCOE/Widgets/CustomizableObjectLODReductionSettings.cpp
+++ b/Engine/Plugins/Experimental/Mutable/Source/CustomizableObjectEditor/Private/MuCOE/Widgets/CustomizableObjectLODReductionSettings.cpp
@@ -26,7 +26,7 @@ void FCustomizableObjectLODReductionSettings::CustomizeHeader(TSharedRef<IProper
 	constexpr bool bRecurse = false;
 
 	BoneNameProperty = StructPropertyHandle->GetChildHandle(GET_MEMBER_NAME_CHECKED(FBoneToRemove, BoneName), bRecurse);
-	IncludeBoneProperty = StructPropertyHandle->GetChildHandle(GET_MEMBER_NAME_CHECKED(FBoneToRemove, bIncludeBone), bRecurse);
+	IncludeBoneProperty = StructPropertyHandle->GetChildHandle(GET_MEMBER_NAME_CHECKED(FBoneToRemove, bOnlyRemoveChildren), bRecurse);
 
 	TSharedPtr<IPropertyHandle> Parent = StructPropertyHandle->GetParentHandle();
 	if (Parent)
@@ -39,8 +39,8 @@ void FCustomizableObjectLODReductionSettings::CustomizeHeader(TSharedRef<IProper
 	{
 		ObjectNode = GetObjectNode();
 
-		bool bIncludeBoneProperty = false;
-		IncludeBoneProperty->GetValue(bIncludeBoneProperty);
+		bool bOnlyRemoveChildrenProperty = false;
+		IncludeBoneProperty->GetValue(bOnlyRemoveChildrenProperty);
 
 		HeaderRow
 		.NameContent()
@@ -97,7 +97,7 @@ void FCustomizableObjectLODReductionSettings::CustomizeHeader(TSharedRef<IProper
 				[
 					SNew(SCheckBox)
 					.OnCheckStateChanged(this, &FCustomizableObjectLODReductionSettings::OnBoneCheckBoxChanged)
-					.IsChecked(bIncludeBoneProperty ? ECheckBoxState::Checked : ECheckBoxState::Unchecked)
+					.IsChecked(bOnlyRemoveChildrenProperty ? ECheckBoxState::Checked : ECheckBoxState::Unchecked)
 				]
 			]
 		];
