[Mutable] Fix incorrect Task scope.

#rnx

diff --git a/Engine/Plugins/Experimental/Mutable/Source/CustomizableObject/Private/MuCO/CustomizableObjectSystem.cpp b/Engine/Plugins/Experimental/Mutable/Source/CustomizableObject/Private/MuCO/CustomizableObjectSystem.cpp
--- a/Engine/Plugins/Experimental/Mutable/Source/CustomizableObject/Private/MuCO/CustomizableObjectSystem.cpp
+++ b/Engine/Plugins/Experimental/Mutable/Source/CustomizableObject/Private/MuCO/CustomizableObjectSystem.cpp
@@ -2015,7 +2015,7 @@ namespace impl
 	/** "Lock Cached Resources" */
 	void Task_Game_LockCache(TSharedPtr<FMutableOperationData> OperationData, const TWeakObjectPtr<UCustomizableObjectInstance>& CustomizableObjectInstancePtr, mu::Ptr<const mu::Parameters> Parameters, bool bBuildParameterDecorations)
 	{
-		MUTABLE_CPUPROFILER_SCOPE(Task_Game_ConvertResources)
+		MUTABLE_CPUPROFILER_SCOPE(Task_Game_LockCache)
 
 		check(IsInGameThread());
 		check(OperationData);
