#rnx

diff --git a/Engine/Plugins/Experimental/Mutable/Source/CustomizableObject/Classes/MuCO/CustomizableObject.h b/Engine/Plugins/Experimental/Mutable/Source/CustomizableObject/Classes/MuCO/CustomizableObject.h
--- a/Engine/Plugins/Experimental/Mutable/Source/CustomizableObject/Classes/MuCO/CustomizableObject.h
+++ b/Engine/Plugins/Experimental/Mutable/Source/CustomizableObject/Classes/MuCO/CustomizableObject.h
@@ -1117,6 +1117,9 @@ public:
 	UPROPERTY()
 	TArray<TSoftObjectPtr<USkeleton>> ReferencedSkeletons;
 
+	UPROPERTY()
+	TArray<TSoftObjectPtr<UTexture2D>> ReferencedPassThroughTextures;
+
 	UPROPERTY(EditAnywhere, Category = CustomizableObject, meta = (DisplayName = "LOD Settings"))
 	FMutableLODSettings LODSettings;
 
diff --git a/Engine/Plugins/Experimental/Mutable/Source/CustomizableObjectEditor/Classes/MuCOE/EdGraphSchema_CustomizableObject.h b/Engine/Plugins/Experimental/Mutable/Source/CustomizableObjectEditor/Classes/MuCOE/EdGraphSchema_CustomizableObject.h
--- a/Engine/Plugins/Experimental/Mutable/Source/CustomizableObjectEditor/Classes/MuCOE/EdGraphSchema_CustomizableObject.h
+++ b/Engine/Plugins/Experimental/Mutable/Source/CustomizableObjectEditor/Classes/MuCOE/EdGraphSchema_CustomizableObject.h
@@ -94,6 +94,7 @@ public:
 	static const FName PC_Mesh;
 	static const FName PC_Layout;
 	static const FName PC_Image;
+	static const FName PC_PassThroughImage;
 	static const FName PC_Projector;
 	static const FName PC_GroupProjector;
 	static const FName PC_Color;
diff --git a/Engine/Plugins/Experimental/Mutable/Source/CustomizableObjectEditor/Classes/MuCOE/Nodes/CustomizableObjectNodeMaterial.h b/Engine/Plugins/Experimental/Mutable/Source/CustomizableObjectEditor/Classes/MuCOE/Nodes/CustomizableObjectNodeMaterial.h
--- a/Engine/Plugins/Experimental/Mutable/Source/CustomizableObjectEditor/Classes/MuCOE/Nodes/CustomizableObjectNodeMaterial.h
+++ b/Engine/Plugins/Experimental/Mutable/Source/CustomizableObjectEditor/Classes/MuCOE/Nodes/CustomizableObjectNodeMaterial.h
@@ -119,6 +119,7 @@ public:
 	TSharedPtr<SGraphNode> CreateVisualWidget() override;
 	virtual void PinConnectionListChanged(UEdGraphPin* Pin) override;
 	virtual void PostPasteNode() override;
+	virtual bool CanConnect(const UEdGraphPin* InOwnedInputPin, const UEdGraphPin* InOutputPin, bool& bOutIsOtherNodeBlocklisted, bool& bOutArePinsCompatible) const override;
 
 	// UCustomizableObjectNode interface
 	virtual void BackwardsCompatibleFixup() override;
diff --git a/Engine/Plugins/Experimental/Mutable/Source/CustomizableObjectEditor/Classes/MuCOE/Nodes/CustomizableObjectNodePassThroughTexture.h b/Engine/Plugins/Experimental/Mutable/Source/CustomizableObjectEditor/Classes/MuCOE/Nodes/CustomizableObjectNodePassThroughTexture.h
new file mode 100644
--- /dev/null
+++ b/Engine/Plugins/Experimental/Mutable/Source/CustomizableObjectEditor/Classes/MuCOE/Nodes/CustomizableObjectNodePassThroughTexture.h
@@ -0,0 +1,23 @@
+// Copyright Epic Games, Inc. All Rights Reserved.
+
+#pragma once
+
+#include "MuCOE/Nodes/CustomizableObjectNodeTexture.h"
+
+#include "CustomizableObjectNodePassThroughTexture.generated.h"
+
+
+UCLASS()
+class CUSTOMIZABLEOBJECTEDITOR_API UCustomizableObjectNodePassThroughTexture : public UCustomizableObjectNodeTextureBase
+{
+public:
+	GENERATED_BODY()
+
+	// UCustomizableObjectNode interface
+	void AllocateDefaultPins(UCustomizableObjectNodeRemapPins* RemapPins) override;
+
+	// Begin EdGraphNode interface
+	FText GetNodeTitle(ENodeTitleType::Type TitleType) const override;
+	FLinearColor GetNodeTitleColor() const override;
+	FText GetTooltipText() const override;
+};
diff --git a/Engine/Plugins/Experimental/Mutable/Source/CustomizableObjectEditor/Classes/MuCOE/Nodes/CustomizableObjectNodePassThroughTextureSwitch.h b/Engine/Plugins/Experimental/Mutable/Source/CustomizableObjectEditor/Classes/MuCOE/Nodes/CustomizableObjectNodePassThroughTextureSwitch.h
new file mode 100644
--- /dev/null
+++ b/Engine/Plugins/Experimental/Mutable/Source/CustomizableObjectEditor/Classes/MuCOE/Nodes/CustomizableObjectNodePassThroughTextureSwitch.h
@@ -0,0 +1,22 @@
+// Copyright Epic Games, Inc. All Rights Reserved.
+
+#pragma once
+
+#include "MuCOE/Nodes/CustomizableObjectNodeSwitchBase.h"
+
+#include "CustomizableObjectNodePassThroughTextureSwitch.generated.h"
+
+
+UCLASS()
+class CUSTOMIZABLEOBJECTEDITOR_API UCustomizableObjectNodePassThroughTextureSwitch : public UCustomizableObjectNodeSwitchBase
+{
+public:
+	GENERATED_BODY()
+
+	// UCustomizableObjectNodeSwitchBase interface
+	FString GetOutputPinName() const override;
+
+	FName GetCategory() const override;
+
+	FString GetPinPrefix() const;
+};
diff --git a/Engine/Plugins/Experimental/Mutable/Source/CustomizableObjectEditor/Classes/MuCOE/Nodes/CustomizableObjectNodeTexture.h b/Engine/Plugins/Experimental/Mutable/Source/CustomizableObjectEditor/Classes/MuCOE/Nodes/CustomizableObjectNodeTexture.h
--- a/Engine/Plugins/Experimental/Mutable/Source/CustomizableObjectEditor/Classes/MuCOE/Nodes/CustomizableObjectNodeTexture.h
+++ b/Engine/Plugins/Experimental/Mutable/Source/CustomizableObjectEditor/Classes/MuCOE/Nodes/CustomizableObjectNodeTexture.h
@@ -48,37 +48,44 @@ public:
 public:
 
 	// Pointer to the NodeTexture that owns this SGraphNode
-	UCustomizableObjectNodeTexture* NodeTexture;
+	UCustomizableObjectNodeTextureBase* NodeTexture;
 
 	// Single property that only draws the combo box widget of the Texture
 	TSharedPtr<class ISinglePropertyView> TextureSelector;
 
 	// Brush to draw the texture to the widget
 	FSlateBrush TextureBrush;
-
 };
 
 
 UCLASS()
-class CUSTOMIZABLEOBJECTEDITOR_API UCustomizableObjectNodeTexture : public UCustomizableObjectNode
+class CUSTOMIZABLEOBJECTEDITOR_API UCustomizableObjectNodeTextureBase : public UCustomizableObjectNode
 {
 public:
 	GENERATED_BODY()
 
-	UPROPERTY(EditAnywhere, Category=CustomizableObject)
+	UPROPERTY(EditAnywhere, Category = CustomizableObject)
 	TObjectPtr<UTexture2D> Texture = nullptr;
 
-	// Begin EdGraphNode interface
-	FText GetNodeTitle(ENodeTitleType::Type TitleType) const override;
-	FLinearColor GetNodeTitleColor() const override;
-	FText GetTooltipText() const override;
+	// Determines if the Node is collapsed or not
+	bool bCollapsed = true;
+
+	// Creates the SGraph Node widget for the thumbnail
+	TSharedPtr<SGraphNode> CreateVisualWidget() override;
+};
+
+
+UCLASS()
+class CUSTOMIZABLEOBJECTEDITOR_API UCustomizableObjectNodeTexture : public UCustomizableObjectNodeTextureBase
+{
+public:
+	GENERATED_BODY()
 
 	// UCustomizableObjectNode interface
 	void AllocateDefaultPins(UCustomizableObjectNodeRemapPins* RemapPins) override;
 
-	// Creates the SGraph Node widget for the thumbnail
-	TSharedPtr<SGraphNode> CreateVisualWidget() override;
-	
-	// Determines if the Node is collapsed or not
-	bool bCollapsed = true;
+	// Begin EdGraphNode interface
+	FText GetNodeTitle(ENodeTitleType::Type TitleType) const override;
+	FLinearColor GetNodeTitleColor() const override;
+	FText GetTooltipText() const override;
 };
diff --git a/Engine/Plugins/Experimental/Mutable/Source/CustomizableObjectEditor/Private/MuCOE/CustomizableObjectCompiler.cpp b/Engine/Plugins/Experimental/Mutable/Source/CustomizableObjectEditor/Private/MuCOE/CustomizableObjectCompiler.cpp
--- a/Engine/Plugins/Experimental/Mutable/Source/CustomizableObjectEditor/Private/MuCOE/CustomizableObjectCompiler.cpp
+++ b/Engine/Plugins/Experimental/Mutable/Source/CustomizableObjectEditor/Private/MuCOE/CustomizableObjectCompiler.cpp
@@ -1172,6 +1172,27 @@ void FCustomizableObjectCompiler::CompileInternal(UCustomizableObject* Object, c
 			Object->ReferencedSkeletons.Add(Skeleton);
 		}
 
+		// Pass-through textures
+		TArray<TSoftObjectPtr<UTexture2D>> NewReferencedPassThroughTextures;
+
+		for (const TPair<TSoftObjectPtr<UTexture2D>, uint32>& Pair : GenerationContext.PassThroughTextureToIndexMap)
+		{
+			check(Pair.Value == NewReferencedPassThroughTextures.Num());
+			NewReferencedPassThroughTextures.Add(Pair.Key);
+		}
+
+		// Mark the object as modified, used to avoid missing assets in packages.
+		if (Object->ReferencedPassThroughTextures != NewReferencedPassThroughTextures)
+		{
+			Object->ReferencedPassThroughTextures.Empty(NewReferencedPassThroughTextures.Num());
+			Object->ReferencedPassThroughTextures = NewReferencedPassThroughTextures;
+
+			if (!ParamNamesToSelectedOptions.Num()) // Don't mark the objects as modified because of a partial compilation
+			{
+				Object->MarkPackageDirty();
+			}
+		}
+
 		if (bIsRootObject && (GenerationContext.MaskOutMaterialCache.Num() > 0 || GenerationContext.MaskOutTextureCache.Num() > 0))
 		{
 			// Load MaskOutCache, if can't, create it.
diff --git a/Engine/Plugins/Experimental/Mutable/Source/CustomizableObjectEditor/Private/MuCOE/EdGraphSchema_CustomizableObject.cpp b/Engine/Plugins/Experimental/Mutable/Source/CustomizableObjectEditor/Private/MuCOE/EdGraphSchema_CustomizableObject.cpp
--- a/Engine/Plugins/Experimental/Mutable/Source/CustomizableObjectEditor/Private/MuCOE/EdGraphSchema_CustomizableObject.cpp
+++ b/Engine/Plugins/Experimental/Mutable/Source/CustomizableObjectEditor/Private/MuCOE/EdGraphSchema_CustomizableObject.cpp
@@ -53,6 +53,7 @@
 #include "MuCOE/Nodes/CustomizableObjectNodeStaticMesh.h"
 #include "MuCOE/Nodes/CustomizableObjectNodeTable.h"
 #include "MuCOE/Nodes/CustomizableObjectNodeTexture.h"
+#include "MuCOE/Nodes/CustomizableObjectNodePassThroughTexture.h"
 #include "MuCOE/Nodes/CustomizableObjectNodeTextureBinarise.h"
 #include "MuCOE/Nodes/CustomizableObjectNodeTextureColourMap.h"
 #include "MuCOE/Nodes/CustomizableObjectNodeTextureFromChannels.h"
@@ -64,6 +65,7 @@
 #include "MuCOE/Nodes/CustomizableObjectNodeTextureProject.h"
 #include "MuCOE/Nodes/CustomizableObjectNodeTextureSample.h"
 #include "MuCOE/Nodes/CustomizableObjectNodeTextureSwitch.h"
+#include "MuCOE/Nodes/CustomizableObjectNodePassThroughTextureSwitch.h"
 #include "MuCOE/Nodes/CustomizableObjectNodeTextureToChannels.h"
 #include "MuCOE/Nodes/CustomizableObjectNodeTextureTransform.h"
 #include "MuCOE/Nodes/CustomizableObjectNodeTextureSaturate.h"
@@ -213,6 +215,7 @@ const FName UEdGraphSchema_CustomizableObject::PC_Material("material");
 const FName UEdGraphSchema_CustomizableObject::PC_Mesh("mesh");
 const FName UEdGraphSchema_CustomizableObject::PC_Layout("layout");
 const FName UEdGraphSchema_CustomizableObject::PC_Image("image");
+const FName UEdGraphSchema_CustomizableObject::PC_PassThroughImage("passThroughImage");
 const FName UEdGraphSchema_CustomizableObject::PC_Projector("projector");
 const FName UEdGraphSchema_CustomizableObject::PC_GroupProjector("groupProjector");
 const FName UEdGraphSchema_CustomizableObject::PC_Color("color");
@@ -415,10 +418,12 @@ void UEdGraphSchema_CustomizableObject::GetGraphContextActions(FGraphContextMenu
 		UCustomizableObjectNode* TextureTemplateNodes[] =
 		{
 			ContextMenuBuilder.CreateTemplateNode<UCustomizableObjectNodeTexture>(),
+			ContextMenuBuilder.CreateTemplateNode<UCustomizableObjectNodePassThroughTexture>(),
 			ContextMenuBuilder.CreateTemplateNode<UCustomizableObjectNodeTextureBinarise>(),
 			ContextMenuBuilder.CreateTemplateNode<UCustomizableObjectNodeTextureInterpolate>(),
 			ContextMenuBuilder.CreateTemplateNode<UCustomizableObjectNodeTextureLayer>(),
 			ContextMenuBuilder.CreateTemplateNode<UCustomizableObjectNodeTextureSwitch>(),
+			ContextMenuBuilder.CreateTemplateNode<UCustomizableObjectNodePassThroughTextureSwitch>(),
 			ContextMenuBuilder.CreateTemplateNode<UCustomizableObjectNodeTextureVariation>(),
 			ContextMenuBuilder.CreateTemplateNode<UCustomizableObjectNodeTextureToChannels>(),
 			ContextMenuBuilder.CreateTemplateNode<UCustomizableObjectNodeTextureFromChannels>(),
@@ -485,7 +490,7 @@ void UEdGraphSchema_CustomizableObject::GetGraphContextActions(FGraphContextMenu
 
 	{
 		// External Pin Nodes
-		TArray<FName> PinTypes({ PC_Material, PC_Mesh, PC_Image, PC_Projector, PC_GroupProjector, PC_Color, PC_Float, PC_Bool, PC_Enum, PC_Stack });
+		TArray<FName> PinTypes({ PC_Material, PC_Mesh, PC_Image, PC_Projector, PC_GroupProjector, PC_Color, PC_Float, PC_Bool, PC_Enum, PC_Stack, PC_PassThroughImage });
 
 		// Add pin types from extensions
 		for (const FRegisteredCustomizableObjectPinType& PinType : ICustomizableObjectModule::Get().GetExtendedPinTypes())
@@ -658,6 +663,10 @@ FLinearColor UEdGraphSchema_CustomizableObject::GetPinTypeColor(const FName& Typ
 	{
 		return FLinearColor(0.353393f, 0.454175f, 1.000000f, 1.000000f);
 	}
+	else if (TypeString == PC_PassThroughImage)
+	{
+		return FLinearColor(0.353393f, 0.454175f, 0.353393f, 1.000000f);
+	}
 	else if (TypeString == PC_Material)
 	{
 		return FLinearColor(0.000000f, 0.100000f, 0.600000f, 1.000000f);
@@ -1052,6 +1061,10 @@ FText UEdGraphSchema_CustomizableObject::GetPinCategoryName(const FName& PinCate
 	{
 		return LOCTEXT("Image_Pin_Category", "Texture");
 	}
+	else if (PinCategory == UEdGraphSchema_CustomizableObject::PC_PassThroughImage)
+	{
+		return LOCTEXT("Image_Pin_Category", "PassThrough Texture");
+	}
 	else if (PinCategory == UEdGraphSchema_CustomizableObject::PC_Projector)
 	{
 		return LOCTEXT("Projector_Pin_Category", "Projector");
diff --git a/Engine/Plugins/Experimental/Mutable/Source/CustomizableObjectEditor/Private/MuCOE/GenerateMutableSource/GenerateMutableSource.cpp b/Engine/Plugins/Experimental/Mutable/Source/CustomizableObjectEditor/Private/MuCOE/GenerateMutableSource/GenerateMutableSource.cpp
--- a/Engine/Plugins/Experimental/Mutable/Source/CustomizableObjectEditor/Private/MuCOE/GenerateMutableSource/GenerateMutableSource.cpp
+++ b/Engine/Plugins/Experimental/Mutable/Source/CustomizableObjectEditor/Private/MuCOE/GenerateMutableSource/GenerateMutableSource.cpp
@@ -24,6 +24,7 @@
 #include "MuCOE/Nodes/CustomizableObjectNodeObjectGroup.h"
 #include "MuCOE/Nodes/CustomizableObjectNodeTable.h"
 #include "MuCOE/Nodes/CustomizableObjectNodeTexture.h"
+#include "MuCOE/Nodes/CustomizableObjectNodePassThroughTexture.h"
 #include "MuCOE/Nodes/CustomizableObjectNodeTextureBinarise.h"
 #include "MuCOE/Nodes/CustomizableObjectNodeTextureColourMap.h"
 #include "MuCOE/Nodes/CustomizableObjectNodeTextureFromChannels.h"
@@ -33,6 +34,7 @@
 #include "MuCOE/Nodes/CustomizableObjectNodeTextureParameter.h"
 #include "MuCOE/Nodes/CustomizableObjectNodeTextureProject.h"
 #include "MuCOE/Nodes/CustomizableObjectNodeTextureSwitch.h"
+#include "MuCOE/Nodes/CustomizableObjectNodePassThroughTextureSwitch.h"
 #include "MuCOE/Nodes/CustomizableObjectNodeTextureToChannels.h"
 #include "MuCOE/Nodes/CustomizableObjectNodeTextureTransform.h"
 #include "MuCOE/Nodes/CustomizableObjectNodeTextureSaturate.h"
@@ -371,6 +373,11 @@ UTexture2D* FindReferenceImage(const UEdGraphPin* Pin, FMutableGraphGenerationCo
 		Result = TypedNodeTex->Texture;
 	}
 
+	else if (const UCustomizableObjectNodePassThroughTexture* TypedNodePassThroughTex = Cast<UCustomizableObjectNodePassThroughTexture>(Node))
+	{
+		Result = TypedNodePassThroughTex->Texture;
+	}
+
 	else if (const UCustomizableObjectNodeTextureParameter* ParamNodeTex = Cast<UCustomizableObjectNodeTextureParameter>(Node))
 	{
 		Result = ParamNodeTex->DefaultValue;
@@ -422,6 +429,17 @@ UTexture2D* FindReferenceImage(const UEdGraphPin* Pin, FMutableGraphGenerationCo
 		}
 	}
 
+	else if (const UCustomizableObjectNodePassThroughTextureSwitch* TypedNodePassThroughSwitch = Cast<UCustomizableObjectNodePassThroughTextureSwitch>(Node))
+	{
+		for (int32 SelectorIndex = 0; !Result && SelectorIndex < TypedNodePassThroughSwitch->GetNumElements(); ++SelectorIndex)
+		{
+			if (const UEdGraphPin* ConnectedPin = FollowInputPin(*TypedNodePassThroughSwitch->GetElementPin(SelectorIndex)))
+			{
+				Result = FindReferenceImage(ConnectedPin, GenerationContext);
+			}
+		}
+	}
+
 	else if (const UCustomizableObjectNodeTextureVariation* TypedNodeVariation = Cast<UCustomizableObjectNodeTextureVariation>(Node))
 	{
 		if (const UEdGraphPin* ConnectedPin = FollowInputPin(*TypedNodeVariation->DefaultPin()))
diff --git a/Engine/Plugins/Experimental/Mutable/Source/CustomizableObjectEditor/Private/MuCOE/GenerateMutableSource/GenerateMutableSource.h b/Engine/Plugins/Experimental/Mutable/Source/CustomizableObjectEditor/Private/MuCOE/GenerateMutableSource/GenerateMutableSource.h
--- a/Engine/Plugins/Experimental/Mutable/Source/CustomizableObjectEditor/Private/MuCOE/GenerateMutableSource/GenerateMutableSource.h
+++ b/Engine/Plugins/Experimental/Mutable/Source/CustomizableObjectEditor/Private/MuCOE/GenerateMutableSource/GenerateMutableSource.h
@@ -520,6 +520,9 @@ struct FMutableGraphGenerationContext
 	// level cache
 	TMap<FGeneratedImageKey, mu::NodeImageConstantPtr> GeneratedImages;
 
+	// Cache of pass-through images and their IDs used in the core to indentify them
+	TMap<TSoftObjectPtr<UTexture2D>, uint32> PassThroughTextureToIndexMap;
+
     // Global morph selection overrides.
     TArray<FRealTimeMorphSelectionOverride> RealTimeMorphTargetsOverrides;
 
diff --git a/Engine/Plugins/Experimental/Mutable/Source/CustomizableObjectEditor/Private/MuCOE/GenerateMutableSource/GenerateMutableSourceImage.cpp b/Engine/Plugins/Experimental/Mutable/Source/CustomizableObjectEditor/Private/MuCOE/GenerateMutableSource/GenerateMutableSourceImage.cpp
--- a/Engine/Plugins/Experimental/Mutable/Source/CustomizableObjectEditor/Private/MuCOE/GenerateMutableSource/GenerateMutableSourceImage.cpp
+++ b/Engine/Plugins/Experimental/Mutable/Source/CustomizableObjectEditor/Private/MuCOE/GenerateMutableSource/GenerateMutableSourceImage.cpp
@@ -13,6 +13,7 @@
 #include "MuCOE/Nodes/CustomizableObjectNodeMesh.h"
 #include "MuCOE/Nodes/CustomizableObjectNodeTable.h"
 #include "MuCOE/Nodes/CustomizableObjectNodeTexture.h"
+#include "MuCOE/Nodes/CustomizableObjectNodePassThroughTexture.h"
 #include "MuCOE/Nodes/CustomizableObjectNodeTextureBinarise.h"
 #include "MuCOE/Nodes/CustomizableObjectNodeTextureColourMap.h"
 #include "MuCOE/Nodes/CustomizableObjectNodeTextureFromChannels.h"
@@ -23,6 +24,7 @@
 #include "MuCOE/Nodes/CustomizableObjectNodeTextureParameter.h"
 #include "MuCOE/Nodes/CustomizableObjectNodeTextureProject.h"
 #include "MuCOE/Nodes/CustomizableObjectNodeTextureSwitch.h"
+#include "MuCOE/Nodes/CustomizableObjectNodePassThroughTextureSwitch.h"
 #include "MuCOE/Nodes/CustomizableObjectNodeTextureToChannels.h"
 #include "MuCOE/Nodes/CustomizableObjectNodeTextureTransform.h"
 #include "MuCOE/Nodes/CustomizableObjectNodeTextureSaturate.h"
@@ -169,6 +171,37 @@ mu::NodeImagePtr GenerateMutableSourceImage(const UEdGraphPin* Pin, FMutableGrap
 		}
 	}
 
+	else if (const UCustomizableObjectNodePassThroughTexture* TypedNodePassThroughTex = Cast<UCustomizableObjectNodePassThroughTexture>(Node))
+	{
+		UTexture2D* BaseTexture = TypedNodePassThroughTex->Texture;
+		if (BaseTexture)
+		{
+			uint32* FoundIndex = GenerationContext.PassThroughTextureToIndexMap.Find(BaseTexture);
+			uint32 NewId;
+
+			if (!FoundIndex)
+			{
+				NewId = GenerationContext.PassThroughTextureToIndexMap.Num();
+				GenerationContext.PassThroughTextureToIndexMap.Add(BaseTexture, NewId);
+			}
+			else
+			{
+				NewId = *FoundIndex;
+			}
+
+			// TODO: Implement mutable core nodes
+			//mu::NodePassThroughImageConstantPtr PassThroughImageNode;
+			//PassThroughImageNode.SetIndex(NewId);
+
+			//Result = PassThroughImageNode;
+
+		}
+		else
+		{
+			GenerationContext.Compiler->CompilerLog(LOCTEXT("MissingImagePassThrough", "Missing image in pass-through texture node."), Node);
+		}
+	}
+
 	else if (const UCustomizableObjectNodeTextureParameter* TypedNodeParam = Cast<UCustomizableObjectNodeTextureParameter>(Node))
 	{
 		mu::Ptr<mu::NodeImageParameter> TextureNode = new mu::NodeImageParameter();
@@ -400,6 +433,73 @@ mu::NodeImagePtr GenerateMutableSourceImage(const UEdGraphPin* Pin, FMutableGrap
 		}(); // invoke lambda.
 	}
 
+	else if (const UCustomizableObjectNodePassThroughTextureSwitch* TypedNodePassThroughTextureSwitch = Cast<UCustomizableObjectNodePassThroughTextureSwitch>(Node))
+	{
+		Result = [&]()
+		{
+			const UEdGraphPin* SwitchParameter = TypedNodePassThroughTextureSwitch->SwitchParameter();
+
+			// Check Switch Parameter arity preconditions.
+			if (const UEdGraphPin* ConnectedPin = FollowInputPin(*SwitchParameter))
+			{
+				mu::NodeScalarPtr SwitchParam = GenerateMutableSourceFloat(ConnectedPin, GenerationContext);
+				// Switch Param not generated
+				if (!SwitchParam)
+				{
+					const FText Message = LOCTEXT("FailedToGenerateSwitchParam", "Could not generate switch enum parameter. Please refesh the switch node and connect an enum.");
+					GenerationContext.Compiler->CompilerLog(Message, Node);
+
+					return Result;
+				}
+
+				if (SwitchParam->GetType() != mu::NodeScalarEnumParameter::GetStaticType())
+				{
+					const FText Message = LOCTEXT("WrongSwitchParamType", "Switch parameter of incorrect type.");
+					GenerationContext.Compiler->CompilerLog(Message, Node);
+
+					return Result;
+				}
+
+				const int32 NumSwitchOptions = TypedNodePassThroughTextureSwitch->GetNumElements();
+
+				mu::NodeScalarEnumParameter* EnumParameter = static_cast<mu::NodeScalarEnumParameter*>(SwitchParam.get());
+				if (NumSwitchOptions != EnumParameter->GetValueCount())
+				{
+					const FText Message = LOCTEXT("MismatchedSwitch", "Switch enum and switch node have different number of options. Please refresh the switch node to make sure the outcomes are labeled properly.");
+					GenerationContext.Compiler->CompilerLog(Message, Node);
+				}
+				
+				// TODO: Implement Mutable core pass-through switch nodes
+				//mu::NodePassThroughImageSwitchPtr SwitchNode = new mu::NodePassThroughImageSwitchPtr;
+				//SwitchNode->SetParameter(SwitchParam);
+				//SwitchNode->SetOptionCount(NumSwitchOptions);
+
+				for (int32 SelectorIndex = 0; SelectorIndex < NumSwitchOptions; ++SelectorIndex)
+				{
+					if (const UEdGraphPin* TexturePin = FollowInputPin(*TypedNodePassThroughTextureSwitch->GetElementPin(SelectorIndex)))
+					{
+						mu::NodeImagePtr PassThroughImage = GenerateMutableSourceImage(TexturePin, GenerationContext, MaxTextureSize);
+						//SwitchNode->SetOption(SelectorIndex, PassThroughImage);
+					}
+					else
+					{
+						const FText Message = LOCTEXT("MissingTexture", "Unable to generate pass-through texture switch node. Required connection not found.");
+						GenerationContext.Compiler->CompilerLog(Message, Node);
+						return Result;
+					}
+				}
+
+				//Result = SwitchNode;
+				return Result;
+			}
+			else
+			{
+				GenerationContext.Compiler->CompilerLog(LOCTEXT("NoEnumParamInSwitch", "Switch nodes must have an enum switch parameter. Please connect an enum and refesh the switch node."), Node);
+				return Result;
+			}
+		}(); // invoke lambda.
+	}
+
 	else if (const UCustomizableObjectNodeTextureVariation* TypedNodeImageVar = Cast<const UCustomizableObjectNodeTextureVariation>(Node))
 	{
 		mu::NodeImageVariationPtr TextureNode = new mu::NodeImageVariation();
diff --git a/Engine/Plugins/Experimental/Mutable/Source/CustomizableObjectEditor/Private/MuCOE/GenerateMutableSource/GenerateMutableSourceSurface.cpp b/Engine/Plugins/Experimental/Mutable/Source/CustomizableObjectEditor/Private/MuCOE/GenerateMutableSource/GenerateMutableSourceSurface.cpp
--- a/Engine/Plugins/Experimental/Mutable/Source/CustomizableObjectEditor/Private/MuCOE/GenerateMutableSource/GenerateMutableSourceSurface.cpp
+++ b/Engine/Plugins/Experimental/Mutable/Source/CustomizableObjectEditor/Private/MuCOE/GenerateMutableSource/GenerateMutableSourceSurface.cpp
@@ -431,324 +431,339 @@ mu::NodeSurfacePtr GenerateMutableSourceSurface(const UEdGraphPin * Pin, FMutabl
 
 		const int32 NumImages = TypedNodeMat->GetNumParameters(EMaterialParameterType::Texture);
 		SurfNode->SetImageCount(NumImages);
+
 		for (int32 ImageIndex = 0; ImageIndex < NumImages; ++ImageIndex)
 		{
-			mu::NodeImagePtr GroupProjectionImg;
-			UTexture2D* GroupProjectionReferenceTexture = nullptr;
 			const UEdGraphPin* ImagePin = TypedNodeMat->GetParameterPin(EMaterialParameterType::Texture, ImageIndex);
-			const FString ImageName = TypedNodeMat->GetParameterName(EMaterialParameterType::Texture, ImageIndex).ToString();
-			const FGuid ImageId = TypedNodeMat->GetParameterId(EMaterialParameterType::Texture, ImageIndex);
 
-			FString MaterialImageId = FGroupProjectorImageInfo::GenerateId(TypedNodeMat, ImageIndex);
-			bool bShareProjectionTexturesBetweenLODs = false;
-			FGroupProjectorImageInfo* ProjectorInfo = GenerationContext.GroupProjectorLODCache.Find(MaterialImageId);
+			if (ImagePin && !TypedNodeMaterial->IsImageMutableMode(ImageIndex))
+			{
+				// This is a connected pass-through texture that simply has to be passed to the core
 
-			if (!ProjectorInfo) // No previous LOD of this material generated the image.
+				if (const UEdGraphPin* ConnectedPin = FollowInputPin(*ImagePin))
+				{
+					mu::NodeImagePtr PassThroughImagePtr = GenerateMutableSourceImage(ConnectedPin, GenerationContext, 0);
+					SurfNode->SetImage(ImageIndex, PassThroughImagePtr);
+				}
+			}
+			else
 			{
-				bool bIsGroupProjectorImage = false;
+				mu::NodeImagePtr GroupProjectionImg;
+				UTexture2D* GroupProjectionReferenceTexture = nullptr;
+				const FString ImageName = TypedNodeMat->GetParameterName(EMaterialParameterType::Texture, ImageIndex).ToString();
+				const FGuid ImageId = TypedNodeMat->GetParameterId(EMaterialParameterType::Texture, ImageIndex);
 
-				GroupProjectionImg = GenerateMutableGroupProjection(LOD, ImageIndex, MeshNode, GenerationContext,
-				TypedNodeMat, bShareProjectionTexturesBetweenLODs, bIsGroupProjectorImage,
-				GroupProjectionReferenceTexture, TextureNameToProjectionResFactor, AlternateResStateName,
-				nullptr);
+				FString MaterialImageId = FGroupProjectorImageInfo::GenerateId(TypedNodeMat, ImageIndex);
+				bool bShareProjectionTexturesBetweenLODs = false;
+				FGroupProjectorImageInfo* ProjectorInfo = GenerationContext.GroupProjectorLODCache.Find(MaterialImageId);
 
-				if (GroupProjectionImg.get() || TypedNodeMaterial->IsImageMutableMode(ImageIndex) || (bMaterialPinConnected && NodeTableParametersGenerated.Contains(ImageId)))
+				if (!ProjectorInfo) // No previous LOD of this material generated the image.
 				{
-					// Get the reference texture
-					UTexture2D* ReferenceTexture = nullptr;
+					bool bIsGroupProjectorImage = false;
 
-					GenerationContext.CurrentMaterialTableParameter = ImageName;
-					GenerationContext.CurrentMaterialTableParameterId = ImageId.ToString();
+					GroupProjectionImg = GenerateMutableGroupProjection(LOD, ImageIndex, MeshNode, GenerationContext,
+						TypedNodeMat, bShareProjectionTexturesBetweenLODs, bIsGroupProjectorImage,
+						GroupProjectionReferenceTexture, TextureNameToProjectionResFactor, AlternateResStateName,
+						nullptr);
 
-					if (GroupProjectionImg.get() && GroupProjectionReferenceTexture)
-					{
-						ReferenceTexture = GroupProjectionReferenceTexture;
-					}
-					else
+					if (GroupProjectionImg.get() || TypedNodeMaterial->IsImageMutableMode(ImageIndex) || (bMaterialPinConnected && NodeTableParametersGenerated.Contains(ImageId)))
 					{
-						ReferenceTexture = TypedNodeMat->GetImageReferenceTexture(ImageIndex);
+						// Get the reference texture
+						UTexture2D* ReferenceTexture = nullptr;
 
-						if (!ReferenceTexture && !GroupProjectionImg.get())
+						GenerationContext.CurrentMaterialTableParameter = ImageName;
+						GenerationContext.CurrentMaterialTableParameterId = ImageId.ToString();
+
+						if (GroupProjectionImg.get() && GroupProjectionReferenceTexture)
 						{
-							if (ImagePin)
+							ReferenceTexture = GroupProjectionReferenceTexture;
+						}
+						else
+						{
+							ReferenceTexture = TypedNodeMat->GetImageReferenceTexture(ImageIndex);
+
+							if (!ReferenceTexture && !GroupProjectionImg.get())
 							{
-								if (const UEdGraphPin* ConnectedPin = FollowInputPin(*ImagePin))
+								if (ImagePin)
 								{
-									ReferenceTexture = FindReferenceImage(ConnectedPin, GenerationContext);
+									if (const UEdGraphPin* ConnectedPin = FollowInputPin(*ImagePin))
+									{
+										ReferenceTexture = FindReferenceImage(ConnectedPin, GenerationContext);
+									}
 								}
-							}
-							if (!ReferenceTexture && bMaterialPinConnected && NodeTableParametersGenerated.Contains(ImageId))
-							{
-								if (const UEdGraphPin* ConnectedPin = FollowInputPin(*TypedNodeMat->GetMaterialAssetPin()))
+								if (!ReferenceTexture && bMaterialPinConnected && NodeTableParametersGenerated.Contains(ImageId))
+								{
+									if (const UEdGraphPin* ConnectedPin = FollowInputPin(*TypedNodeMat->GetMaterialAssetPin()))
+									{
+										ReferenceTexture = FindReferenceImage(ConnectedPin, GenerationContext);
+									}
+								}
+								if (!ReferenceTexture)
 								{
-									ReferenceTexture = FindReferenceImage(ConnectedPin, GenerationContext);
+									ReferenceTexture = TypedNodeMat->GetImageValue(ImageIndex);
 								}
 							}
-							if (!ReferenceTexture)
+						}
+
+						FGeneratedImageProperties Props;
+						if (ReferenceTexture)
+						{
+							// Store properties for the generated images
+							Props.TextureParameterName = ImageName;
+							Props.CompressionSettings = ReferenceTexture->CompressionSettings;
+							Props.Filter = ReferenceTexture->Filter;
+							Props.SRGB = ReferenceTexture->SRGB;
+							Props.LODBias = 0;
+							Props.MipGenSettings = ReferenceTexture->MipGenSettings;
+							Props.MaxTextureSize = GetMaxTextureSize(ReferenceTexture, GenerationContext);
+							Props.LODGroup = ReferenceTexture->LODGroup;
+							Props.AddressX = ReferenceTexture->AddressX;
+							Props.AddressY = ReferenceTexture->AddressY;
+							Props.bFlipGreenChannel = ReferenceTexture->bFlipGreenChannel;
+
+							// TODO: MTBL-1081
+							// TextureGroup::TEXTUREGROUP_UI does not support streaming. If we generate a texture that requires streaming and set this group, it will crash when initializing the resource. 
+							// If LODGroup == TEXTUREGROUP_UI, UTexture::IsPossibleToStream() will return false and UE will assume all mips are loaded, when they're not, and crash.
+							if (Props.LODGroup == TEXTUREGROUP_UI)
 							{
-								ReferenceTexture = TypedNodeMat->GetImageValue(ImageIndex);
+								Props.LODGroup = TextureGroup::TEXTUREGROUP_Character;
+
+								FString msg = FString::Printf(TEXT("The Reference texture [%s] is using TEXTUREGROUP_UI which does not support streaming. Please set a different TEXTURE group."),
+									*ReferenceTexture->GetName(), *ImageName);
+								GenerationContext.Compiler->CompilerLog(FText::FromString(msg), Node, EMessageSeverity::Info);
 							}
 						}
-					}
-
-					FGeneratedImageProperties Props;
-					if (ReferenceTexture)
-					{
-						// Store properties for the generated images
-						Props.TextureParameterName = ImageName;
-						Props.CompressionSettings = ReferenceTexture->CompressionSettings;
-						Props.Filter = ReferenceTexture->Filter;
-						Props.SRGB = ReferenceTexture->SRGB;
- 						Props.LODBias = 0;
-						Props.MipGenSettings = ReferenceTexture->MipGenSettings;
-						Props.MaxTextureSize = GetMaxTextureSize(ReferenceTexture, GenerationContext);
-						Props.LODGroup = ReferenceTexture->LODGroup;
-						Props.AddressX = ReferenceTexture->AddressX;
-						Props.AddressY = ReferenceTexture->AddressY;
-						Props.bFlipGreenChannel = ReferenceTexture->bFlipGreenChannel;
-
-						// TODO: MTBL-1081
-						// TextureGroup::TEXTUREGROUP_UI does not support streaming. If we generate a texture that requires streaming and set this group, it will crash when initializing the resource. 
-						// If LODGroup == TEXTUREGROUP_UI, UTexture::IsPossibleToStream() will return false and UE will assume all mips are loaded, when they're not, and crash.
-						if (Props.LODGroup == TEXTUREGROUP_UI)
+						else if (!GroupProjectionImg.get())
 						{
-							Props.LODGroup = TextureGroup::TEXTUREGROUP_Character;
-							
-							FString msg = FString::Printf(TEXT("The Reference texture [%s] is using TEXTUREGROUP_UI which does not support streaming. Please set a different TEXTURE group."),
-								*ReferenceTexture->GetName(), *ImageName);
-							GenerationContext.Compiler->CompilerLog(FText::FromString(msg), Node, EMessageSeverity::Info);
+							// warning!
+							FString msg = FString::Printf(TEXT("The Reference texture for material image [%s] is not set and it couldn't be found automatically."), *ImageName);
+							GenerationContext.Compiler->CompilerLog(FText::FromString(msg), Node);
 						}
-					}
-					else if (!GroupProjectionImg.get())
-					{
-						// warning!
-						FString msg = FString::Printf(TEXT("The Reference texture for material image [%s] is not set and it couldn't be found automatically."), *ImageName);
-						GenerationContext.Compiler->CompilerLog(FText::FromString(msg), Node);
-					}
 
-					GenerationContext.ImageProperties.Add(Props);
-					SurfaceData.ImageProperties = Props;
-					
-					// Calculate the LODBias for this texture
-					int32 LODBias = ComputeLODBias(GenerationContext, ReferenceTexture, Props.MaxTextureSize, TypedNodeMat, ImageIndex);
+						GenerationContext.ImageProperties.Add(Props);
+						SurfaceData.ImageProperties = Props;
 
-					GenerationContext.CurrentTextureLODBias = LODBias;
+						// Calculate the LODBias for this texture
+						int32 LODBias = ComputeLODBias(GenerationContext, ReferenceTexture, Props.MaxTextureSize, TypedNodeMat, ImageIndex);
 
-					// Generate the texture nodes
-					mu::NodeImagePtr ImageNode = [&]()
-					{
-						if (TypedNodeMaterial->IsImageMutableMode(ImageIndex))
+						GenerationContext.CurrentTextureLODBias = LODBias;
+
+						// Generate the texture nodes
+						mu::NodeImagePtr ImageNode = [&]()
 						{
-							if (ImagePin)
+							if (TypedNodeMaterial->IsImageMutableMode(ImageIndex))
 							{
-								if (const UEdGraphPin* ConnectedPin = FollowInputPin(*ImagePin))
+								if (ImagePin)
 								{
-									return GenerateMutableSourceImage(ConnectedPin, GenerationContext, Props.MaxTextureSize);
+									if (const UEdGraphPin* ConnectedPin = FollowInputPin(*ImagePin))
+									{
+										return GenerateMutableSourceImage(ConnectedPin, GenerationContext, Props.MaxTextureSize);
+									}
 								}
-							}
-							
-							if (bMaterialPinConnected && NodeTableParametersGenerated.Contains(ImageId))
-							{
-								if (const UEdGraphPin* ConnectedPin = FollowInputPin(*TypedNodeMat->GetMaterialAssetPin()))
+
+								if (bMaterialPinConnected && NodeTableParametersGenerated.Contains(ImageId))
 								{
-									return GenerateMutableSourceImage(ConnectedPin, GenerationContext, Props.MaxTextureSize);
+									if (const UEdGraphPin* ConnectedPin = FollowInputPin(*TypedNodeMat->GetMaterialAssetPin()))
+									{
+										return GenerateMutableSourceImage(ConnectedPin, GenerationContext, Props.MaxTextureSize);
+									}
 								}
-							}
 
-							// Else
-							{
-								UTexture2D* Texture2D = TypedNodeMat->GetImageValue(ImageIndex);
+								// Else
+								{
+									UTexture2D* Texture2D = TypedNodeMat->GetImageValue(ImageIndex);
 
-								const mu::NodeImageConstantPtr ConstImageNode = new mu::NodeImageConstant();
-								GenerationContext.ArrayTextureUnrealToMutableTask.Add(FTextureUnrealToMutableTask(ConstImageNode, Texture2D, Node));
+									const mu::NodeImageConstantPtr ConstImageNode = new mu::NodeImageConstant();
+									GenerationContext.ArrayTextureUnrealToMutableTask.Add(FTextureUnrealToMutableTask(ConstImageNode, Texture2D, Node));
 
-								return ResizeToMaxTextureSize(Props.MaxTextureSize, Texture2D, ConstImageNode);
+									return ResizeToMaxTextureSize(Props.MaxTextureSize, Texture2D, ConstImageNode);
+								}
 							}
-						}
-						else
+							else
+							{
+								return mu::NodeImagePtr();
+							}
+						}();
+
+						if (GroupProjectionImg.get())
 						{
-							return mu::NodeImagePtr();
+							ImageNode = GroupProjectionImg;
 						}
-					}();
-
-					if (GroupProjectionImg.get())
-					{
-						ImageNode = GroupProjectionImg;
-					}
 
-					if (GenerationContext.Options.TargetPlatform && ReferenceTexture)
-					{
-						int LayerIndex = 0;
+						if (GenerationContext.Options.TargetPlatform && ReferenceTexture)
+						{
+							int LayerIndex = 0;
 
-						TArray< TArray<FName> > PlatformFormats;
-						GenerationContext.Options.TargetPlatform->GetTextureFormats(ReferenceTexture, PlatformFormats);
+							TArray< TArray<FName> > PlatformFormats;
+							GenerationContext.Options.TargetPlatform->GetTextureFormats(ReferenceTexture, PlatformFormats);
 
-						bool bHaveFormatPlatformFormats = PlatformFormats.Num() > 0;
+							bool bHaveFormatPlatformFormats = PlatformFormats.Num() > 0;
 
-						mu::NodeImagePtr LastImage = ImageNode;
+							mu::NodeImagePtr LastImage = ImageNode;
 
-						// To apply LOD bias						
-						if (LODBias > 0)
-						{
-							mu::NodeImageResizePtr ResizeImage = new mu::NodeImageResize();
-							ResizeImage->SetBase(LastImage.get());
-							ResizeImage->SetRelative(true);
-							float factor = FMath::Pow(0.5f, LODBias);
-							ResizeImage->SetSize(factor, factor);
-							ResizeImage->SetMessageContext(Node);
-							LastImage = ResizeImage;
-						}
-						
-						mu::NodeImageMipmapPtr MipmapImage = new mu::NodeImageMipmap();
-						MipmapImage->SetSource(LastImage.get());
-						MipmapImage->SetMipmapGenerationSettings(mu::EMipmapFilterType::MFT_SimpleAverage, mu::EAddressMode::AM_NONE, 1.0f, false);
+							// To apply LOD bias						
+							if (LODBias > 0)
+							{
+								mu::NodeImageResizePtr ResizeImage = new mu::NodeImageResize();
+								ResizeImage->SetBase(LastImage.get());
+								ResizeImage->SetRelative(true);
+								float factor = FMath::Pow(0.5f, LODBias);
+								ResizeImage->SetSize(factor, factor);
+								ResizeImage->SetMessageContext(Node);
+								LastImage = ResizeImage;
+							}
 
-						MipmapImage->SetMessageContext(Node);
-						LastImage = MipmapImage;
+							mu::NodeImageMipmapPtr MipmapImage = new mu::NodeImageMipmap();
+							MipmapImage->SetSource(LastImage.get());
+							MipmapImage->SetMipmapGenerationSettings(mu::EMipmapFilterType::MFT_SimpleAverage, mu::EAddressMode::AM_NONE, 1.0f, false);
 
-						// Apply composite image. This needs to be computed after mipmaps generation. 	
-						if (ReferenceTexture && ReferenceTexture->CompositeTexture && ReferenceTexture->CompositeTextureMode != CTM_Disabled)
-						{
-							mu::NodeImageNormalCompositePtr CompositedImage = new mu::NodeImageNormalComposite();
-							CompositedImage->SetBase( LastImage.get() );
-							CompositedImage->SetPower( ReferenceTexture->CompositePower );
+							MipmapImage->SetMessageContext(Node);
+							LastImage = MipmapImage;
 
-							mu::ECompositeImageMode CompositeImageMode = [CompositeTextureMode = ReferenceTexture->CompositeTextureMode]() -> mu::ECompositeImageMode
+							// Apply composite image. This needs to be computed after mipmaps generation. 	
+							if (ReferenceTexture && ReferenceTexture->CompositeTexture && ReferenceTexture->CompositeTextureMode != CTM_Disabled)
 							{
-								switch(CompositeTextureMode)
+								mu::NodeImageNormalCompositePtr CompositedImage = new mu::NodeImageNormalComposite();
+								CompositedImage->SetBase(LastImage.get());
+								CompositedImage->SetPower(ReferenceTexture->CompositePower);
+
+								mu::ECompositeImageMode CompositeImageMode = [CompositeTextureMode = ReferenceTexture->CompositeTextureMode]() -> mu::ECompositeImageMode
 								{
-									case CTM_NormalRoughnessToRed  : return mu::ECompositeImageMode::CIM_NormalRoughnessToRed;
+									switch (CompositeTextureMode)
+									{
+									case CTM_NormalRoughnessToRed: return mu::ECompositeImageMode::CIM_NormalRoughnessToRed;
 									case CTM_NormalRoughnessToGreen: return mu::ECompositeImageMode::CIM_NormalRoughnessToGreen;
-									case CTM_NormalRoughnessToBlue : return mu::ECompositeImageMode::CIM_NormalRoughnessToBlue;
+									case CTM_NormalRoughnessToBlue: return mu::ECompositeImageMode::CIM_NormalRoughnessToBlue;
 									case CTM_NormalRoughnessToAlpha: return mu::ECompositeImageMode::CIM_NormalRoughnessToAlpha;
-									
+
 									default: return mu::ECompositeImageMode::CIM_Disabled;
-								}
-							}();
-				
-							CompositedImage->SetMode( CompositeImageMode );							
+									}
+								}();
 
-							mu::NodeImageConstantPtr CompositeNormalImage = new mu::NodeImageConstant();
+								CompositedImage->SetMode(CompositeImageMode);
 
-							UTexture2D* ReferenceCompositeNormalTexture = Cast<UTexture2D>( ReferenceTexture->CompositeTexture );
-							if (ReferenceCompositeNormalTexture)
-							{
-								GenerationContext.ArrayTextureUnrealToMutableTask.Add(FTextureUnrealToMutableTask(CompositeNormalImage, ReferenceCompositeNormalTexture, Node, true));
+								mu::NodeImageConstantPtr CompositeNormalImage = new mu::NodeImageConstant();
 
-								mu::NodeImageMipmapPtr NormalCompositeMipmapImage = new mu::NodeImageMipmap();
-								NormalCompositeMipmapImage->SetSource( CompositeNormalImage );
-								NormalCompositeMipmapImage->SetMipmapGenerationSettings(mu::EMipmapFilterType::MFT_SimpleAverage, mu::EAddressMode::AM_NONE, 1.0f, true);
-								
-								CompositedImage->SetNormal( NormalCompositeMipmapImage );
-							}
+								UTexture2D* ReferenceCompositeNormalTexture = Cast<UTexture2D>(ReferenceTexture->CompositeTexture);
+								if (ReferenceCompositeNormalTexture)
+								{
+									GenerationContext.ArrayTextureUnrealToMutableTask.Add(FTextureUnrealToMutableTask(CompositeNormalImage, ReferenceCompositeNormalTexture, Node, true));
 
-							LastImage = CompositedImage;
-						}
+									mu::NodeImageMipmapPtr NormalCompositeMipmapImage = new mu::NodeImageMipmap();
+									NormalCompositeMipmapImage->SetSource(CompositeNormalImage);
+									NormalCompositeMipmapImage->SetMipmapGenerationSettings(mu::EMipmapFilterType::MFT_SimpleAverage, mu::EAddressMode::AM_NONE, 1.0f, true);
 
-						if (bHaveFormatPlatformFormats)
-						{
-							mu::NodeImageFormatPtr FormatImage = new mu::NodeImageFormat();
-							FormatImage->SetSource(LastImage.get());
-							FormatImage->SetFormat(mu::EImageFormat::IF_RGBA_UBYTE);
-							FormatImage->SetMessageContext(Node);
-							LastImage = FormatImage;
+									CompositedImage->SetNormal(NormalCompositeMipmapImage);
+								}
+
+								LastImage = CompositedImage;
+							}
 
-							if (GenerationContext.Options.bTextureCompression)
+							if (bHaveFormatPlatformFormats)
 							{
-								check(PlatformFormats[0].Num() > LayerIndex);
+								mu::NodeImageFormatPtr FormatImage = new mu::NodeImageFormat();
+								FormatImage->SetSource(LastImage.get());
+								FormatImage->SetFormat(mu::EImageFormat::IF_RGBA_UBYTE);
+								FormatImage->SetMessageContext(Node);
+								LastImage = FormatImage;
 
-								const FString PlatformFormat = PlatformFormats[0][LayerIndex].ToString();
+								if (GenerationContext.Options.bTextureCompression)
+								{
+									check(PlatformFormats[0].Num() > LayerIndex);
 
-								// Remove platform prefix
-								FString FormatWithoutPrefix = PlatformFormat;
-								PlatformFormat.Split(TEXT("_"), nullptr, &FormatWithoutPrefix, ESearchCase::IgnoreCase, ESearchDir::FromEnd);
+									const FString PlatformFormat = PlatformFormats[0][LayerIndex].ToString();
 
-								mu::EImageFormat mutableFormat = mu::EImageFormat::IF_NONE;
-								mu::EImageFormat mutableFormatIfAlpha = mu::EImageFormat::IF_NONE;
+									// Remove platform prefix
+									FString FormatWithoutPrefix = PlatformFormat;
+									PlatformFormat.Split(TEXT("_"), nullptr, &FormatWithoutPrefix, ESearchCase::IgnoreCase, ESearchDir::FromEnd);
 
-								if (FormatWithoutPrefix == TEXT("AutoDXT"))
-								{
-									mutableFormat = mu::EImageFormat::IF_BC1;
-									mutableFormatIfAlpha = mu::EImageFormat::IF_BC3;
-								}
-								else if (FormatWithoutPrefix == TEXT("DXT1")) mutableFormat = mu::EImageFormat::IF_BC1;
-								else if (FormatWithoutPrefix == TEXT("DXT3")) mutableFormat = mu::EImageFormat::IF_BC2;
-								else if (FormatWithoutPrefix == TEXT("DXT5")) mutableFormat = mu::EImageFormat::IF_BC3;
-								else if (FormatWithoutPrefix == TEXT("BC1")) mutableFormat = mu::EImageFormat::IF_BC1;
-								else if (FormatWithoutPrefix == TEXT("BC2")) mutableFormat = mu::EImageFormat::IF_BC2;
-								else if (FormatWithoutPrefix == TEXT("BC3")) mutableFormat = mu::EImageFormat::IF_BC3;
-								else if (FormatWithoutPrefix == TEXT("BC4")) mutableFormat = mu::EImageFormat::IF_BC4;
-								else if (FormatWithoutPrefix == TEXT("BC5")) mutableFormat = mu::EImageFormat::IF_BC5;
-								else if (FormatWithoutPrefix == TEXT("G8")) mutableFormat = mu::EImageFormat::IF_L_UBYTE;
-								else if (FormatWithoutPrefix == TEXT("BGRA8")) mutableFormat = mu::EImageFormat::IF_RGBA_UBYTE;
-								else if (PlatformFormat.Contains(TEXT("ASTC")))
-								{
-									if ((FormatWithoutPrefix == TEXT("AutoASTC")) || (FormatWithoutPrefix == TEXT("RGBAuto")))
+									mu::EImageFormat mutableFormat = mu::EImageFormat::IF_NONE;
+									mu::EImageFormat mutableFormatIfAlpha = mu::EImageFormat::IF_NONE;
+
+									if (FormatWithoutPrefix == TEXT("AutoDXT"))
 									{
-										mutableFormat = mu::EImageFormat::IF_ASTC_4x4_RGB_LDR;
-										mutableFormatIfAlpha = mu::EImageFormat::IF_ASTC_4x4_RGBA_LDR;
+										mutableFormat = mu::EImageFormat::IF_BC1;
+										mutableFormatIfAlpha = mu::EImageFormat::IF_BC3;
+									}
+									else if (FormatWithoutPrefix == TEXT("DXT1")) mutableFormat = mu::EImageFormat::IF_BC1;
+									else if (FormatWithoutPrefix == TEXT("DXT3")) mutableFormat = mu::EImageFormat::IF_BC2;
+									else if (FormatWithoutPrefix == TEXT("DXT5")) mutableFormat = mu::EImageFormat::IF_BC3;
+									else if (FormatWithoutPrefix == TEXT("BC1")) mutableFormat = mu::EImageFormat::IF_BC1;
+									else if (FormatWithoutPrefix == TEXT("BC2")) mutableFormat = mu::EImageFormat::IF_BC2;
+									else if (FormatWithoutPrefix == TEXT("BC3")) mutableFormat = mu::EImageFormat::IF_BC3;
+									else if (FormatWithoutPrefix == TEXT("BC4")) mutableFormat = mu::EImageFormat::IF_BC4;
+									else if (FormatWithoutPrefix == TEXT("BC5")) mutableFormat = mu::EImageFormat::IF_BC5;
+									else if (FormatWithoutPrefix == TEXT("G8")) mutableFormat = mu::EImageFormat::IF_L_UBYTE;
+									else if (FormatWithoutPrefix == TEXT("BGRA8")) mutableFormat = mu::EImageFormat::IF_RGBA_UBYTE;
+									else if (PlatformFormat.Contains(TEXT("ASTC")))
+									{
+										if ((FormatWithoutPrefix == TEXT("AutoASTC")) || (FormatWithoutPrefix == TEXT("RGBAuto")))
+										{
+											mutableFormat = mu::EImageFormat::IF_ASTC_4x4_RGB_LDR;
+											mutableFormatIfAlpha = mu::EImageFormat::IF_ASTC_4x4_RGBA_LDR;
+										}
+										else if (FormatWithoutPrefix == TEXT("RGB")) mutableFormat = mu::EImageFormat::IF_ASTC_4x4_RGB_LDR;
+										else if (FormatWithoutPrefix == TEXT("RGBA")) mutableFormat = mu::EImageFormat::IF_ASTC_4x4_RGBA_LDR;
+										else if (FormatWithoutPrefix == TEXT("NormalRG")) mutableFormat = mu::EImageFormat::IF_ASTC_4x4_RG_LDR;
 									}
-									else if (FormatWithoutPrefix == TEXT("RGB")) mutableFormat = mu::EImageFormat::IF_ASTC_4x4_RGB_LDR;
-									else if (FormatWithoutPrefix == TEXT("RGBA")) mutableFormat = mu::EImageFormat::IF_ASTC_4x4_RGBA_LDR;
-									else if (FormatWithoutPrefix == TEXT("NormalRG")) mutableFormat = mu::EImageFormat::IF_ASTC_4x4_RG_LDR;
-								}
 
-								if (mutableFormat == mu::EImageFormat::IF_NONE)
-								{
-									// Format not supported by Mutable, use RBGA_UBYTE as default.
-									mutableFormat = mu::EImageFormat::IF_RGBA_UBYTE;
+									if (mutableFormat == mu::EImageFormat::IF_NONE)
+									{
+										// Format not supported by Mutable, use RBGA_UBYTE as default.
+										mutableFormat = mu::EImageFormat::IF_RGBA_UBYTE;
 
-									const FString UnexpectedImageFormatMsg = FString::Printf(TEXT("In object [%s] Unexpected image format [%s], RGBA_UBYTE will be used instead."), *GenerationContext.Object->GetName(), *FormatWithoutPrefix);
-									const FText UnexpectedImageFormatText = FText::FromString(UnexpectedImageFormatMsg);
-									GenerationContext.Compiler->CompilerLog(UnexpectedImageFormatText, Node);
-									UE_LOG(LogMutable, Warning, TEXT("%s"), *UnexpectedImageFormatMsg);
-								}
+										const FString UnexpectedImageFormatMsg = FString::Printf(TEXT("In object [%s] Unexpected image format [%s], RGBA_UBYTE will be used instead."), *GenerationContext.Object->GetName(), *FormatWithoutPrefix);
+										const FText UnexpectedImageFormatText = FText::FromString(UnexpectedImageFormatMsg);
+										GenerationContext.Compiler->CompilerLog(UnexpectedImageFormatText, Node);
+										UE_LOG(LogMutable, Warning, TEXT("%s"), *UnexpectedImageFormatMsg);
+									}
 
-								FormatImage->SetFormat(mutableFormat, mutableFormatIfAlpha);
+									FormatImage->SetFormat(mutableFormat, mutableFormatIfAlpha);
+								}
 							}
-						}
 
-						ImageNode = LastImage;
-					}
+							ImageNode = LastImage;
+						}
 
-					mu::NodeImagePtr ImageNodePtr = ImageNode;
-					SurfNode->SetImage(ImageIndex, ImageNodePtr);
+						mu::NodeImagePtr ImageNodePtr = ImageNode;
+						SurfNode->SetImage(ImageIndex, ImageNodePtr);
 
-					FString SurfNodeImageName = FString::Printf(TEXT("%d"), GenerationContext.ImageProperties.Num() - 1);
+						FString SurfNodeImageName = FString::Printf(TEXT("%d"), GenerationContext.ImageProperties.Num() - 1);
 
-					// Encoding material layer in mutable name
-					const int32 LayerIndex = TypedNodeMat->GetParameterLayerIndex(EMaterialParameterType::Texture, ImageIndex);
-					if (LayerIndex != -1)
-					{
-						SurfNodeImageName += "-MutableLayerParam:" + FString::FromInt(LayerIndex);
-					}
+						// Encoding material layer in mutable name
+						const int32 LayerIndex = TypedNodeMat->GetParameterLayerIndex(EMaterialParameterType::Texture, ImageIndex);
+						if (LayerIndex != -1)
+						{
+							SurfNodeImageName += "-MutableLayerParam:" + FString::FromInt(LayerIndex);
+						}
 
-					SurfNode->SetImageName(ImageIndex, StringCast<ANSICHAR>(*SurfNodeImageName).Get());
-					int32 UVLayout = TypedNodeMat->GetImageUVLayout(ImageIndex);
-					SurfNode->SetImageLayoutIndex(ImageIndex, UVLayout);
-					SurfNode->SetImageAdditionalNames(ImageIndex, StringCast<ANSICHAR>(*TypedNodeMat->Material->GetName()).Get(), StringCast<ANSICHAR>(*ImageName).Get());
+						SurfNode->SetImageName(ImageIndex, StringCast<ANSICHAR>(*SurfNodeImageName).Get());
+						int32 UVLayout = TypedNodeMat->GetImageUVLayout(ImageIndex);
+						SurfNode->SetImageLayoutIndex(ImageIndex, UVLayout);
+						SurfNode->SetImageAdditionalNames(ImageIndex, StringCast<ANSICHAR>(*TypedNodeMat->Material->GetName()).Get(), StringCast<ANSICHAR>(*ImageName).Get());
 
-					if (bShareProjectionTexturesBetweenLODs && bIsGroupProjectorImage)
-					{
-						// Add to the GroupProjectorLODCache to potentially reuse this projection texture in higher LODs
-						ensure(LOD == GenerationContext.FirstLODAvailable);
-						float* AlternateProjectionResFactor = TextureNameToProjectionResFactor.Find(ImageName);
-						GenerationContext.GroupProjectorLODCache.Add(MaterialImageId,
-							FGroupProjectorImageInfo(ImageNodePtr, ImageName, ImageName, TypedNodeMat,
-							AlternateProjectionResFactor ? *AlternateProjectionResFactor : 0.f, AlternateResStateName, SurfNode, UVLayout));
+						if (bShareProjectionTexturesBetweenLODs && bIsGroupProjectorImage)
+						{
+							// Add to the GroupProjectorLODCache to potentially reuse this projection texture in higher LODs
+							ensure(LOD == GenerationContext.FirstLODAvailable);
+							float* AlternateProjectionResFactor = TextureNameToProjectionResFactor.Find(ImageName);
+							GenerationContext.GroupProjectorLODCache.Add(MaterialImageId,
+								FGroupProjectorImageInfo(ImageNodePtr, ImageName, ImageName, TypedNodeMat,
+									AlternateProjectionResFactor ? *AlternateProjectionResFactor : 0.f, AlternateResStateName, SurfNode, UVLayout));
+						}
 					}
 				}
-			}
-			else
-			{
-				ensure(LOD > GenerationContext.FirstLODAvailable);
-				check(ProjectorInfo->SurfNode->GetImage(ImageIndex) == ProjectorInfo->ImageNode);
-				SurfNode->SetImage(ImageIndex, ProjectorInfo->ImageNode);
-				SurfNode->SetImageName(ImageIndex, StringCast<ANSICHAR>(*ProjectorInfo->TextureName).Get());
-				SurfNode->SetImageLayoutIndex(ImageIndex, ProjectorInfo->UVLayout);
-
-				TextureNameToProjectionResFactor.Add(ProjectorInfo->RealTextureName, ProjectorInfo->AlternateProjectionResolutionFactor);
-				AlternateResStateName = ProjectorInfo->AlternateResStateName;
+				else
+				{
+					ensure(LOD > GenerationContext.FirstLODAvailable);
+					check(ProjectorInfo->SurfNode->GetImage(ImageIndex) == ProjectorInfo->ImageNode);
+					SurfNode->SetImage(ImageIndex, ProjectorInfo->ImageNode);
+					SurfNode->SetImageName(ImageIndex, StringCast<ANSICHAR>(*ProjectorInfo->TextureName).Get());
+					SurfNode->SetImageLayoutIndex(ImageIndex, ProjectorInfo->UVLayout);
+
+					TextureNameToProjectionResFactor.Add(ProjectorInfo->RealTextureName, ProjectorInfo->AlternateProjectionResolutionFactor);
+					AlternateResStateName = ProjectorInfo->AlternateResStateName;
+				}
 			}
 		}
 
diff --git a/Engine/Plugins/Experimental/Mutable/Source/CustomizableObjectEditor/Private/MuCOE/Nodes/CustomizableObjectNodeMaterial.cpp b/Engine/Plugins/Experimental/Mutable/Source/CustomizableObjectEditor/Private/MuCOE/Nodes/CustomizableObjectNodeMaterial.cpp
--- a/Engine/Plugins/Experimental/Mutable/Source/CustomizableObjectEditor/Private/MuCOE/Nodes/CustomizableObjectNodeMaterial.cpp
+++ b/Engine/Plugins/Experimental/Mutable/Source/CustomizableObjectEditor/Private/MuCOE/Nodes/CustomizableObjectNodeMaterial.cpp
@@ -155,8 +155,15 @@ EPinMode UCustomizableObjectNodeMaterial::NodePinModeToImagePinMode(const ENodeP
 
 EPinMode UCustomizableObjectNodeMaterial::GetImagePinMode(const UEdGraphPin& Pin) const
 {
-	if (FollowInputPin(Pin))
+	const UEdGraphSchema_CustomizableObject* Schema = GetDefault<UEdGraphSchema_CustomizableObject>();
+
+	if (UEdGraphPin *FollowedPin = FollowInputPin(Pin))
 	{
+		if (FollowedPin->PinType.PinCategory == Schema->PC_PassThroughImage)
+		{
+			return EPinMode::Passthrough;
+		}
+
 		return EPinMode::Mutable;
 	}
 	else if (const UEdGraphPin* MaterialAssetPin = GetMaterialAssetPin())
@@ -1204,6 +1211,20 @@ void UCustomizableObjectNodeMaterial::PostPasteNode()
 }
 
 
+bool UCustomizableObjectNodeMaterial::CanConnect(const UEdGraphPin* InOwnedInputPin, const UEdGraphPin* InOutputPin, bool& bOutIsOtherNodeBlocklisted, bool& bOutArePinsCompatible) const
+{
+	const UEdGraphSchema_CustomizableObject* Schema = GetDefault<UEdGraphSchema_CustomizableObject>();
+
+	if (InOwnedInputPin && InOwnedInputPin->PinType.PinCategory == Schema->PC_Image &&
+		InOutputPin && InOutputPin->PinType.PinCategory == Schema->PC_PassThroughImage)
+	{
+		return true;
+	}
+	
+	return Super::CanConnect(InOwnedInputPin, InOutputPin, bOutIsOtherNodeBlocklisted, bOutArePinsCompatible);
+}
+
+
 void UCustomizableObjectNodeMaterial::MeshPostEditChangeProperty(FPostEditChangePropertyDelegateParameters& Parameters)
 {
 	if (const UEdGraphPin* MeshPin = FindPin(TEXT("Mesh")))
diff --git a/Engine/Plugins/Experimental/Mutable/Source/CustomizableObjectEditor/Private/MuCOE/Nodes/CustomizableObjectNodePassThroughTexture.cpp b/Engine/Plugins/Experimental/Mutable/Source/CustomizableObjectEditor/Private/MuCOE/Nodes/CustomizableObjectNodePassThroughTexture.cpp
new file mode 100644
--- /dev/null
+++ b/Engine/Plugins/Experimental/Mutable/Source/CustomizableObjectEditor/Private/MuCOE/Nodes/CustomizableObjectNodePassThroughTexture.cpp
@@ -0,0 +1,48 @@
+// Copyright Epic Games, Inc. All Rights Reserved.
+
+#include "MuCOE/Nodes/CustomizableObjectNodePassThroughTexture.h"
+
+#include "MuCOE/EdGraphSchema_CustomizableObject.h"
+
+#define LOCTEXT_NAMESPACE "CustomizableObjectEditor"
+
+
+void UCustomizableObjectNodePassThroughTexture::AllocateDefaultPins(UCustomizableObjectNodeRemapPins* RemapPins)
+{
+	const UEdGraphSchema_CustomizableObject* Schema = GetDefault<UEdGraphSchema_CustomizableObject>();
+
+	FString PinName = TEXT("Texture");
+	UEdGraphPin* PinImagePin = CustomCreatePin(EGPD_Output, Schema->PC_PassThroughImage, FName(*PinName));
+	PinImagePin->bDefaultValueIsIgnored = true;
+}
+
+
+FText UCustomizableObjectNodePassThroughTexture::GetNodeTitle(ENodeTitleType::Type TitleType) const
+{
+	if (Texture)
+	{
+		FFormatNamedArguments Args;
+		Args.Add(TEXT("TextureName"), FText::FromString(Texture->GetName()));
+
+		return FText::Format(LOCTEXT("PassThrough Texture_Title", "{TextureName}\nPassThrough Texture"), Args);
+	}
+	else
+	{
+		return LOCTEXT("PassThrough Texture", "PassThrough Texture");
+	}
+}
+
+
+FLinearColor UCustomizableObjectNodePassThroughTexture::GetNodeTitleColor() const
+{
+	const UEdGraphSchema_CustomizableObject* Schema = GetDefault<UEdGraphSchema_CustomizableObject>();
+	return Schema->GetPinTypeColor(Schema->PC_PassThroughImage);
+}
+
+
+FText UCustomizableObjectNodePassThroughTexture::GetTooltipText() const
+{
+	return LOCTEXT("PassThrough_Texture_Tooltip", "Defines a pass-through texture. It will not be modified by Mutable in any way, just referenced as a UE asset. It's much cheaper than a Mutable texture, but you cannot make any operations on it, just switch it.");
+}
+
+#undef LOCTEXT_NAMESPACE
diff --git a/Engine/Plugins/Experimental/Mutable/Source/CustomizableObjectEditor/Private/MuCOE/Nodes/CustomizableObjectNodePassThroughTextureSwitch.cpp b/Engine/Plugins/Experimental/Mutable/Source/CustomizableObjectEditor/Private/MuCOE/Nodes/CustomizableObjectNodePassThroughTextureSwitch.cpp
new file mode 100644
--- /dev/null
+++ b/Engine/Plugins/Experimental/Mutable/Source/CustomizableObjectEditor/Private/MuCOE/Nodes/CustomizableObjectNodePassThroughTextureSwitch.cpp
@@ -0,0 +1,29 @@
+// Copyright Epic Games, Inc. All Rights Reserved.
+
+#include "MuCOE/Nodes/CustomizableObjectNodePassThroughTextureSwitch.h"
+
+#include "MuCOE/EdGraphSchema_CustomizableObject.h"
+
+#define LOCTEXT_NAMESPACE "CustomizableObjectEditor"
+
+
+FString UCustomizableObjectNodePassThroughTextureSwitch::GetOutputPinName() const
+{
+	return TEXT("PassThrough Texture");
+}
+
+
+FName UCustomizableObjectNodePassThroughTextureSwitch::GetCategory() const
+{
+	return UEdGraphSchema_CustomizableObject::PC_PassThroughImage;
+}
+
+
+FString UCustomizableObjectNodePassThroughTextureSwitch::GetPinPrefix() const
+{
+	return TEXT("PassThrough Texture ");
+}
+
+
+#undef LOCTEXT_NAMESPACE
+
diff --git a/Engine/Plugins/Experimental/Mutable/Source/CustomizableObjectEditor/Private/MuCOE/Nodes/CustomizableObjectNodeTexture.cpp b/Engine/Plugins/Experimental/Mutable/Source/CustomizableObjectEditor/Private/MuCOE/Nodes/CustomizableObjectNodeTexture.cpp
--- a/Engine/Plugins/Experimental/Mutable/Source/CustomizableObjectEditor/Private/MuCOE/Nodes/CustomizableObjectNodeTexture.cpp
+++ b/Engine/Plugins/Experimental/Mutable/Source/CustomizableObjectEditor/Private/MuCOE/Nodes/CustomizableObjectNodeTexture.cpp
@@ -17,6 +17,12 @@ struct FGeometry;
 #define LOCTEXT_NAMESPACE "CustomizableObjectEditor"
 
 
+TSharedPtr<SGraphNode> UCustomizableObjectNodeTextureBase::CreateVisualWidget()
+{
+	return SNew(SGraphNodeTexture, this);
+}
+
+
 void UCustomizableObjectNodeTexture::AllocateDefaultPins(UCustomizableObjectNodeRemapPins* RemapPins)
 {
 	const UEdGraphSchema_CustomizableObject* Schema = GetDefault<UEdGraphSchema_CustomizableObject>();
@@ -56,16 +62,10 @@ FText UCustomizableObjectNodeTexture::GetTooltipText() const
 }
 
 
-TSharedPtr<SGraphNode> UCustomizableObjectNodeTexture::CreateVisualWidget()
-{
-	return SNew(SGraphNodeTexture, this);
-}
-
-
 void SGraphNodeTexture::Construct(const FArguments& InArgs, UEdGraphNode* InGraphNode)
 {
 	GraphNode = InGraphNode;
-	NodeTexture = Cast< UCustomizableObjectNodeTexture >(GraphNode);
+	NodeTexture = Cast<UCustomizableObjectNodeTextureBase>(GraphNode);
 
 	FPropertyEditorModule& PropPlugin = FModuleManager::LoadModuleChecked<FPropertyEditorModule>("PropertyEditor");
 
