Fixed unnecessary crash in Mutable in SetIntParameterSelectedOption.  Note that none of the other types of Parameters (i.e. of the form Set...ParameterSelectedOption) had a check on the index being valid, and they all have conditional guards to work around it safely.

Added Error logs to both the name (FString) and index versions of SetIntParameterSelectionOption if the named parameter isn't found or the index isn't valid.

Also, changed two comments very slightly to make them clearer (and arguably correct), because as written they read as stating the opposite of the check() they were associated with.

#UE5

diff --git a/Engine/Plugins/Experimental/Mutable/Source/CustomizableObject/Private/MuCO/CustomizableObjectInstanceDescriptor.cpp b/Engine/Plugins/Experimental/Mutable/Source/CustomizableObject/Private/MuCO/CustomizableObjectInstanceDescriptor.cpp
--- a/Engine/Plugins/Experimental/Mutable/Source/CustomizableObject/Private/MuCO/CustomizableObjectInstanceDescriptor.cpp
+++ b/Engine/Plugins/Experimental/Mutable/Source/CustomizableObject/Private/MuCO/CustomizableObjectInstanceDescriptor.cpp
@@ -1163,8 +1163,6 @@ void FCustomizableObjectInstanceDescriptor::SetIntParameterSelectedOption(const
 {
 	check(CustomizableObject);
 
-	check(IntParameters.IsValidIndex(IntParamIndex));
-
 	if (IntParameters.IsValidIndex(IntParamIndex))
 	{
 		const int32 ParameterIndexInObject = CustomizableObject->FindParameter(IntParameters[IntParamIndex].ParameterName);
@@ -1185,12 +1183,14 @@ void FCustomizableObjectInstanceDescriptor::SetIntParameterSelectedOption(const
 			
 			if (RangeIndex == -1)
 			{
-				check(!CustomizableObject->IsParameterMultidimensional(ParameterIndexInObject)); // This param is multidimensional, it must have a RangeIndex of 0 or more
+				// If this param were multidimensional, it must have a RangeIndex of 0 or more
+				check(!CustomizableObject->IsParameterMultidimensional(ParameterIndexInObject)); 
 				IntParameters[IntParamIndex].ParameterValueName = SelectedOption;
 			}
 			else
 			{
-				check(CustomizableObject->IsParameterMultidimensional(ParameterIndexInObject)); // This param is not multidimensional, it must have a RangeIndex of -1
+				// If this param were not multidimensional, it must have a RangeIndex of -1
+				check(CustomizableObject->IsParameterMultidimensional(ParameterIndexInObject));
 
 				if (!IntParameters[IntParamIndex].ParameterRangeValueNames.IsValidIndex(RangeIndex))
 				{
@@ -1204,13 +1204,30 @@ void FCustomizableObjectInstanceDescriptor::SetIntParameterSelectedOption(const
 			}
 		}
 	}
+	else
+	{
+		UE_LOG(LogMutable, Error,
+				TEXT("%hs: IntParamIndex (%d) is not valid on CO (%s).  (SelectedOptionName: %s, RangeIndex: %d).  Cannot set parameter."),
+				IntParamIndex, *GetNameSafe(CustomizableObject), *SelectedOption, RangeIndex
+			  );
+	}
 }
 
 
 void FCustomizableObjectInstanceDescriptor::SetIntParameterSelectedOption(const FString& ParamName, const FString& SelectedOptionName, const int32 RangeIndex)
 {
 	const int32 IntParamIndex = FindIntParameterNameIndex(ParamName);
-	SetIntParameterSelectedOption(IntParamIndex, SelectedOptionName, RangeIndex);
+	if (IntParamIndex != INDEX_NONE)
+	{
+		SetIntParameterSelectedOption(IntParamIndex, SelectedOptionName, RangeIndex);
+	}
+	else
+	{
+		UE_LOG(LogMutable, Error,
+				TEXT("%hs: Failed to find valid parameter index for ParamName (%s) on CO (%s).  (SelectedOptionName: %s, RangeIndex: %d).  Cannot set parameter."),
+				*ParamName, *GetNameSafe(CustomizableObject), *SelectedOptionName, RangeIndex
+			  );
+	}
 }
 
 
