Fixed bug in recent addition of two UE_LOGs: Adding missing __FUNCTION__ argument to UE_LOG calls.

#UE5

diff --git a/Engine/Plugins/Experimental/Mutable/Source/CustomizableObject/Private/MuCO/CustomizableObjectInstanceDescriptor.cpp b/Engine/Plugins/Experimental/Mutable/Source/CustomizableObject/Private/MuCO/CustomizableObjectInstanceDescriptor.cpp
--- a/Engine/Plugins/Experimental/Mutable/Source/CustomizableObject/Private/MuCO/CustomizableObjectInstanceDescriptor.cpp
+++ b/Engine/Plugins/Experimental/Mutable/Source/CustomizableObject/Private/MuCO/CustomizableObjectInstanceDescriptor.cpp
@@ -1208,7 +1208,7 @@ void FCustomizableObjectInstanceDescriptor::SetIntParameterSelectedOption(const
 	{
 		UE_LOG(LogMutable, Error,
 				TEXT("%hs: IntParamIndex (%d) is not valid on CO (%s).  (SelectedOptionName: %s, RangeIndex: %d).  Cannot set parameter."),
-				IntParamIndex, *GetNameSafe(CustomizableObject), *SelectedOption, RangeIndex
+				__FUNCTION__, IntParamIndex, *GetNameSafe(CustomizableObject), *SelectedOption, RangeIndex
 			  );
 	}
 }
@@ -1225,7 +1225,7 @@ void FCustomizableObjectInstanceDescriptor::SetIntParameterSelectedOption(const
 	{
 		UE_LOG(LogMutable, Error,
 				TEXT("%hs: Failed to find valid parameter index for ParamName (%s) on CO (%s).  (SelectedOptionName: %s, RangeIndex: %d).  Cannot set parameter."),
-				*ParamName, *GetNameSafe(CustomizableObject), *SelectedOptionName, RangeIndex
+				__FUNCTION__, *ParamName, *GetNameSafe(CustomizableObject), *SelectedOptionName, RangeIndex
 			  );
 	}
 }
