[mutable] Enable open DGGUI command to work even in packaged builds, as long as they are not Shipping config

#rnx

diff --git a/Engine/Plugins/Experimental/Mutable/Source/CustomizableObject/Private/MuCO/CustomizableObjectDGGUI.cpp b/Engine/Plugins/Experimental/Mutable/Source/CustomizableObject/Private/MuCO/CustomizableObjectDGGUI.cpp
--- a/Engine/Plugins/Experimental/Mutable/Source/CustomizableObject/Private/MuCO/CustomizableObjectDGGUI.cpp
+++ b/Engine/Plugins/Experimental/Mutable/Source/CustomizableObject/Private/MuCO/CustomizableObjectDGGUI.cpp
@@ -9,7 +9,7 @@
 
 void UDGGUI::OpenDGGUI(const int32 SlotID, UCustomizableSkeletalComponent* SelectedCustomizableSkeletalComponent, const UWorld* CurrentWorld, const int32 PlayerIndex)
 {
-#if WITH_EDITOR
+#if !UE_BUILD_SHIPPING
 	if (APlayerController* Player = UGameplayStatics::GetPlayerController(CurrentWorld, PlayerIndex))
 	{
 		FSoftClassPath DGUIPath(TEXT("/Mutable/UI/DynamicallyGeneratedGUI_DGGUI/DynamicallyGeneratedGUI_DGGUI.DynamicallyGeneratedGUI_DGGUI_C"));
@@ -24,12 +24,12 @@ void UDGGUI::OpenDGGUI(const int32 SlotID, UCustomizableSkeletalComponent* Selec
 			}
 		}
 	}
-#endif // WITH_EDITOR
+#endif // !UE_BUILD_SHIPPING
 }
 
 bool UDGGUI::CloseExistingDGGUI(const UWorld* CurrentWorld)
 {
-#if WITH_EDITOR
+#if !UE_BUILD_SHIPPING
 	bool bClosing = false;
 	for (TObjectIterator<UDGGUI> PreviousGUI; PreviousGUI; ++PreviousGUI)
 	{
@@ -51,6 +51,6 @@ bool UDGGUI::CloseExistingDGGUI(const UWorld* CurrentWorld)
 		}
 		return true;
 	}
-#endif // WITH_EDITOR
+#endif // !UE_BUILD_SHIPPING
 	return false;
 }
diff --git a/Engine/Plugins/Experimental/Mutable/Source/CustomizableObject/Private/MuCO/CustomizableObjectModule.cpp b/Engine/Plugins/Experimental/Mutable/Source/CustomizableObject/Private/MuCO/CustomizableObjectModule.cpp
--- a/Engine/Plugins/Experimental/Mutable/Source/CustomizableObject/Private/MuCO/CustomizableObjectModule.cpp
+++ b/Engine/Plugins/Experimental/Mutable/Source/CustomizableObject/Private/MuCO/CustomizableObjectModule.cpp
@@ -2,9 +2,12 @@
 
 #include "Algo/Find.h"
 #include "Interfaces/IPluginManager.h"
+#include "Kismet/GameplayStatics.h"
 #include "Misc/ConfigCacheIni.h"
 #include "MuCO/CustomizableObject.h"
+#include "MuCO/CustomizableObjectDGGUI.h"
 #include "MuCO/CustomizableObjectExtension.h"
+#include "MuCO/CustomizableSkeletalComponent.h"
 #include "MuCO/ICustomizableObjectModule.h"
 #include "UObject/StrongObjectPtr.h"
 
@@ -31,6 +34,10 @@ public:
 private:
 	void RefreshExtensionData();
 
+	// Command to look for Customizable Object Instance in the player pawn of the current world and open a DGGUI to edit its parameters
+	IConsoleCommand* LaunchDGGUICommand;
+	static void ToggleDGGUI(const TArray<FString>& Arguments);
+
 	// Ensure extensions aren't garbage collected
 	TArray<TStrongObjectPtr<const UCustomizableObjectExtension>> StrongExtensions;
 	// For returning from GetRegisteredExtensions
@@ -45,6 +52,10 @@ IMPLEMENT_MODULE( FCustomizableObjectModule, CustomizableObject );
 
 void FCustomizableObjectModule::StartupModule()
 {
+	LaunchDGGUICommand = IConsoleManager::Get().RegisterConsoleCommand(
+		TEXT("mutable.ToggleDGGUI"),
+		TEXT("Looks for a Customizable Object Instance within the player pawn and opens a UI to modify its parameters, or closes it if it's open. Specify slot ID to control which component is modified."),
+		FConsoleCommandWithArgsDelegate::CreateStatic(&FCustomizableObjectModule::ToggleDGGUI));
 }
 
 
@@ -166,3 +177,91 @@ void FCustomizableObjectModule::RefreshExtensionData()
 		}
 	}
 }
+
+
+UCustomizableSkeletalComponent* GetPlayerCustomizableSkeletalComponent(const int32 SlotID, const UWorld* CurrentWorld, const int32 PlayerIndex)
+{
+	// Get customizable skeletal component attached to player pawn
+	UCustomizableSkeletalComponent* SelectedCustomizableSkeletalComponent = nullptr;
+	{
+		APawn* PlayerPawn = UGameplayStatics::GetPlayerPawn(CurrentWorld, PlayerIndex);
+		int32 IndexFound = INDEX_NONE;
+		for (TObjectIterator<UCustomizableSkeletalComponent> CustomizableSkeletalComponent; CustomizableSkeletalComponent; ++CustomizableSkeletalComponent)
+		{
+			if (CustomizableSkeletalComponent->IsValidLowLevel() && !CustomizableSkeletalComponent->IsTemplate())
+			{
+				AActor* CustomizableActor = CustomizableSkeletalComponent->GetAttachmentRootActor();
+				if (CustomizableActor && PlayerPawn == CustomizableActor)
+				{
+					++IndexFound;
+					SelectedCustomizableSkeletalComponent = *CustomizableSkeletalComponent;
+					if (IndexFound == SlotID)
+					{
+						break;
+					}
+				}
+			}
+		}
+	}
+
+
+	// If none found, try getting a component without caring about the actor
+	if (!SelectedCustomizableSkeletalComponent)
+	{
+		APawn* PlayerPawn = UGameplayStatics::GetPlayerPawn(CurrentWorld, PlayerIndex);
+		int32 IndexFound = INDEX_NONE;
+		for (TObjectIterator<UCustomizableSkeletalComponent> CustomizableSkeletalComponent; CustomizableSkeletalComponent; ++CustomizableSkeletalComponent)
+		{
+			if (CustomizableSkeletalComponent->IsValidLowLevel() && !CustomizableSkeletalComponent->IsTemplate())
+			{
+				++IndexFound;
+				SelectedCustomizableSkeletalComponent = *CustomizableSkeletalComponent;
+				if (IndexFound == SlotID)
+				{
+					break;
+				}
+			}
+		}
+	}
+
+	return SelectedCustomizableSkeletalComponent;
+}
+
+
+void FCustomizableObjectModule::ToggleDGGUI(const TArray<FString>& Arguments)
+{
+	int32 SlotID = INDEX_NONE;
+	if (Arguments.Num() >= 1)
+	{
+		SlotID = FCString::Atoi(*Arguments[0]);
+	}
+
+	const UWorld* CurrentWorld = []() -> const UWorld*
+	{
+		UWorld* WorldForCurrentCOI = nullptr;
+		const TIndirectArray<FWorldContext>& WorldContexts = GEngine->GetWorldContexts();
+		for (const FWorldContext& Context : WorldContexts)
+		{
+			if ((Context.WorldType == EWorldType::Game) && (Context.World() != NULL))
+			{
+				WorldForCurrentCOI = Context.World();
+			}
+		}
+		// Fall back to GWorld if we don't actually have a world.
+		if (WorldForCurrentCOI == nullptr)
+		{
+			WorldForCurrentCOI = GWorld;
+		}
+		return WorldForCurrentCOI;
+	}();
+
+	const int32 PlayerIndex = 0;
+	if (UDGGUI::CloseExistingDGGUI(CurrentWorld))
+	{
+		return;
+	}
+	else if (UCustomizableSkeletalComponent* SelectedCustomizableSkeletalComponent = GetPlayerCustomizableSkeletalComponent(SlotID, CurrentWorld, PlayerIndex))
+	{
+		UDGGUI::OpenDGGUI(SlotID, SelectedCustomizableSkeletalComponent, CurrentWorld, PlayerIndex);
+	}
+}
\ No newline at end of file
diff --git a/Engine/Plugins/Experimental/Mutable/Source/CustomizableObject/Public/MuCO/ICustomizableObjectModule.h b/Engine/Plugins/Experimental/Mutable/Source/CustomizableObject/Public/MuCO/ICustomizableObjectModule.h
--- a/Engine/Plugins/Experimental/Mutable/Source/CustomizableObject/Public/MuCO/ICustomizableObjectModule.h
+++ b/Engine/Plugins/Experimental/Mutable/Source/CustomizableObject/Public/MuCO/ICustomizableObjectModule.h
@@ -6,6 +6,8 @@
 #include "Modules/ModuleManager.h"
 #include "MuCO/CustomizableObjectExtension.h"
 
+class UCustomizableSkeletalComponent;
+
 struct FRegisteredCustomizableObjectPinType
 {
 	TWeakObjectPtr<const UCustomizableObjectExtension> Extension;
@@ -76,3 +78,7 @@ public:
 	virtual TArrayView<const FRegisteredObjectNodeInputPin> GetAdditionalObjectNodePins() const = 0;
 };
 
+
+/* Utility function for commands*/
+CUSTOMIZABLEOBJECT_API
+UCustomizableSkeletalComponent* GetPlayerCustomizableSkeletalComponent(const int32 SlotID, const UWorld* CurrentWorld, const int32 PlayerIndex = 0);
diff --git a/Engine/Plugins/Experimental/Mutable/Source/CustomizableObjectEditor/Private/MuCOE/CustomizableObjectEditorModule.cpp b/Engine/Plugins/Experimental/Mutable/Source/CustomizableObjectEditor/Private/MuCOE/CustomizableObjectEditorModule.cpp
--- a/Engine/Plugins/Experimental/Mutable/Source/CustomizableObjectEditor/Private/MuCOE/CustomizableObjectEditorModule.cpp
+++ b/Engine/Plugins/Experimental/Mutable/Source/CustomizableObjectEditor/Private/MuCOE/CustomizableObjectEditorModule.cpp
@@ -6,12 +6,11 @@
 #include "GameFramework/Pawn.h"
 #include "ISettingsModule.h"
 #include "ISettingsSection.h"
-#include "Kismet/GameplayStatics.h"
 #include "MessageLogModule.h"
-#include "MuCO/CustomizableObjectDGGUI.h"
 #include "MuCO/CustomizableObjectSystem.h"		// For defines related to memory function replacements.
 #include "MuCO/CustomizableSkeletalComponent.h"
 #include "MuCO/CustomizableSkeletalMeshActor.h"
+#include "MuCO/ICustomizableObjectModule.h"		// For instance editor command utility function
 #include "MuCOE/CustomizableInstanceDetails.h"
 #include "MuCOE/CustomizableObjectCustomSettings.h"
 #include "MuCOE/CustomizableObjectCustomSettingsDetails.h"
@@ -124,10 +123,6 @@ private:
 
 	FCustomizableObjectEditorLogger Logger;
 
-	// Command to look for Customizable Object Instance in the player pawn of the current world and open a DGGUI to edit its parameters
-	IConsoleCommand* LaunchDGGUICommand;
-	static void ToggleDGGUI(const TArray<FString>& Arguments);
-
 	// Command to look for Customizable Object Instance in the player pawn of the current world and open its Customizable Object Instance Editor
 	IConsoleCommand* LaunchCOIECommand;
 	static void OpenCOIE(const TArray<FString>& Arguments);
@@ -213,11 +208,6 @@ void FCustomizableObjectEditorModule::StartupModule()
 	CustomizableObjectEditor_ToolBarExtensibilityManager = MakeShareable(new FExtensibilityManager);
 	CustomizableObjectEditor_MenuExtensibilityManager = MakeShareable(new FExtensibilityManager);
 
-	LaunchDGGUICommand = IConsoleManager::Get().RegisterConsoleCommand(
-		TEXT("mutable.ToggleDGGUI"),
-		TEXT("Looks for a Customizable Object Instance within the player pawn and opens a UI to modify its parameters, or closes it if it's open. Specify slot ID to control which component is modified."),
-		FConsoleCommandWithArgsDelegate::CreateStatic(&FCustomizableObjectEditorModule::ToggleDGGUI));
-
 	LaunchCOIECommand = IConsoleManager::Get().RegisterConsoleCommand(
 		TEXT("mutable.OpenCOIE"),
 		TEXT("Looks for a Customizable Object Instance within the player pawn and opens its Customizable Object Instance Editor. Specify slot ID to control which component is edited."),
@@ -339,103 +329,32 @@ void FCustomizableObjectEditorModule::RegisterCustomDetails(FPropertyEditorModul
 }
 
 
-const UWorld* GetWorldForCurrentCOI()
-{
-	UWorld* CurrentWorld = nullptr;
-	const TIndirectArray<FWorldContext>& WorldContexts = GEngine->GetWorldContexts();
-	for (const FWorldContext& Context : WorldContexts)
-	{
-		if ((Context.WorldType == EWorldType::Game) && (Context.World() != NULL))
-		{
-			CurrentWorld = Context.World();
-		}
-	}
-	// Fall back to GWorld if we don't actually have a world.
-	if (CurrentWorld == nullptr)
-	{
-		CurrentWorld = GWorld;
-	}
-	return CurrentWorld;
-}
-
-
-UCustomizableSkeletalComponent* GetPlayerCustomizableSkeletalComponent(const int32 SlotID, const UWorld* CurrentWorld, const int32 PlayerIndex = 0)
-{
-	// Get customizable skeletal component attached to player pawn
-	UCustomizableSkeletalComponent* SelectedCustomizableSkeletalComponent = nullptr;
-	{
-		APawn* PlayerPawn = UGameplayStatics::GetPlayerPawn(CurrentWorld, PlayerIndex);
-		int32 IndexFound = INDEX_NONE;
-		for (TObjectIterator<UCustomizableSkeletalComponent> CustomizableSkeletalComponent; CustomizableSkeletalComponent; ++CustomizableSkeletalComponent)
-		{
-			if (CustomizableSkeletalComponent->IsValidLowLevel() && !CustomizableSkeletalComponent->IsTemplate())
-			{
-				AActor* CustomizableActor = CustomizableSkeletalComponent->GetAttachmentRootActor();
-				if (CustomizableActor && PlayerPawn == CustomizableActor)
-				{
-					++IndexFound;
-					SelectedCustomizableSkeletalComponent = *CustomizableSkeletalComponent;
-					if (IndexFound == SlotID)
+void FCustomizableObjectEditorModule::OpenCOIE(const TArray<FString>& Arguments)
 					{
-						break;
-					}
-				}
-			}
-		}
-	}
-
-
-	// If none found, try getting a component without caring about the actor
-	if (!SelectedCustomizableSkeletalComponent)
-	{
-		APawn* PlayerPawn = UGameplayStatics::GetPlayerPawn(CurrentWorld, PlayerIndex);
-		int32 IndexFound = INDEX_NONE;
-		for (TObjectIterator<UCustomizableSkeletalComponent> CustomizableSkeletalComponent; CustomizableSkeletalComponent; ++CustomizableSkeletalComponent)
-		{
-			if (CustomizableSkeletalComponent->IsValidLowLevel() && !CustomizableSkeletalComponent->IsTemplate())
-			{
-				++IndexFound;
-				SelectedCustomizableSkeletalComponent = *CustomizableSkeletalComponent;
-				if (IndexFound == SlotID)
-				{
-					break;
-				}
-			}
-		}
-	}
-
-	return SelectedCustomizableSkeletalComponent;
-}
-
-
-void FCustomizableObjectEditorModule::ToggleDGGUI(const TArray<FString>& Arguments)
-{
 	int32 SlotID = INDEX_NONE;
 	if (Arguments.Num() >= 1)
 	{
 		SlotID = FCString::Atoi(*Arguments[0]);
 	}
-	const UWorld* CurrentWorld = GetWorldForCurrentCOI();
-	const int32 PlayerIndex = 0;
-	if (UDGGUI::CloseExistingDGGUI(CurrentWorld))
+
+	const UWorld* CurrentWorld = []() -> const UWorld*
 	{
-		return;
-	}
-	else if (UCustomizableSkeletalComponent* SelectedCustomizableSkeletalComponent = GetPlayerCustomizableSkeletalComponent(SlotID, CurrentWorld, PlayerIndex))
+		UWorld* WorldForCurrentCOI = nullptr;
+		const TIndirectArray<FWorldContext>& WorldContexts = GEngine->GetWorldContexts();
+		for (const FWorldContext& Context : WorldContexts)
 	{
-		UDGGUI::OpenDGGUI(SlotID, SelectedCustomizableSkeletalComponent, CurrentWorld, PlayerIndex);
+			if ((Context.WorldType == EWorldType::Game) && (Context.World() != NULL))
+			{
+				WorldForCurrentCOI = Context.World();
 	}
 }
-
-
-void FCustomizableObjectEditorModule::OpenCOIE(const TArray<FString>& Arguments)
+		// Fall back to GWorld if we don't actually have a world.
+		if (WorldForCurrentCOI == nullptr)
 {
-	int32 SlotID = INDEX_NONE;
-	if (Arguments.Num() >= 1)
-	{
-		SlotID = FCString::Atoi(*Arguments[0]);
+			WorldForCurrentCOI = GWorld;
 	}
-	const UWorld* CurrentWorld = GetWorldForCurrentCOI();
+		return WorldForCurrentCOI;
+	}();
 	const int32 PlayerIndex = 0;
 
 	// Open the Customizable Object Instance Editor
