#rnx

diff --git a/Engine/Plugins/Experimental/Mutable/Source/CustomizableObjectEditor/Private/MuCOE/EdGraphSchema_CustomizableObject.cpp b/Engine/Plugins/Experimental/Mutable/Source/CustomizableObjectEditor/Private/MuCOE/EdGraphSchema_CustomizableObject.cpp
--- a/Engine/Plugins/Experimental/Mutable/Source/CustomizableObjectEditor/Private/MuCOE/EdGraphSchema_CustomizableObject.cpp
+++ b/Engine/Plugins/Experimental/Mutable/Source/CustomizableObjectEditor/Private/MuCOE/EdGraphSchema_CustomizableObject.cpp
@@ -1063,7 +1063,7 @@ FText UEdGraphSchema_CustomizableObject::GetPinCategoryName(const FName& PinCate
 	}
 	else if (PinCategory == UEdGraphSchema_CustomizableObject::PC_PassThroughImage)
 	{
-		return LOCTEXT("Image_Pin_Category", "PassThrough Texture");
+		return LOCTEXT("PassThrough_Image_Pin_Category", "PassThrough Texture");
 	}
 	else if (PinCategory == UEdGraphSchema_CustomizableObject::PC_Projector)
 	{
diff --git a/Engine/Plugins/Experimental/Mutable/Source/CustomizableObjectEditor/Private/MuCOE/GenerateMutableSource/GenerateMutableSourceImage.cpp b/Engine/Plugins/Experimental/Mutable/Source/CustomizableObjectEditor/Private/MuCOE/GenerateMutableSource/GenerateMutableSourceImage.cpp
--- a/Engine/Plugins/Experimental/Mutable/Source/CustomizableObjectEditor/Private/MuCOE/GenerateMutableSource/GenerateMutableSourceImage.cpp
+++ b/Engine/Plugins/Experimental/Mutable/Source/CustomizableObjectEditor/Private/MuCOE/GenerateMutableSource/GenerateMutableSourceImage.cpp
@@ -483,7 +483,7 @@ mu::NodeImagePtr GenerateMutableSourceImage(const UEdGraphPin* Pin, FMutableGrap
 					}
 					else
 					{
-						const FText Message = LOCTEXT("MissingTexture", "Unable to generate pass-through texture switch node. Required connection not found.");
+						const FText Message = LOCTEXT("MissingPassThroughTexture", "Unable to generate pass-through texture switch node. Required connection not found.");
 						GenerationContext.Compiler->CompilerLog(Message, Node);
 						return Result;
 					}
