[Mutable] Added GenerateMutabelSource grammar documentation.

#rnx

diff --git a/Engine/Plugins/Experimental/Mutable/Source/CustomizableObjectEditor/Private/MuCOE/Compiler/GrammarImplicits.txt b/Engine/Plugins/Experimental/Mutable/Source/CustomizableObjectEditor/Private/MuCOE/Compiler/GrammarImplicits.txt
new file mode 100644
--- /dev/null
+++ b/Engine/Plugins/Experimental/Mutable/Source/CustomizableObjectEditor/Private/MuCOE/Compiler/GrammarImplicits.txt
@@ -0,0 +1,98 @@
+// GRAMMAR IMPLICITS
+//
+// Grammar with implicit connections (e.g., parent materials).
+// Equivalent to what GenerateMutableSource is traversing.
+
+ENUM -> NodeEnumParameter
+
+CURVEASSET -> UCurve
+
+FLOAT -> NodeFloatConstant
+FLOAT -> NodeFloatParameter
+FLOAT -> NodeCurve CURVEASSET
+FLOAT -> NodeFloatSwitch ENUM FLOAT*
+FLOAT -> NodeFloatVariation FLOAT FLOAT*
+FLOAT -> NodeTable
+
+LAYOUT -> MeshLayoutNode
+
+MORPHSTACK -> NodeMeshMorphStackDefinition MESH
+
+STATICMESHASSET -> UStaticMesh
+
+SKELETASLMESHASSET-> USkeletalMesh
+
+MESH -> NodeMeshGeometryOperation MESH:MeshA MESH:MeshB FLOAT:ScalarA FLOAT:ScalarB
+MESH -> NodeMeshMorph MESH FLOAT
+MESH -> NodeMeshMorphStackApplication MORPHSTACK MESH
+MESH -> NodeMeshReshape MESH:BaseMesh MESH:BaseShape MESH:TargetShape
+MESH -> NodeMeshSwitch ENUM MESH*
+MESH -> NodeMeshVariation MESH MESH*
+MESH -> NodeSkeletalMesh SKELETALMESH LAYOUT*
+MESH -> NodeStaticMesh STATICMESHASSET
+MESH -> NodeAnimationPose MESH
+MESH -> NodeTable SKELETASLMESHASSET* STATICMESHASSET*
+
+PROJECTOR -> NodeProjectorConstant
+PROJECTOR -> NodeProjectorParameter
+
+TEXTUREASSET -> UTexture2D
+
+TEXTURE -> NodeSkeletalMesh TEXTUREASSET
+TEXTURE -> NodeStaticMesh TEXTUREASSET
+TEXTURE -> NodeTextureToChannels TEXTURE:R
+TEXTURE -> NodeTextureToChannels TEXTURE:G
+TEXTURE -> NodeTextureToChannels TEXTURE:B
+TEXTURE -> NodeTextureToChannels TEXTURE:A
+TEXTURE -> NodeTextureFromChannels TEXTURE:R TEXTURE:G TEXTURE:B TEXTURE:A
+TEXTURE -> NodeTexture TEXTUREASSET // Texture Constant
+TEXTURE -> NodeTextureBinarise TEXTURE FLOAT
+TEXTURE -> NodeTextureColourMap TEXTURE:Base TEXTURE:Mask TEXTURE:Map
+TEXTURE -> NodeTextureInterpolate FLOAT TEXTURE:Target0 TEXTURE:Target1
+TEXTURE -> NodeTextureInvert TEXTURE
+TEXTURE -> NodeTextureLayer TEXTURE
+TEXTURE -> NodeTextureParameter
+TEXTURE -> NodeTextureProject TEXTURE:Base MESH TEXTURE:Mask FLOAT:FadeStartAngle FLOAT:FadeEndAngle PROJECTOR
+TEXTURE -> NodeTextureSaturate TEXTURE FLOAT
+TEXTURE -> NodeTextureSwitch ENUM TEXTURE*
+TEXTURE -> NodeTextureVariation TEXTURE TEXTURE*
+TEXTURE -> NodeTextureTransform TEXTURE FLOAT:OffsetX FLOAT:OffsetY FLOAT:ScaleX FLOAT:ScaleY FLOAT:Rotation
+TEXTURE -> NodeTable TEXTUREASSET*
+
+COLOR -> NodeColorConstant
+COLOR -> NodeColorParameter
+COLOR -> NodeColorArithmeticOp COLOR:A COLOR:B
+COLOR -> NodeColorFromFloats FLOAT:R FLOAT:G FLOAT:B FLOAT:A
+COLOR -> NodeColorSwitch ENUM COLOR*
+COLOR -> NodeColorVariation COLOR COLOR*
+COLOR -> NodeTextureSample TEXTURE FLOAT:X FLOAT:Y
+COLOR -> NodeTable
+
+MATERIALASSET -> UMaterial
+
+MATERIALINSTANCE -> NodeTableMaterial MATERIALASSET MATERIALASSET*
+
+SURFACENODE -> NodeMaterial MATERIALASSET MESH:Mesh MATERIALINSTANCE TEXTURE* COLOR* FLOAT*
+
+SURFACE -> NodeEditMaterial SURFACENODE TEXTURE*
+SURFACE -> NodeExtendMaterial SURFACENODE TEXTURE*
+SURFACE -> NodeMaterialVariation SURFACENODE SURFACENODE*
+SURFACE -> NodeCopyMaterial MESH SURFACENODE
+SURFACE -> NodeRemoveMesh SURFACENODE MESH
+SURFACE -> NodeRemoveMeshBlocks SURFACENODE 
+SURFACE -> NodeMorphMaterial SURFACENODE
+
+MODIFIER -> NodeMeshClipDeform MESH
+MODIFIER -> NodeMeshClipWithMesh MESH SURFACENODE
+MODIFIER -> NodeMeshClipMorph
+
+PROJECTORGROUP -> NodeGroupProjectorParameter
+
+OBJECT -> NodeObject SURFACE* MODIFIER*
+
+GROUP -> NodeObjectGroup PROJECTORGROUP* OBJECT* 
+
+
+// Observations:
+// - Core allows to chain Edit Materials.
+// - Currently there we do not differentiate MODIFIERs and SURFACEs
\ No newline at end of file
