[mutable] Fix actor cast warning.

#rnx

diff --git a/Engine/Plugins/Experimental/Mutable/Source/CustomizableObject/Private/MuCO/CustomizableObjectModule.cpp b/Engine/Plugins/Experimental/Mutable/Source/CustomizableObject/Private/MuCO/CustomizableObjectModule.cpp
--- a/Engine/Plugins/Experimental/Mutable/Source/CustomizableObject/Private/MuCO/CustomizableObjectModule.cpp
+++ b/Engine/Plugins/Experimental/Mutable/Source/CustomizableObject/Private/MuCO/CustomizableObjectModule.cpp
@@ -3,6 +3,7 @@
 #include "Algo/Find.h"
 #include "Engine/Engine.h"
 #include "Interfaces/IPluginManager.h"
+#include "GameFramework/Pawn.h"
 #include "Kismet/GameplayStatics.h"
 #include "Misc/ConfigCacheIni.h"
 #include "MuCO/CustomizableObject.h"
@@ -185,7 +186,7 @@ UCustomizableSkeletalComponent* GetPlayerCustomizableSkeletalComponent(const int
 	// Get customizable skeletal component attached to player pawn
 	UCustomizableSkeletalComponent* SelectedCustomizableSkeletalComponent = nullptr;
 	{
-		AActor* PlayerPawn = reinterpret_cast<AActor*>(UGameplayStatics::GetPlayerPawn(CurrentWorld, PlayerIndex));
+		AActor* PlayerPawn = Cast<AActor>(UGameplayStatics::GetPlayerPawn(CurrentWorld, PlayerIndex));
 		int32 IndexFound = INDEX_NONE;
 		for (TObjectIterator<UCustomizableSkeletalComponent> CustomizableSkeletalComponent; CustomizableSkeletalComponent; ++CustomizableSkeletalComponent)
 		{
@@ -209,7 +210,7 @@ UCustomizableSkeletalComponent* GetPlayerCustomizableSkeletalComponent(const int
 	// If none found, try getting a component without caring about the actor
 	if (!SelectedCustomizableSkeletalComponent)
 	{
-		AActor* PlayerPawn = reinterpret_cast<AActor*>(UGameplayStatics::GetPlayerPawn(CurrentWorld, PlayerIndex));
+		AActor* PlayerPawn = Cast<AActor>(UGameplayStatics::GetPlayerPawn(CurrentWorld, PlayerIndex));
 		int32 IndexFound = INDEX_NONE;
 		for (TObjectIterator<UCustomizableSkeletalComponent> CustomizableSkeletalComponent; CustomizableSkeletalComponent; ++CustomizableSkeletalComponent)
 		{
