[Mutable] Fix incorrect check in default values.

#rnx

diff --git a/Engine/Plugins/Experimental/Mutable/Source/MutableRuntime/Private/MuR/Model.cpp b/Engine/Plugins/Experimental/Mutable/Source/MutableRuntime/Private/MuR/Model.cpp
--- a/Engine/Plugins/Experimental/Mutable/Source/MutableRuntime/Private/MuR/Model.cpp
+++ b/Engine/Plugins/Experimental/Mutable/Source/MutableRuntime/Private/MuR/Model.cpp
@@ -347,8 +347,8 @@ namespace mu
 		check(m_pD->m_program.m_parameters[Index].m_type == PARAMETER_TYPE::T_BOOL);
 
         // Early out in case of invalid parameters
-        if (m_pD->m_program.m_parameters.IsValidIndex(Index) ||
-            m_pD->m_program.m_parameters[Index].m_type == PARAMETER_TYPE::T_BOOL)
+        if (!m_pD->m_program.m_parameters.IsValidIndex(Index) ||
+            m_pD->m_program.m_parameters[Index].m_type != PARAMETER_TYPE::T_BOOL)
         {
             return false;
         }
@@ -364,8 +364,8 @@ namespace mu
 		check(m_pD->m_program.m_parameters[Index].m_type == PARAMETER_TYPE::T_INT);
 
         // Early out in case of invalid parameters
-        if (m_pD->m_program.m_parameters.IsValidIndex(Index) ||
-            m_pD->m_program.m_parameters[Index].m_type == PARAMETER_TYPE::T_INT)
+        if (!m_pD->m_program.m_parameters.IsValidIndex(Index) ||
+            m_pD->m_program.m_parameters[Index].m_type != PARAMETER_TYPE::T_INT)
         {
             return 0;
         }
@@ -381,8 +381,8 @@ namespace mu
 		check(m_pD->m_program.m_parameters[Index].m_type == PARAMETER_TYPE::T_FLOAT);
 
         // Early out in case of invalid parameters
-        if (m_pD->m_program.m_parameters.IsValidIndex(Index) ||
-            m_pD->m_program.m_parameters[Index].m_type == PARAMETER_TYPE::T_FLOAT)
+        if (!m_pD->m_program.m_parameters.IsValidIndex(Index) ||
+            m_pD->m_program.m_parameters[Index].m_type != PARAMETER_TYPE::T_FLOAT)
         {
             return 0.0f;
         }
@@ -397,8 +397,8 @@ namespace mu
 		check(m_pD->m_program.m_parameters[Index].m_type == PARAMETER_TYPE::T_COLOUR);
 
         // Early out in case of invalid parameters
-        if (m_pD->m_program.m_parameters.IsValidIndex(Index) ||
-            m_pD->m_program.m_parameters[Index].m_type == PARAMETER_TYPE::T_COLOUR)
+        if (!m_pD->m_program.m_parameters.IsValidIndex(Index) ||
+            m_pD->m_program.m_parameters[Index].m_type != PARAMETER_TYPE::T_COLOUR)
         {
             return;
         }
@@ -417,8 +417,8 @@ namespace mu
 		check(m_pD->m_program.m_parameters[Index].m_type == PARAMETER_TYPE::T_PROJECTOR);
 
         // Early out in case of invalid parameters
-        if (m_pD->m_program.m_parameters.IsValidIndex(Index) ||
-            m_pD->m_program.m_parameters[Index].m_type == PARAMETER_TYPE::T_PROJECTOR)
+        if (!m_pD->m_program.m_parameters.IsValidIndex(Index) ||
+            m_pD->m_program.m_parameters[Index].m_type != PARAMETER_TYPE::T_PROJECTOR)
         {
             return;
         }
@@ -440,8 +440,8 @@ namespace mu
 		check(m_pD->m_program.m_parameters[Index].m_type == PARAMETER_TYPE::T_IMAGE);
 
         // Early out in case of invalid parameters
-        if (m_pD->m_program.m_parameters.IsValidIndex(Index) ||
-            m_pD->m_program.m_parameters[Index].m_type == PARAMETER_TYPE::T_IMAGE)
+        if (!m_pD->m_program.m_parameters.IsValidIndex(Index) ||
+            m_pD->m_program.m_parameters[Index].m_type != PARAMETER_TYPE::T_IMAGE)
         {
             return 0;
         }
