[mutable] Added support for pass-though textures to Mutable tables.

#rnx

diff --git a/Engine/Plugins/Experimental/Mutable/Source/CustomizableObject/Public/MuCO/CustomizableObjectCustomVersion.h b/Engine/Plugins/Experimental/Mutable/Source/CustomizableObject/Public/MuCO/CustomizableObjectCustomVersion.h
--- a/Engine/Plugins/Experimental/Mutable/Source/CustomizableObject/Public/MuCO/CustomizableObjectCustomVersion.h
+++ b/Engine/Plugins/Experimental/Mutable/Source/CustomizableObject/Public/MuCO/CustomizableObjectCustomVersion.h
@@ -68,6 +68,8 @@ struct CUSTOMIZABLEOBJECT_API FCustomizableObjectCustomVersion
 
 		StateTextureCompressionStrategyEnum,
 
+		AddedTableNodesTextureMode,
+
 		// -----<new versions can be added above this line>-------------------------------------------------
 		VersionPlusOne,
 		LatestVersion = VersionPlusOne - 1
diff --git a/Engine/Plugins/Experimental/Mutable/Source/CustomizableObjectEditor/Classes/MuCOE/Nodes/CustomizableObjectNodeTable.h b/Engine/Plugins/Experimental/Mutable/Source/CustomizableObjectEditor/Classes/MuCOE/Nodes/CustomizableObjectNodeTable.h
--- a/Engine/Plugins/Experimental/Mutable/Source/CustomizableObjectEditor/Classes/MuCOE/Nodes/CustomizableObjectNodeTable.h
+++ b/Engine/Plugins/Experimental/Mutable/Source/CustomizableObjectEditor/Classes/MuCOE/Nodes/CustomizableObjectNodeTable.h
@@ -4,6 +4,7 @@
 
 #include "Engine/DataTable.h"
 #include "MuCOE/Nodes/CustomizableObjectNode.h"
+#include "MuCOE/RemapPins/CustomizableObjectNodeRemapPinsByName.h"
 
 #include "CustomizableObjectNodeTable.generated.h"
 
@@ -11,6 +12,7 @@ namespace ENodeTitleType { enum Type : int; }
 
 class UCustomizableObjectLayout;
 class UCustomizableObjectNodeRemapPins;
+class UCustomizableObjectNodeRemapPinsByName;
 class UEdGraphPin;
 class UObject;
 class USkeletalMesh;
@@ -18,6 +20,15 @@ class UTexture2D;
 struct FGuid;
 
 
+/** Enum class for the different types of image pins */
+UENUM()
+enum class ETableTextureType : uint8
+{
+	PASSTHROUGH_TEXTURE = 0 UMETA(DisplayName = "Passthrough"),
+	MUTABLE_TEXTURE = 1 UMETA(DisplayName = "Mutable")
+};
+
+
 /** Base class for all Table Pins. */
 UCLASS()
 class CUSTOMIZABLEOBJECTEDITOR_API UCustomizableObjectNodeTableObjectPinData : public UCustomizableObjectNodePinData
@@ -33,7 +44,30 @@ public:
 };
 
 
-/** Additional data for a Mesh pins. */
+/** Additional data for Image pins. */
+UCLASS()
+class CUSTOMIZABLEOBJECTEDITOR_API UCustomizableObjectNodeTableImagePinData : public UCustomizableObjectNodeTableObjectPinData
+{
+	GENERATED_BODY()
+
+public:
+
+	bool IsDefaultImageMode() { return bIsDefault; }
+	void SetDefaultImageMode(bool bValue) { bIsDefault = bValue; }
+
+	// Pin Type
+	UPROPERTY()
+	ETableTextureType ImageMode = ETableTextureType::MUTABLE_TEXTURE;
+
+private:
+
+	UPROPERTY()
+	bool bIsDefault = true;
+
+};
+
+
+/** Additional data for Mesh pins. */
 UCLASS()
 class CUSTOMIZABLEOBJECTEDITOR_API UCustomizableObjectNodeTableMeshPinData : public UCustomizableObjectNodeTableObjectPinData
 {
@@ -70,6 +104,20 @@ public:
 };
 
 
+UCLASS()
+class UCustomizableObjectNodeTableRemapPins : public UCustomizableObjectNodeRemapPinsByName
+{
+	GENERATED_BODY()
+public:
+
+	// Specific method to decide when two pins are equal
+	virtual bool Equal(const UEdGraphPin& OldPin, const UEdGraphPin& NewPin) const override;
+
+	// Method to use in the RemapPins step of the node reconstruction process
+	virtual void RemapPins(const TArray<UEdGraphPin*>& OldPins, const TArray<UEdGraphPin*>& NewPins, TMap<UEdGraphPin*, UEdGraphPin*>& PinsToRemap, TArray<UEdGraphPin*>& PinsToOrphan) override;
+};
+
+
 UCLASS()
 class CUSTOMIZABLEOBJECTEDITOR_API UCustomizableObjectNodeTable : public UCustomizableObjectNode
 {
@@ -92,7 +140,13 @@ public:
 	/** If there is a bool column in the table, checked rows will not be compiled */
 	UPROPERTY(EditAnywhere, Category = TableProperties)
 	bool bDisableCheckedRows = true;
-	
+
+	/** Decides the default type of the texture pins (passtrhough or mutable)
+	*   Right click on a non-linked image pin to customize its image mode
+	*/
+	UPROPERTY(EditAnywhere, Category = TableProperties)
+	ETableTextureType DefaultImageMode = ETableTextureType::MUTABLE_TEXTURE;
+
 	UPROPERTY(EditAnywhere, Category = UI, meta = (DisplayName = "Parameter UI Metadata"))
 	FMutableParamUIMetadata ParamUIMetadata;
 
@@ -105,14 +159,17 @@ public:
 	FText GetNodeTitle(ENodeTitleType::Type TitleType) const override;
 	FLinearColor GetNodeTitleColor() const override;
 	FText GetTooltipText() const override;
+	virtual void PinConnectionListChanged(UEdGraphPin* Pin) override;
 	
 	// UCustomizableObjectNode interface
 	virtual void PostBackwardsCompatibleFixup() override;
+	virtual void BackwardsCompatibleFixup() override;
 	void AllocateDefaultPins(UCustomizableObjectNodeRemapPins* RemapPins) override;
 	bool IsNodeOutDatedAndNeedsRefresh() override;
 	FString GetRefreshMessage() const override;
 	virtual bool ProvidesCustomPinRelevancyTest() const override;
 	virtual bool IsPinRelevant(const UEdGraphPin* Pin) const override;
+	UCustomizableObjectNodeTableRemapPins* CreateRemapPinsDefault() const;
 
 	/*** Allows to perform work when remapping the pin data. */
 	virtual void RemapPinsData(const TMap<UEdGraphPin*, UEdGraphPin*>& PinsToRemap) override;
@@ -227,6 +284,13 @@ public:
 	// Returns the name if the row has a bool column set as false (true == disabled)
 	TArray<FName> GetRowNames() const;
 
+	// Changes the image mode of a pin
+	// bSetDefault param: if true sets the pin to be equal to the default mode (same as node)
+	void ChangeImagePinMode(UEdGraphPin* Pin, bool bSetDefault = false);
+
+	// Returns true if the pin is in the default mode (same as node)
+	bool IsImagePinDefault(UEdGraphPin* Pin);
+
 private:
 
 	/** Number of properties to know when the node needs an update */
@@ -240,4 +304,6 @@ private:
 
 	// Checks if a pin already exists and if it has the same type as before the node refresh
 	bool CheckPinUpdated(const FString& PinName, const FName& PinType) const;
+
+
 };
diff --git a/Engine/Plugins/Experimental/Mutable/Source/CustomizableObjectEditor/Private/MuCOE/EdGraphSchema_CustomizableObject.cpp b/Engine/Plugins/Experimental/Mutable/Source/CustomizableObjectEditor/Private/MuCOE/EdGraphSchema_CustomizableObject.cpp
--- a/Engine/Plugins/Experimental/Mutable/Source/CustomizableObjectEditor/Private/MuCOE/EdGraphSchema_CustomizableObject.cpp
+++ b/Engine/Plugins/Experimental/Mutable/Source/CustomizableObjectEditor/Private/MuCOE/EdGraphSchema_CustomizableObject.cpp
@@ -719,48 +719,91 @@ bool UEdGraphSchema_CustomizableObject::ShouldHidePinDefaultValue(UEdGraphPin* P
 
 void UEdGraphSchema_CustomizableObject::GetContextMenuActions(class UToolMenu* Menu, class UGraphNodeContextMenuContext* Context) const
 {
-	if (Context && !Context->Pin && Context->Node)
+	if (Context && Context->Node)
 	{
-		if (!Context->bIsDebugging)
+		if (!Context->Pin) // On Node right click
 		{
-			// Node contextual actions
-			FToolMenuSection& Section = Menu->AddSection("EdGraphSchemaNodeActions", LOCTEXT("NodeActionsMenuHeader", "Node Actions"));
-			Section.AddMenuEntry(FGenericCommands::Get().Delete);
-			Section.AddMenuEntry(FGenericCommands::Get().Cut);
-			Section.AddMenuEntry(FGenericCommands::Get().Copy);
-			Section.AddMenuEntry(FGenericCommands::Get().Duplicate);
-			Section.AddMenuEntry(FGraphEditorCommands::Get().ReconstructNodes);
-			Section.AddMenuEntry(FGraphEditorCommands::Get().BreakNodeLinks);
-
-			// In the case of a UCustomizableObjectNodeObjectGroup, add the option to refresh all Customizable Object Material Node nodes of all the children of this node
-			const UCustomizableObjectNodeObjectGroup* TypedNode = Cast<UCustomizableObjectNodeObjectGroup>(Context->Node);
-			if (TypedNode != nullptr)
+			if (!Context->bIsDebugging)
 			{
-				Section.AddMenuEntry( FCustomizableObjectEditorNodeContextCommands::Get().RefreshMaterialNodesInAllChildren );
+				// Node contextual actions
+				FToolMenuSection& Section = Menu->AddSection("EdGraphSchemaNodeActions", LOCTEXT("NodeActionsMenuHeader", "Node Actions"));
+				Section.AddMenuEntry(FGenericCommands::Get().Delete);
+				Section.AddMenuEntry(FGenericCommands::Get().Cut);
+				Section.AddMenuEntry(FGenericCommands::Get().Copy);
+				Section.AddMenuEntry(FGenericCommands::Get().Duplicate);
+				Section.AddMenuEntry(FGraphEditorCommands::Get().ReconstructNodes);
+				Section.AddMenuEntry(FGraphEditorCommands::Get().BreakNodeLinks);
+
+				// In the case of a UCustomizableObjectNodeObjectGroup, add the option to refresh all Customizable Object Material Node nodes of all the children of this node
+				const UCustomizableObjectNodeObjectGroup* TypedNode = Cast<UCustomizableObjectNodeObjectGroup>(Context->Node);
+				if (TypedNode != nullptr)
+				{
+					Section.AddMenuEntry(FCustomizableObjectEditorNodeContextCommands::Get().RefreshMaterialNodesInAllChildren);
+				}
 			}
-		}
 
-		struct SCommentUtility
+			struct SCommentUtility
+			{
+				static void CreateComment(const UEdGraphSchema_CustomizableObject* Schema, UEdGraph* Graph)
+				{
+					if (Schema && Graph)
+					{
+						Schema->AddComment(Graph, NULL, FVector2D::ZeroVector, true);
+					}
+				}
+			};
+
+			FToolMenuSection& section = Menu->AddSection("SchemaActionComment", LOCTEXT("MultiCommentHeader", "Comment Group"));
+			section.AddMenuEntry("MultiCommentDesc", LOCTEXT("MultiCommentDesc", "Create Comment from Selection"),
+				LOCTEXT("CommentToolTip", "Create a resizable comment box around selection."),
+				FSlateIcon(),
+				FUIAction(FExecuteAction::CreateStatic(SCommentUtility::CreateComment, this, const_cast<UEdGraph*>(ToRawPtr(Context->Graph)))));
+		}
+		else // On Pin right click
 		{
-			static void CreateComment(const UEdGraphSchema_CustomizableObject* Schema, UEdGraph* Graph)
+			UCustomizableObjectNodeTable* TableNode = Cast<UCustomizableObjectNodeTable>(Context->Node);
+			UEdGraphPin* TexturePin = (UEdGraphPin*)Context->Pin;
+
+			if (TableNode && TexturePin && !TexturePin->LinkedTo.Num() && (TexturePin->PinType.PinCategory == PC_Image || TexturePin->PinType.PinCategory == PC_PassThroughImage))
 			{
-				if (Schema && Graph)
+				FText ActionText = FText::Format(LOCTEXT("ChangeTexturePinMode_Label", 
+					"Set as {0} texture"), Context->Pin->PinType.PinCategory == PC_PassThroughImage ? FText::FromString("Mutable") : FText::FromString("Pass Through"));
+				
+				FText ToolTipText = FText::Format(LOCTEXT("ChangeTexturePinMode_Tooltip",
+					"Set the texture pin as {0} texture."), Context->Pin->PinType.PinCategory == PC_PassThroughImage ? FText::FromString("Mutable") : FText::FromString("Pass Through"));
+
+				FToolMenuSection& Section = Menu->FindOrAddSection("EdGraphSchemaPinActions");
+				Section.InitSection("EdGraphSchemaPinActions", LOCTEXT("PinActionsMenuHeader", "Pin Actions"), FToolMenuInsert());
+
+				Section.AddMenuEntry
+				(
+					"ChangeTexturePinMode", ActionText, ToolTipText, FSlateIcon(),
+					FUIAction(FExecuteAction::CreateLambda([TableNode, TexturePin]()
+						{
+							TableNode->ChangeImagePinMode(TexturePin);
+						}))
+				);
+
+				if (!TableNode->IsImagePinDefault(TexturePin))
 				{
-					Schema->AddComment(Graph, NULL, FVector2D::ZeroVector, true);
+					Section.AddMenuEntry
+					(
+						"SetTexturePinModeDefault",
+						LOCTEXT("SetTexturePinModeDefault_Label","Set pin as default."),
+						LOCTEXT("SetTexturePinModeDefault_Tooltip","Set the selected texture pin to use the default node mode."),
+						FSlateIcon(),
+						FUIAction(FExecuteAction::CreateLambda([TableNode, TexturePin]()
+							{
+								TableNode->ChangeImagePinMode(TexturePin, true);
+							}))
+					);
 				}
 			}
-		};
-
-		FToolMenuSection& section = Menu->AddSection("SchemaActionComment", LOCTEXT("MultiCommentHeader", "Comment Group"));
-		section.AddMenuEntry("MultiCommentDesc", LOCTEXT("MultiCommentDesc", "Create Comment from Selection"),
-			LOCTEXT("CommentToolTip", "Create a resizable comment box around selection."),
-			FSlateIcon(),
-			FUIAction(FExecuteAction::CreateStatic(SCommentUtility::CreateComment, this, const_cast<UEdGraph*>(ToRawPtr(Context->Graph)))));
+		}
 	}
 }
 
 
-
 void UEdGraphSchema_CustomizableObject::BreakPinLinks(UEdGraphPin& TargetPin, bool bSendsNodeNotification) const
 {
 	const FScopedTransaction Transaction(NSLOCTEXT("UnrealEd", "GraphEd_BreakPinLinks", "Break Pin Links"));
@@ -768,6 +811,7 @@ void UEdGraphSchema_CustomizableObject::BreakPinLinks(UEdGraphPin& TargetPin, bo
 	Super::BreakPinLinks(TargetPin, bSendsNodeNotification);
 }
 
+
 void UEdGraphSchema_CustomizableObject::BreakSinglePinLink(UEdGraphPin* SourcePin, UEdGraphPin* TargetPin) const
 {
 	const FScopedTransaction Transaction(NSLOCTEXT("UnrealEd", "GraphEd_BreakSinglePinLink", "Break Pin Link"));
@@ -794,6 +838,7 @@ TSharedPtr<ICustomizableObjectEditor> UEdGraphSchema_CustomizableObject::GetCust
 	return nullptr;
 }
 
+
 UEdGraphNode* UEdGraphSchema_CustomizableObject::AddComment(class UEdGraph* ParentGraph, UEdGraphPin* FromPin, const FVector2D Location, bool bSelectNewNode/* = true*/) const
 {
 	TSharedPtr<ICustomizableObjectEditor> CustomizableObjectEditor = GetCustomizableObjectEditor(ParentGraph);
@@ -923,6 +968,7 @@ void UEdGraphSchema_CustomizableObject::DroppedAssetsOnGraph(const TArray<FAsset
 	}
 }
 
+
 bool UEdGraphSchema_CustomizableObject::IsSpawnableAsset(const FAssetData& InAsset, ESpawnableObjectType& OutObjectType) const
 {
 	UObject* Object = InAsset.GetAsset();
@@ -1009,6 +1055,7 @@ void UEdGraphSchema_CustomizableObject::GetAssetsGraphHoverMessage(const TArray<
 	}
 }
 
+
 bool UEdGraphSchema_CustomizableObject::TryCreateConnection(UEdGraphPin* PinA, UEdGraphPin* PinB) const
 {
 	bool Result = Super::TryCreateConnection(PinA, PinB);
diff --git a/Engine/Plugins/Experimental/Mutable/Source/CustomizableObjectEditor/Private/MuCOE/Nodes/CustomizableObjectNodeTable.cpp b/Engine/Plugins/Experimental/Mutable/Source/CustomizableObjectEditor/Private/MuCOE/Nodes/CustomizableObjectNodeTable.cpp
--- a/Engine/Plugins/Experimental/Mutable/Source/CustomizableObjectEditor/Private/MuCOE/Nodes/CustomizableObjectNodeTable.cpp
+++ b/Engine/Plugins/Experimental/Mutable/Source/CustomizableObjectEditor/Private/MuCOE/Nodes/CustomizableObjectNodeTable.cpp
@@ -14,6 +14,7 @@
 #include "MuCOE/MutableUtils.h"
 #include "MuCOE/Nodes/CustomizableObjectNodeMaterial.h"
 #include "MuCOE/SCustomizableObjectNodeLayoutBlocksEditor.h"
+#include "MuCOE/UnrealEditorPortabilityHelpers.h"
 #include "Rendering/SkeletalMeshLODModel.h"
 #include "Rendering/SkeletalMeshModel.h"
 
@@ -23,6 +24,77 @@ class UCustomizableObjectNodeRemapPins;
 #define LOCTEXT_NAMESPACE "CustomizableObjectEditor"
 
 
+bool UCustomizableObjectNodeTableRemapPins::Equal(const UEdGraphPin& OldPin, const UEdGraphPin& NewPin) const
+{
+	if (OldPin.PinType.PinCategory == UEdGraphSchema_CustomizableObject::PC_Image ||
+		OldPin.PinType.PinCategory == UEdGraphSchema_CustomizableObject::PC_PassThroughImage)
+	{
+		// In this case pin type may have changed but we consider them the same type
+		return Helper_GetPinName(&OldPin) == Helper_GetPinName(&NewPin) &&
+			OldPin.Direction == NewPin.Direction;
+	}
+
+	return Helper_GetPinName(&OldPin) == Helper_GetPinName(&NewPin) &&
+		OldPin.PinType == NewPin.PinType &&
+		OldPin.Direction == NewPin.Direction;
+}
+
+
+void UCustomizableObjectNodeTableRemapPins::RemapPins(const TArray<UEdGraphPin*>& OldPins, const TArray<UEdGraphPin*>& NewPins, TMap<UEdGraphPin*, UEdGraphPin*>& PinsToRemap, TArray<UEdGraphPin*>& PinsToOrphan)
+{
+	for (UEdGraphPin* OldPin : OldPins)
+	{
+		bool bFound = false;
+
+		for (UEdGraphPin* NewPin : NewPins)
+		{
+			if (Equal(*OldPin, *NewPin))
+			{
+				bFound = true;
+
+				PinsToRemap.Add(OldPin, NewPin);
+				break;
+			}
+		}
+
+		if (!bFound && OldPin->LinkedTo.Num())
+		{
+			PinsToOrphan.Add(OldPin);
+		}
+	}
+}
+
+
+// Table Node -------
+
+void UCustomizableObjectNodeTable::BackwardsCompatibleFixup()
+{
+	Super::BackwardsCompatibleFixup();
+
+	const int32 CustomizableObjectCustomVersion = GetLinkerCustomVersion(FCustomizableObjectCustomVersion::GUID);
+
+	if (CustomizableObjectCustomVersion < FCustomizableObjectCustomVersion::AddedTableNodesTextureMode)
+	{
+		for (UEdGraphPin* Pin : Pins)
+		{
+			if (Pin->PinType.PinCategory == UEdGraphSchema_CustomizableObject::PC_Image)
+			{
+				UCustomizableObjectNodeTableObjectPinData* OldPinData = Cast<UCustomizableObjectNodeTableObjectPinData>(GetPinData(*(Pin)));
+				UCustomizableObjectNodeTableImagePinData* NewPinData = NewObject<UCustomizableObjectNodeTableImagePinData>(this);
+
+				if (OldPinData && NewPinData)
+				{
+					NewPinData->ColumnName = OldPinData->ColumnName;
+					NewPinData->ImageMode = ETableTextureType::MUTABLE_TEXTURE; // Old pin type by default
+				}
+
+				AddPinData(*Pin, *NewPinData);
+			}
+		}
+	}
+}
+
+
 void UCustomizableObjectNodeTable::PostBackwardsCompatibleFixup()
 {
 	Super::PostBackwardsCompatibleFixup();
@@ -92,14 +164,58 @@ FText UCustomizableObjectNodeTable::GetTooltipText() const
 }
 
 
+void UCustomizableObjectNodeTable::PinConnectionListChanged(UEdGraphPin* Pin)
+{
+	if (Pin && !Pin->LinkedTo.Num() && 
+		(Pin->PinType.PinCategory == UEdGraphSchema_CustomizableObject::PC_Image || Pin->PinType.PinCategory == UEdGraphSchema_CustomizableObject::PC_PassThroughImage))
+	{
+		if (UCustomizableObjectNodeTableImagePinData* ImagePinData = Cast<UCustomizableObjectNodeTableImagePinData>(GetPinData(*(Pin))))
+		{
+			if (ImagePinData->IsDefaultImageMode() && ImagePinData->ImageMode != DefaultImageMode)
+			{
+				ImagePinData->ImageMode = DefaultImageMode;
+				ReconstructNode();
+			}
+		}
+	}
+}
+
+
 void UCustomizableObjectNodeTable::PostEditChangeProperty(FPropertyChangedEvent& PropertyChangedEvent)
 {
 	Super::PostEditChangeProperty(PropertyChangedEvent);
 
 	FProperty* PropertyThatChanged = PropertyChangedEvent.Property;
-	if (PropertyThatChanged && PropertyThatChanged->GetName() == TEXT("Table"))
+
+	if (PropertyThatChanged)
 	{
-		MarkForReconstruct();
+		if (PropertyThatChanged->GetName() == TEXT("Table"))
+		{
+			MarkForReconstruct();
+		}
+		else if (PropertyThatChanged->GetName() == TEXT("DefaultImageMode"))
+		{
+			bool bReconstruct = false;
+
+			for (UEdGraphPin* Pin : Pins)
+			{
+				if (UCustomizableObjectNodeTableImagePinData* TexturePinData = Cast<UCustomizableObjectNodeTableImagePinData>(GetPinData(*(Pin))))
+				{
+					if (!Pin->LinkedTo.Num() && TexturePinData->IsDefaultImageMode())
+					{
+						ETableTextureType Mode = TexturePinData->ImageMode == ETableTextureType::PASSTHROUGH_TEXTURE ? ETableTextureType::MUTABLE_TEXTURE : ETableTextureType::PASSTHROUGH_TEXTURE;
+						TexturePinData->ImageMode = Mode;
+
+						bReconstruct = true;
+					}
+				}
+			}
+
+			if (bReconstruct)
+			{
+				ReconstructNode();
+			}
+		}
 	}
 }
 
@@ -129,6 +245,8 @@ void UCustomizableObjectNodeTable::AllocateDefaultPins(UCustomizableObjectNodeRe
 		return;
 	}
 
+	TArray<UEdGraphPin*> OldPins(Pins);
+
 	TableStruct->InitializeStruct(DefaultRowData);
 	
 	for (TFieldIterator<FProperty> It(TableStruct); It; ++It)
@@ -168,10 +286,24 @@ void UCustomizableObjectNodeTable::AllocateDefaultPins(UCustomizableObjectNodeRe
 
 				else if (Object->IsA(UTexture2D::StaticClass()))
 				{
-					UCustomizableObjectNodeTableObjectPinData* PinData = NewObject<UCustomizableObjectNodeTableObjectPinData>(this);
+					UCustomizableObjectNodeTableImagePinData* PinData = NewObject<UCustomizableObjectNodeTableImagePinData>(this);
 					PinData->ColumnName = ColumnName;
+					FName PinCategory = DefaultImageMode == ETableTextureType::PASSTHROUGH_TEXTURE ? Schema->PC_PassThroughImage : Schema->PC_Image;
 
-					OutPin = CustomCreatePin(EGPD_Output, Schema->PC_Image, FName(*PinName), PinData);
+					for (UEdGraphPin* Pin : OldPins)
+					{
+						// Checking if this column already exist
+						if (UCustomizableObjectNodeTableImagePinData* OldPinData = Cast<UCustomizableObjectNodeTableImagePinData>(GetPinData(*(Pin))))
+						{
+							if (OldPinData->ColumnName == ColumnName)
+							{
+								PinCategory = OldPinData->ImageMode == ETableTextureType::PASSTHROUGH_TEXTURE ? Schema->PC_PassThroughImage : Schema->PC_Image;
+								break;
+							}
+						}
+					}
+
+					OutPin = CustomCreatePin(EGPD_Output, PinCategory, FName(*PinName), PinData);
 				}
 
 				else if (Object->IsA(UMaterialInstance::StaticClass()))
@@ -236,6 +368,7 @@ void UCustomizableObjectNodeTable::AllocateDefaultPins(UCustomizableObjectNodeRe
 	FMemory::Free(DefaultRowData);
 }
 
+
 void UCustomizableObjectNodeTable::GenerateMeshPins(UObject* Mesh, FString Name)
 {
 	if (USkeletalMesh* SkeletalMesh = Cast<USkeletalMesh>(Mesh))
@@ -439,7 +572,7 @@ bool UCustomizableObjectNodeTable::IsNodeOutDatedAndNeedsRefresh()
 				{
 					FString PinName = DataTableUtils::GetPropertyExportName(ColumnProperty);
 					
-					if (CheckPinUpdated(PinName, Schema->PC_Image))
+					if (CheckPinUpdated(PinName, Schema->PC_Image) && CheckPinUpdated(PinName, Schema->PC_PassThroughImage))
 					{
 						bNeedsUpdate = true;
 					}
@@ -564,6 +697,17 @@ void UCustomizableObjectNodeTable::RemapPinsData(const TMap<UEdGraphPin*, UEdGra
 				}
 			}
 		}
+		else if (Pair.Key->PinType.PinCategory == Schema->PC_Image || Pair.Key->PinType.PinCategory == Schema->PC_PassThroughImage)
+		{
+			UCustomizableObjectNodeTableImagePinData* PinDataOldPin = Cast<UCustomizableObjectNodeTableImagePinData>(GetPinData(*(Pair.Key)));
+			UCustomizableObjectNodeTableImagePinData* PinDataNewPin = Cast<UCustomizableObjectNodeTableImagePinData>(GetPinData(*(Pair.Value)));
+
+			if (PinDataOldPin && PinDataNewPin)
+			{
+				PinDataNewPin->ImageMode = PinDataOldPin->ImageMode;
+				PinDataNewPin->SetDefaultImageMode(PinDataOldPin->IsDefaultImageMode());
+			}
+		}
 	}
 }
 
@@ -581,6 +725,7 @@ bool UCustomizableObjectNodeTable::IsPinRelevant(const UEdGraphPin* Pin) const
 	return Pin->Direction == EGPD_Input &&
 		(Pin->PinType.PinCategory == Schema->PC_MaterialAsset ||
 		Pin->PinType.PinCategory == Schema->PC_Image ||
+		Pin->PinType.PinCategory == Schema->PC_PassThroughImage ||
 		Pin->PinType.PinCategory == Schema->PC_Color ||
 		Pin->PinType.PinCategory == Schema->PC_Float ||
 		Pin->PinType.PinCategory == Schema->PC_Mesh);
@@ -981,4 +1126,51 @@ TArray<FName> UCustomizableObjectNodeTable::GetRowNames() const
 }
 
 
+void UCustomizableObjectNodeTable::ChangeImagePinMode(UEdGraphPin* Pin, bool bSetDefault)
+{
+	if (UCustomizableObjectNodeTableImagePinData* PinData = Cast<UCustomizableObjectNodeTableImagePinData>(GetPinData(*(Pin))))
+	{
+		if (PinData->IsDefaultImageMode() && !bSetDefault)
+		{
+			ETableTextureType Mode = DefaultImageMode == ETableTextureType::PASSTHROUGH_TEXTURE ? ETableTextureType::MUTABLE_TEXTURE : ETableTextureType::PASSTHROUGH_TEXTURE;
+			PinData->ImageMode = Mode;
+			PinData->SetDefaultImageMode(false);
+
+			ReconstructNode();
+		}
+		else
+		{
+			if (bSetDefault)
+			{
+				PinData->ImageMode = DefaultImageMode;
+				PinData->SetDefaultImageMode(true);
+			}
+			else
+			{
+				ETableTextureType Mode = PinData->ImageMode == ETableTextureType::PASSTHROUGH_TEXTURE ? ETableTextureType::MUTABLE_TEXTURE : ETableTextureType::PASSTHROUGH_TEXTURE;
+				PinData->ImageMode = Mode;
+			}
+
+			ReconstructNode();
+		}
+	}
+}
+
+
+UCustomizableObjectNodeTableRemapPins* UCustomizableObjectNodeTable::CreateRemapPinsDefault() const
+{
+	return NewObject<UCustomizableObjectNodeTableRemapPins>();
+}
+
+
+bool UCustomizableObjectNodeTable::IsImagePinDefault(UEdGraphPin* Pin)
+{
+	if (UCustomizableObjectNodeTableImagePinData* PinData = Cast<UCustomizableObjectNodeTableImagePinData>(GetPinData(*(Pin))))
+	{
+		return PinData->IsDefaultImageMode();
+	}
+
+	return true;
+}
+
 #undef LOCTEXT_NAMESPACE
