[Mutable] Fix copied AssetUserData preventing baked meshes from being saved.

#rnx

diff --git a/Engine/Plugins/Experimental/Mutable/Source/CustomizableObjectEditor/Private/MuCOE/CustomizableObjectEditorViewportClient.cpp b/Engine/Plugins/Experimental/Mutable/Source/CustomizableObjectEditor/Private/MuCOE/CustomizableObjectEditorViewportClient.cpp
--- a/Engine/Plugins/Experimental/Mutable/Source/CustomizableObjectEditor/Private/MuCOE/CustomizableObjectEditorViewportClient.cpp
+++ b/Engine/Plugins/Experimental/Mutable/Source/CustomizableObjectEditor/Private/MuCOE/CustomizableObjectEditorViewportClient.cpp
@@ -2708,6 +2708,20 @@ void FCustomizableObjectEditorViewportClient::BakeInstance(UCustomizableObjectIn
 
 					SkeletalMesh->GetImportedModel()->SkeletalMeshModelGUID = FGuid::NewGuid();
 
+					// Duplicate AssetUserData
+					{
+						const TArray<UAssetUserData*>* AssetUserDataArray = Mesh->GetAssetUserDataArray();
+						for (const UAssetUserData* AssetUserData : *AssetUserDataArray)
+						{
+							if (AssetUserData)
+							{
+								// Duplicate to change ownership
+								UAssetUserData* NewAssetUserData = Cast<UAssetUserData>(StaticDuplicateObject(AssetUserData, SkeletalMesh));
+								SkeletalMesh->AddAssetUserData(NewAssetUserData);
+							}
+						}
+					}
+
 					// Generate render data
 					SkeletalMesh->Build();
 				}
