OnlyEditorOnly support: Mark UCustomizableObject's resources as EditorOnly by default.
#rnx

diff --git a/Engine/Plugins/Experimental/Mutable/Source/CustomizableObjectEditor/Private/MuCOE/CustomizableObjectCompiler.cpp b/Engine/Plugins/Experimental/Mutable/Source/CustomizableObjectEditor/Private/MuCOE/CustomizableObjectCompiler.cpp
--- a/Engine/Plugins/Experimental/Mutable/Source/CustomizableObjectEditor/Private/MuCOE/CustomizableObjectCompiler.cpp
+++ b/Engine/Plugins/Experimental/Mutable/Source/CustomizableObjectEditor/Private/MuCOE/CustomizableObjectCompiler.cpp
@@ -18,6 +18,7 @@
 #include "MuCOE/GraphTraversal.h"
 #include "MuCOE/ICustomizableObjectPopulationModule.h"
 #include "MuCOE/Nodes/CustomizableObjectNodeObjectGroup.h"
+#include "UObject/ICookInfo.h"
 #include "UObject/UObjectIterator.h"
 #include "Widgets/Notifications/SNotificationList.h"
 
@@ -247,6 +248,9 @@ void FCustomizableObjectCompiler::Compile(UCustomizableObject& Object, const FCo
 
 	bool bAssetsLoaded = true;
 
+	// Customizations are marked as editoronly on load and are not packaged into the runtime game by default.
+ 	// The ones that need to be kept will be copied into SoftObjectPath on the object during save.
+	FCookLoadScope CookLoadScope(ECookLoadType::EditorOnly);
 	if (ArrayAssetToStream.Num() > 0)
 	{
 		FStreamableManager& Streamable = System->GetStreamableManager();
