#rnx

diff --git a/Engine/Plugins/Experimental/Mutable/Source/CustomizableObject/Classes/MuCO/CustomizableObjectInstance.h b/Engine/Plugins/Experimental/Mutable/Source/CustomizableObject/Classes/MuCO/CustomizableObjectInstance.h
--- a/Engine/Plugins/Experimental/Mutable/Source/CustomizableObject/Classes/MuCO/CustomizableObjectInstance.h
+++ b/Engine/Plugins/Experimental/Mutable/Source/CustomizableObject/Classes/MuCO/CustomizableObjectInstance.h
@@ -110,9 +110,9 @@ DECLARE_MULTICAST_DELEGATE_OneParam(FBeginDestroyNativeDelegate, UCustomizableOb
 
 DECLARE_DELEGATE_OneParam(FProjectorStateChangedDelegate, FString);
 
-DECLARE_DYNAMIC_DELEGATE_TwoParams(FEachComponentAnimInstanceClassDelegate, int32, SlotIndex, TSubclassOf<UAnimInstance>, AnimInstClass);
+DECLARE_DYNAMIC_DELEGATE_TwoParams(FEachComponentAnimInstanceClassDelegate, FName, SlotIndex, TSubclassOf<UAnimInstance>, AnimInstClass);
 
-DECLARE_DELEGATE_TwoParams(FEachComponentAnimInstanceClassNativeDelegate, int32 /*SlotIndex*/, TSubclassOf<UAnimInstance> /*AnimInstClass*/);
+DECLARE_DELEGATE_TwoParams(FEachComponentAnimInstanceClassNativeDelegate, FName /*SlotIndex*/, TSubclassOf<UAnimInstance> /*AnimInstClass*/);
 
 
 UCLASS( Blueprintable, BlueprintType, HideCategories=(CustomizableObjectInstance) )
@@ -563,7 +563,7 @@ public:
 	
 	// Returns the animation BP for the parameter component and slot, gathered from all the meshes that compose this instance
 	UFUNCTION(BlueprintCallable, Category = CustomizableObjectInstance)
-	TSubclassOf<UAnimInstance> GetAnimBP(int32 ComponentIndex, int32 SlotIndex) const;
+	TSubclassOf<UAnimInstance> GetAnimBP(int32 ComponentIndex, const FName& Slot) const;
 
 	UFUNCTION(BlueprintCallable, Category = CustomizableObjectInstance)
 	const FGameplayTagContainer& GetAnimationGameplayTags() const;
diff --git a/Engine/Plugins/Experimental/Mutable/Source/CustomizableObject/Private/MuCO/CustomizableInstancePrivateData.h b/Engine/Plugins/Experimental/Mutable/Source/CustomizableObject/Private/MuCO/CustomizableInstancePrivateData.h
--- a/Engine/Plugins/Experimental/Mutable/Source/CustomizableObject/Private/MuCO/CustomizableInstancePrivateData.h
+++ b/Engine/Plugins/Experimental/Mutable/Source/CustomizableObject/Private/MuCO/CustomizableInstancePrivateData.h
@@ -125,7 +125,7 @@ struct FCustomizableInstanceComponentData
 
 	// AnimBP data gathered for a component from its constituent meshes
 	UPROPERTY(Transient, Category = CustomizableObjectInstance, editfixedsize, VisibleAnywhere)
-	TMap<int32, TSoftClassPtr<UAnimInstance>> AnimSlotToBP;
+	TMap<FName, TSoftClassPtr<UAnimInstance>> AnimSlotToBP;
 
 #if WITH_EDITORONLY_DATA
 	// Just used for mutable.EnableMutableAnimInfoDebugging command
diff --git a/Engine/Plugins/Experimental/Mutable/Source/CustomizableObject/Private/MuCO/CustomizableObjectInstance.cpp b/Engine/Plugins/Experimental/Mutable/Source/CustomizableObject/Private/MuCO/CustomizableObjectInstance.cpp
--- a/Engine/Plugins/Experimental/Mutable/Source/CustomizableObject/Private/MuCO/CustomizableObjectInstance.cpp
+++ b/Engine/Plugins/Experimental/Mutable/Source/CustomizableObject/Private/MuCO/CustomizableObjectInstance.cpp
@@ -5252,9 +5252,9 @@ FGraphEventRef UCustomizableInstancePrivateData::LoadAdditionalAssetsAsync(const
 
 					if (Tag.Split(TEXT("_Slot_"), &SlotIndexString, &AssetPath))
 					{
-						if (SlotIndexString.IsNumeric())
+						if (!SlotIndexString.IsEmpty())
 						{
-							int32 SlotIndex = FCString::Atoi(*SlotIndexString);
+							FName SlotIndex = *SlotIndexString;
 
 							TSoftClassPtr<UAnimInstance>* AnimBPAsset = CustomizableObject->AnimBPAssetsMap.Find(AssetPath);
 
@@ -5276,7 +5276,7 @@ FGraphEventRef UCustomizableInstancePrivateData::LoadAdditionalAssetsAsync(const
 								else
 								{
 									// Two submeshes should not have the same animation slot index
-									FString ErrorMsg = FString::Printf(TEXT("Two submeshes have the same anim slot index [%d] in a Mutable Instance."), SlotIndex);
+									FString ErrorMsg = FString::Printf(TEXT("Two submeshes have the same anim slot index [%s] in a Mutable Instance."), *SlotIndex.ToString());
 									UE_LOG(LogMutable, Error, TEXT("%s"), *ErrorMsg);
 #if WITH_EDITOR
 									FMessageLogModule& MessageLogModule = FModuleManager::LoadModuleChecked<FMessageLogModule>("MessageLog");
@@ -5430,7 +5430,7 @@ void UCustomizableInstancePrivateData::AdditionalAssetsAsyncLoaded(UCustomizable
 		ComponentData.ClothingPhysicsAssetsToStream.Empty();
 
 		// Loaded anim BPs
-		for (TPair<int32, TSoftClassPtr<UAnimInstance>>& SlotAnimBP : ComponentData.AnimSlotToBP)
+		for (TPair<FName, TSoftClassPtr<UAnimInstance>>& SlotAnimBP : ComponentData.AnimSlotToBP)
 		{
 			if (TSubclassOf<UAnimInstance> AnimBP = SlotAnimBP.Value.Get())
 			{
@@ -6514,7 +6514,7 @@ bool UCustomizableObjectInstance::HasAnyParameters() const
 }
 
 
-TSubclassOf<UAnimInstance> UCustomizableObjectInstance::GetAnimBP(int32 ComponentIndex, int32 SlotIndex) const
+TSubclassOf<UAnimInstance> UCustomizableObjectInstance::GetAnimBP(int32 ComponentIndex, const FName& SlotName) const
 {
 	FCustomizableInstanceComponentData* ComponentData =	GetPrivate()->GetComponentData(ComponentIndex);
 	
@@ -6533,7 +6533,7 @@ TSubclassOf<UAnimInstance> UCustomizableObjectInstance::GetAnimBP(int32 Componen
 		return nullptr;
 	}
 
-	TSoftClassPtr<UAnimInstance>* Result = ComponentData->AnimSlotToBP.Find(SlotIndex);
+	TSoftClassPtr<UAnimInstance>* Result = ComponentData->AnimSlotToBP.Find(SlotName);
 
 	return Result ? Result->Get() : nullptr;
 }
@@ -6584,9 +6584,9 @@ void UCustomizableObjectInstance::ForEachAnimInstance(int32 ComponentIndex, FEac
 		return;
 	}
 
-	for (const TPair<int32, TSoftClassPtr<UAnimInstance>>& MapElem : ComponentData->AnimSlotToBP)
+	for (const TPair<FName, TSoftClassPtr<UAnimInstance>>& MapElem : ComponentData->AnimSlotToBP)
 	{
-		const int32 Index = MapElem.Key;
+		const FName& Index = MapElem.Key;
 		const TSoftClassPtr<UAnimInstance>& AnimBP = MapElem.Value;
 
 		// if this _can_ resolve to a real AnimBP
@@ -6731,9 +6731,9 @@ void UCustomizableObjectInstance::ForEachAnimInstance(int32 ComponentIndex, FEac
 		return;
 	}
 
-	for (const TPair<int32, TSoftClassPtr<UAnimInstance>>& MapElem : ComponentData->AnimSlotToBP)
+	for (const TPair<FName, TSoftClassPtr<UAnimInstance>>& MapElem : ComponentData->AnimSlotToBP)
 	{
-		const int32 Index = MapElem.Key;
+		const FName& Index = MapElem.Key;
 		const TSoftClassPtr<UAnimInstance>& AnimBP = MapElem.Value;
 
 		// if this _can_ resolve to a real AnimBP
diff --git a/Engine/Plugins/Experimental/Mutable/Source/CustomizableObject/Private/MuCO/CustomizableObjectSystem.cpp b/Engine/Plugins/Experimental/Mutable/Source/CustomizableObject/Private/MuCO/CustomizableObjectSystem.cpp
--- a/Engine/Plugins/Experimental/Mutable/Source/CustomizableObject/Private/MuCO/CustomizableObjectSystem.cpp
+++ b/Engine/Plugins/Experimental/Mutable/Source/CustomizableObject/Private/MuCO/CustomizableObjectSystem.cpp
@@ -641,11 +641,11 @@ void FCustomizableObjectSystemPrivate::UpdateStats()
 
 						if (ComponentData)
 						{
-							for (TPair<int32, TSoftClassPtr<UAnimInstance>>& Entry : ComponentData->AnimSlotToBP)
+							for (TPair<FName, TSoftClassPtr<UAnimInstance>>& Entry : ComponentData->AnimSlotToBP)
 							{
 								FString AnimBPSlot;
 
-								AnimBPSlot += FString::Printf(TEXT("%d"), Entry.Key) + FString("-") + Entry.Value.GetAssetName();
+								AnimBPSlot += Entry.Key.ToString() + FString("-") + Entry.Value.GetAssetName();
 								GEngine->AddOnScreenDebugMessage(MsgIndex++, .0f, FColor::Green, AnimBPSlot);
 							}
 						}
diff --git a/Engine/Plugins/Experimental/Mutable/Source/CustomizableObject/Public/MuCO/CustomizableObjectCustomVersion.h b/Engine/Plugins/Experimental/Mutable/Source/CustomizableObject/Public/MuCO/CustomizableObjectCustomVersion.h
--- a/Engine/Plugins/Experimental/Mutable/Source/CustomizableObject/Public/MuCO/CustomizableObjectCustomVersion.h
+++ b/Engine/Plugins/Experimental/Mutable/Source/CustomizableObject/Public/MuCO/CustomizableObjectCustomVersion.h
@@ -69,6 +69,8 @@ struct CUSTOMIZABLEOBJECT_API FCustomizableObjectCustomVersion
 		StateTextureCompressionStrategyEnum,
 
 		AddedTableNodesTextureMode,
+		
+		ConvertAnimationSlotToFName,
 
 		// -----<new versions can be added above this line>-------------------------------------------------
 		VersionPlusOne,
diff --git a/Engine/Plugins/Experimental/Mutable/Source/CustomizableObjectEditor/Classes/MuCOE/Nodes/CustomizableObjectNodeSkeletalMesh.h b/Engine/Plugins/Experimental/Mutable/Source/CustomizableObjectEditor/Classes/MuCOE/Nodes/CustomizableObjectNodeSkeletalMesh.h
--- a/Engine/Plugins/Experimental/Mutable/Source/CustomizableObjectEditor/Classes/MuCOE/Nodes/CustomizableObjectNodeSkeletalMesh.h
+++ b/Engine/Plugins/Experimental/Mutable/Source/CustomizableObjectEditor/Classes/MuCOE/Nodes/CustomizableObjectNodeSkeletalMesh.h
@@ -148,9 +148,12 @@ public:
 	UPROPERTY(EditAnywhere, Category = CustomizableObject)
 	TSoftClassPtr<UAnimInstance> AnimInstance;
 
+	UPROPERTY()
+	int32 AnimBlueprintSlot_DEPRECATED;
+	
 	/** The anim slot associated with the AnimInstance */
 	UPROPERTY(EditAnywhere, Category = CustomizableObject)
-	int32 AnimBlueprintSlot = 0;
+	FName AnimBlueprintSlotName;
 
 	/** Animation tags that will be gathered by a Generated instance if it contains this skeletal mesh part,
 		it will not be grouped by component or AnimBlueprintSlot */
diff --git a/Engine/Plugins/Experimental/Mutable/Source/CustomizableObjectEditor/Private/MuCOE/CompilationMessageCache.cpp b/Engine/Plugins/Experimental/Mutable/Source/CustomizableObjectEditor/Private/MuCOE/CompilationMessageCache.cpp
--- a/Engine/Plugins/Experimental/Mutable/Source/CustomizableObjectEditor/Private/MuCOE/CompilationMessageCache.cpp
+++ b/Engine/Plugins/Experimental/Mutable/Source/CustomizableObjectEditor/Private/MuCOE/CompilationMessageCache.cpp
@@ -4,13 +4,13 @@
 
 
 
-void FCompilationMessageCache::AddMessage(const FText& InMessage, const TArray<const UCustomizableObjectNode*>& InArrayNode, EMessageSeverity::Type MessageSeverity /* = EMessageSeverity::Warning */)
+bool FCompilationMessageCache::AddMessage(const FText& InMessage, const TArray<const UCustomizableObjectNode*>& InArrayNode, EMessageSeverity::Type MessageSeverity /* = EMessageSeverity::Warning */)
 {
 	// Skip message if an identical one has already been reported
 	const FLoggedMessage Cached = { InMessage, InArrayNode, MessageSeverity };
 	if (LoggedMessages.Contains(Cached))
 	{
-		return;
+		return false;
 	}
 
 	LoggedMessages.Add(Cached);
@@ -33,6 +33,8 @@ void FCompilationMessageCache::AddMessage(const FText& InMessage, const TArray<c
 		checkNoEntry();
 		break;
 	}
+
+	return true;
 }
 
 
diff --git a/Engine/Plugins/Experimental/Mutable/Source/CustomizableObjectEditor/Private/MuCOE/CustomizableObjectCompiler.cpp b/Engine/Plugins/Experimental/Mutable/Source/CustomizableObjectEditor/Private/MuCOE/CustomizableObjectCompiler.cpp
--- a/Engine/Plugins/Experimental/Mutable/Source/CustomizableObjectEditor/Private/MuCOE/CustomizableObjectCompiler.cpp
+++ b/Engine/Plugins/Experimental/Mutable/Source/CustomizableObjectEditor/Private/MuCOE/CustomizableObjectCompiler.cpp
@@ -1551,14 +1551,14 @@ void FCustomizableObjectCompiler::NotifyCompilationErrors() const
 
 void FCustomizableObjectCompiler::CompilerLog(const FText& Message, const TArray<const UCustomizableObjectNode*>& ArrayNode, EMessageSeverity::Type MessageSeverity, bool bAddBaseObjectInfo)
 {
-	// Cache the message for later reference
-	CompilationLogsContainer.AddMessage(Message,ArrayNode,MessageSeverity);
-	
-	FCustomizableObjectEditorLogger::CreateLog(Message)
-	.Severity(MessageSeverity)
-	.Nodes(ArrayNode)
-	.BaseObject(bAddBaseObjectInfo)
-	.Log();
+	if (CompilationLogsContainer.AddMessage(Message, ArrayNode, MessageSeverity)) // Cache the message for later reference
+	{
+		FCustomizableObjectEditorLogger::CreateLog(Message)
+			.Severity(MessageSeverity)
+			.Nodes(ArrayNode)
+			.BaseObject(bAddBaseObjectInfo)
+			.Log();
+	}
 }
 
 
diff --git a/Engine/Plugins/Experimental/Mutable/Source/CustomizableObjectEditor/Private/MuCOE/GenerateMutableSource/GenerateMutableSource.cpp b/Engine/Plugins/Experimental/Mutable/Source/CustomizableObjectEditor/Private/MuCOE/GenerateMutableSource/GenerateMutableSource.cpp
--- a/Engine/Plugins/Experimental/Mutable/Source/CustomizableObjectEditor/Private/MuCOE/GenerateMutableSource/GenerateMutableSource.cpp
+++ b/Engine/Plugins/Experimental/Mutable/Source/CustomizableObjectEditor/Private/MuCOE/GenerateMutableSource/GenerateMutableSource.cpp
@@ -1283,9 +1283,9 @@ int32 AddTagToMutableMeshUnique(mu::Mesh& MutableMesh, const FString& Tag)
 	return TagCount;
 }
 
-FString GenerateAnimationInstanceTag(const FString& AnimInstance, int32 SlotIndex)
+FString GenerateAnimationInstanceTag(const FString& AnimInstance, const FName& SlotIndex)
 {
-	return FString("__AnimBP:") + FString::Printf(TEXT("%d_Slot_"), SlotIndex) + AnimInstance;
+	return FString("__AnimBP:") + FString::Printf(TEXT("%s_Slot_"), *SlotIndex.ToString()) + AnimInstance;
 }
 
 
diff --git a/Engine/Plugins/Experimental/Mutable/Source/CustomizableObjectEditor/Private/MuCOE/GenerateMutableSource/GenerateMutableSource.h b/Engine/Plugins/Experimental/Mutable/Source/CustomizableObjectEditor/Private/MuCOE/GenerateMutableSource/GenerateMutableSource.h
--- a/Engine/Plugins/Experimental/Mutable/Source/CustomizableObjectEditor/Private/MuCOE/GenerateMutableSource/GenerateMutableSource.h
+++ b/Engine/Plugins/Experimental/Mutable/Source/CustomizableObjectEditor/Private/MuCOE/GenerateMutableSource/GenerateMutableSource.h
@@ -762,7 +762,7 @@ int32 AddTagToMutableMeshUnique(mu::Mesh& MutableMesh, const FString& Tag);
 void AddSocketTagsToMesh(const USkeletalMesh* SourceMesh, mu::MeshPtr MutableMesh, FMutableGraphGenerationContext& GenerationContext);
 
 // Generates the tag for an animation instance
-FString GenerateAnimationInstanceTag(const FString& AnimInstance, int32 SlotIndex);
+FString GenerateAnimationInstanceTag(const FString& AnimInstance, const FName& SlotIndex);
 
 
 FString GenerateGameplayTag(const FString& GameplayTag);
diff --git a/Engine/Plugins/Experimental/Mutable/Source/CustomizableObjectEditor/Private/MuCOE/GenerateMutableSource/GenerateMutableSourceMesh.cpp b/Engine/Plugins/Experimental/Mutable/Source/CustomizableObjectEditor/Private/MuCOE/GenerateMutableSource/GenerateMutableSourceMesh.cpp
--- a/Engine/Plugins/Experimental/Mutable/Source/CustomizableObjectEditor/Private/MuCOE/GenerateMutableSource/GenerateMutableSourceMesh.cpp
+++ b/Engine/Plugins/Experimental/Mutable/Source/CustomizableObjectEditor/Private/MuCOE/GenerateMutableSource/GenerateMutableSourceMesh.cpp
@@ -2799,7 +2799,7 @@ mu::NodeMeshPtr GenerateMutableSourceMesh(const UEdGraphPin * Pin,
 
 			if (!TypedNodeSkel->AnimInstance.IsNull())
 			{
-				int32 SlotIndex = TypedNodeSkel->AnimBlueprintSlot;
+				FName SlotIndex = TypedNodeSkel->AnimBlueprintSlotName;
 				GenerationContext.AnimBPAssetsMap.Add(TypedNodeSkel->AnimInstance.ToString(), TypedNodeSkel->AnimInstance);
 
 				AnimBPAssetTag = GenerateAnimationInstanceTag(TypedNodeSkel->AnimInstance.ToString(), SlotIndex);
diff --git a/Engine/Plugins/Experimental/Mutable/Source/CustomizableObjectEditor/Private/MuCOE/GenerateMutableSource/GenerateMutableSourceTable.cpp b/Engine/Plugins/Experimental/Mutable/Source/CustomizableObjectEditor/Private/MuCOE/GenerateMutableSource/GenerateMutableSourceTable.cpp
--- a/Engine/Plugins/Experimental/Mutable/Source/CustomizableObjectEditor/Private/MuCOE/GenerateMutableSource/GenerateMutableSourceTable.cpp
+++ b/Engine/Plugins/Experimental/Mutable/Source/CustomizableObjectEditor/Private/MuCOE/GenerateMutableSource/GenerateMutableSourceTable.cpp
@@ -52,7 +52,7 @@ void FillTableColumn(const UCustomizableObjectNodeTable* TableNode,	mu::TablePtr
 
 						if (AnimRowData)
 						{
-							int32 SlotIndex = -1;
+							FName SlotIndex;
 
 							// Getting animation slot row value from data table
 							if (FProperty* AnimSlotProperty = TableNode->Table->FindTableProperty(FName(*AnimSlot)))
@@ -63,12 +63,21 @@ void FillTableColumn(const UCustomizableObjectNodeTable* TableNode,	mu::TablePtr
 								{
 									if (const FIntProperty* IntProperty = CastField<FIntProperty>(AnimSlotProperty))
 									{
-										SlotIndex = IntProperty->GetPropertyValue(AnimSlotData);
+										FString Message = FString::Printf(
+											TEXT("The column with name [%s] for the Anim Slot property should be an FName instead of an Integer, it will be internally converted to FName but should probaly be converted in the table itself."), 
+											*AnimBP);
+										GenerationContext.Compiler->CompilerLog(FText::FromString(Message), TableNode, EMessageSeverity::Info);
+
+										SlotIndex = FName(FString::FromInt(IntProperty->GetPropertyValue(AnimSlotData)));
+									}
+									else if (const FNameProperty* NameProperty = CastField<FNameProperty>(AnimSlotProperty))
+									{
+										SlotIndex = NameProperty->GetPropertyValue(AnimSlotData);
 									}
 								}
 							}
 
-							if (SlotIndex > -1)
+							if (SlotIndex.GetStringLength() != 0)
 							{
 								// Getting animation instance soft class from data table
 								if (FProperty* AnimBPProperty = TableNode->Table->FindTableProperty(FName(*AnimBP)))
@@ -93,7 +102,7 @@ void FillTableColumn(const UCustomizableObjectNodeTable* TableNode,	mu::TablePtr
 							}
 							else
 							{
-								FString msg = FString::Printf(TEXT("Could not found the Slot column of the animation blueprint column [%s] for the mesh column [%s] row [%s]."), *AnimBP, *ColumnName, *RowName);
+								FString msg = FString::Printf(TEXT("Could not find the Slot column of the animation blueprint column [%s] for the mesh column [%s] row [%s]."), *AnimBP, *ColumnName, *RowName);
 								GenerationContext.Compiler->CompilerLog(FText::FromString(msg), TableNode);
 							}
 						}
diff --git a/Engine/Plugins/Experimental/Mutable/Source/CustomizableObjectEditor/Private/MuCOE/Nodes/CustomizableObjectNodeSkeletalMesh.cpp b/Engine/Plugins/Experimental/Mutable/Source/CustomizableObjectEditor/Private/MuCOE/Nodes/CustomizableObjectNodeSkeletalMesh.cpp
--- a/Engine/Plugins/Experimental/Mutable/Source/CustomizableObjectEditor/Private/MuCOE/Nodes/CustomizableObjectNodeSkeletalMesh.cpp
+++ b/Engine/Plugins/Experimental/Mutable/Source/CustomizableObjectEditor/Private/MuCOE/Nodes/CustomizableObjectNodeSkeletalMesh.cpp
@@ -607,6 +607,15 @@ void UCustomizableObjectNodeSkeletalMesh::BackwardsCompatibleFixup()
 	{
 		ReconstructNode();
 	}
+
+	if (CustomizableObjectCustomVersion < FCustomizableObjectCustomVersion::ConvertAnimationSlotToFName)
+	{
+		if (AnimBlueprintSlotName.IsNone() && AnimBlueprintSlot_DEPRECATED != -1)
+		{
+			AnimBlueprintSlotName = FName(FString::FromInt(AnimBlueprintSlot_DEPRECATED));
+			AnimBlueprintSlot_DEPRECATED = -1; // Unnecessary, just in case anyone tried to use it later in this method.
+		}
+	}
 }
 
 
diff --git a/Engine/Plugins/Experimental/Mutable/Source/CustomizableObjectEditor/Public/MuCOE/CompilationMessageCache.h b/Engine/Plugins/Experimental/Mutable/Source/CustomizableObjectEditor/Public/MuCOE/CompilationMessageCache.h
--- a/Engine/Plugins/Experimental/Mutable/Source/CustomizableObjectEditor/Public/MuCOE/CompilationMessageCache.h
+++ b/Engine/Plugins/Experimental/Mutable/Source/CustomizableObjectEditor/Public/MuCOE/CompilationMessageCache.h
@@ -16,8 +16,9 @@ public:
 	 * @param InMessage The message's text.
 	 * @param InArrayNode The node that is related with the InMessage provided.
 	 * @param MessageSeverity Severity of the message.
+	 * @return True if the message is new and had not been added before. If it returns false it means it has not been logged again.
 	 */
-	 void AddMessage(const FText& InMessage, const TArray<const UCustomizableObjectNode*>& InArrayNode, EMessageSeverity::Type MessageSeverity = EMessageSeverity::Warning);
+	 bool AddMessage(const FText& InMessage, const TArray<const UCustomizableObjectNode*>& InArrayNode, EMessageSeverity::Type MessageSeverity = EMessageSeverity::Warning);
 	
 	
 	/** Provides the caller with separated arrays with all the warning and error messages.
