[mutable] Remove parameter decorations. Restore parameter relevancy calculations.
#rnx

diff --git a/Engine/Plugins/Experimental/Mutable/Source/CustomizableObject/Classes/MuCO/CustomizableObject.h b/Engine/Plugins/Experimental/Mutable/Source/CustomizableObject/Classes/MuCO/CustomizableObject.h
--- a/Engine/Plugins/Experimental/Mutable/Source/CustomizableObject/Classes/MuCO/CustomizableObject.h
+++ b/Engine/Plugins/Experimental/Mutable/Source/CustomizableObject/Classes/MuCO/CustomizableObject.h
@@ -383,9 +383,6 @@ struct FMutableModelParameterProperties
 	UPROPERTY()
 	EMutableParameterType Type = EMutableParameterType::None;
 
-	UPROPERTY()
-	int ImageDescriptionCount = 0;
-
 	UPROPERTY()
 	TArray<FMutableModelParameterValue> PossibleValues;
 };
@@ -1263,10 +1260,7 @@ public:
 	const FString& GetParameterName(int32 ParamIndex) const;
 
 	// Get the number of description images available for a parameter
-	int32 GetParameterDescriptionCount(int32 ParamIndex) const;
-
-	// Get the number of description images available for a parameter
-	UFUNCTION(BlueprintCallable, Category = CustomizableObject)
+	UFUNCTION(BlueprintCallable, Category = CustomizableObject, meta = (DeprecatedFunction, DeprecationMessage = "Parameter decorations have been removed. This method will be removed in future versions."))
 	int32 GetParameterDescriptionCount(const FString& ParamName) const;
 
 	// Returns how many possible options an int parameter has
@@ -1416,7 +1410,7 @@ private:
 	// This is a manual version number for the binary blobs in this asset.
 	// Increasing it invalidates all the previously compiled models.
 	// Warning: If while merging code both versions have changed, take the highest+1.
-	static const int32 CurrentSupportedVersion = 381;
+	static const int32 CurrentSupportedVersion = 382;
 
 public:
 
diff --git a/Engine/Plugins/Experimental/Mutable/Source/CustomizableObject/Classes/MuCO/CustomizableObjectInstance.h b/Engine/Plugins/Experimental/Mutable/Source/CustomizableObject/Classes/MuCO/CustomizableObjectInstance.h
--- a/Engine/Plugins/Experimental/Mutable/Source/CustomizableObject/Classes/MuCO/CustomizableObjectInstance.h
+++ b/Engine/Plugins/Experimental/Mutable/Source/CustomizableObject/Classes/MuCO/CustomizableObjectInstance.h
@@ -194,12 +194,18 @@ public:
 	UFUNCTION(BlueprintCallable, Category = CustomizableObjectInstance)
 	UCustomizableObject* GetCustomizableObject() const;
 
-	UFUNCTION(BlueprintCallable, Category = CustomizableObjectInstance)
+	UFUNCTION(BlueprintCallable, Category = CustomizableObjectInstance, meta = (DeprecatedFunction, DeprecationMessage = "Parameter decorations have been removed. This method will be removed in future versions."))
 	bool GetBuildParameterDecorations() const;
 
-	UFUNCTION(BlueprintCallable, Category = CustomizableObjectInstance)
+	UFUNCTION(BlueprintCallable, Category = CustomizableObjectInstance, meta = (DeprecatedFunction, DeprecationMessage = "Parameter decorations have been removed. This method will be removed in future versions."))
 	void SetBuildParameterDecorations(bool Value);
 
+	UFUNCTION(BlueprintCallable, Category = CustomizableObjectInstance)
+	bool GetBuildParameterRelevancy() const;
+
+	UFUNCTION(BlueprintCallable, Category = CustomizableObjectInstance)
+	void SetBuildParameterRelevancy(bool Value);
+
 	int32 GetState() const;
 	void SetState(int32 InState);
 
@@ -298,20 +304,15 @@ public:
 	// Releases all the mutable resources this instance holds, should only be called when it is not going to be used any more.
 	void ReleaseMutableResources(bool bCalledFromBeginDestroy);
 
-	// Get a description image for a parameter.
-	// The param index is in the range of o to CustomizableObject->GetParameterCount.
-	// This will only be valid if bBuildParameterDecorations was set to true before the last update.
-	UTexture2D* GetParameterDescription(int32 ParamIndex, int32 DescIndex);
-
 	// Returns de description texture (ex: color bar) for this parameter and DescIndex
 	// This will only be valid if bBuildParameterDecorations was set to true before the last update.
-	UFUNCTION(BlueprintCallable, Category = CustomizableObjectInstance)
+	UFUNCTION(BlueprintCallable, Category = CustomizableObjectInstance, meta = (DeprecatedFunction, DeprecationMessage = "Parameter decorations have been removed. This method will be removed in future versions."))
 	UTexture2D* GetParameterDescription(const FString& ParamName, int32 DescIndex);
 
-	//! This is only valid if bBuildParameterDecorations has been set before the last update.
+	//! 
 	bool IsParameterRelevant( int32 ParameterIndex ) const;
 
-	//! This is only valid if bBuildParameterDecorations has been set before the last update.
+	//! 
 	UFUNCTION(BlueprintCallable, Category = CustomizableObjectInstance)
 	bool IsParameterRelevant(const FString& ParamName) const;
 
@@ -753,8 +754,8 @@ private:
    	UPROPERTY()
    	TMap<FName, FMultilayerProjector> MultilayerProjectors_DEPRECATED;
 
-	UPROPERTY()
-	bool bBuildParameterDecorations_DEPRECATED;
+	/** If this is set to true, when updating the instance an additional step will be performed to calculate the list of instance parameters that are relevant for the current parameter vaules. */
+	bool bBuildParameterRelevancy = false;
 };
 
 #if WITH_EDITOR
diff --git a/Engine/Plugins/Experimental/Mutable/Source/CustomizableObject/Classes/MuCO/CustomizableObjectSystem.h b/Engine/Plugins/Experimental/Mutable/Source/CustomizableObject/Classes/MuCO/CustomizableObjectSystem.h
--- a/Engine/Plugins/Experimental/Mutable/Source/CustomizableObject/Classes/MuCO/CustomizableObjectSystem.h
+++ b/Engine/Plugins/Experimental/Mutable/Source/CustomizableObject/Classes/MuCO/CustomizableObjectSystem.h
@@ -52,7 +52,6 @@ DECLARE_DWORD_ACCUMULATOR_STAT_EXTERN(TEXT("Num Allocated Mutable Textures"), ST
 DECLARE_DWORD_ACCUMULATOR_STAT_EXTERN(TEXT("Texture Resource Memory"), STAT_MutableTextureResourceMemory, STATGROUP_Mutable, );
 DECLARE_DWORD_ACCUMULATOR_STAT_EXTERN(TEXT("Texture Generated Memory"), STAT_MutableTextureGeneratedMemory, STATGROUP_Mutable, );
 DECLARE_DWORD_ACCUMULATOR_STAT_EXTERN(TEXT("Texture Locked Memory"), STAT_MutableTextureCacheMemory, STATGROUP_Mutable, );
-DECLARE_DWORD_ACCUMULATOR_STAT_EXTERN(TEXT("Texture Parameter Decoration Memory"), STAT_MutableTextureParameterDecorationMemory, STATGROUP_Mutable, );
 DECLARE_DWORD_ACCUMULATOR_STAT_EXTERN(TEXT("Pending Instance Updates"), STAT_MutablePendingInstanceUpdates, STATGROUP_Mutable, );
 DECLARE_DWORD_ACCUMULATOR_STAT_EXTERN(TEXT("Abandoned Instance Updates"), STAT_MutableAbandonedInstanceUpdates, STATGROUP_Mutable, );
 DECLARE_DWORD_ACCUMULATOR_STAT_EXTERN(TEXT("Last Instance Build Time"), STAT_MutableInstanceBuildTime, STATGROUP_Mutable, );
diff --git a/Engine/Plugins/Experimental/Mutable/Source/CustomizableObject/Private/MuCO/CustomizableInstancePrivateData.h b/Engine/Plugins/Experimental/Mutable/Source/CustomizableObject/Private/MuCO/CustomizableInstancePrivateData.h
--- a/Engine/Plugins/Experimental/Mutable/Source/CustomizableObject/Private/MuCO/CustomizableInstancePrivateData.h
+++ b/Engine/Plugins/Experimental/Mutable/Source/CustomizableObject/Private/MuCO/CustomizableInstancePrivateData.h
@@ -75,17 +75,6 @@ struct FGeneratedMaterial
 };
 
 
-// Unreal-side data for the parameter decorations.
-USTRUCT()
-struct FParameterDecorations
-{
-	GENERATED_USTRUCT_BODY();
-
-	UPROPERTY(Transient)
-	TArray< TObjectPtr<UTexture2D> > Images;
-};
-
-
 USTRUCT()
 struct FReferencedPhysicsAssets
 {
@@ -189,10 +178,6 @@ public:
 	// Projector state for each parameter name
 	TMap<TPair<FString, int32>, EProjectorState::Type> ProjectorStates;
 
-	// Parameter decoration unreal-side data if generated.
-	UPROPERTY(Transient)
-	TArray<FParameterDecorations> ParameterDecorations;
-
 	// Indices of the parameters that are relevant for the given parameter values.
 	// This only gets updated if parameter decorations are generated.
 	TArray<int> RelevantParameters;
diff --git a/Engine/Plugins/Experimental/Mutable/Source/CustomizableObject/Private/MuCO/CustomizableObject.cpp b/Engine/Plugins/Experimental/Mutable/Source/CustomizableObject/Private/MuCO/CustomizableObject.cpp
--- a/Engine/Plugins/Experimental/Mutable/Source/CustomizableObject/Private/MuCO/CustomizableObject.cpp
+++ b/Engine/Plugins/Experimental/Mutable/Source/CustomizableObject/Private/MuCO/CustomizableObject.cpp
@@ -1277,7 +1277,6 @@ void UCustomizableObject::UpdateParameterPropertiesFromModel()
 
 			Data.Name = MutableParameters->GetName(paramIndex);
 			Data.Type = EMutableParameterType::None;
-			Data.ImageDescriptionCount = MutableParameters->GetAdditionalImageCount(paramIndex);
 
 			mu::PARAMETER_TYPE mutableType = MutableParameters->GetType(paramIndex);
 			switch (mutableType)
@@ -1345,24 +1344,9 @@ void UCustomizableObject::UpdateParameterPropertiesFromModel()
 }
 
 
-int UCustomizableObject::GetParameterDescriptionCount(int32 ParamIndex) const
-{
-	if (ParamIndex >= 0 && ParamIndex < ParameterProperties.Num())
-	{
-		return ParameterProperties[ParamIndex].ImageDescriptionCount;
-	}
-	else
-	{
-		UE_LOG(LogMutable, Warning, TEXT("Index [%d] out of ParameterProperties bounds at GetParameterDescriptionCount at CO %s."), ParamIndex, *GetName());
-	}
-
-	return 0;
-}
-
-
 int32 UCustomizableObject::GetParameterDescriptionCount(const FString& ParamName) const
 {
-	return GetParameterDescriptionCount(FindParameter(ParamName));
+	return 0;
 }
 
 
diff --git a/Engine/Plugins/Experimental/Mutable/Source/CustomizableObject/Private/MuCO/CustomizableObjectInstance.cpp b/Engine/Plugins/Experimental/Mutable/Source/CustomizableObject/Private/MuCO/CustomizableObjectInstance.cpp
--- a/Engine/Plugins/Experimental/Mutable/Source/CustomizableObject/Private/MuCO/CustomizableObjectInstance.cpp
+++ b/Engine/Plugins/Experimental/Mutable/Source/CustomizableObject/Private/MuCO/CustomizableObjectInstance.cpp
@@ -449,11 +449,6 @@ void UCustomizableObjectInstance::Serialize(FArchive& Ar)
 		Descriptor.VectorParameters = VectorParameters_DEPRECATED;
 		Descriptor.ProjectorParameters = ProjectorParameters_DEPRECATED;
 	}
-	
-	if (CustomizableObjectCustomVersion < FCustomizableObjectCustomVersion::DescriptorBuildParameterDecorations)
-	{
-		Descriptor.bBuildParameterDecorations = bBuildParameterDecorations_DEPRECATED;
-	}
 
 	if (CustomizableObjectCustomVersion < FCustomizableObjectCustomVersion::DescriptorMultilayerProjectors)
 	{
@@ -611,15 +606,28 @@ UCustomizableObject* UCustomizableObjectInstance::GetCustomizableObject() const
 }
 
 
+// Marked as deprecated. Will be removed in future versions.
 bool UCustomizableObjectInstance::GetBuildParameterDecorations() const
 {
-	return Descriptor.GetBuildParameterDecorations();
+	return false;
 }
 
 
+// Marked as deprecated. Will be removed in future versions.
 void UCustomizableObjectInstance::SetBuildParameterDecorations(const bool Value)
 {
-	Descriptor.SetBuildParameterDecorations(Value);
+}
+
+
+bool UCustomizableObjectInstance::GetBuildParameterRelevancy() const
+{
+	return bBuildParameterRelevancy;
+}
+
+
+void UCustomizableObjectInstance::SetBuildParameterRelevancy(bool Value)
+{
+	bBuildParameterRelevancy = Value;
 }
 
 
@@ -683,39 +691,10 @@ EProjectorState::Type UCustomizableObjectInstance::GetProjectorState(const FStri
 }
 
 
-UTexture2D* UCustomizableObjectInstance::GetParameterDescription(int32 ParamIndex, int32 DescIndex)
-{
-	const UCustomizableObject* CustomizableObject = GetCustomizableObject();
-    if (!CustomizableObject)
-    {
-    	return nullptr;
-    }
-	
-	if (CustomizableObject->GetPrivate()->GetModel())
-	{
-		return nullptr;
-	}
-
-	// TODO, make sure the lines below aren't actually needed. They were commented because they made SDetailsViewBase::UpdatePropertyMaps() crash
-	// Make sure this instance is updated
-	//UCustomizableObjectSystem* CustomizableObjectSystem = UCustomizableObjectSystem::GetInstance();
-	//CustomizableObjectSystem->WaitForInstance(this);
-
-	if (GetPrivate()->ParameterDecorations.IsValidIndex(ParamIndex)
-		&&
-		GetPrivate()->ParameterDecorations[ParamIndex].Images.IsValidIndex(DescIndex))
-	{
-		return GetPrivate()->ParameterDecorations[ParamIndex].Images[DescIndex];
-	}
-
-	//UE_LOG(LogMutable, Warning, TEXT("Expected parameter decoration not found in instance [%s]."), *GetName());
-	return nullptr;
-}
-
-
+// Marked as deprecated. Will be removed in future versions.
 UTexture2D* UCustomizableObjectInstance::GetParameterDescription(const FString& ParamName, int32 DescIndex)
 {
-	return GetParameterDescription(GetCustomizableObject()->FindParameter(ParamName), DescIndex);
+	return nullptr;
 }
 
 
@@ -6141,54 +6120,6 @@ void UCustomizableInstancePrivateData::ProcessTextureCoverageQueries(const TShar
 }
 
 
-
-void UCustomizableInstancePrivateData::UpdateParameterDecorationsEngineResources(const TSharedPtr<FMutableOperationData>& OperationData)
-{
-	// This should run in the unreal thread.
-	ParameterDecorations.SetNum(OperationData->ParametersUpdateData.Parameters.Num());
-	for (int32 ParamIndex = 0; ParamIndex < OperationData->ParametersUpdateData.Parameters.Num(); ++ParamIndex)
-	{
-		int32 DescCount = OperationData->ParametersUpdateData.Parameters[ParamIndex].Images.Num();
-		ParameterDecorations[ParamIndex].Images.SetNum(DescCount);
-		for (int32 DescIndex = 0; DescIndex < DescCount; ++DescIndex)
-		{
-			UTexture2D* Tex = nullptr;
-
-			auto Image = OperationData->ParametersUpdateData.Parameters[ParamIndex].Images[DescIndex];
-			if (Image)
-			{
-				check(Image->GetSizeX()>0 && Image->GetSizeY()>0);
-
-				Tex = NewObject<UTexture2D>(UTexture2D::StaticClass());
-				FMutableModelImageProperties Props;
-				Props.Filter = TF_Bilinear;
-				Props.SRGB = true;
-				Props.LODBias = 0;
-
-				ConvertImage(Tex, Image, Props);
-
-#if !PLATFORM_DESKTOP
-				for (int32 i = 0; i < Tex->GetPlatformData()->Mips.Num(); ++i)
-				{
-					uint32 DataFlags = Tex->GetPlatformData()->Mips[i].BulkData.GetBulkDataFlags();
-					Tex->GetPlatformData()->Mips[i].BulkData.SetBulkDataFlags(DataFlags | BULKDATA_SingleUse);
-				}
-#endif
-
-				Tex->UpdateResource();
-			}
-
-			ParameterDecorations[ParamIndex].Images[DescIndex] = Tex;
-		}
-	}
-
-	OperationData->ParametersUpdateData.Clear();
-
-	// Relevancy
-	RelevantParameters = OperationData->RelevantParametersInProgress;
-}
-
-
 void UCustomizableObjectInstance::SetReplacePhysicsAssets(bool bReplaceEnabled)
 {
 	bReplaceEnabled ? GetPrivate()->SetCOInstanceFlags(ReplacePhysicsAssets) : GetPrivate()->ClearCOInstanceFlags(ReplacePhysicsAssets);
diff --git a/Engine/Plugins/Experimental/Mutable/Source/CustomizableObject/Private/MuCO/CustomizableObjectInstanceDescriptor.cpp b/Engine/Plugins/Experimental/Mutable/Source/CustomizableObject/Private/MuCO/CustomizableObjectInstanceDescriptor.cpp
--- a/Engine/Plugins/Experimental/Mutable/Source/CustomizableObject/Private/MuCO/CustomizableObjectInstanceDescriptor.cpp
+++ b/Engine/Plugins/Experimental/Mutable/Source/CustomizableObject/Private/MuCO/CustomizableObjectInstanceDescriptor.cpp
@@ -440,15 +440,15 @@ void FCustomizableObjectInstanceDescriptor::SetCustomizableObject(UCustomizableO
 }
 
 
-bool FCustomizableObjectInstanceDescriptor::GetBuildParameterDecorations() const
+bool FCustomizableObjectInstanceDescriptor::GetBuildParameterRelevancy() const
 {
-	return bBuildParameterDecorations;
+	return bBuildParameterRelevancy;
 }
 
 
-void FCustomizableObjectInstanceDescriptor::SetBuildParameterDecorations(const bool Value)
+void FCustomizableObjectInstanceDescriptor::SetBuildParameterRelevancy(const bool Value)
 {
-	bBuildParameterDecorations = Value;
+	bBuildParameterRelevancy = Value;
 }
 
 
@@ -2292,7 +2292,7 @@ FDescriptorHash::FDescriptorHash(const FCustomizableObjectInstanceDescriptor& De
 	}
 	
 	Hash = HashCombine(Hash, GetTypeHash(Descriptor.State));
-	Hash = HashCombine(Hash, GetTypeHash(Descriptor.bBuildParameterDecorations));
+	Hash = HashCombine(Hash, GetTypeHash(Descriptor.bBuildParameterRelevancy));
 
 	for (const TTuple<FName, FMultilayerProjector>& Pair : Descriptor.MultilayerProjectors)
 	{
diff --git a/Engine/Plugins/Experimental/Mutable/Source/CustomizableObject/Private/MuCO/CustomizableObjectSystem.cpp b/Engine/Plugins/Experimental/Mutable/Source/CustomizableObject/Private/MuCO/CustomizableObjectSystem.cpp
--- a/Engine/Plugins/Experimental/Mutable/Source/CustomizableObject/Private/MuCO/CustomizableObjectSystem.cpp
+++ b/Engine/Plugins/Experimental/Mutable/Source/CustomizableObject/Private/MuCO/CustomizableObjectSystem.cpp
@@ -66,7 +66,6 @@ DEFINE_STAT(STAT_MutableNumAllocatedTextures);
 DEFINE_STAT(STAT_MutableTextureResourceMemory);
 DEFINE_STAT(STAT_MutableTextureGeneratedMemory);
 DEFINE_STAT(STAT_MutableTextureCacheMemory);
-DEFINE_STAT(STAT_MutableTextureParameterDecorationMemory);
 DEFINE_STAT(STAT_MutablePendingInstanceUpdates);
 DEFINE_STAT(STAT_MutableAbandonedInstanceUpdates);
 DEFINE_STAT(STAT_MutableInstanceBuildTime);
@@ -544,7 +543,6 @@ void FCustomizableObjectSystemPrivate::UpdateStats()
 
 		TextureMemoryUsed = int64_t(Size / 1024);
 
-		uint64 SizeDecoration = 0;
 		uint64 SizeGenerated = 0;
 		for (TObjectIterator<UCustomizableObjectInstance> CustomizableObjectInstance; CustomizableObjectInstance; ++CustomizableObjectInstance)
 		{
@@ -558,17 +556,6 @@ void FCustomizableObjectSystemPrivate::UpdateStats()
 
 				if (bHasResourceForRendering)
 				{
-					for (const FParameterDecorations& ParameterDecoration : CustomizableObjectInstance->GetPrivate()->ParameterDecorations)
-					{
-						for (const UTexture2D* ParameterDecorationImage : ParameterDecoration.Images)
-						{
-							if (ParameterDecorationImage && ParameterDecorationImage->IsValidLowLevel())
-							{
-								SizeDecoration += ParameterDecorationImage->CalcTextureMemorySizeEnum(TMC_AllMips);
-							}
-						}
-					}
-
 					for (const FGeneratedTexture& GeneratedTextures : CustomizableObjectInstance->GetPrivate()->GeneratedTextures)
 					{
 						if (GeneratedTextures.Texture)
@@ -598,7 +585,6 @@ void FCustomizableObjectSystemPrivate::UpdateStats()
 		SET_DWORD_STAT(STAT_MutablePendingInstanceUpdates, MutableOperationQueue.Num());
 		SET_DWORD_STAT(STAT_MutableNumAllocatedTextures, CountAllocated);
 		SET_DWORD_STAT(STAT_MutableTextureResourceMemory, Size / 1024.f);
-		SET_DWORD_STAT(STAT_MutableTextureParameterDecorationMemory, SizeDecoration / 1024.f);
 		SET_DWORD_STAT(STAT_MutableTextureGeneratedMemory, SizeGenerated / 1024.f);
 #endif
 
@@ -1174,59 +1160,34 @@ void FCustomizableObjectSystemPrivate::UpdateStreamingLimit()
 namespace impl
 {
 
-	void Subtask_Mutable_UpdateParameterDecorations(const TSharedPtr<FMutableOperationData>& OperationData, const TSharedPtr<mu::Model, ESPMode::ThreadSafe>& MutableModel, const mu::Parameters* MutableParameters)
+	void Subtask_Mutable_UpdateParameterRelevancy(const TSharedPtr<FMutableOperationData>& OperationData, const mu::Parameters* MutableParameters)
 	{
-		MUTABLE_CPUPROFILER_SCOPE(Subtask_Mutable_UpdateParameterDecorations)
+		MUTABLE_CPUPROFILER_SCOPE(Subtask_Mutable_UpdateParameterRelevancy)
 
 		check(OperationData);
-		OperationData->ParametersUpdateData.Clear();
 		OperationData->RelevantParametersInProgress.Empty();
 
 		check(MutableParameters);
+		check(OperationData->InstanceID != 0);
 
 		// This must run in the mutable thread.
 		check(UCustomizableObjectSystem::GetInstance() != nullptr);
 		check(UCustomizableObjectSystem::GetInstance()->GetPrivate() != nullptr);
 		mu::SystemPtr MutableSystem = UCustomizableObjectSystem::GetInstance()->GetPrivate()->MutableSystem;
 
-		// Decorations
 		int32 NumParameters = MutableParameters->GetCount();
 
-		{
-			MUTABLE_CPUPROFILER_SCOPE(ParameterDecorations)
-
-				// Generate all the images that decorate the parameters UI, only in the editor.
-				for (int32 ParamIndex = 0; ParamIndex < NumParameters; ++ParamIndex)
-				{
-					FParameterDecorationsUpdateData::FParameterDesc ParamData;
-
-					for (int32 ImageIndex = 0; ImageIndex < MutableParameters->GetAdditionalImageCount(ParamIndex); ++ImageIndex)
-					{
-						mu::ImagePtrConst Image = MutableSystem->BuildParameterAdditionalImage(
-							MutableModel,
-							MutableParameters,
-							ParamIndex, ImageIndex);
-
-						ParamData.Images.Add(Image);
-					};
-
-					OperationData->ParametersUpdateData.Parameters.Add(ParamData);
-				}
-		}
-
 		// Update the parameter relevancy.
 		{
 			MUTABLE_CPUPROFILER_SCOPE(ParameterRelevancy)
 
 			TArray<bool> Relevant;
 			Relevant.SetNumZeroed(NumParameters);
-			// This prevents live update cache reusal
-			//MutableSystem->GetParameterRelevancy(MutableModel, MutableParameters, Relevant.GetData());
+			MutableSystem->GetParameterRelevancy(OperationData->InstanceID, MutableParameters, Relevant.GetData());
 
 			for (int32 ParamIndex = 0; ParamIndex < NumParameters; ++ParamIndex)
 			{
-				// This prevents live update cache reusal
-				//if (Relevant[ParamIndex])
+				if (Relevant[ParamIndex])
 				{
 					OperationData->RelevantParametersInProgress.Add(ParamIndex);
 				}
@@ -1695,7 +1656,7 @@ namespace impl
 
 
 	// This runs in a worker thread.
-	void Task_Mutable_Update_GetMesh(TSharedPtr<FMutableOperationData> OperationData, TSharedPtr<mu::Model, ESPMode::ThreadSafe> Model, mu::ParametersPtrConst Parameters, bool bBuildParameterDecorations, int32 State)
+	void Task_Mutable_Update_GetMesh(TSharedPtr<FMutableOperationData> OperationData, TSharedPtr<mu::Model, ESPMode::ThreadSafe> Model, mu::ParametersPtrConst Parameters, bool bBuildParameterRelevancy, int32 State)
 	{
 		MUTABLE_CPUPROFILER_SCOPE(Task_Mutable_Update_GetMesh)
 
@@ -1705,22 +1666,20 @@ namespace impl
 
 		check(OperationData.IsValid());
 
-		if (bBuildParameterDecorations)
+		Subtask_Mutable_BeginUpdate_GetMesh(OperationData, Model, Parameters.get(), State);
+
+		// TODO: Not strictly mutable: move to another worker thread task to free mutable access?
+		Subtask_Mutable_PrepareSkeletonData(OperationData);
+
+		if (bBuildParameterRelevancy)
 		{
-			// This cannot happen inside the beginupdate-endupdate cycle.
-			Subtask_Mutable_UpdateParameterDecorations(OperationData, Model, Parameters.get());
+			Subtask_Mutable_UpdateParameterRelevancy(OperationData, Parameters.get());
 		}
 		else
 		{
-			OperationData->ParametersUpdateData.Clear();
 			OperationData->RelevantParametersInProgress.Reset();
 		}
 
-		Subtask_Mutable_BeginUpdate_GetMesh(OperationData, Model, Parameters.get(), State);
-
-		// TODO: Not strictly mutable: move to another worker thread task to free mutable access?
-		Subtask_Mutable_PrepareSkeletonData(OperationData);
-
 #if WITH_EDITOR
 		uint32 EndCycles = FPlatformTime::Cycles();
 		OperationData->MutableRuntimeCycles = EndCycles - StartCycles;
@@ -2013,7 +1972,7 @@ namespace impl
 
 
 	/** "Lock Cached Resources" */
-	void Task_Game_LockCache(TSharedPtr<FMutableOperationData> OperationData, const TWeakObjectPtr<UCustomizableObjectInstance>& CustomizableObjectInstancePtr, mu::Ptr<const mu::Parameters> Parameters, bool bBuildParameterDecorations)
+	void Task_Game_LockCache(TSharedPtr<FMutableOperationData> OperationData, const TWeakObjectPtr<UCustomizableObjectInstance>& CustomizableObjectInstancePtr, mu::Ptr<const mu::Parameters> Parameters, bool bBuildParameterRelevancy)
 	{
 		MUTABLE_CPUPROFILER_SCOPE(Task_Game_LockCache)
 
@@ -2059,10 +2018,10 @@ namespace impl
 
 		check(OperationData.IsValid());
 		
-		// Process the parameter decorations if requested
-		if (bBuildParameterDecorations)
+		if (bBuildParameterRelevancy)
 		{
-			ObjectInstancePrivateData->UpdateParameterDecorationsEngineResources(OperationData);
+			// Relevancy
+			ObjectInstancePrivateData->RelevantParameters = OperationData->RelevantParametersInProgress;
 		}
 
 		if (const TObjectPtr<UDefaultImageProvider> DefaultImageProvider = System->GetDefaultImageProvider())
@@ -2389,15 +2348,15 @@ namespace impl
 			check(CustomizableObjectInstancePtr.IsValid());
 			TSharedPtr<FMutableOperation> CurrentMutableOperation = SystemPrivateData->CurrentMutableOperation;
 			check(CurrentMutableOperation);
-			bool bBuildParameterDecorations = CurrentMutableOperation->IsBuildParameterDecorations();
+			bool bBuildParameterRelevancy = CurrentMutableOperation->IsBuildParameterRelevancy();
 			TSharedPtr<mu::Model, ESPMode::ThreadSafe> Model = CustomizableObject->GetPrivate()->GetModel();
 			int32 State = CustomizableObjectInstancePtr->GetState();
 
 			Mutable_GetMeshTask = SystemPrivateData->AddMutableThreadTask(
 				TEXT("Task_Mutable_Update_GetMesh"),
-				[CurrentOperationData, bBuildParameterDecorations, Parameters, Model, State]()
+				[CurrentOperationData, bBuildParameterRelevancy, Parameters, Model, State]()
 				{
-					impl::Task_Mutable_Update_GetMesh(CurrentOperationData, Model, Parameters, bBuildParameterDecorations, State);
+					impl::Task_Mutable_Update_GetMesh(CurrentOperationData, Model, Parameters, bBuildParameterRelevancy, State);
 				},
 				UE::Tasks::ETaskPriority::BackgroundHigh);
 		}
@@ -2411,14 +2370,14 @@ namespace impl
 			check(CustomizableObjectInstancePtr.IsValid());
 			TSharedPtr<FMutableOperation> CurrentMutableOperation = SystemPrivateData->CurrentMutableOperation;
 			check(CurrentMutableOperation);
-			bool bBuildParameterDecorations = CurrentMutableOperation->IsBuildParameterDecorations();
+			bool bBuildParameterRelevancy = CurrentMutableOperation->IsBuildParameterRelevancy();
 
 			SystemPrivateData->AddGameThreadTask(
 				{
 				FMutableTaskDelegate::CreateLambda(
-					[CurrentOperationData, CustomizableObjectInstancePtr, bBuildParameterDecorations, Parameters]()
+					[CurrentOperationData, CustomizableObjectInstancePtr, bBuildParameterRelevancy, Parameters]()
 					{
-						impl::Task_Game_LockCache(CurrentOperationData, CustomizableObjectInstancePtr, Parameters, bBuildParameterDecorations);
+						impl::Task_Game_LockCache(CurrentOperationData, CustomizableObjectInstancePtr, Parameters, bBuildParameterRelevancy);
 					}),
 				Mutable_GetMeshTask
 				});
@@ -2733,7 +2692,7 @@ FMutableOperation FMutableOperation::CreateInstanceUpdate(UCustomizableObjectIns
 	Op.CustomizableObjectInstance = InCustomizableObjectInstance;
 	Op.InstanceDescriptorRuntimeHash = InCustomizableObjectInstance->GetUpdateDescriptorRuntimeHash();
 	Op.bStarted = false;
-	Op.bBuildParameterDecorations = InCustomizableObjectInstance->GetBuildParameterDecorations();
+	Op.bBuildParameterRelevancy = InCustomizableObjectInstance->GetBuildParameterRelevancy();
 	Op.Parameters = InCustomizableObjectInstance->GetPrivate()->GetParameters(InCustomizableObjectInstance);
 	if (UpdateCallback)
 	{
diff --git a/Engine/Plugins/Experimental/Mutable/Source/CustomizableObject/Private/MuCO/CustomizableObjectSystemPrivate.h b/Engine/Plugins/Experimental/Mutable/Source/CustomizableObject/Private/MuCO/CustomizableObjectSystemPrivate.h
--- a/Engine/Plugins/Experimental/Mutable/Source/CustomizableObject/Private/MuCO/CustomizableObjectSystemPrivate.h
+++ b/Engine/Plugins/Experimental/Mutable/Source/CustomizableObject/Private/MuCO/CustomizableObjectSystemPrivate.h
@@ -34,7 +34,7 @@ class FMutableOperation
 	/** Instance parameters at the time of the operation request. */
 	mu::ParametersPtr Parameters; 
 	
-	bool bBuildParameterDecorations = false;
+	bool bBuildParameterRelevancy = false;
 	bool bMeshNeedsUpdate = false;
 
 	/** Protected constructor. */
@@ -94,9 +94,9 @@ public:
 
 	FInstanceUpdateDelegate UpdateCallback;
 	
-	bool IsBuildParameterDecorations() const
+	bool IsBuildParameterRelevancy() const
 	{
-		return bBuildParameterDecorations;
+		return bBuildParameterRelevancy;
 	}
 
 	/** Read-only access to the mutable instance parameters for this operation. */
@@ -439,24 +439,6 @@ struct FInstanceUpdateData
 };
 
 
-// Mutable data generated during the update steps.
-struct FParameterDecorationsUpdateData
-{
-	struct FParameterDesc
-	{
-		TArray<mu::ImagePtrConst> Images;
-	};
-
-	//! Information about additional images used for parameter UI decoration
-	TArray<FParameterDesc> Parameters;
-
-	void Clear()
-	{
-		Parameters.Empty();
-	}
-};
-
-
 struct FTextureCoverageQueryData
 {
 	uint32 CoveredTexels;
@@ -486,7 +468,6 @@ struct FMutableOperationData
 	FInstanceGeneratedData LastUpdateData;
 
 	FInstanceUpdateData InstanceUpdateData;
-	FParameterDecorationsUpdateData ParametersUpdateData;
 	TArray<int> RelevantParametersInProgress;
 
 	TArray<FString> LowPriorityTextures;
diff --git a/Engine/Plugins/Experimental/Mutable/Source/CustomizableObject/Private/MuCO/LogInformationUtil.cpp b/Engine/Plugins/Experimental/Mutable/Source/CustomizableObject/Private/MuCO/LogInformationUtil.cpp
--- a/Engine/Plugins/Experimental/Mutable/Source/CustomizableObject/Private/MuCO/LogInformationUtil.cpp
+++ b/Engine/Plugins/Experimental/Mutable/Source/CustomizableObject/Private/MuCO/LogInformationUtil.cpp
@@ -145,7 +145,6 @@ void LogInformationUtil::LogShowInstanceDataFull(const UCustomizableObjectInstan
 	FString MessageChunk;
 
 	MessageChunk += "\n\t";
-	MessageChunk += FString::Printf(TEXT("        bBuildParameterDecorations = %d\n"), CustomizableObjectInstance->GetBuildParameterDecorations());
 	MessageChunk += FString::Printf(TEXT("        bShowOnlyRuntimeParameters = %d\n"), CustomizableObjectInstance->bShowOnlyRuntimeParameters);
 	MessageChunk += FString::Printf(TEXT("        bShowOnlyRelevantParameters = %d\n"), CustomizableObjectInstance->bShowOnlyRelevantParameters);
 	MessageChunk += FString::Printf(TEXT("        MinSquareDistFromComponentToPlayer = %.2f\n"), CustomizableObjectInstance->GetPrivate()->MinSquareDistFromComponentToPlayer);
diff --git a/Engine/Plugins/Experimental/Mutable/Source/CustomizableObject/Public/MuCO/CustomizableObjectCustomVersion.h b/Engine/Plugins/Experimental/Mutable/Source/CustomizableObject/Public/MuCO/CustomizableObjectCustomVersion.h
--- a/Engine/Plugins/Experimental/Mutable/Source/CustomizableObject/Public/MuCO/CustomizableObjectCustomVersion.h
+++ b/Engine/Plugins/Experimental/Mutable/Source/CustomizableObject/Public/MuCO/CustomizableObjectCustomVersion.h
@@ -72,6 +72,8 @@ struct CUSTOMIZABLEOBJECT_API FCustomizableObjectCustomVersion
 		
 		ConvertAnimationSlotToFName,
 
+		RemovedParameterDecorations,
+
 		// -----<new versions can be added above this line>-------------------------------------------------
 		VersionPlusOne,
 		LatestVersion = VersionPlusOne - 1
diff --git a/Engine/Plugins/Experimental/Mutable/Source/CustomizableObject/Public/MuCO/CustomizableObjectInstanceDescriptor.h b/Engine/Plugins/Experimental/Mutable/Source/CustomizableObject/Public/MuCO/CustomizableObjectInstanceDescriptor.h
--- a/Engine/Plugins/Experimental/Mutable/Source/CustomizableObject/Public/MuCO/CustomizableObjectInstanceDescriptor.h
+++ b/Engine/Plugins/Experimental/Mutable/Source/CustomizableObject/Public/MuCO/CustomizableObjectInstanceDescriptor.h
@@ -51,9 +51,9 @@ struct CUSTOMIZABLEOBJECT_API FCustomizableObjectInstanceDescriptor
 
 	void SetCustomizableObject(UCustomizableObject& InCustomizableObject);
 	
-	bool GetBuildParameterDecorations() const;
+	bool GetBuildParameterRelevancy() const;
 	
-	void SetBuildParameterDecorations(bool Value);
+	void SetBuildParameterRelevancy(bool Value);
 	
 	/** Update all parameters to be up to date with the Mutable Core parameters. */
 	void ReloadParameters();
@@ -343,9 +343,8 @@ private:
 	/** Mutable parameters optimization state. */
 	int32 State = 0;
 	
-	/** Flag to control the build of the parameter description images required for customization UI.
-     * These descriptions get generated with every instance update, so it should be disabled when not needed. */
-	bool bBuildParameterDecorations = false;
+	/** If this is set to true, when updating the instance an additional step will be performed to calculate the list of instance parameters that are relevant for the current parameter vaules. */
+	bool bBuildParameterRelevancy = false;
 
 	/** These are the LODs Mutable can generate, they MUST NOT be used in an update (Mutable thread). */
 	int32 MinLOD = 0;
diff --git a/Engine/Plugins/Experimental/Mutable/Source/CustomizableObjectEditor/Classes/MuCOE/Nodes/CustomizableObjectNodeFloatParameter.h b/Engine/Plugins/Experimental/Mutable/Source/CustomizableObjectEditor/Classes/MuCOE/Nodes/CustomizableObjectNodeFloatParameter.h
--- a/Engine/Plugins/Experimental/Mutable/Source/CustomizableObjectEditor/Classes/MuCOE/Nodes/CustomizableObjectNodeFloatParameter.h
+++ b/Engine/Plugins/Experimental/Mutable/Source/CustomizableObjectEditor/Classes/MuCOE/Nodes/CustomizableObjectNodeFloatParameter.h
@@ -43,28 +43,16 @@ public:
 	
 	// UObject interface.
 	void PostEditChangeProperty(FPropertyChangedEvent& PropertyChangedEvent) override;
-	virtual void Serialize(FArchive& Ar) override;
 	
 	// Begin EdGraphNode interface
 	FText GetNodeTitle(ENodeTitleType::Type TitleType) const override;
 	FLinearColor GetNodeTitleColor() const override;
 	FText GetTooltipText() const override;
-	void PostPasteNode() override;
 
 	// UCustomizableObjectNode interface
 	void AllocateDefaultPins(UCustomizableObjectNodeRemapPins* RemapPins) override;
-	virtual bool IsAffectedByLOD() const override;
+	bool IsAffectedByLOD() const override;
+	void BackwardsCompatibleFixup() override;
 
-	// Own interface
-	int32 GetNumDescriptionImage() const;
-	
-	const UEdGraphPin* GetDescriptionImagePin(int32 Index) const;
-
-private:
-	UPROPERTY()
-	TArray<FEdGraphPinReference> DescriptionImagePinsReferences;
-	
-	UPROPERTY(EditAnywhere, Category=CustomizableObject)
-	TArray<FCustomizableObjectNodeFloatDescription> DescriptionImage;
 };
 
diff --git a/Engine/Plugins/Experimental/Mutable/Source/CustomizableObjectEditor/Private/MuCOE/CustomizableObjectEditor.cpp b/Engine/Plugins/Experimental/Mutable/Source/CustomizableObjectEditor/Private/MuCOE/CustomizableObjectEditor.cpp
--- a/Engine/Plugins/Experimental/Mutable/Source/CustomizableObjectEditor/Private/MuCOE/CustomizableObjectEditor.cpp
+++ b/Engine/Plugins/Experimental/Mutable/Source/CustomizableObjectEditor/Private/MuCOE/CustomizableObjectEditor.cpp
@@ -436,7 +436,7 @@ void FCustomizableObjectEditor::CreatePreviewInstance()
 				// Asset loading works in the Main Thread, there's a risk the sync loading could put to sleep the thread while
 				// waiting for the asset registry to load the content (that never gets executed)
 				Viewport->SetAssetRegistryLoaded(true);
-				PreviewInstance->SetBuildParameterDecorations(true);
+				PreviewInstance->SetBuildParameterRelevancy(true);
 				PreviewInstance->UpdateSkeletalMeshAsync(true, true);
 			}
 			else
diff --git a/Engine/Plugins/Experimental/Mutable/Source/CustomizableObjectEditor/Private/MuCOE/GenerateMutableSource/GenerateMutableSourceFloat.cpp b/Engine/Plugins/Experimental/Mutable/Source/CustomizableObjectEditor/Private/MuCOE/GenerateMutableSource/GenerateMutableSourceFloat.cpp
--- a/Engine/Plugins/Experimental/Mutable/Source/CustomizableObjectEditor/Private/MuCOE/GenerateMutableSource/GenerateMutableSourceFloat.cpp
+++ b/Engine/Plugins/Experimental/Mutable/Source/CustomizableObjectEditor/Private/MuCOE/GenerateMutableSource/GenerateMutableSourceFloat.cpp
@@ -95,23 +95,6 @@ mu::NodeScalarPtr GenerateMutableSourceFloat(const UEdGraphPin* Pin, FMutableGra
 		ScalarNode->SetUid(StringCast<ANSICHAR>(*GenerationContext.GetNodeIdUnique(Node).ToString()).Get());
 		ScalarNode->SetDefaultValue(FloatParameterNode->DefaultValue);
 
-		const int32 NumDescriptionImage = FloatParameterNode->GetNumDescriptionImage();
-		ScalarNode->SetDescImageCount(NumDescriptionImage);
-
-		for (int d = 0; d < NumDescriptionImage; ++d)
-		{
-			const UEdGraphPin* DescriptionImagePin = FloatParameterNode->GetDescriptionImagePin(d);
-			if (!DescriptionImagePin)
-			{
-				GenerationContext.Compiler->CompilerLog(LOCTEXT("WrongFloatParameterDescriptionImagePin", "Float parameter in inconsistent state, please refresh the node."), FloatParameterNode);
-			}
-			else if (const UEdGraphPin* ConnectedPin = FollowInputPin(*DescriptionImagePin))
-			{
-				mu::NodeImagePtr TextureNode = GenerateMutableSourceImage(ConnectedPin, GenerationContext, 0.f);
-				ScalarNode->SetDescImage(d, TextureNode);
-			}
-		}
-
 		GenerationContext.ParameterUIDataMap.Add(FloatParameterNode->ParameterName, FParameterUIData(
 			FloatParameterNode->ParameterName,
 			FloatParameterNode->ParamUIMetadata,
diff --git a/Engine/Plugins/Experimental/Mutable/Source/CustomizableObjectEditor/Private/MuCOE/Nodes/CustomizableObjectNodeFloatParameter.cpp b/Engine/Plugins/Experimental/Mutable/Source/CustomizableObjectEditor/Private/MuCOE/Nodes/CustomizableObjectNodeFloatParameter.cpp
--- a/Engine/Plugins/Experimental/Mutable/Source/CustomizableObjectEditor/Private/MuCOE/Nodes/CustomizableObjectNodeFloatParameter.cpp
+++ b/Engine/Plugins/Experimental/Mutable/Source/CustomizableObjectEditor/Private/MuCOE/Nodes/CustomizableObjectNodeFloatParameter.cpp
@@ -10,12 +10,6 @@ class UCustomizableObjectNodeRemapPins;
 #define LOCTEXT_NAMESPACE "CustomizableObjectEditor"
 
 
-const UEdGraphPin* UCustomizableObjectNodeFloatParameter::GetDescriptionImagePin(int32 Index) const
-{
-	return DescriptionImagePinsReferences[Index].Get();
-}
-
-
 void UCustomizableObjectNodeFloatParameter::PostEditChangeProperty(FPropertyChangedEvent& PropertyChangedEvent)
 {
 	FProperty* PropertyThatChanged = PropertyChangedEvent.Property;
@@ -28,42 +22,12 @@ void UCustomizableObjectNodeFloatParameter::PostEditChangeProperty(FPropertyChan
 }
 
 
-void UCustomizableObjectNodeFloatParameter::Serialize(FArchive& Ar)
-{
-	Super::Serialize(Ar);
-
-	Ar.UsingCustomVersion(FCustomizableObjectCustomVersion::GUID);
-
-	if (Ar.CustomVer(FCustomizableObjectCustomVersion::GUID) < FCustomizableObjectCustomVersion::PinsNamesImageToTexture)
-	{
-		for (int32 Index = 0; Index < DescriptionImage.Num(); ++Index)
-		{
-			DescriptionImagePinsReferences.Add(FindPin(FString::Printf(TEXT("Description Image %d"), Index)));
-		}
-	}
-}
-
-
 void UCustomizableObjectNodeFloatParameter::AllocateDefaultPins(UCustomizableObjectNodeRemapPins* RemapPins)
 {
 	const UEdGraphSchema_CustomizableObject* Schema = GetDefault<UEdGraphSchema_CustomizableObject>();
 
 	UEdGraphPin* ValuePin = CustomCreatePin(EGPD_Output, Schema->PC_Float, FName("Value"));
 	ValuePin->bDefaultValueIsIgnored = true;
-
-	DescriptionImagePinsReferences.Reset();
-	
-	for (int32 Index = 0; Index <  DescriptionImage.Num(); ++Index )
-	{
-		FString PinName = DescriptionImage[Index].Name.IsEmpty() ?
-			FString::Printf(TEXT("Description Texture %d"), Index) :
-			DescriptionImage[Index].Name;
-		
-		UEdGraphPin* DescriptionImagePin = CustomCreatePin(EGPD_Input, Schema->PC_Image, FName(*PinName));
-		DescriptionImagePin->bDefaultValueIsIgnored = true;
-
-		DescriptionImagePinsReferences.Add(DescriptionImagePin);
-	}
 }
 
 
@@ -73,12 +37,6 @@ bool UCustomizableObjectNodeFloatParameter::IsAffectedByLOD() const
 }
 
 
-int32 UCustomizableObjectNodeFloatParameter::GetNumDescriptionImage() const
-{
-	return DescriptionImagePinsReferences.Num();
-}
-
-
 FText UCustomizableObjectNodeFloatParameter::GetNodeTitle(ENodeTitleType::Type TitleType) const
 {
 	return LOCTEXT("Float_Parameter", "Float Parameter");
@@ -92,27 +50,23 @@ FLinearColor UCustomizableObjectNodeFloatParameter::GetNodeTitleColor() const
 }
 
 
-void UCustomizableObjectNodeFloatParameter::PostPasteNode()
+FText UCustomizableObjectNodeFloatParameter::GetTooltipText() const
 {
-	TArray<UEdGraphPin*> InputPins;
-	GetInputPins(InputPins);
-
-	DescriptionImage.SetNum(InputPins.Num());
-	const int32 MaxIndex = InputPins.Num();
-	for (int32 i = 0; i < MaxIndex; ++i)
-	{
-		const UEdGraphPin* DescriptionImagePin = InputPins[i];
-		if (DescriptionImagePin->LinkedTo.Num())
-		{
-			DescriptionImage[i].Name = Helper_GetPinName(InputPins[i]);
-		}
-	}
+	return LOCTEXT("Float_Parameter_Tooltip", "Expose a numeric parameter from the Customizable Object that can be modified at runtime.");
 }
 
 
-FText UCustomizableObjectNodeFloatParameter::GetTooltipText() const
+void UCustomizableObjectNodeFloatParameter::BackwardsCompatibleFixup()
 {
-	return LOCTEXT("Float_Parameter_Tooltip", "Expose a numeric parameter from the Customizable Object that can be modified at runtime.");
+	Super::BackwardsCompatibleFixup();
+
+	const int32 CustomizableObjectCustomVersion = GetLinkerCustomVersion(FCustomizableObjectCustomVersion::GUID);
+
+	if (CustomizableObjectCustomVersion < FCustomizableObjectCustomVersion::RemovedParameterDecorations)
+	{
+		ReconstructNode();
+	}
 }
 
+
 #undef LOCTEXT_NAMESPACE
diff --git a/Engine/Plugins/Experimental/Mutable/Source/CustomizableObjectEditor/Private/MuCOE/Nodes/CustomizableObjectNodeGroupProjectorParameter.cpp b/Engine/Plugins/Experimental/Mutable/Source/CustomizableObjectEditor/Private/MuCOE/Nodes/CustomizableObjectNodeGroupProjectorParameter.cpp
--- a/Engine/Plugins/Experimental/Mutable/Source/CustomizableObjectEditor/Private/MuCOE/Nodes/CustomizableObjectNodeGroupProjectorParameter.cpp
+++ b/Engine/Plugins/Experimental/Mutable/Source/CustomizableObjectEditor/Private/MuCOE/Nodes/CustomizableObjectNodeGroupProjectorParameter.cpp
@@ -116,11 +116,6 @@ void UCustomizableObjectNodeGroupProjectorParameter::BackwardsCompatibleFixup()
 		}
 	}
 
-	if (CustomizableObjectCustomVersion < FCustomizableObjectCustomVersion::GroupProjectorImagePin)
-	{
-		ReconstructNode();		
-	}
-
 	if (CustomizableObjectCustomVersion < FCustomizableObjectCustomVersion::GroupProjectorImagePinRemoved)
 	{
 		ReconstructNode();		
diff --git a/Engine/Plugins/Experimental/Mutable/Source/CustomizableObjectEditor/Private/MuCOE/SCustomizableInstanceProperties.cpp b/Engine/Plugins/Experimental/Mutable/Source/CustomizableObjectEditor/Private/MuCOE/SCustomizableInstanceProperties.cpp
--- a/Engine/Plugins/Experimental/Mutable/Source/CustomizableObjectEditor/Private/MuCOE/SCustomizableInstanceProperties.cpp
+++ b/Engine/Plugins/Experimental/Mutable/Source/CustomizableObjectEditor/Private/MuCOE/SCustomizableInstanceProperties.cpp
@@ -657,25 +657,6 @@ void SCustomizableInstanceProperties::AddParameter(int32 ParamIndexInObject)
 			SAssignNew(SliderBox, SVerticalBox)
 		];
 
-		// Color bar?
-		if (CustomizableObject->GetParameterDescriptionCount(ParamIndexInObject) > 0)
-		{
-			if (UTexture2D* Bar = CustomInstance->GetParameterDescription(ParamIndexInObject, 0))
-			{
-				TSharedRef<FDeferredCleanupSlateBrush> Brush = FDeferredCleanupSlateBrush::CreateBrush(Bar, FVector2D(Bar->GetSizeX(), 12), FLinearColor(1.0f, 1.0f, 1.0f, 1.0f), ESlateBrushTileType::NoTile, ESlateBrushImageType::Linear);
-				DynamicBrushes.Add(Brush);
-
-				SliderBox->AddSlot()
-				.AutoHeight()
-				.FillHeight(1.0f)
-				.HAlign(HAlign_Fill)
-				.Padding(0.0f, 0.0f, 0.0f, 2.0f)
-				[
-					SNew(SImage)
-					.Image(Brush->GetSlateBrush())
-				];
-			}
-		}
 
 		TSharedPtr<SSpinBox<float>> Slider;
 
@@ -1257,9 +1238,9 @@ void SCustomizableInstanceProperties::Tick(const FGeometry& AllottedGeometry, co
 	SCompoundWidget::Tick(AllottedGeometry, InCurrentTime, InDeltaTime);
 
 	// We set this every frame, in case other widgets are also modifying this flag
-	if (!CustomInstance->GetBuildParameterDecorations())
+	if (!CustomInstance->GetBuildParameterRelevancy())
 	{
-		CustomInstance->SetBuildParameterDecorations(true);
+		CustomInstance->SetBuildParameterRelevancy(true);
 		CustomInstance->UpdateSkeletalMeshAsync(true);
 	}
 
diff --git a/Engine/Plugins/Experimental/Mutable/Source/CustomizableObjectPopulationEditor/Private/MuCOPE/CustomizableObjectPopulationClassDetails.cpp b/Engine/Plugins/Experimental/Mutable/Source/CustomizableObjectPopulationEditor/Private/MuCOPE/CustomizableObjectPopulationClassDetails.cpp
--- a/Engine/Plugins/Experimental/Mutable/Source/CustomizableObjectPopulationEditor/Private/MuCOPE/CustomizableObjectPopulationClassDetails.cpp
+++ b/Engine/Plugins/Experimental/Mutable/Source/CustomizableObjectPopulationEditor/Private/MuCOPE/CustomizableObjectPopulationClassDetails.cpp
@@ -1268,20 +1268,12 @@ void SPopulationClassConstraint::BuildRangeWidget()
 	// Adding a Texture if the parameter option has one
 	if (ParameterIndex >=0 && PopulationClass->CustomizableObject && PopulationClass->CustomizableObjectInstance)
 	{
-		UTexture2D* Bar = nullptr;
-
-		if (PopulationClass->CustomizableObject->GetParameterDescriptionCount(ParameterIndex) > 0)
-		{
-			Bar = PopulationClass->CustomizableObjectInstance->GetParameterDescription(ParameterIndex, 0);
-		}
-
 		RangeWidget->AddSlot()
 		.AutoHeight()
 		.Padding(0.0f, 10.0f, 0.0f, 5.0f)
 		[
 			SNew(SRangeSquare)
 			.Ranges(&PopulationClass->Characteristics[CharactericticIndex].Constraints[ConstraintIndex].Ranges)
-			.Texture(Bar)
 			.bDiscrete(bDiscrete)
 		];
 	}
diff --git a/Engine/Plugins/Experimental/Mutable/Source/CustomizableObjectPopulationEditor/Private/MuCOPE/CustomizableObjectPopulationClassEditor.cpp b/Engine/Plugins/Experimental/Mutable/Source/CustomizableObjectPopulationEditor/Private/MuCOPE/CustomizableObjectPopulationClassEditor.cpp
--- a/Engine/Plugins/Experimental/Mutable/Source/CustomizableObjectPopulationEditor/Private/MuCOPE/CustomizableObjectPopulationClassEditor.cpp
+++ b/Engine/Plugins/Experimental/Mutable/Source/CustomizableObjectPopulationEditor/Private/MuCOPE/CustomizableObjectPopulationClassEditor.cpp
@@ -130,7 +130,7 @@ void FCustomizableObjectPopulationClassEditor::InitCustomizableObjectPopulationC
 		if (PopulationClass->CustomizableObject->IsCompiled() && !PopulationClass->CustomizableObjectInstance)
 		{
 			PopulationClass->CustomizableObjectInstance = PopulationClass->CustomizableObject->CreateInstance();
-			PopulationClass->CustomizableObjectInstance->SetBuildParameterDecorations(true);
+			PopulationClass->CustomizableObjectInstance->SetBuildParameterRelevancy(true);
 			PopulationClass->CustomizableObjectInstance->UpdateSkeletalMeshAsync(true, true);
 
 			bRefreshDetailsView = true;
@@ -267,7 +267,7 @@ void FCustomizableObjectPopulationClassEditor::Tick(float InDeltaTime)
 		if (PopulationClass->CustomizableObject->IsCompiled() && !PopulationClass->CustomizableObjectInstance)
 		{
 			PopulationClass->CustomizableObjectInstance = PopulationClass->CustomizableObject->CreateInstance();
-			PopulationClass->CustomizableObjectInstance->SetBuildParameterDecorations(true);
+			PopulationClass->CustomizableObjectInstance->SetBuildParameterRelevancy(true);
 			PopulationClass->CustomizableObjectInstance->UpdateSkeletalMeshAsync(true, true);
 
 			bRefreshDetailsView = true;
@@ -791,7 +791,7 @@ void FCustomizableObjectPopulationClassEditor::TestPopulationClass()
 						// Attaching the Customizable Skeletal Component to the Skeletal Mesh Component
 						if (PreviewCustomizableSkeletalComponent && SkeletalMeshComponent && ColliderComponent)
 						{
-							ViewportInstances[i]->SetBuildParameterDecorations(true);
+							ViewportInstances[i]->SetBuildParameterRelevancy(true);
 							ViewportInstances[i]->UpdateSkeletalMeshAsync(true, true);
 
 							PreviewCustomizableSkeletalComponent->CustomizableObjectInstance = ViewportInstances[i];
diff --git a/Engine/Plugins/Experimental/Mutable/Source/CustomizableObjectPopulationEditor/Private/MuCOPE/CustomizableObjectPopulationEditor.cpp b/Engine/Plugins/Experimental/Mutable/Source/CustomizableObjectPopulationEditor/Private/MuCOPE/CustomizableObjectPopulationEditor.cpp
--- a/Engine/Plugins/Experimental/Mutable/Source/CustomizableObjectPopulationEditor/Private/MuCOPE/CustomizableObjectPopulationEditor.cpp
+++ b/Engine/Plugins/Experimental/Mutable/Source/CustomizableObjectPopulationEditor/Private/MuCOPE/CustomizableObjectPopulationEditor.cpp
@@ -472,7 +472,7 @@ void FCustomizableObjectPopulationEditor::TestPopulation()
 				// Attaching the Customizable Skeletal Component to the Skeletal Mesh Component
 				if (PreviewCustomizableSkeletalComponent && PreviewSkeletalMeshComponent && ColliderComponent)
 				{
-					ViewportInstances[i]->SetBuildParameterDecorations(true);
+					ViewportInstances[i]->SetBuildParameterRelevancy(true);
 					ViewportInstances[i]->UpdateSkeletalMeshAsync(true, true);
 
 					PreviewCustomizableSkeletalComponent->CustomizableObjectInstance = ViewportInstances[i];
diff --git a/Engine/Plugins/Experimental/Mutable/Source/MutableRuntime/Private/MuR/CodeVisitor.h b/Engine/Plugins/Experimental/Mutable/Source/MutableRuntime/Private/MuR/CodeVisitor.h
--- a/Engine/Plugins/Experimental/Mutable/Source/MutableRuntime/Private/MuR/CodeVisitor.h
+++ b/Engine/Plugins/Experimental/Mutable/Source/MutableRuntime/Private/MuR/CodeVisitor.h
@@ -104,29 +104,16 @@ class Parameters;
         }
 
 
-        void Traverse( OP::ADDRESS root, FProgram& program, bool visitDecorators = true )
+        void Traverse( OP::ADDRESS root, FProgram& program )
         {
             m_pending.Reserve( program.m_opAddress.Num() );
 
             // Visit the given root
             m_pending.Add( PENDING(root,0) );
             Recurse( program );
+         }
 
-            if (visitDecorators)
-            {
-                // Fix the code used in the parameter descriptions
-                for ( std::size_t p=0; p<program.m_parameters.Num(); ++p )
-                {
-                    for ( std::size_t d=0; d<program.m_parameters[p].m_descImages.Num(); ++d )
-                    {
-                        m_pending.Add(PENDING(program.m_parameters[p].m_descImages[d],0) );
-                        Recurse( program );
-                    }
-                }
-            }
-        }
-
-        void FullTraverse( FProgram& program, bool visitDecorators = true )
+        void FullTraverse( FProgram& program )
         {
             // Visit all the state roots
             for ( std::size_t p=0; p<program.m_states.Num(); ++p )
@@ -134,19 +121,6 @@ class Parameters;
                 m_pending.Add(PENDING(program.m_states[p].m_root,0) );
                 Recurse( program );
             }
-
-            if (visitDecorators)
-            {
-                // Fix the code used in the parameter descriptions
-                for ( std::size_t p=0; p<program.m_parameters.Num(); ++p )
-                {
-                    for ( std::size_t d=0; d<program.m_parameters[p].m_descImages.Num(); ++d )
-                    {
-                        m_pending.Add(PENDING(program.m_parameters[p].m_descImages[d],0) );
-                        Recurse( program );
-                    }
-                }
-            }
         }
 
 
@@ -307,29 +281,16 @@ class Parameters;
         }
 
 
-        void Traverse( OP::ADDRESS root, FProgram& program, bool visitDecorators = true )
+        void Traverse( OP::ADDRESS root, FProgram& program )
         {
             m_pending.reserve( program.m_opAddress.Num() );
 
             // Visit the given root
             m_pending.Add( PENDING(root,0) );
             Recurse( program );
-
-            if (visitDecorators)
-            {
-                // Fix the code used in the parameter descriptions
-                for ( std::size_t p=0; p<program.m_parameters.Num(); ++p )
-                {
-                    for ( std::size_t d=0; d<program.m_parameters[p].m_descImages.Num(); ++d )
-                    {
-                        m_pending.Add(PENDING(program.m_parameters[p].m_descImages[d],0) );
-                        Recurse( program );
-                    }
-                }
-            }
         }
 
-        void FullTraverse( FProgram& program, bool visitDecorators = true )
+        void FullTraverse( FProgram& program )
         {
             // Visit all the state roots
             for ( std::size_t p=0; p<program.m_states.Num(); ++p )
@@ -337,19 +298,6 @@ class Parameters;
                 m_pending.Add(PENDING(program.m_states[p].m_root,0) );
                 Recurse( program );
             }
-
-            if (visitDecorators)
-            {
-                // Fix the code used in the parameter descriptions
-                for ( std::size_t p=0; p<program.m_parameters.Num(); ++p )
-                {
-                    for ( std::size_t d=0; d<program.m_parameters[p].m_descImages.Num(); ++d )
-                    {
-                        m_pending.Add(PENDING(program.m_parameters[p].m_descImages[d],0) );
-                        Recurse( program );
-                    }
-                }
-            }
         }
 
 
diff --git a/Engine/Plugins/Experimental/Mutable/Source/MutableRuntime/Private/MuR/Parameters.cpp b/Engine/Plugins/Experimental/Mutable/Source/MutableRuntime/Private/MuR/Parameters.cpp
--- a/Engine/Plugins/Experimental/Mutable/Source/MutableRuntime/Private/MuR/Parameters.cpp
+++ b/Engine/Plugins/Experimental/Mutable/Source/MutableRuntime/Private/MuR/Parameters.cpp
@@ -312,15 +312,6 @@ namespace mu
     }
 
 
-	//---------------------------------------------------------------------------------------------
-	int Parameters::GetAdditionalImageCount( int index ) const
-	{
-		const FProgram& program = m_pD->m_pModel->GetPrivate()->m_program;
-		check( index>=0 && index<(int)program.m_parameters.Num() );
-		return (int)program.m_parameters[index].m_descImages.Num();
-	}
-
-
 	//---------------------------------------------------------------------------------------------
     bool Parameters::GetBoolValue( int index,
                                    const Ptr<const RangeIndex>& pos ) const
diff --git a/Engine/Plugins/Experimental/Mutable/Source/MutableRuntime/Private/MuR/ParametersPrivate.h b/Engine/Plugins/Experimental/Mutable/Source/MutableRuntime/Private/MuR/ParametersPrivate.h
--- a/Engine/Plugins/Experimental/Mutable/Source/MutableRuntime/Private/MuR/ParametersPrivate.h
+++ b/Engine/Plugins/Experimental/Mutable/Source/MutableRuntime/Private/MuR/ParametersPrivate.h
@@ -316,9 +316,6 @@ namespace mu
         //! vector of range descriptors.
 		TArray<uint32> m_ranges;
 
-        //! Additional parameter description information
-		TArray<OP::ADDRESS> m_descImages;
-
         //! Possible values of the parameter in case of being an integer, and its names
         struct INT_VALUE_DESC
         {
@@ -356,14 +353,14 @@ namespace mu
         {
             return m_name == other.m_name && m_uid == other.m_uid && m_type == other.m_type &&
                    m_defaultValue == other.m_defaultValue &&
-                   m_ranges == other.m_ranges && m_descImages == other.m_descImages &&
+                   m_ranges == other.m_ranges &&
                    m_possibleValues == other.m_possibleValues;
         }
 
         //!
         void Serialise( OutputArchive& arch ) const
         {
-            const int32 ver = 6;
+            const int32 ver = 7;
             arch << ver;
 
 			arch << m_name;
@@ -371,7 +368,6 @@ namespace mu
             arch << m_type;
             arch << m_defaultValue;
             arch << m_ranges;
-            arch << m_descImages;
             arch << m_possibleValues;
         }
 
@@ -380,14 +376,19 @@ namespace mu
         {
             int32 ver;
             arch >> ver;
-            check( ver == 6 );
+            check( ver >= 6 );
 
 			arch >> m_name;
             arch >> m_uid;
             arch >> m_type;
             arch >> m_defaultValue;
             arch >> m_ranges;
-            arch >> m_descImages;
+
+			if (ver<=6)
+			{ 
+				TArray<OP::ADDRESS> UnusedDescImages;
+				arch >> UnusedDescImages;
+			}
             arch >> m_possibleValues;
         }
     };
diff --git a/Engine/Plugins/Experimental/Mutable/Source/MutableRuntime/Private/MuR/System.cpp b/Engine/Plugins/Experimental/Mutable/Source/MutableRuntime/Private/MuR/System.cpp
--- a/Engine/Plugins/Experimental/Mutable/Source/MutableRuntime/Private/MuR/System.cpp
+++ b/Engine/Plugins/Experimental/Mutable/Source/MutableRuntime/Private/MuR/System.cpp
@@ -460,37 +460,6 @@ namespace mu
 	}
 
 
-    //---------------------------------------------------------------------------------------------
-    ImagePtrConst System::BuildParameterAdditionalImage(const TSharedPtr<const Model>& pModel,
-                                                         const Ptr<const Parameters>& pParams,
-                                                         int32 parameter,
-                                                         int32 ImageIndex )
-    {
-		LLM_SCOPE_BYNAME(TEXT("MutableRuntime"));
-		MUTABLE_CPUPROFILER_SCOPE(BuildParameterAdditionalImage);
-
-		check(parameter >= 0
-			&&
-			parameter < (int)pModel->GetPrivate()->m_program.m_parameters.Num());
-
-		const FParameterDesc& param = pModel->GetPrivate()->m_program.m_parameters[parameter];
-		check(ImageIndex >= 0 && ImageIndex < (int)param.m_descImages.Num());
-
-		OP::ADDRESS imageAddress = param.m_descImages[ImageIndex];
-
-		ImagePtrConst pResult;
-
-		if (imageAddress)
-		{
-			m_pD->BeginBuild(pModel);
-			pResult = m_pD->BuildImage( pModel, pParams.get(), imageAddress, 0, 0 );
-			m_pD->EndBuild();
-		}
-
-		return pResult;
-	}
-
-
     //---------------------------------------------------------------------------------------------
     class RelevantParameterVisitor : public UniqueDiscreteCoveredCodeVisitor<>
     {
@@ -507,15 +476,11 @@ namespace mu
         {
             m_pFlags = pFlags;
 
-            memset( pFlags, 0, sizeof(bool)*pParams->GetCount() );
+            FMemory::Memset( pFlags, 0, sizeof(bool)*pParams->GetCount() );
 
             OP::ADDRESS at = pModel->GetPrivate()->m_program.m_states[0].m_root;
 
-            pSystem->BeginBuild(pModel);
-
             Run( at );
-
-			pSystem->EndBuild();
         }
 
 
@@ -552,61 +517,68 @@ namespace mu
 
 
     //---------------------------------------------------------------------------------------------
-    void System::GetParameterRelevancy(const TSharedPtr<const Model>& pModel,
-                                        const ParametersPtrConst& pParameters,
-                                        bool* pFlags )
+    void System::GetParameterRelevancy( Instance::ID InstanceID,
+                                        const ParametersPtrConst& Parameters,
+                                        bool* Flags )
     {
 		LLM_SCOPE_BYNAME(TEXT("MutableRuntime"));
 
-        RelevantParameterVisitor visitor( m_pD, pModel, pParameters, pFlags );
+		// Find the live instance
+		Private::FLiveInstance* pLiveInstance = m_pD->FindLiveInstance(InstanceID);
+		check(pLiveInstance);
+		m_pD->m_memory = pLiveInstance->m_memory;
+		
+		RelevantParameterVisitor visitor( m_pD, pLiveInstance->m_pModel, Parameters, Flags );
+
+		m_pD->m_memory = nullptr;
     }
 
 
     //---------------------------------------------------------------------------------------------
     //---------------------------------------------------------------------------------------------
     //---------------------------------------------------------------------------------------------
-    bool System::Private::CheckUpdatedParameters( const FLiveInstance* pLiveInstance,
-                                 const Ptr<const Parameters>& pParams,
-                                 uint64& updatedParameters )
+    bool System::Private::CheckUpdatedParameters( const FLiveInstance* LiveInstance,
+                                 const Ptr<const Parameters>& Params,
+                                 uint64& UpdatedParameters)
     {
-        bool fullBuild = false;
+        bool bFullBuild = false;
 
-		if (!pLiveInstance->m_pOldParameters)
+		if (!LiveInstance->m_pOldParameters)
 		{
-			updatedParameters = 0xffffffff;
+			UpdatedParameters = AllParametersMask;
 			return true;
 		}
 
         // check what parameters have changed
-        updatedParameters = 0;
-        const FProgram& program = pLiveInstance->m_pModel->GetPrivate()->m_program;
-        const TArray<int>& runtimeParams = program.m_states[ pLiveInstance->m_state ].m_runtimeParameters;
+		UpdatedParameters = 0;
+        const FProgram& program = LiveInstance->m_pModel->GetPrivate()->m_program;
+        const TArray<int>& runtimeParams = program.m_states[ LiveInstance->m_state ].m_runtimeParameters;
 
-        check( pParams->GetCount() == (int)program.m_parameters.Num() );
-        check( !pLiveInstance->m_pOldParameters
-                        ||
-                        pParams->GetCount() == pLiveInstance->m_pOldParameters->GetCount() );
+        check( Params->GetCount() == (int)program.m_parameters.Num() );
+        check( !LiveInstance->m_pOldParameters
+			||
+			Params->GetCount() == LiveInstance->m_pOldParameters->GetCount() );
 
-        for ( int p=0; p<(int)program.m_parameters.Num() && !fullBuild; ++p )
+        for ( int32 p=0; p<program.m_parameters.Num() && !bFullBuild; ++p )
         {
             bool isRuntime = runtimeParams.Contains( p );
-            bool changed = !pParams->HasSameValue( p, pLiveInstance->m_pOldParameters, p );
+            bool changed = !Params->HasSameValue( p, LiveInstance->m_pOldParameters, p );
 
             if (changed && isRuntime)
             {
 				uint64 runtimeIndex = runtimeParams.IndexOfByKey(p);
-                updatedParameters |= uint64(1) << runtimeIndex;
+				UpdatedParameters |= uint64(1) << runtimeIndex;
             }
             else if (changed)
             {
                 // A non-runtime parameter has changed, we need a full build.
                 // TODO: report, or log somehow.
-                fullBuild = true;
-                updatedParameters = 0xffffffffffffffff;
+				bFullBuild = true;
+                UpdatedParameters = AllParametersMask;
             }
         }
 
-        return fullBuild;
+        return bFullBuild;
     }
 
 
diff --git a/Engine/Plugins/Experimental/Mutable/Source/MutableRuntime/Private/MuR/SystemPrivate.h b/Engine/Plugins/Experimental/Mutable/Source/MutableRuntime/Private/MuR/SystemPrivate.h
--- a/Engine/Plugins/Experimental/Mutable/Source/MutableRuntime/Private/MuR/SystemPrivate.h
+++ b/Engine/Plugins/Experimental/Mutable/Source/MutableRuntime/Private/MuR/SystemPrivate.h
@@ -18,6 +18,9 @@ namespace mu
 {
 	class ExtensionDataStreamer;
 
+	constexpr uint64 AllParametersMask = TNumericLimits<uint64> ::Max();
+
+
 	//! Reference-counted colour to be stored in the cache
 	class Colour : public RefCounted
 	{
diff --git a/Engine/Plugins/Experimental/Mutable/Source/MutableRuntime/Public/MuR/Model.h b/Engine/Plugins/Experimental/Mutable/Source/MutableRuntime/Public/MuR/Model.h
--- a/Engine/Plugins/Experimental/Mutable/Source/MutableRuntime/Public/MuR/Model.h
+++ b/Engine/Plugins/Experimental/Mutable/Source/MutableRuntime/Public/MuR/Model.h
@@ -13,7 +13,7 @@
 //! data structures. Compiled models are not necessarily compatible when the runtime is updated,
 //! so this version number can be used externally to verify this. It is not used internally, and
 //! serializing models from different versions than this runtime will probably result in a crash.
-#define MUTABLE_COMPILED_MODEL_CODE_VERSION		uint32( 63 )
+#define MUTABLE_COMPILED_MODEL_CODE_VERSION		uint32( 64 )
 #define MUTABLE_PARAMETERS_VERSION              uint32( 2 )
 
 
diff --git a/Engine/Plugins/Experimental/Mutable/Source/MutableRuntime/Public/MuR/Parameters.h b/Engine/Plugins/Experimental/Mutable/Source/MutableRuntime/Public/MuR/Parameters.h
--- a/Engine/Plugins/Experimental/Mutable/Source/MutableRuntime/Public/MuR/Parameters.h
+++ b/Engine/Plugins/Experimental/Mutable/Source/MutableRuntime/Public/MuR/Parameters.h
@@ -239,8 +239,7 @@ namespace mu
 		//! \pre The parameter specified by index is a T_BOOL.
 		//! \param index Index of the parameter from 0 to GetCount()-1
         //! \param pos Only for multidimensional parameters: relevant position to get in the ranges
-        bool GetBoolValue( int index,
-                           const Ptr<const RangeIndex>& pos=nullptr ) const;
+        bool GetBoolValue( int index, const Ptr<const RangeIndex>& pos=nullptr ) const;
 
 		//! Set the value of a parameter of type boolean.
 		//! \pre The parameter specified by index is a T_BOOL.
@@ -248,8 +247,7 @@ namespace mu
 		//! \param value new value of the parameter
         //! \param pos optional parameter to set a specific value for the given multidimensional
         //! position. If null, the value is set for all possible positions.
-        void SetBoolValue( int index, bool value,
-                           const Ptr<const RangeIndex>& pos=nullptr );
+        void SetBoolValue( int index, bool value, const Ptr<const RangeIndex>& pos=nullptr );
 
 		//! Get the number of possible values for this integer parameter.
 		//! If the number is zero, it means any integer value is accepted.
@@ -275,30 +273,26 @@ namespace mu
 		//! \pre The parameter specified by index is a T_INT.
 		//! \param index Index of the parameter from 0 to GetCount()-1
         //! \param pos Only for multidimensional parameters: relevant position to get in the ranges
-        int GetIntValue( int index,
-                         const Ptr<const RangeIndex>& pos=nullptr ) const;
+        int GetIntValue( int index, const Ptr<const RangeIndex>& pos=nullptr ) const;
 
 		//! If the parameter is of the integer type, set its value.
 		//! \param index Index of the parameter from 0 to GetCount()-1
 		//! \param value new value of the parameter. It must be in the possible values for this
 		//! parameter (see GetIntPossibleValue), or the method will leave the value unchanged.
         //! \param pos Only for multidimensional parameters: relevant position to set in the ranges
-        void SetIntValue( int index, int value,
-                          const Ptr<const RangeIndex>& pos=nullptr );
+        void SetIntValue( int index, int value, const Ptr<const RangeIndex>& pos=nullptr );
 
 		//! Return the value of a float parameter.
 		//! \pre The parameter specified by index is a T_FLOAT.
 		//! \param index Index of the parameter from 0 to GetCount()-1
         //! \param pos Only for multidimensional parameters: relevant position to get in the ranges
-        float GetFloatValue( int index,
-                             const Ptr<const RangeIndex>& pos=nullptr ) const;
+        float GetFloatValue( int index, const Ptr<const RangeIndex>& pos=nullptr ) const;
 
 		//! If the parameter is of the float type, set its value.
 		//! \param index Index of the parameter from 0 to GetCount()-1
 		//! \param value new value of the parameter
         //! \param pos Only for multidimensional parameters: relevant position to set in the ranges
-        void SetFloatValue( int index, float value,
-                            const Ptr<const RangeIndex>& pos=nullptr );
+        void SetFloatValue( int index, float value, const Ptr<const RangeIndex>& pos=nullptr );
 
 		//! Return the value of a colour parameter.
 		//! \pre The parameter specified by index is a T_FLOAT.
@@ -377,16 +371,9 @@ namespace mu
         //! \param pos Only for multidimensional parameters: relevant position to set in the ranges
         void SetStringValue( int index, const char* value, const Ptr<const RangeIndex>& pos = nullptr );
 
-        //! Get the number of additional description images defined in the parameter.
-		//! The use of these depends on the game, but it can be used to provide colour bars, etc.
-		//! These images can be built using a system object and the model of these parameters.
-		int GetAdditionalImageCount( int index ) const;
-
         //! Utility method to compare the values of a specific parameter with the values of another
         //! Parameters object. It returns false if type or values are different.
-        bool HasSameValue( int thisParamIndex,
-                           const ParametersPtrConst& other,
-                           int otherParamIndex ) const;
+        bool HasSameValue( int thisParamIndex, const ParametersPtrConst& other, int otherParamIndex ) const;
 
 		//-----------------------------------------------------------------------------------------
 		// Interface pattern
diff --git a/Engine/Plugins/Experimental/Mutable/Source/MutableRuntime/Public/MuR/System.h b/Engine/Plugins/Experimental/Mutable/Source/MutableRuntime/Public/MuR/System.h
--- a/Engine/Plugins/Experimental/Mutable/Source/MutableRuntime/Public/MuR/System.h
+++ b/Engine/Plugins/Experimental/Mutable/Source/MutableRuntime/Public/MuR/System.h
@@ -188,13 +188,6 @@ namespace mu
         //! \param instance The id of the instance to destroy.
         void ReleaseInstance( Instance::ID InstanceID);
 
-		//! Build one of the images defined in a model parameter as additional description. These
-		//! images can be used for colour bars, icons, etc...
-		Ptr<const Image> BuildParameterAdditionalImage(const TSharedPtr<const Model>& Model,
-														const Ptr<const Parameters>& Params,
-														int32 Parameter,
-														int32 ImageIndex );
-
 		//! Calculate the relevancy of every parameter. Some parameters may be unused depending on
 		//! the values of other parameters. This method will set to true the flags for parameters
 		//! that are relevant, and to false otherwise. This is useful to hide irrelevant parameters
@@ -203,7 +196,7 @@ namespace mu
         //! \param pParameters Parameter set that we want to find the relevancy of.
         //! \param pFlags is a pointer to a preallocated array of booleans that contains at least
 		//! pParameters->GetCount() elements.
-        void GetParameterRelevancy(const TSharedPtr<const Model>& Model,
+        void GetParameterRelevancy( Instance::ID InstanceID,
 									const Ptr<const Parameters>& Parameters,
 									bool* Flags );
 
diff --git a/Engine/Plugins/Experimental/Mutable/Source/MutableTools/Private/MuT/ASTOpParameter.cpp b/Engine/Plugins/Experimental/Mutable/Source/MutableTools/Private/MuT/ASTOpParameter.cpp
--- a/Engine/Plugins/Experimental/Mutable/Source/MutableTools/Private/MuT/ASTOpParameter.cpp
+++ b/Engine/Plugins/Experimental/Mutable/Source/MutableTools/Private/MuT/ASTOpParameter.cpp
@@ -31,10 +31,6 @@ namespace mu
 		{
 			f(c.rangeSize);
 		}
-		for (auto& c : additionalImages)
-		{
-			f(c);
-		}
 	}
 
 
@@ -44,8 +40,7 @@ namespace mu
 		{
 			return type == other->type &&
 				parameter == other->parameter &&
-				ranges == other->ranges &&
-				additionalImages == other->additionalImages;
+				ranges == other->ranges;
 		}
 		return false;
 	}
@@ -60,10 +55,6 @@ namespace mu
 		{
 			n->ranges.Emplace(n.get(), mapChild(c.rangeSize.child()), c.rangeName, c.rangeUID);
 		}
-		for (const auto& c : additionalImages)
-		{
-			n->additionalImages.Emplace(n, mapChild(c.child()));
-		}
 		return n;
 	}
 
@@ -103,12 +94,6 @@ namespace mu
 				program.m_parameters.Last().m_ranges.Add(rangeId);
 			}
 
-			for (const auto& d : additionalImages)
-			{
-				OP::ADDRESS descAt = d ? d->linkedAddress : 0;
-				program.m_parameters.Last().m_descImages.Add(descAt);
-			}
-
 			linkedAddress = (OP::ADDRESS)program.m_opAddress.Num();
 			//program.m_code.push_back(op);
 
diff --git a/Engine/Plugins/Experimental/Mutable/Source/MutableTools/Private/MuT/ASTOpParameter.h b/Engine/Plugins/Experimental/Mutable/Source/MutableTools/Private/MuT/ASTOpParameter.h
--- a/Engine/Plugins/Experimental/Mutable/Source/MutableTools/Private/MuT/ASTOpParameter.h
+++ b/Engine/Plugins/Experimental/Mutable/Source/MutableTools/Private/MuT/ASTOpParameter.h
@@ -31,9 +31,6 @@ namespace mu
 		//** Ranges adding dimensions to this parameter. */
 		TArray<FRangeData> ranges;
 
-		//! Additional images attached to the parameter
-		TArray<ASTChild> additionalImages;
-
 		/** Index of the parameter in the program parameter list. Generated ar link time. */
 		int32 LinkedParameterIndex = -1;
 
diff --git a/Engine/Plugins/Experimental/Mutable/Source/MutableTools/Private/MuT/CodeGenerator_Scalar.cpp b/Engine/Plugins/Experimental/Mutable/Source/MutableTools/Private/MuT/CodeGenerator_Scalar.cpp
--- a/Engine/Plugins/Experimental/Mutable/Source/MutableTools/Private/MuT/CodeGenerator_Scalar.cpp
+++ b/Engine/Plugins/Experimental/Mutable/Source/MutableTools/Private/MuT/CodeGenerator_Scalar.cpp
@@ -158,32 +158,6 @@ namespace mu
 				op->ranges.Emplace(op.get(), rangeResult.sizeOp, rangeResult.rangeName, rangeResult.rangeUID);
 			}
 
-			// Generate the code for the additional images in the parameter
-			for (int32 a = 0; a < node.m_additionalImages.Num(); ++a)
-			{
-				Ptr<ASTOp> descAd;
-				if (node.m_additionalImages[a])
-				{
-					// We take whatever size will be produced
-					IMAGE_STATE newState;
-					FImageDesc desc = CalculateImageDesc(*node.m_additionalImages[a]->GetBasePrivate());
-					newState.m_imageSize = UE::Math::TIntVector2<int32>(desc.m_size);
-					newState.m_imageRect.min[0] = 0;
-					newState.m_imageRect.min[1] = 0;
-					newState.m_imageRect.size = UE::Math::TIntVector2<int32>(desc.m_size);
-					newState.m_layoutBlockId = -1;
-					m_imageState.Add(newState);
-
-					// Generate
-					descAd = Generate(node.m_additionalImages[a]);
-					check(descAd);
-
-					// Restore rect
-					m_imageState.Pop();
-				}
-				op->additionalImages.Emplace(op, descAd);
-			}
-
 			m_nodeVariables[node.m_pNode] = op;
 		}
 		else
diff --git a/Engine/Plugins/Experimental/Mutable/Source/MutableTools/Private/MuT/NodeScalarParameter.cpp b/Engine/Plugins/Experimental/Mutable/Source/MutableTools/Private/MuT/NodeScalarParameter.cpp
--- a/Engine/Plugins/Experimental/Mutable/Source/MutableTools/Private/MuT/NodeScalarParameter.cpp
+++ b/Engine/Plugins/Experimental/Mutable/Source/MutableTools/Private/MuT/NodeScalarParameter.cpp
@@ -30,9 +30,7 @@ namespace mu
 	//---------------------------------------------------------------------------------------------
 	int NodeScalarParameter::GetInputCount() const
 	{
-        return int( m_pD->m_additionalImages.Num()
-                    +
-                    m_pD->m_ranges.Num() );
+        return m_pD->m_ranges.Num();
 	}
 
 
@@ -41,16 +39,10 @@ namespace mu
 	{
         check( i<GetInputCount() );
 
-        int imageCount = int(m_pD->m_additionalImages.Num());
         int rangeCount = int(m_pD->m_ranges.Num());
-        if (i<imageCount)
+		if ( i < rangeCount )
         {
-            return m_pD->m_additionalImages[i].get();
-        }
-        else if ( i < imageCount + rangeCount )
-        {
-            int r = i - imageCount;
-            return m_pD->m_ranges[r].get();
+            return m_pD->m_ranges[i].get();
         }
         return nullptr;
 	}
@@ -60,16 +52,10 @@ namespace mu
     void NodeScalarParameter::SetInputNode( int i, NodePtr n )
 	{
         check( i<GetInputCount() );
-        int imageCount = int(m_pD->m_additionalImages.Num());
         int rangeCount = int(m_pD->m_ranges.Num());
-        if (i<imageCount)
-        {
-            m_pD->m_additionalImages[i] = dynamic_cast<NodeImage*>(n.get());
-        }
-        else if ( i < imageCount + rangeCount )
+        if ( i < rangeCount )
         {
-            int r = i - imageCount;
-            m_pD->m_ranges[r] = dynamic_cast<NodeRange*>(n.get());
+            m_pD->m_ranges[i] = dynamic_cast<NodeRange*>(n.get());
         }
     }
 
@@ -143,22 +129,6 @@ namespace mu
 	}
 
 
-	//---------------------------------------------------------------------------------------------
-	void NodeScalarParameter::SetDescImageCount( int i )
-	{
-        check(i>=0);
-        m_pD->m_additionalImages.SetNum(i);
-	}
-
-
-	//---------------------------------------------------------------------------------------------
-	void NodeScalarParameter::SetDescImage( int i, NodeImagePtr pImage )
-	{
-		check( i>=0 && i<(int)m_pD->m_additionalImages.Num() );
-		m_pD->m_additionalImages[i] = pImage;
-	}
-
-
     //---------------------------------------------------------------------------------------------
     void NodeScalarParameter::SetRangeCount( int i )
     {
diff --git a/Engine/Plugins/Experimental/Mutable/Source/MutableTools/Private/MuT/NodeScalarParameterPrivate.h b/Engine/Plugins/Experimental/Mutable/Source/MutableTools/Private/MuT/NodeScalarParameterPrivate.h
--- a/Engine/Plugins/Experimental/Mutable/Source/MutableTools/Private/MuT/NodeScalarParameterPrivate.h
+++ b/Engine/Plugins/Experimental/Mutable/Source/MutableTools/Private/MuT/NodeScalarParameterPrivate.h
@@ -29,20 +29,17 @@ namespace mu
 		string m_uid;
 		PARAMETER_DETAILED_TYPE m_detailedType = PARAMETER_DETAILED_TYPE::UNKNOWN;
 
-        TArray<Ptr<NodeImage>> m_additionalImages;
-
 		TArray<Ptr<NodeRange>> m_ranges;
 
 		//!
 		void Serialise( OutputArchive& arch ) const
 		{
-            uint32_t ver = 4;
+            uint32_t ver = 5;
 			arch << ver;
 
 			arch << m_defaultValue;
 			arch << m_name;
 			arch << m_uid;
-			arch << m_additionalImages;
 			arch << m_detailedType;
             arch << m_ranges;
         }
@@ -52,12 +49,16 @@ namespace mu
 		{
             uint32_t ver;
 			arch >> ver;
-            check(ver==4);
+            check(ver>=4 && ver<=5);
 
 			arch >> m_defaultValue;
 			arch >> m_name;
             arch >> m_uid;
-			arch >> m_additionalImages;
+			if (ver < 5)
+			{
+				TArray<Ptr<NodeImage>> TempAdditionalImages;
+				arch >> TempAdditionalImages;
+			}
             arch >> m_detailedType;
             arch >> m_ranges;
         }
diff --git a/Engine/Plugins/Experimental/Mutable/Source/MutableTools/Public/MuT/NodeScalarParameter.h b/Engine/Plugins/Experimental/Mutable/Source/MutableTools/Public/MuT/NodeScalarParameter.h
--- a/Engine/Plugins/Experimental/Mutable/Source/MutableTools/Public/MuT/NodeScalarParameter.h
+++ b/Engine/Plugins/Experimental/Mutable/Source/MutableTools/Public/MuT/NodeScalarParameter.h
@@ -47,8 +47,6 @@ namespace mu
 		// Node Interface
 		//-----------------------------------------------------------------------------------------
 
-        
-
         const NODE_TYPE* GetType() const override;
 		static const NODE_TYPE* GetStaticType();
 
@@ -76,12 +74,6 @@ namespace mu
 		PARAMETER_DETAILED_TYPE GetDetailedType() const ;
 		void SetDetailedType( PARAMETER_DETAILED_TYPE type );
 
-        //! Set the number of additional description images associated to this parameter.
-        //! These images can be built from a parameter and, for example, they can be used to build
-        //! user-interfaces for customisation.
-        void SetDescImageCount( int i );
-        void SetDescImage( int i, NodeImagePtr pImage );
-
         //! Set the number of ranges (dimensions) for this parameter.
         //! By default a parameter has 0 ranges, meaning it only has one value.
         void SetRangeCount( int i );
