Fixed mismatched copy constructor/assignment operator for mu::MeshBufferIteratorConst.

diff --git a/Engine/Plugins/Experimental/Mutable/Source/MutableRuntime/Private/MuR/MeshPrivate.h b/Engine/Plugins/Experimental/Mutable/Source/MutableRuntime/Private/MuR/MeshPrivate.h
--- a/Engine/Plugins/Experimental/Mutable/Source/MutableRuntime/Private/MuR/MeshPrivate.h
+++ b/Engine/Plugins/Experimental/Mutable/Source/MutableRuntime/Private/MuR/MeshPrivate.h
@@ -613,11 +613,6 @@ namespace mu
 			}
 		}
 
-		inline MeshBufferIteratorConst( const MeshBufferIteratorConst& other )
-			: UntypedMeshBufferIteratorConst( other )
-		{
-		}
-
 		inline const vec<CTYPE,COMPONENTS>& operator*() const
 		{
 			return *reinterpret_cast< const vec<CTYPE,COMPONENTS>* >( m_pBuf );
