Fixed mismatched copy constructor/assignment operator in FCustomizableObjectInstanceDescriptor.

diff --git a/Engine/Plugins/Experimental/Mutable/Source/CustomizableObject/Public/MuCO/CustomizableObjectInstanceDescriptor.h b/Engine/Plugins/Experimental/Mutable/Source/CustomizableObject/Public/MuCO/CustomizableObjectInstanceDescriptor.h
--- a/Engine/Plugins/Experimental/Mutable/Source/CustomizableObject/Public/MuCO/CustomizableObjectInstanceDescriptor.h
+++ b/Engine/Plugins/Experimental/Mutable/Source/CustomizableObject/Public/MuCO/CustomizableObjectInstanceDescriptor.h
@@ -39,8 +39,6 @@ struct CUSTOMIZABLEOBJECT_API FCustomizableObjectInstanceDescriptor
 	
 	explicit FCustomizableObjectInstanceDescriptor(UCustomizableObject& Object);
 
-	FCustomizableObjectInstanceDescriptor(const FCustomizableObjectInstanceDescriptor& Other) = default;
-
 	/** Serialize this object. Does not support Multilayer Projectors! */
 	void SaveDescriptor(FArchive &Ar);
 
