#rnx

diff --git a/Engine/Plugins/Experimental/Mutable/Source/CustomizableObject/Classes/MuCO/CustomizableObject.h b/Engine/Plugins/Experimental/Mutable/Source/CustomizableObject/Classes/MuCO/CustomizableObject.h
--- a/Engine/Plugins/Experimental/Mutable/Source/CustomizableObject/Classes/MuCO/CustomizableObject.h
+++ b/Engine/Plugins/Experimental/Mutable/Source/CustomizableObject/Classes/MuCO/CustomizableObject.h
@@ -47,7 +47,7 @@ namespace mu
 }
 
 
-CUSTOMIZABLEOBJECT_API DECLARE_LOG_CATEGORY_EXTERN(LogMutable, Warning, All);
+CUSTOMIZABLEOBJECT_API DECLARE_LOG_CATEGORY_EXTERN(LogMutable, Log, All);
 
 USTRUCT()
 struct FFParameterOptionsTags
diff --git a/Engine/Plugins/Experimental/Mutable/Source/CustomizableObject/Private/MuCO/CustomizableObject.cpp b/Engine/Plugins/Experimental/Mutable/Source/CustomizableObject/Private/MuCO/CustomizableObject.cpp
--- a/Engine/Plugins/Experimental/Mutable/Source/CustomizableObject/Private/MuCO/CustomizableObject.cpp
+++ b/Engine/Plugins/Experimental/Mutable/Source/CustomizableObject/Private/MuCO/CustomizableObject.cpp
@@ -852,7 +852,7 @@ FString UCustomizableObject::GetDesc()
 
 void UCustomizableObject::SaveEmbeddedData(FArchive& Ar)
 {
-	UE_LOG(LogMutable, Log, TEXT("Saving embedded data for Customizable Object [%s] now at position %d."), *GetName(), int(Ar.Tell()));
+	UE_LOG(LogMutable, Verbose, TEXT("Saving embedded data for Customizable Object [%s] now at position %d."), *GetName(), int(Ar.Tell()));
 
 	int32 InternalVersion = PrivateData->GetModel() ? CurrentSupportedVersion : -1;
 	Ar << InternalVersion;
@@ -880,7 +880,7 @@ void UCustomizableObject::SaveEmbeddedData(FArchive& Ar)
 			mu::Model::Serialise(PrivateData->GetModel().Get(), arch);
 		}
 
-		UE_LOG(LogMutable, Log, TEXT("Saved embedded data for Customizable Object [%s] now at position %d."), *GetName(), int(Ar.Tell()));
+		UE_LOG(LogMutable, Verbose, TEXT("Saved embedded data for Customizable Object [%s] now at position %d."), *GetName(), int(Ar.Tell()));
 	}
 }
 
@@ -1155,7 +1155,7 @@ UCustomizableObjectInstance* UCustomizableObject::CreateInstance()
 	PreviewInstance->SetObject(this);
 	PreviewInstance->bShowOnlyRuntimeParameters = false;
 
-	UE_LOG(LogMutable, Log, TEXT("Created Customizable Object Instance."));
+	UE_LOG(LogMutable, Verbose, TEXT("Created Customizable Object Instance."));
 
 	return PreviewInstance;
 }
@@ -1701,7 +1701,7 @@ void UCustomizableObject::ApplyStateForcedValuesToParameters( int32 State, mu::P
 			}
 			default:
 			{
-				UE_LOG(LogMutable, Log, TEXT("Forced parameter type not supported."));
+				UE_LOG(LogMutable, Warning, TEXT("Forced parameter type not supported."));
 				break;
 			}
 			}
@@ -1764,7 +1764,7 @@ EDataValidationResult UCustomizableObject::IsDataValid(FDataValidationContext& C
 		return Result;
 	}
 
-	UE_LOG(LogMutable,Display,TEXT("Running data validation checks for %s CO."),*this->GetName());
+	UE_LOG(LogMutable,Verbose,TEXT("Running data validation checks for %s CO."),*this->GetName());
 
 	// Bind the post validation method to the post validation delegate if not bound already to be able to know when the validation
 	// operation (for all assets) concludes
diff --git a/Engine/Plugins/Experimental/Mutable/Source/CustomizableObject/Private/MuCO/CustomizableObjectInstance.cpp b/Engine/Plugins/Experimental/Mutable/Source/CustomizableObject/Private/MuCO/CustomizableObjectInstance.cpp
--- a/Engine/Plugins/Experimental/Mutable/Source/CustomizableObject/Private/MuCO/CustomizableObjectInstance.cpp
+++ b/Engine/Plugins/Experimental/Mutable/Source/CustomizableObject/Private/MuCO/CustomizableObjectInstance.cpp
@@ -3020,7 +3020,7 @@ FTexturePlatformData* UCustomizableInstancePrivateData::MutableCreateImagePlatfo
 
 	if (SizeX <= 0 || SizeY <= 0)
 	{
-		UE_LOG(LogMutable, Verbose, TEXT("Invalid parameters specified for UCustomizableInstancePrivateData::MutableCreateImagePlatformData()"));
+		UE_LOG(LogMutable, Warning, TEXT("Invalid parameters specified for UCustomizableInstancePrivateData::MutableCreateImagePlatformData()"));
 		return nullptr;
 	}
 
diff --git a/Engine/Plugins/Experimental/Mutable/Source/CustomizableObject/Private/MuCO/CustomizableObjectSystem.cpp b/Engine/Plugins/Experimental/Mutable/Source/CustomizableObject/Private/MuCO/CustomizableObjectSystem.cpp
--- a/Engine/Plugins/Experimental/Mutable/Source/CustomizableObject/Private/MuCO/CustomizableObjectSystem.cpp
+++ b/Engine/Plugins/Experimental/Mutable/Source/CustomizableObject/Private/MuCO/CustomizableObjectSystem.cpp
@@ -202,7 +202,7 @@ UCustomizableObjectSystem* UCustomizableObjectSystem::GetInstance()
 {
 	if (!FCustomizableObjectSystemPrivate::SSystem)
 	{
-		UE_LOG(LogMutable, Log, TEXT("Creating system."));
+		UE_LOG(LogMutable, Log, TEXT("Creating Mutable Customizable Object System."));
 
 		check(IsInGameThread());
 
@@ -274,7 +274,7 @@ void UCustomizableObjectSystem::LogShowData(bool bFullInfo, bool ShowMaterialInf
 		{
 			LogInformationUtil::LogShowInstanceData(ArrayData[i], LogData);
 		}
-		UE_LOG(LogMutable, Warning, TEXT("%s"), *LogData);
+		UE_LOG(LogMutable, Log, TEXT("%s"), *LogData);
 
 		UWorld* World = GWorld;
 
@@ -1246,7 +1246,7 @@ namespace impl
 
 			if (!Instance)
 			{
-				UE_LOG(LogMutable, Log, TEXT("An Instace update has failed."));
+				UE_LOG(LogMutable, Warning, TEXT("An Instace update has failed."));
 				return;
 			}
 		}
@@ -1484,7 +1484,7 @@ namespace impl
 
 			if (bCached)
 			{
-				UE_LOG(LogMutable, Verbose, TEXT("Texture resource with id [%d] is cached."), Image.ImageID);
+				UE_LOG(LogMutable, VeryVerbose, TEXT("Texture resource with id [%d] is cached."), Image.ImageID);
 				INC_DWORD_STAT(STAT_MutableNumCachedTextures);
 			}
 			else
diff --git a/Engine/Plugins/Experimental/Mutable/Source/CustomizableObject/Private/MuCO/LogInformationUtil.cpp b/Engine/Plugins/Experimental/Mutable/Source/CustomizableObject/Private/MuCO/LogInformationUtil.cpp
--- a/Engine/Plugins/Experimental/Mutable/Source/CustomizableObject/Private/MuCO/LogInformationUtil.cpp
+++ b/Engine/Plugins/Experimental/Mutable/Source/CustomizableObject/Private/MuCO/LogInformationUtil.cpp
@@ -37,7 +37,7 @@ void LogInformationUtil::PrintGeneratedTextures(const TArray<struct FGeneratedTe
 
 	if (Max > 1)
 	{
-		UE_LOG(LogMutable, Warning, TEXT("Case"));
+		UE_LOG(LogMutable, Log, TEXT("Case"));
 	}
 
 	for (i = 0; i < Max; ++i)
@@ -190,7 +190,7 @@ void LogInformationUtil::LogShowInstanceDataFull(const UCustomizableObjectInstan
 		}
 	}
 
-	UE_LOG(LogMutable, Warning, TEXT("%s"), *LogData);
+	UE_LOG(LogMutable, Log, TEXT("%s"), *LogData);
 
 	UWorld* World = GWorld;
 
diff --git a/Engine/Plugins/Experimental/Mutable/Source/CustomizableObject/Private/MuCO/UnrealMutableModelDiskStreamer.cpp b/Engine/Plugins/Experimental/Mutable/Source/CustomizableObject/Private/MuCO/UnrealMutableModelDiskStreamer.cpp
--- a/Engine/Plugins/Experimental/Mutable/Source/CustomizableObject/Private/MuCO/UnrealMutableModelDiskStreamer.cpp
+++ b/Engine/Plugins/Experimental/Mutable/Source/CustomizableObject/Private/MuCO/UnrealMutableModelDiskStreamer.cpp
@@ -283,7 +283,7 @@ mu::ModelStreamer::OPERATION_ID FUnrealMutableModelBulkStreamer::BeginReadBlock(
 {
 	MUTABLE_CPUPROFILER_SCOPE(FUnrealMutableModelBulkStreamer::OpenReadFile);
 
-	UE_LOG(LogMutable, Verbose, TEXT("Streaming: reading data %08lx."), Key);
+	UE_LOG(LogMutable, VeryVerbose, TEXT("Streaming: reading data %08lx."), Key);
 
 	// Find the object we are streaming for
 	FObjectData* ObjectData = Objects.FindByPredicate(
diff --git a/Engine/Plugins/Experimental/Mutable/Source/CustomizableObjectEditor/Private/MuCOE/CustomizableObjectCompileRunnable.cpp b/Engine/Plugins/Experimental/Mutable/Source/CustomizableObjectEditor/Private/MuCOE/CustomizableObjectCompileRunnable.cpp
--- a/Engine/Plugins/Experimental/Mutable/Source/CustomizableObjectEditor/Private/MuCOE/CustomizableObjectCompileRunnable.cpp
+++ b/Engine/Plugins/Experimental/Mutable/Source/CustomizableObjectEditor/Private/MuCOE/CustomizableObjectCompileRunnable.cpp
@@ -45,7 +45,7 @@ uint32 FCustomizableObjectCompileRunnable::Run()
 
 	if (Options.OptimizationLevel > 3)
 	{
-		UE_LOG(LogMutable, Verbose, TEXT("Mutable compile optimization level out of range. Clamping to maximum."));
+		UE_LOG(LogMutable, Log, TEXT("Mutable compile optimization level out of range. Clamping to maximum."));
 		Options.OptimizationLevel = 3;
 	}
 
diff --git a/Engine/Plugins/Experimental/Mutable/Source/CustomizableObjectEditor/Private/MuCOE/CustomizableObjectCompiler.cpp b/Engine/Plugins/Experimental/Mutable/Source/CustomizableObjectEditor/Private/MuCOE/CustomizableObjectCompiler.cpp
--- a/Engine/Plugins/Experimental/Mutable/Source/CustomizableObjectEditor/Private/MuCOE/CustomizableObjectCompiler.cpp
+++ b/Engine/Plugins/Experimental/Mutable/Source/CustomizableObjectEditor/Private/MuCOE/CustomizableObjectCompiler.cpp
@@ -50,7 +50,7 @@ bool FCustomizableObjectCompiler::Tick()
 
 	if (CompileTask.IsValid() && CompileTask->IsCompleted())
 	{
-		UE_LOG(LogMutable, Log, TEXT("PROFILE: [ %16.8f ] Finishing Compilation task for Object %s."), FPlatformTime::Seconds(), *CurrentObject->GetName());
+		UE_LOG(LogMutable, Verbose, TEXT("PROFILE: [ %16.8f ] Finishing Compilation task for Object %s."), FPlatformTime::Seconds(), *CurrentObject->GetName());
 
 		FinishCompilation();
 
@@ -67,7 +67,7 @@ bool FCustomizableObjectCompiler::Tick()
 
 			NotifyCompilationErrors();
 		}
-		
+
 		UE_LOG(LogMutable, Verbose, TEXT("PROFILE: [ %16.8f ] Finished Compilation task."), FPlatformTime::Seconds());
 		UE_LOG(LogMutable, Verbose, TEXT("PROFILE: -----------------------------------------------------------"));
 
@@ -123,7 +123,7 @@ void FCustomizableObjectCompiler::MutableIsDisabledCase(UCustomizableObject* Obj
 		return;
 	}
 
-	UE_LOG(LogMutable, Warning, TEXT("Mutable has been disabled. To reenable it, please deactivate the Disable Mutale Option in the Plugins -> Mutable option"), FPlatformTime::Seconds());
+	UE_LOG(LogMutable, Log, TEXT("Mutable has been disabled. To reenable it, please deactivate the Disable Mutale Option in the Plugins -> Mutable option"), FPlatformTime::Seconds());
 	CompileTask = MakeShareable(new FCustomizableObjectCompileRunnable(nullptr, true));
 	CompileTask->MutableIsDisabled = true;
 	LaunchMutableCompile(false);
@@ -157,7 +157,7 @@ bool FCustomizableObjectCompiler::IsRootObject(const UCustomizableObject* Object
 
 void FCustomizableObjectCompiler::PreloadingReferencerAssetsCallback(UCustomizableObject* Object, FCustomizableObjectCompiler* CustomizableObjectCompiler, const FCompilationOptions Options, bool bAsync)
 {
-	UE_LOG(LogMutable, Log, TEXT("PROFILE: [ %16.8f ] Preload asynchronously assets end."), FPlatformTime::Seconds());
+	UE_LOG(LogMutable, Verbose, TEXT("PROFILE: [ %16.8f ] Preload asynchronously assets end."), FPlatformTime::Seconds());
 
 	if (CustomizableObjectCompiler->GetCurrentGAsyncLoadingTimeLimit() != -1.0f)
 	{
@@ -188,7 +188,7 @@ void FCustomizableObjectCompiler::Compile(UCustomizableObject& Object, const FCo
 		return;
 	}
 
-	UE_LOG(LogMutable, Log, TEXT("PROFILE: [ %16.8f ] Preload asynchronously assets start."), FPlatformTime::Seconds());
+	UE_LOG(LogMutable, Verbose, TEXT("PROFILE: [ %16.8f ] Preload asynchronously assets start."), FPlatformTime::Seconds());
 
 	FString Message = FString::Printf(TEXT("Customizable Object %s is already being compiled or updated. Please wait a few seconds and try again."), *Object.GetName());
 	FNotificationInfo Info(LOCTEXT("CustomizableObjectBeingCompilerOrUpdated", "Customizable Object compile and/or update still in process. Please wait a few seconds and try again."));
@@ -537,7 +537,7 @@ mu::NodeObjectPtr FCustomizableObjectCompiler::GenerateMutableRoot(
 		// We cannot load while saving. This should only happen in cooking and all assets should have been preloaded.
 		if (!GIsSavingPackage)
 		{
-			UE_LOG(LogMutable, Log, TEXT("PROFILE: [ %16.8f ] Begin search for children."), FPlatformTime::Seconds());
+			UE_LOG(LogMutable, Verbose, TEXT("PROFILE: [ %16.8f ] Begin search for children."), FPlatformTime::Seconds());
 
 			TArray<UCustomizableObject*> VisitedObjects;
 			ActualRootObject = Root->ParentObject ? GetFullGraphRootObject(Root, VisitedObjects) : Object;
@@ -552,7 +552,7 @@ mu::NodeObjectPtr FCustomizableObjectCompiler::GenerateMutableRoot(
 			// The object doesn't reference a root object but is a root object, look for all the objects that reference it and get their root nodes
 			FAssetRegistryModule& AssetRegistryModule = FModuleManager::LoadModuleChecked<FAssetRegistryModule>("AssetRegistry");
 			ProcessChildObjectsRecursively(ActualRootObject, AssetRegistryModule, GenerationContext);
-			UE_LOG(LogMutable, Log, TEXT("PROFILE: [ %16.8f ] End search for children."), FPlatformTime::Seconds());
+			UE_LOG(LogMutable, Verbose, TEXT("PROFILE: [ %16.8f ] End search for children."), FPlatformTime::Seconds());
 		}
 	}
 	else
@@ -653,9 +653,9 @@ mu::NodeObjectPtr FCustomizableObjectCompiler::GenerateMutableRoot(
     GenerationContext.RealTimeMorphTargetsOverrides.Reset();
 
 	// Generate the object expression
-	UE_LOG(LogMutable, Log, TEXT("PROFILE: [ %16.8f ] GenerateMutableSource start."), FPlatformTime::Seconds());
+	UE_LOG(LogMutable, Verbose, TEXT("PROFILE: [ %16.8f ] GenerateMutableSource start."), FPlatformTime::Seconds());
 	mu::NodeObjectPtr MutableRoot = GenerateMutableSource(ActualRoot->OutputPin(), GenerationContext, !bOutIsRootObject);
-	UE_LOG(LogMutable, Log, TEXT("PROFILE: [ %16.8f ] GenerateMutableSource end."), FPlatformTime::Seconds());
+	UE_LOG(LogMutable, Verbose, TEXT("PROFILE: [ %16.8f ] GenerateMutableSource end."), FPlatformTime::Seconds());
 
     ActualRoot->RealTimeMorphSelectionOverrides = GenerationContext.RealTimeMorphTargetsOverrides;
 	GenerationContext.GenerateClippingCOInternalTags();
@@ -922,7 +922,7 @@ void FCustomizableObjectCompiler::CompileInternal(UCustomizableObject* Object, c
 		return;
 	}
 
-	UE_LOG(LogMutable, Display, TEXT("Started Customizable Object Compile %s."), *Object->GetName());	
+	UE_LOG(LogMutable, Log, TEXT("Started Customizable Object Compile %s."), *Object->GetName());	
 
 	CompilationLogsContainer.ClearMessageCounters();
 
@@ -1345,7 +1345,7 @@ void FCustomizableObjectCompiler::CompileInternal(UCustomizableObject* Object, c
 
 mu::NodePtr FCustomizableObjectCompiler::Export(UCustomizableObject* Object, const FCompilationOptions& InCompilerOptions)
 {
-	UE_LOG(LogMutable, Display, TEXT("Started Customizable Object Export %s."), *Object->GetName());
+	UE_LOG(LogMutable, Log, TEXT("Started Customizable Object Export %s."), *Object->GetName());
 
 	FNotificationInfo Info(LOCTEXT("CustomizableObjectExportInProgress", "Exported Customizable Object"));
 	Info.bFireAndForget = true;
diff --git a/Engine/Plugins/Experimental/Mutable/Source/CustomizableObjectEditor/Private/MuCOE/CustomizableObjectEditor.cpp b/Engine/Plugins/Experimental/Mutable/Source/CustomizableObjectEditor/Private/MuCOE/CustomizableObjectEditor.cpp
--- a/Engine/Plugins/Experimental/Mutable/Source/CustomizableObjectEditor/Private/MuCOE/CustomizableObjectEditor.cpp
+++ b/Engine/Plugins/Experimental/Mutable/Source/CustomizableObjectEditor/Private/MuCOE/CustomizableObjectEditor.cpp
@@ -1428,7 +1428,7 @@ void FCustomizableObjectEditor::CompileObject()
 		Compiler.Compile(*CustomizableObject, Options, true);
 	}
 
-	UE_LOG(LogMutable, Log, TEXT("PROFILE: [ %16.8f ] FCustomizableObjectEditor::CompileObject end."), FPlatformTime::Seconds());
+	UE_LOG(LogMutable, Verbose, TEXT("PROFILE: [ %16.8f ] FCustomizableObjectEditor::CompileObject end."), FPlatformTime::Seconds());
 }
 
 
diff --git a/Engine/Plugins/Experimental/Mutable/Source/CustomizableObjectEditor/Private/MuCOE/CustomizableObjectInstanceEditor.cpp b/Engine/Plugins/Experimental/Mutable/Source/CustomizableObjectEditor/Private/MuCOE/CustomizableObjectInstanceEditor.cpp
--- a/Engine/Plugins/Experimental/Mutable/Source/CustomizableObjectEditor/Private/MuCOE/CustomizableObjectInstanceEditor.cpp
+++ b/Engine/Plugins/Experimental/Mutable/Source/CustomizableObjectEditor/Private/MuCOE/CustomizableObjectInstanceEditor.cpp
@@ -786,8 +786,8 @@ void FCustomizableObjectInstanceEditor::OnUpdatePreviewInstance()
 
 void FCustomizableObjectInstanceEditor::CompileObject(UCustomizableObject* Object)
 {
-	UE_LOG(LogMutable, Log, TEXT("PROFILE: -----------------------------------------------------------"));
-	UE_LOG(LogMutable, Log, TEXT("PROFILE: [ %16.8f ] FCustomizableObjectInstanceEditor::CompileObject start."), FPlatformTime::Seconds());
+	UE_LOG(LogMutable, Verbose, TEXT("PROFILE: -----------------------------------------------------------"));
+	UE_LOG(LogMutable, Verbose, TEXT("PROFILE: [ %16.8f ] FCustomizableObjectInstanceEditor::CompileObject start."), FPlatformTime::Seconds());
 
 	if (!AssetRegistryLoaded)
 	{
@@ -809,7 +809,7 @@ void FCustomizableObjectInstanceEditor::CompileObject(UCustomizableObject* Objec
 		Compiler->Compile(*Object, Options, true);
 	}
 
-	UE_LOG(LogMutable, Log, TEXT("PROFILE: [ %16.8f ] FCustomizableObjectInstanceEditor::CompileObject end."), FPlatformTime::Seconds());
+	UE_LOG(LogMutable, Verbose, TEXT("PROFILE: [ %16.8f ] FCustomizableObjectInstanceEditor::CompileObject end."), FPlatformTime::Seconds());
 }
 
 
diff --git a/Engine/Plugins/Experimental/Mutable/Source/CustomizableObjectEditor/Private/MuCOE/GenerateMutableSource/GenerateMutableSourceMesh.cpp b/Engine/Plugins/Experimental/Mutable/Source/CustomizableObjectEditor/Private/MuCOE/GenerateMutableSource/GenerateMutableSourceMesh.cpp
--- a/Engine/Plugins/Experimental/Mutable/Source/CustomizableObjectEditor/Private/MuCOE/GenerateMutableSource/GenerateMutableSourceMesh.cpp
+++ b/Engine/Plugins/Experimental/Mutable/Source/CustomizableObjectEditor/Private/MuCOE/GenerateMutableSource/GenerateMutableSourceMesh.cpp
@@ -1106,7 +1106,7 @@ mu::MeshPtr ConvertSkeletalMeshToMutable(USkeletalMesh* InSkeletalMesh, const TS
 				}
 			}
 
-			UE_LOG(LogMutable, Log, TEXT("Processing morph targets took %.2f ms"), (FPlatformTime::Seconds() - StartTime) * 1000.0);
+			UE_LOG(LogMutable, Verbose, TEXT("Processing morph targets took %.2f ms"), (FPlatformTime::Seconds() - StartTime) * 1000.0);
 		}
 	}
 
diff --git a/Engine/Plugins/Experimental/Mutable/Source/CustomizableObjectEditor/Private/MuCOE/GenerateMutableSource/GenerateMutableSourceSurface.cpp b/Engine/Plugins/Experimental/Mutable/Source/CustomizableObjectEditor/Private/MuCOE/GenerateMutableSource/GenerateMutableSourceSurface.cpp
--- a/Engine/Plugins/Experimental/Mutable/Source/CustomizableObjectEditor/Private/MuCOE/GenerateMutableSource/GenerateMutableSourceSurface.cpp
+++ b/Engine/Plugins/Experimental/Mutable/Source/CustomizableObjectEditor/Private/MuCOE/GenerateMutableSource/GenerateMutableSourceSurface.cpp
@@ -93,7 +93,7 @@ void SetSurfaceFormat( FMutableGraphGenerationContext& GenerationContext,
 	
 	if (MutableBonesPerVertex != MeshData.MaxNumBonesPerVertex)
 	{
-		UE_LOG(LogMutable, Log, TEXT("In object [%s] Mesh bone number adjusted from %d to %d."), *GenerationContext.Object->GetName(), MeshData.MaxNumBonesPerVertex, MutableBonesPerVertex);
+		UE_LOG(LogMutable, Verbose, TEXT("In object [%s] Mesh bone number adjusted from %d to %d."), *GenerationContext.Object->GetName(), MeshData.MaxNumBonesPerVertex, MutableBonesPerVertex);
 	}
 
 	int MutableBufferCount = MUTABLE_VERTEXBUFFER_TEXCOORDS + 1;
diff --git a/Engine/Plugins/Experimental/Mutable/Source/CustomizableObjectEditor/Private/MuCOE/MutableMeshPreviewUtils.cpp b/Engine/Plugins/Experimental/Mutable/Source/CustomizableObjectEditor/Private/MuCOE/MutableMeshPreviewUtils.cpp
--- a/Engine/Plugins/Experimental/Mutable/Source/CustomizableObjectEditor/Private/MuCOE/MutableMeshPreviewUtils.cpp
+++ b/Engine/Plugins/Experimental/Mutable/Source/CustomizableObjectEditor/Private/MuCOE/MutableMeshPreviewUtils.cpp
@@ -661,7 +661,7 @@ namespace MutableMeshPreviewUtils
 					return nullptr;
 				}
 				
-				UE_LOG(LogMutable,Log,TEXT("Buffers Restructuring was succesfull."));
+				UE_LOG(LogMutable, Verbose, TEXT("Buffers Restructuring was succesfull."));
 			}
 
 			// We should now have a mutable mesh following a buffer structure where the required buffers are there to be
diff --git a/Engine/Plugins/Experimental/Mutable/Source/CustomizableObjectEditor/Private/MuCOE/SCustomizableInstanceProperties.cpp b/Engine/Plugins/Experimental/Mutable/Source/CustomizableObjectEditor/Private/MuCOE/SCustomizableInstanceProperties.cpp
--- a/Engine/Plugins/Experimental/Mutable/Source/CustomizableObjectEditor/Private/MuCOE/SCustomizableInstanceProperties.cpp
+++ b/Engine/Plugins/Experimental/Mutable/Source/CustomizableObjectEditor/Private/MuCOE/SCustomizableInstanceProperties.cpp
@@ -1333,7 +1333,7 @@ FReply SCustomizableInstanceProperties::OnCopyAllParameters()
 {
 	ToBinary = new FBufferArchive(false, TEXT("BufferArchive"));
 
-	UE_LOG(LogMutable, Warning, TEXT("Parameters have been copied."));
+	UE_LOG(LogMutable, Log, TEXT("Parameters have been copied."));
 
 	CustomInstance->SaveDescriptor(*ToBinary);
 
diff --git a/Engine/Plugins/Experimental/Mutable/Source/CustomizableObjectEditor/Private/MuCOE/SMutableCodeViewer.cpp b/Engine/Plugins/Experimental/Mutable/Source/CustomizableObjectEditor/Private/MuCOE/SMutableCodeViewer.cpp
--- a/Engine/Plugins/Experimental/Mutable/Source/CustomizableObjectEditor/Private/MuCOE/SMutableCodeViewer.cpp
+++ b/Engine/Plugins/Experimental/Mutable/Source/CustomizableObjectEditor/Private/MuCOE/SMutableCodeViewer.cpp
@@ -2493,7 +2493,7 @@ void SMutableCodeViewer::Tick(const FGeometry& AllottedGeometry, const double In
 	
 	default:
 #if UE_BUILD_DEBUG || UE_BUILD_DEVELOPMENT
-		UE_LOG(LogMutable,Display,TEXT("There is no previewer for the selected type of Mutable object"))
+		UE_LOG(LogMutable, Log, TEXT("There is no previewer for the selected type of Mutable object"))
 #endif
 		// There is no viewer for this type.
 		break;
@@ -2649,13 +2649,13 @@ void SMutableCodeViewer::PreviewMutableCurve(const mu::Curve* Curve)
 // TODO: Implement matrix viewer
 void SMutableCodeViewer::PreviewMutableMatrix(const mu::mat4f* Mat)
 {
-	UE_LOG(LogMutable,Warning,TEXT("Previewer for Mutable Matrices not yet implemented"))
+	UE_LOG(LogMutable, Warning, TEXT("Previewer for Mutable Matrices not yet implemented"))
 }
 
 // TODO: Implement shape viewer
 void SMutableCodeViewer::PreviewMutableShape(const mu::FShape* Shape)
 {
-	UE_LOG(LogMutable,Warning,TEXT("Previewer for Mutable Shapes not yet implemented"))
+	UE_LOG(LogMutable, Warning, TEXT("Previewer for Mutable Shapes not yet implemented"))
 }
 
 
diff --git a/Engine/Plugins/Experimental/Mutable/Source/CustomizableObjectEditor/Private/MuCOE/StressTest.cpp b/Engine/Plugins/Experimental/Mutable/Source/CustomizableObjectEditor/Private/MuCOE/StressTest.cpp
--- a/Engine/Plugins/Experimental/Mutable/Source/CustomizableObjectEditor/Private/MuCOE/StressTest.cpp
+++ b/Engine/Plugins/Experimental/Mutable/Source/CustomizableObjectEditor/Private/MuCOE/StressTest.cpp
@@ -157,17 +157,17 @@ bool FRunningStressTest::RunStressTestTick(float DeltaTime)
 
 		if (NumUpdatedDone == NumUpdateRequested)
 		{
-			UE_LOG(LogMutable, Warning, TEXT("\t\t\t Final Stress Test results for model %s"), *CustomizableObject->GetName());
+			UE_LOG(LogMutable, Log, TEXT("\t\t\t Final Stress Test results for model %s"), *CustomizableObject->GetName());
 		}
 
 		WriteTestResults();
 
 		PrintMostExpensivePerLOD();
 
-		UE_LOG(LogMutable, Warning, TEXT("///////////////////////////////////////////////////////////////////////////////////"));
+		UE_LOG(LogMutable, Log, TEXT("///////////////////////////////////////////////////////////////////////////////////"));
 
-		UE_LOG(LogMutable, Warning, TEXT("Finished stress test for model [%s]."), *CustomizableObject->GetName());
-		UE_LOG(LogMutable, Warning, TEXT("\n"));
+		UE_LOG(LogMutable, Log, TEXT("Finished stress test for model [%s]."), *CustomizableObject->GetName());
+		UE_LOG(LogMutable, Log, TEXT("\n"));
 
 		FinishTest();
 
@@ -237,7 +237,7 @@ bool FRunningStressTest::RunStressTestTick(float DeltaTime)
 				NextInstanceTimeMs -= remainingTime;
 			}
 
-			UE_LOG(LogMutable, Warning, TEXT("Stress test created instance. [%d active]"), LiveInstances.Num());
+			UE_LOG(LogMutable, Log, TEXT("Stress test created instance. [%d active]"), LiveInstances.Num());
 		}
 
 		// Update instances that need it
@@ -283,7 +283,7 @@ bool FRunningStressTest::RunStressTestTick(float DeltaTime)
 
 		if (ToRemove.Num())
 		{
-			UE_LOG(LogMutable, Warning, TEXT("Stress test removed instance. [%d active]"), LiveInstances.Num());
+			UE_LOG(LogMutable, Log, TEXT("Stress test removed instance. [%d active]"), LiveInstances.Num());
 		}
 	}
 
@@ -326,7 +326,7 @@ void FRunningStressTest::WriteTestResults()
 {
 	FString Line;
 	Line = "///////////////////////////////////////////////////////////////////////////////////";
-	UE_LOG(LogMutable, Warning, TEXT("%s"), *Line);
+	UE_LOG(LogMutable, Log, TEXT("%s"), *Line);
 
 	for (uint32 i = 0; i < NumLOD; ++i)
 	{
@@ -343,7 +343,7 @@ void FRunningStressTest::WriteTestResults()
 
 		Line = "Results for LOD level ";
 		Line += FString::Printf(TEXT("%d"), i);
-		UE_LOG(LogMutable, Warning, TEXT("%s:"), *Line);
+		UE_LOG(LogMutable, Log, TEXT("%s:"), *Line);
 
 		uint32 MaxLength1 = MaxOffsetName + 25;
 		uint32 MaxLength2 = MaxLength1 + 25;
@@ -355,10 +355,10 @@ void FRunningStressTest::WriteTestResults()
 		Line += "Mean";
 		AppendWhitespace(Line, MaxLength2 - Line.Len());
 		Line += "Max";
-		UE_LOG(LogMutable, Warning, TEXT("%s"), *Line);
+		UE_LOG(LogMutable, Log, TEXT("%s"), *Line);
 
 		Line = "\t ---------------------------------------------------------------------------------------------";
-		UE_LOG(LogMutable, Warning, TEXT("%s"), *Line);
+		UE_LOG(LogMutable, Log, TEXT("%s"), *Line);
 
 		for (TPair<FString, MeasuredData>& it : ArrayMeasuredDataPerLOD[i])
 		{
@@ -375,7 +375,7 @@ void FRunningStressTest::WriteTestResults()
 			Line += FString::Printf(TEXT("%d"), int32(it.Value.MeanTime));
 			AppendWhitespace(Line, MaxLength2 - Line.Len());
 			Line += FString::Printf(TEXT("%d"), int32(it.Value.MaxTime));
-			UE_LOG(LogMutable, Warning, TEXT("%s"), *Line);
+			UE_LOG(LogMutable, Log, TEXT("%s"), *Line);
 		}
 
 		// Texture size information
@@ -392,7 +392,7 @@ void FRunningStressTest::WriteTestResults()
 		Value = (MaxTextureSize > 1048576) ? float(MaxTextureSize) / 1048576.0f : float(MaxTextureSize) / 1024.0f;
 		Line += FString::Printf(TEXT("%.2f"), Value);
 		Line += (MaxTextureSize > 1048576) ? "(MB)" : "(KB)";
-		UE_LOG(LogMutable, Warning, TEXT("%s\n"), *Line);
+		UE_LOG(LogMutable, Log, TEXT("%s\n"), *Line);
 	}
 }
 
@@ -459,7 +459,7 @@ void FRunningStressTest::AddInstanceData(MeasureDataMap MapMeasuredParam, Materi
 
 	FString Line = "Test update instance number ";
 	Line += FString::Printf(TEXT("%d"), NumUpdatedDone);
-	UE_LOG(LogMutable, Warning, TEXT("%s"), *Line);
+	UE_LOG(LogMutable, Log, TEXT("%s"), *Line);
 
 	WriteTestResults();
 }
@@ -597,19 +597,19 @@ void FRunningStressTest::PrintMostExpensivePerLOD()
 	for (uint32 i = 0; i < MaxIndexLOD; ++i)
 	{
 		FString Line = "///////////////////////////////////////////////////////////////////////////////////";
-		UE_LOG(LogMutable, Warning, TEXT("%s:"), *Line);
+		UE_LOG(LogMutable, Log, TEXT("%s:"), *Line);
 
 		Line = "Instance detailed information for LOD level ";
 		Line += FString::Printf(TEXT("%d:"), i);
-		UE_LOG(LogMutable, Warning, TEXT("%s\n"), *Line);
+		UE_LOG(LogMutable, Log, TEXT("%s\n"), *Line);
 
 		Line = "\t Information for most expensive instance in terms of geometry size";
-		UE_LOG(LogMutable, Warning, TEXT("%s:"), *Line);
+		UE_LOG(LogMutable, Log, TEXT("%s:"), *Line);
 		MostExpensiveInstanceData& DataGeometry = ArrayLODMostExpensiveGeometry[i];
 		PrintMostExpensiveInstanceData(DataGeometry, i);
 
 		Line = "\t Information for most expensive instance in terms of material texture size";
-		UE_LOG(LogMutable, Warning, TEXT("%s:"), *Line);
+		UE_LOG(LogMutable, Log, TEXT("%s:"), *Line);
 		MostExpensiveInstanceData& DataTexture = ArrayLODMostExpensiveTexture[i];
 		PrintMostExpensiveInstanceData(DataTexture, i);
 	}
@@ -631,7 +631,7 @@ void FRunningStressTest::PrintParameterData(const MostExpensiveInstanceData& Dat
 			Line = FString::Printf(TEXT("\t\t\t %s"), *it.ParameterName);
 			AppendWhitespace(Line, MaxLength - Line.Len());
 			Line += FString::Printf(TEXT("%d"), it.ParameterValue);
-			UE_LOG(LogMutable, Warning, TEXT("%s"), *Line);
+			UE_LOG(LogMutable, Log, TEXT("%s"), *Line);
 		}
 	}
 
@@ -643,7 +643,7 @@ void FRunningStressTest::PrintParameterData(const MostExpensiveInstanceData& Dat
 			Line = FString::Printf(TEXT("\t\t\t %s"), *it.ParameterName);
 			AppendWhitespace(Line, MaxLength - Line.Len());
 			Line += FString::Printf(TEXT("%s"), *it.ParameterValueName);
-			UE_LOG(LogMutable, Warning, TEXT("%s"), *Line);
+			UE_LOG(LogMutable, Log, TEXT("%s"), *Line);
 		}
 	}
 
@@ -655,7 +655,7 @@ void FRunningStressTest::PrintParameterData(const MostExpensiveInstanceData& Dat
 			Line = FString::Printf(TEXT("\t\t\t %s"), *it.ParameterName);
 			AppendWhitespace(Line, MaxLength - Line.Len());
 			Line += FString::Printf(TEXT("%.2f"), it.ParameterValue);
-			UE_LOG(LogMutable, Warning, TEXT("%s"), *Line);
+			UE_LOG(LogMutable, Log, TEXT("%s"), *Line);
 		}
 	}
 
@@ -667,7 +667,7 @@ void FRunningStressTest::PrintParameterData(const MostExpensiveInstanceData& Dat
 			Line = FString::Printf(TEXT("\t\t\t %s"), *it.ParameterName);
 			AppendWhitespace(Line, MaxLength - Line.Len());
 			Line += FString::Printf(TEXT("%d"), it.ParameterValue);
-			UE_LOG(LogMutable, Warning, TEXT("%s"), *Line);
+			UE_LOG(LogMutable, Log, TEXT("%s"), *Line);
 		}
 	}
 
@@ -682,7 +682,7 @@ void FRunningStressTest::PrintParameterData(const MostExpensiveInstanceData& Dat
 			Line += FString::Printf(TEXT("%.2f,"), it.ParameterValue.G);
 			Line += FString::Printf(TEXT("%.2f,"), it.ParameterValue.B);
 			Line += FString::Printf(TEXT("%.2f)"), it.ParameterValue.A);
-			UE_LOG(LogMutable, Warning, TEXT("%s"), *Line);
+			UE_LOG(LogMutable, Log, TEXT("%s"), *Line);
 		}
 	}
 
@@ -696,7 +696,7 @@ void FRunningStressTest::PrintParameterData(const MostExpensiveInstanceData& Dat
 			Line += FString::Printf(TEXT("(%.2f,"), it.Value.Position.X);
 			Line += FString::Printf(TEXT("%.2f,"), it.Value.Position.Y);
 			Line += FString::Printf(TEXT("%.2f)"), it.Value.Position.Z);
-			UE_LOG(LogMutable, Warning, TEXT("%s"), *Line);
+			UE_LOG(LogMutable, Log, TEXT("%s"), *Line);
 		}
 	}
 }
@@ -742,15 +742,15 @@ void FRunningStressTest::PrintParameterHeader(FString Param, int32 Value, bool A
 {
 	if (AddNewLine)
 	{
-		UE_LOG(LogMutable, Warning, TEXT("\n"));
+		UE_LOG(LogMutable, Log, TEXT("\n"));
 	}
 	FString Line = Param;
-	UE_LOG(LogMutable, Warning, TEXT("\t\t\t %s"), *Line);
+	UE_LOG(LogMutable, Log, TEXT("\t\t\t %s"), *Line);
 	Line = "\t\t\t Name";
 	AppendWhitespace(Line, Value - Line.Len());
 	Line += "Value";
-	UE_LOG(LogMutable, Warning, TEXT("%s"), *Line);
-	UE_LOG(LogMutable, Warning, TEXT("\t\t\t ------------------------------------------"));
+	UE_LOG(LogMutable, Log, TEXT("%s"), *Line);
+	UE_LOG(LogMutable, Log, TEXT("\t\t\t ------------------------------------------"));
 }
 
 
@@ -759,7 +759,7 @@ void FRunningStressTest::PrintMostExpensiveInstanceData(const MostExpensiveInsta
 	FString Line;
 	Line = FString::Printf(TEXT("\t\t Number of faces: "));
 	Line += FString::Printf(TEXT("%d"), Data.CurrentLODNumFaces);
-	UE_LOG(LogMutable, Warning, TEXT("%s"), *Line);
+	UE_LOG(LogMutable, Log, TEXT("%s"), *Line);
 
 	// Texture size information
 	Line = "\t\t Texture size: ";
@@ -769,7 +769,7 @@ void FRunningStressTest::PrintMostExpensiveInstanceData(const MostExpensiveInsta
 	Value = (MaxTextureSize > 1048576) ? float(MaxTextureSize) / 1048576.0f : float(MaxTextureSize) / 1024.0f;
 	Line += FString::Printf(TEXT(", max=%.2f"), Value);
 	Line += (MaxTextureSize > 1048576) ? "(MB)" : "(KB)";
-	UE_LOG(LogMutable, Warning, TEXT("%s\n"), *Line);
+	UE_LOG(LogMutable, Log, TEXT("%s\n"), *Line);
 
 	// Material information
 	const TMap<FString, MaterialBriefInfo>& MapMaterialLOD = Data.MapMaterial;
@@ -785,7 +785,7 @@ void FRunningStressTest::PrintMostExpensiveInstanceData(const MostExpensiveInsta
 		Line += FString::Printf(TEXT(" (%d elements):"), MapMaterialLOD.Num());
 	}
 
-	UE_LOG(LogMutable, Warning, TEXT("%s"), *Line);
+	UE_LOG(LogMutable, Log, TEXT("%s"), *Line);
 
 	int32 MaxOffsetName = 0;
 	for (const TPair<FString, MaterialBriefInfo>& itMaterial : MapMaterialLOD)
@@ -800,7 +800,7 @@ void FRunningStressTest::PrintMostExpensiveInstanceData(const MostExpensiveInsta
 	for (const TPair<FString, MaterialBriefInfo>& itMaterial : MapMaterialLOD)
 	{
 		Line = FString::Printf(TEXT("\t\t\t Material %s: "), *(itMaterial.Key));
-		UE_LOG(LogMutable, Warning, TEXT("%s"), *Line);
+		UE_LOG(LogMutable, Log, TEXT("%s"), *Line);
 
 		int32 MaxIndex = itMaterial.Value.ArrayTextureName.Num();
 		for (int32 i = 0; i < MaxIndex; ++i)
@@ -812,17 +812,17 @@ void FRunningStressTest::PrintMostExpensiveInstanceData(const MostExpensiveInsta
 			AppendWhitespace(Line, SecondOffsetTexture - Line.Len());
 			TextureFormatString = StringyfyEPixelFormat(itMaterial.Value.ArrayTextureFormat[i]);
 			Line += FString::Printf(TEXT("format=%s"), *TextureFormatString);
-			UE_LOG(LogMutable, Warning, TEXT("%s"), *Line);
+			UE_LOG(LogMutable, Log, TEXT("%s"), *Line);
 		}
 	}
 
-	UE_LOG(LogMutable, Warning, TEXT("\n"));
+	UE_LOG(LogMutable, Log, TEXT("\n"));
 
 	if ((Data.BoolParameters.Num() > 0) || (Data.IntParameters.Num() > 0) || (Data.FloatParameters.Num() > 0) ||
 		(Data.TextureParameters.Num() > 0) || (Data.VectorParameters.Num() > 0) || (Data.ProjectorParameters.Num() > 0))
 	{
 		Line = FString::Printf(TEXT("\t\t Parameter information"));
-		UE_LOG(LogMutable, Warning, TEXT("%s"), *Line);
+		UE_LOG(LogMutable, Log, TEXT("%s"), *Line);
 
 		// Parameter information
 		PrintParameterData(Data);
