Added logging an Error to UCustomizableObjectSystem::AddUncompiledCOWarning, and also added an optional parameter for the caller to add more info to the log (so we can easily log the source of the call to figure out what is being called that is trying to use a CO that hasn't been compiled yet).

Created a private helper function FCustomizableObjectInstanceDescriptor::AddUncompiledCOWarning(...), which makes sure the System exists, is valid, and isn't being destroyed before calling to AddUncompiledCOWarning on the System.

SetIntParameterSelectedOption(FString ParameterName,...) now ensures and calls to AddUncompiledCOWarning (which logs an error, makes an on-screen display, and sends a notfiication in Editor) if the CO that the descriptor is associated with is not compiled.  If it is compiled, but the named parameter is not found, it ensures (in non-shipping builds only) and logs an error (in all builds).

SetIntParameterSelectedOption((const int32 IntParamIndex, ...) now also ensure and calls AddUncompiledCOWarning  as the string version does if its CO is not compiled.  If it is compiled, but the IntParamIndex is invalid, it logs an error and bails.

(We should make similar ensure/logging changes to all other parameter type setters.
 I'll investigate that in the near future.)

#UE5

diff --git a/Engine/Plugins/Experimental/Mutable/Source/CustomizableObject/Classes/MuCO/CustomizableObjectSystem.h b/Engine/Plugins/Experimental/Mutable/Source/CustomizableObject/Classes/MuCO/CustomizableObjectSystem.h
--- a/Engine/Plugins/Experimental/Mutable/Source/CustomizableObject/Classes/MuCO/CustomizableObjectSystem.h
+++ b/Engine/Plugins/Experimental/Mutable/Source/CustomizableObject/Classes/MuCO/CustomizableObjectSystem.h
@@ -307,8 +307,9 @@ private:
 
 public:
     
-	// Show a warning when a CustomizableObject is being used and it's not compiled 
-	void AddUncompiledCOWarning(UCustomizableObject* InObject);
+	// Show a warning on-screen and via a notification (if in Editor) and log an error when a CustomizableObject is
+	// being used and it's not compiled.  Callers can add additional information to the error log.
+	void AddUncompiledCOWarning(UCustomizableObject* InObject, FString const* OptionalLogInfo = nullptr);
 
 	// Enables the collection of internal mutabe performance data. It has a performance cost.
 	void EnableBenchmark();
diff --git a/Engine/Plugins/Experimental/Mutable/Source/CustomizableObject/Private/MuCO/CustomizableObjectInstanceDescriptor.cpp b/Engine/Plugins/Experimental/Mutable/Source/CustomizableObject/Private/MuCO/CustomizableObjectInstanceDescriptor.cpp
--- a/Engine/Plugins/Experimental/Mutable/Source/CustomizableObject/Private/MuCO/CustomizableObjectInstanceDescriptor.cpp
+++ b/Engine/Plugins/Experimental/Mutable/Source/CustomizableObject/Private/MuCO/CustomizableObjectInstanceDescriptor.cpp
@@ -1155,64 +1155,110 @@ const FString& FCustomizableObjectInstanceDescriptor::GetIntParameterSelectedOpt
 	return FCustomizableObjectBoolParameterValue::DEFAULT_PARAMETER_VALUE_NAME;
 }
 
+void FCustomizableObjectInstanceDescriptor::AddUncompiledCOWarning(const FString& AdditionalLoggingInfo)
+{
+	// Send a warning (on-screen notification, log error, and in-editor notification)
+	UCustomizableObjectSystem* System = UCustomizableObjectSystem::GetInstance();
+	if (!System || !System->IsValidLowLevel() || System->HasAnyFlags(RF_BeginDestroyed))
+	{
+		return;
+	}
+
+	System->AddUncompiledCOWarning(CustomizableObject.Get(), &AdditionalLoggingInfo);
+}
 
 void FCustomizableObjectInstanceDescriptor::SetIntParameterSelectedOption(const int32 IntParamIndex, const FString& SelectedOption, const int32 RangeIndex)
 {
 	check(CustomizableObject);
 
-	if (IntParameters.IsValidIndex(IntParamIndex))
+	if (!ensureMsgf(CustomizableObject->IsCompiled(), TEXT("Customizable Object (%s) was not compiled."),
+					*GetNameSafe(CustomizableObject)
+				   )
+	   )
 	{
-		const int32 ParameterIndexInObject = CustomizableObject->FindParameter(IntParameters[IntParamIndex].ParameterName);
-		if (ParameterIndexInObject >= 0)
-		{
-			const bool bValid = SelectedOption == TEXT("None") || CustomizableObject->FindIntParameterValue(ParameterIndexInObject, SelectedOption) >= 0;
-			if (!bValid)
-			{
-				const FString Message = FString::Printf(
-					TEXT("LogMutable: Tried to set the invalid value [%s] to parameter [%d, %s]! Value index=[%d]. Correct values=[%s]."), 
-					*SelectedOption, ParameterIndexInObject,
-					*IntParameters[IntParamIndex].ParameterName, 
-					CustomizableObject->FindIntParameterValue(ParameterIndexInObject, SelectedOption), 
-					*GetAvailableOptionsString(*CustomizableObject, ParameterIndexInObject)
-				);
-				UE_LOG(LogMutable, Error, TEXT("%s"), *Message);
-			}
-			
-			if (RangeIndex == -1)
-			{
-				// If this param were multidimensional, it must have a RangeIndex of 0 or more
-				check(!CustomizableObject->IsParameterMultidimensional(ParameterIndexInObject)); 
-				IntParameters[IntParamIndex].ParameterValueName = SelectedOption;
-			}
-			else
-			{
-				// If this param were not multidimensional, it must have a RangeIndex of -1
-				check(CustomizableObject->IsParameterMultidimensional(ParameterIndexInObject));
+		FString AdditionalLoggingInfo = FString::Printf(
+			TEXT("Calling function: %hs.  Cannot set parameter (Parameter Index: %d, SelectedOptionName: %s, RangeIndex: %d)."),
+			__FUNCTION__, IntParamIndex, *SelectedOption, RangeIndex
+													   );
 
-				if (!IntParameters[IntParamIndex].ParameterRangeValueNames.IsValidIndex(RangeIndex))
-				{
-					const int32 InsertionIndex = IntParameters[IntParamIndex].ParameterRangeValueNames.Num();
-					const int32 NumInsertedElements = RangeIndex + 1 - IntParameters[IntParamIndex].ParameterRangeValueNames.Num();
-					IntParameters[IntParamIndex].ParameterRangeValueNames.InsertDefaulted(InsertionIndex, NumInsertedElements);
-				}
+		// Send a warning (on-screen notification, log error, and in-editor notification) for each occurrence.
+		AddUncompiledCOWarning(AdditionalLoggingInfo);
 
-				check(IntParameters[IntParamIndex].ParameterRangeValueNames.IsValidIndex(RangeIndex));
-				IntParameters[IntParamIndex].ParameterRangeValueNames[RangeIndex] = SelectedOption;
-			}
-		}
+		// Early out since we can't set any parameters on an uncompiled CO.
+		return;
 	}
-	else
+
+	if (!IntParameters.IsValidIndex(IntParamIndex))
 	{
 		UE_LOG(LogMutable, Error,
 				TEXT("%hs: IntParamIndex (%d) is not valid on CO (%s).  (SelectedOptionName: %s, RangeIndex: %d).  Cannot set parameter."),
 				__FUNCTION__, IntParamIndex, *GetNameSafe(CustomizableObject), *SelectedOption, RangeIndex
 			  );
+
+		return;
 	}
-}
 
+	const int32 ParameterIndexInObject = CustomizableObject->FindParameter(IntParameters[IntParamIndex].ParameterName);
+	if (ParameterIndexInObject >= 0)
+	{
+		const bool bValid = SelectedOption == TEXT("None") || CustomizableObject->FindIntParameterValue(ParameterIndexInObject, SelectedOption) >= 0;
+		if (!bValid)
+		{
+			const FString Message = FString::Printf(
+				TEXT("LogMutable: Tried to set the invalid value [%s] to parameter [%d, %s]! Value index=[%d]. Correct values=[%s]."), 
+				*SelectedOption, ParameterIndexInObject,
+				*IntParameters[IntParamIndex].ParameterName, 
+				CustomizableObject->FindIntParameterValue(ParameterIndexInObject, SelectedOption), 
+				*GetAvailableOptionsString(*CustomizableObject, ParameterIndexInObject)
+			);
+			UE_LOG(LogMutable, Error, TEXT("%s"), *Message);
+		}
+			
+		if (RangeIndex == -1)
+		{
+			// If this param were multidimensional, it must have a RangeIndex of 0 or more
+			check(!CustomizableObject->IsParameterMultidimensional(ParameterIndexInObject)); 
+			IntParameters[IntParamIndex].ParameterValueName = SelectedOption;
+		}
+		else
+		{
+			// If this param were not multidimensional, it must have a RangeIndex of -1
+			check(CustomizableObject->IsParameterMultidimensional(ParameterIndexInObject));
+
+			if (!IntParameters[IntParamIndex].ParameterRangeValueNames.IsValidIndex(RangeIndex))
+			{
+				const int32 InsertionIndex = IntParameters[IntParamIndex].ParameterRangeValueNames.Num();
+				const int32 NumInsertedElements = RangeIndex + 1 - IntParameters[IntParamIndex].ParameterRangeValueNames.Num();
+				IntParameters[IntParamIndex].ParameterRangeValueNames.InsertDefaulted(InsertionIndex, NumInsertedElements);
+			}
+
+			check(IntParameters[IntParamIndex].ParameterRangeValueNames.IsValidIndex(RangeIndex));
+			IntParameters[IntParamIndex].ParameterRangeValueNames[RangeIndex] = SelectedOption;
+		}
+	}
+}
 
 void FCustomizableObjectInstanceDescriptor::SetIntParameterSelectedOption(const FString& ParamName, const FString& SelectedOptionName, const int32 RangeIndex)
 {
+	check(CustomizableObject);
+
+	if (!ensureMsgf(CustomizableObject->IsCompiled(), TEXT("Customizable Object (%s) was not compiled."),
+					*GetNameSafe(CustomizableObject)
+				   )
+	   )
+	{
+		FString AdditionalLoggingInfo = FString::Printf(
+			TEXT("Calling function: %hs.  Cannot set parameter (Parameter Name: %s, SelectedOptionName: %s, RangeIndex: %d)."),
+			__FUNCTION__, *ParamName, *SelectedOptionName, RangeIndex
+													   );
+
+		// Send a warning (on-screen notification, log error, and in-editor notification) for each occurrence.
+		AddUncompiledCOWarning(AdditionalLoggingInfo);
+
+		// Early out since we can't set any parameters on an uncompiled CO.
+		return;
+	}
+
 	const int32 IntParamIndex = FindIntParameterNameIndex(ParamName);
 	if (IntParamIndex != INDEX_NONE)
 	{
@@ -1220,6 +1266,15 @@ void FCustomizableObjectInstanceDescriptor::SetIntParameterSelectedOption(const
 	}
 	else
 	{
+#if !UE_BUILD_SHIPPING
+		// Ensuring to help make sure we catch these critical data mismatches, but since we can handle these gracefully,
+		// don't ensure in Shipping builds.  We need to always log an error afterwards since the ensure will only fire
+		// once per session, and we don't want to always ensure since that might interfere with unrelated debugging.
+		ensureMsgf(false, 
+				TEXT("Failed to find valid parameter index for ParamName (%s) on CO (%s).  (SelectedOptionName: %s, RangeIndex: %d).  Cannot set parameter."),
+				*ParamName, *GetNameSafe(CustomizableObject), *SelectedOptionName, RangeIndex
+				  );
+#endif
 		UE_LOG(LogMutable, Error,
 				TEXT("%hs: Failed to find valid parameter index for ParamName (%s) on CO (%s).  (SelectedOptionName: %s, RangeIndex: %d).  Cannot set parameter."),
 				__FUNCTION__, *ParamName, *GetNameSafe(CustomizableObject), *SelectedOptionName, RangeIndex
diff --git a/Engine/Plugins/Experimental/Mutable/Source/CustomizableObject/Private/MuCO/CustomizableObjectSystem.cpp b/Engine/Plugins/Experimental/Mutable/Source/CustomizableObject/Private/MuCO/CustomizableObjectSystem.cpp
--- a/Engine/Plugins/Experimental/Mutable/Source/CustomizableObject/Private/MuCO/CustomizableObjectSystem.cpp
+++ b/Engine/Plugins/Experimental/Mutable/Source/CustomizableObject/Private/MuCO/CustomizableObjectSystem.cpp
@@ -2809,7 +2809,7 @@ void UCustomizableObjectSystem::SetOnlyGenerateRequestedLODsEnabled(bool bIsEnab
 }
 
 
-void UCustomizableObjectSystem::AddUncompiledCOWarning(UCustomizableObject* InObject)
+void UCustomizableObjectSystem::AddUncompiledCOWarning(UCustomizableObject* InObject, FString const* OptionalLogInfo)
 {
 	if (!InObject)
 	{
@@ -2820,6 +2820,19 @@ void UCustomizableObjectSystem::AddUncompiledCOWarning(UCustomizableObject* InOb
 	Msg += FString::Printf(TEXT("Warning: Customizable Object [%s] not compiled. Please compile and save the object."), *InObject->GetName());
 	GEngine->AddOnScreenDebugMessage((uint64)((PTRINT)InObject), 10.0f, FColor::Red, Msg);
 
+#if WITH_EDITOR
+	const FString ErrorString = FString::Printf(
+		TEXT("Customizable Object [%s] not compiled.  Compile via the editor or via code before instancing.  %s"),
+						   *InObject->GetName(), OptionalLogInfo ? **OptionalLogInfo : TEXT(""));
+#else // !WITH_EDITOR
+	const FString ErrorString = FString::Printf(
+		TEXT("Customizable Object [%s] not compiled.  This is not an Editor build, so this is an unrecoverable bad state; could be due to code or a cook failure.  %s")
+						   *InObject->GetName(), OptionalLogInfo ? **OptionalLogInfo : TEXT(""));
+#endif
+
+	// Also log an error so if this happens as part of a bug report we'll have this info.
+	UE_LOG(LogMutable, Error, TEXT("%s"), *ErrorString);
+
 #if WITH_EDITOR
 	if (UncompiledCustomizableObjectIds.Find(InObject->GetVersionId()) == INDEX_NONE)
 	{
diff --git a/Engine/Plugins/Experimental/Mutable/Source/CustomizableObject/Public/MuCO/CustomizableObjectInstanceDescriptor.h b/Engine/Plugins/Experimental/Mutable/Source/CustomizableObject/Public/MuCO/CustomizableObjectInstanceDescriptor.h
--- a/Engine/Plugins/Experimental/Mutable/Source/CustomizableObject/Public/MuCO/CustomizableObjectInstanceDescriptor.h
+++ b/Engine/Plugins/Experimental/Mutable/Source/CustomizableObject/Public/MuCO/CustomizableObjectInstanceDescriptor.h
@@ -317,6 +317,8 @@ struct CUSTOMIZABLEOBJECT_API FCustomizableObjectInstanceDescriptor
 	void MultilayerProjectorUpdateVirtualLayer(const FName& ProjectorParamName, const FName& Id, const FMultilayerProjectorVirtualLayer& Layer);
 
 private:
+	void AddUncompiledCOWarning(const FString& AdditionalLoggingInfo);
+
 	UPROPERTY()
 	TObjectPtr<UCustomizableObject> CustomizableObject = nullptr;
 
