Fixed compile error in non-Editor builds.

#UE5

diff --git a/Engine/Plugins/Experimental/Mutable/Source/CustomizableObject/Private/MuCO/CustomizableObjectSystem.cpp b/Engine/Plugins/Experimental/Mutable/Source/CustomizableObject/Private/MuCO/CustomizableObjectSystem.cpp
--- a/Engine/Plugins/Experimental/Mutable/Source/CustomizableObject/Private/MuCO/CustomizableObjectSystem.cpp
+++ b/Engine/Plugins/Experimental/Mutable/Source/CustomizableObject/Private/MuCO/CustomizableObjectSystem.cpp
@@ -2826,7 +2826,7 @@ void UCustomizableObjectSystem::AddUncompiledCOWarning(UCustomizableObject* InOb
 						   *InObject->GetName(), OptionalLogInfo ? **OptionalLogInfo : TEXT(""));
 #else // !WITH_EDITOR
 	const FString ErrorString = FString::Printf(
-		TEXT("Customizable Object [%s] not compiled.  This is not an Editor build, so this is an unrecoverable bad state; could be due to code or a cook failure.  %s")
+		TEXT("Customizable Object [%s] not compiled.  This is not an Editor build, so this is an unrecoverable bad state; could be due to code or a cook failure.  %s"),
 						   *InObject->GetName(), OptionalLogInfo ? **OptionalLogInfo : TEXT(""));
 #endif
 
