[Mutable] Add trilinear sampling to texture projections with MinFilterMethod different than none and Bilinear sampling method.

#rnx

diff --git a/Engine/Plugins/Experimental/Mutable/Source/MutableRuntime/Private/MuR/CodeRunner.cpp b/Engine/Plugins/Experimental/Mutable/Source/MutableRuntime/Private/MuR/CodeRunner.cpp
--- a/Engine/Plugins/Experimental/Mutable/Source/MutableRuntime/Private/MuR/CodeRunner.cpp
+++ b/Engine/Plugins/Experimental/Mutable/Source/MutableRuntime/Private/MuR/CodeRunner.cpp
@@ -72,6 +72,17 @@
 #include "Templates/Tuple.h"
 #include "Trace/Detail/Channel.h"
 
+namespace
+{
+	
+bool bForceBestQualityProjectionSampling = false;
+static FAutoConsoleVariableRef CVarForceBestQualityProjectionSampling (
+	TEXT("mutable.ForceBestQualityProjectionSampling"),
+	bForceBestQualityProjectionSampling,
+	TEXT("force mutable to use the projection algorithm with best projection quality, BiLinear with TotalAreaHeuristic"),
+	ECVF_Default);
+}
+
 
 namespace mu
 {
@@ -4035,10 +4046,17 @@ namespace mu
 				const int32 MipsToSkip = item.ExecutionOptions;
 				int32 ProjectionMip = MipsToSkip;
 
+				FScheduledOpData Data;
+				Data.RasterMesh.Mip = ProjectionMip;
+				Data.RasterMesh.MipValue = static_cast<float>(ProjectionMip);
 				Ptr<const Projector> pProjector = GetMemory().GetProjector(FScheduledOp::FromOpAndOptions(args.projector, item, 0));
 				if (pProjector)
 				{
-					if (static_cast<EMinFilterMethod>(args.MinFilterMethod) == EMinFilterMethod::TotalAreaHeuristic)
+					EMinFilterMethod MinFilterMethod = bForceBestQualityProjectionSampling
+						? EMinFilterMethod::TotalAreaHeuristic
+						: static_cast<EMinFilterMethod>(args.MinFilterMethod);
+
+					if (MinFilterMethod == EMinFilterMethod::TotalAreaHeuristic)
 					{
 						FVector2f TargetImageSizeF = FVector2f(
 							FMath::Max(args.sizeX >> MipsToSkip, 1),
@@ -4046,8 +4064,11 @@ namespace mu
 						FVector2f SourceImageSizeF = FVector2f(args.SourceSizeX, args.SourceSizeY);
 						
 						if (pMesh)
-						{
-							ProjectionMip = ComputeProjectedFootprintBestMip(pMesh.get(), pProjector->m_value, TargetImageSizeF, SourceImageSizeF);
+						{ 
+							Data.RasterMesh.MipValue =  
+								ComputeProjectedFootprintBestMip(pMesh.get(), pProjector->m_value, TargetImageSizeF, SourceImageSizeF);
+
+							Data.RasterMesh.Mip = static_cast<uint8>(FMath::FloorToInt32(Data.RasterMesh.MipValue));
 						}
 					}
 				}
@@ -4057,11 +4078,14 @@ namespace mu
 				// TODO: Review if this is a good way of reusing data between stages.
 				GetMemory().SetMesh(FScheduledOp::FromOpAndOptions(args.mesh, item, 0), pMesh);
 				GetMemory().SetProjector(FScheduledOp::FromOpAndOptions(args.projector, item, 0), pProjector);
+				
 
-				AddOp(FScheduledOp(item.At, item, 2, ProjectionMip),
+				int32 DataHeapAddress = m_heapData.Add(Data);
+
+				AddOp(FScheduledOp(item.At, item, 2, DataHeapAddress),
 					FScheduledOp::FromOpAndOptions(args.mesh, item, 0),
 					FScheduledOp::FromOpAndOptions(args.projector, item, 0),
-					FScheduledOp::FromOpAndOptions(args.image, item, ProjectionMip),
+					FScheduledOp::FromOpAndOptions(args.image, item, Data.RasterMesh.Mip),
 					FScheduledOp(args.mask, item),
 					FScheduledOp::FromOpAndOptions(args.angleFadeProperties, item, 0));
 
@@ -4111,9 +4135,11 @@ namespace mu
 					--MipsToDrop;
 				}
 
+				FScheduledOpData Data = m_heapData[item.CustomState];
+
 				// Raster with projection
 				Ptr<const Image> pSource = GetMemory().GetImage(
-					FCacheAddress(args.image, item.ExecutionIndex, static_cast<uint8>(item.CustomState)));
+					FCacheAddress(args.image, item.ExecutionIndex, Data.RasterMesh.Mip));
 
 				Ptr<const Image> pMask = nullptr;
 				if (args.mask)
@@ -4158,7 +4184,26 @@ namespace mu
 				// moving the right layout channel to the 0.
 				int layout = 0;
 
-				ESamplingMethod SamplingMethod = static_cast<ESamplingMethod>(args.SamplingMethod);
+				ESamplingMethod SamplingMethod = bForceBestQualityProjectionSampling
+						? ESamplingMethod::BiLinear
+						: static_cast<ESamplingMethod>(args.SamplingMethod);
+
+				if (SamplingMethod == ESamplingMethod::BiLinear)
+				{
+					if (pSource->GetLODCount() < 2 && pSource->GetSizeX() > 1 && pSource->GetSizeY() > 1)
+					{
+						MUTABLE_CPUPROFILER_SCOPE(RunCode_RasterMesh_BilinearMipGen);
+
+						Ptr<Image> NewImage = new Image(pSource->GetSizeX(), pSource->GetSizeY(), 2, pSource->GetFormat());
+
+						check(NewImage->GetDataSize() >= pSource->GetDataSize());
+						FMemory::Memcpy(NewImage->GetData(), pSource->GetData(), pSource->GetDataSize());
+
+						ImageMipmapInPlace(0, NewImage.get(), FMipmapGenerationSettings{});
+
+						pSource = NewImage;
+					}
+				}
 
 				if (args.projector && pSource && pSource->GetSizeX() > 0 && pSource->GetSizeY() > 0)
 				{
@@ -4172,7 +4217,7 @@ namespace mu
 								pSource.get(), pMask.get(),
 								args.bIsRGBFadingEnabled, args.bIsAlphaFadingEnabled,
 								SamplingMethod,
-								FadeStartRad, FadeEndRad,
+								FadeStartRad, FadeEndRad, FMath::Frac(Data.RasterMesh.MipValue),
 								layout, args.blockIndex,
 								CropMin, UncroppedSize,
 								&scratch);
@@ -4183,7 +4228,7 @@ namespace mu
 								pSource.get(), pMask.get(),
 								args.bIsRGBFadingEnabled, args.bIsAlphaFadingEnabled,
 								SamplingMethod,
-								FadeStartRad, FadeEndRad,
+								FadeStartRad, FadeEndRad, FMath::Frac(Data.RasterMesh.MipValue),
 								layout, args.blockIndex,
 								CropMin, UncroppedSize,
 								&scratch);
@@ -4653,7 +4698,6 @@ namespace mu
         }
     }
 
-
     //---------------------------------------------------------------------------------------------
     void CodeRunner::RunCode_Int(const FScheduledOp& item, const Parameters* pParams, const Model* pModel )
     {
diff --git a/Engine/Plugins/Experimental/Mutable/Source/MutableRuntime/Private/MuR/CodeRunner.h b/Engine/Plugins/Experimental/Mutable/Source/MutableRuntime/Private/MuR/CodeRunner.h
--- a/Engine/Plugins/Experimental/Mutable/Source/MutableRuntime/Private/MuR/CodeRunner.h
+++ b/Engine/Plugins/Experimental/Mutable/Source/MutableRuntime/Private/MuR/CodeRunner.h
@@ -78,6 +78,12 @@ namespace  mu
 					int32 ResultDescAt;
 					int32 SourceDescAt;
 				} ResizeLike;
+
+				struct
+				{
+					uint8 Mip;
+					float MipValue;
+				} RasterMesh;
 			};
             Ptr<RefCounted> Resource;
         };
diff --git a/Engine/Plugins/Experimental/Mutable/Source/MutableRuntime/Private/MuR/OpImageProject.cpp b/Engine/Plugins/Experimental/Mutable/Source/MutableRuntime/Private/MuR/OpImageProject.cpp
--- a/Engine/Plugins/Experimental/Mutable/Source/MutableRuntime/Private/MuR/OpImageProject.cpp
+++ b/Engine/Plugins/Experimental/Mutable/Source/MutableRuntime/Private/MuR/OpImageProject.cpp
@@ -206,14 +206,29 @@ namespace
 			const uint8* pTargetData,
 			const uint8* pMaskData,
 			bool bInIsRGBFadingEnabled, bool bInIsAlphaFadingEnabled,
-			float FadeStart, float FadeEnd )
+			float FadeStart, float FadeEnd, float MipSamplingFactor )
         {
-            SourceSizeX = static_cast<uint16>(Source->GetSizeX());
-            SourceSizeY = static_cast<uint16>(Source->GetSizeY());
 
-            SourceDataPtr = (SourceSizeX > 0 && SourceSizeY > 0) ? Source->GetData() : nullptr;
+            Source0SizeX = static_cast<uint16>(Source->GetSizeX());
+            Source0SizeY = static_cast<uint16>(Source->GetSizeY());
+
+			// If we don't have the next mip, use the first mip as a fallback.
+			const bool bHasNextMip = Source->GetLODCount() >= 2;
+
+            Source0DataPtr = (Source0SizeX > 0 && Source0SizeY > 0) ? Source->GetData() : nullptr;
+			Source1DataPtr = bHasNextMip ? Source->GetMipData(1) : Source0DataPtr;
+	
+			FIntVector2 Source1Size = Source->CalculateMipSize(1);
+            Source1SizeX = bHasNextMip ? static_cast<uint16>(Source1Size.X) : Source0SizeX;
+            Source1SizeY = bHasNextMip ? static_cast<uint16>(Source1Size.Y) : Source0SizeY;
+
+			// Invalidate SourceDataPtrs if either of the two is not valid. 
+			if (!(Source0DataPtr && Source1DataPtr))
+			{
+				Source0DataPtr = nullptr;
+				Source1DataPtr = nullptr;
+			}
 
-			// \TODO: to template parameters?
 			bIsRGBFadingEnabled = bInIsRGBFadingEnabled;
 			bIsAlphaFadingEnabled = bInIsAlphaFadingEnabled;
 
@@ -230,14 +245,13 @@ namespace
             TargetDataPtr = pTargetData;
             MaskDataPtr = pMaskData;
 
-
             check(GetImageFormatData(Source->GetFormat() ).m_bytesPerBlock == PIXEL_SIZE);
         }
 
         //-----------------------------------------------------------------------------------------
         inline void ProcessPixel( uint8* BufferPosPtr, float Varying[4] ) const
         {
-			if (!SourceDataPtr)
+			if (!Source0DataPtr)
 			{
 				return;
 			}
@@ -265,68 +279,100 @@ namespace
             if (Factor > 0)
             {
 				const FVector2f Uv = { Varying[0], Varying[1] };
-				const FVector2f SourceSizeF = FVector2f(SourceSizeX, SourceSizeY);
+				const FVector2f Source0SizeF = FVector2f(Source0SizeX, Source0SizeY);
+				const FVector2f Source1SizeF = FVector2f(Source1SizeX, Source1SizeY);
 
                 if ((Uv.X >= 0.0f) & (Uv.X < 1.0f) & (Uv.Y >= 0.0f) & (Uv.Y < 1.0f))
                 {
-					const float CoordsX = FMath::Clamp(Uv.X * SourceSizeF.X - 0.5f, 0.0f, SourceSizeF.X - 1.0f);
-					const float CoordsY = FMath::Clamp(Uv.Y * SourceSizeF.Y - 0.5f, 0.0f, SourceSizeF.Y - 1.0f);
+					using FUInt16Vector2 = UE::Math::TIntVector2<uint16>;
+					struct FPixelData
+					{
+						alignas(8) uint16 Data[PIXEL_SIZE];
+					};
 
-					auto ComputeInterpolator = [](float T) -> uint16
+					auto SampleImageBilinear = [](FVector2f Uv, FUInt16Vector2 Size, const uint8* DataPtr)
 					{
-						if constexpr (bUseSmoothInterpolator)
+						auto ComputeInterpolator = [](float T) -> uint16
 						{
-							return static_cast<uint16>(255.0f * T * T * (3.0f - 2.0f * T));
-						}
-						else
-						{
-							return static_cast<uint16>(255.0f * T);
-						}
-					};
+							if constexpr (bUseSmoothInterpolator)
+							{
+								return static_cast<uint16>(255.0f * T * T * (3.0f - 2.0f * T));
+							}
+							else
+							{
+								return static_cast<uint16>(255.0f * T);
+							}
+						};
 
-					const uint16 FracX = ComputeInterpolator(FMath::Frac(CoordsX));
-					const uint16 FracY = ComputeInterpolator(FMath::Frac(CoordsY));
+						const FVector2f SizeF(Size.X, Size.Y);
 
-					const uint32 X = static_cast<uint32>(CoordsX);
-					const uint32 Y = static_cast<uint32>(CoordsY);
+						const FVector2f CoordsF = FVector2f(
+							FMath::Clamp(Uv.X * SizeF.X - 0.5f, 0.0f, SizeF.X - 1.0f),
+							FMath::Clamp(Uv.Y * SizeF.Y - 0.5f, 0.0f, SizeF.Y - 1.0f));
 
-					const uint32 XPlusOne = FMath::Min<uint32>(SourceSizeX - 1, X + 1);
-					const uint32 YPlusOne = FMath::Min<uint32>(SourceSizeY - 1, Y + 1);
+						const FUInt16Vector2 Frac = FUInt16Vector2(
+							ComputeInterpolator(FMath::Frac(CoordsF.X)),
+							ComputeInterpolator(FMath::Frac(CoordsF.Y)));
 
-					const uint8* Pixel00Ptr = SourceDataPtr + (Y        * SourceSizeX + X)        * PIXEL_SIZE;
-					const uint8* Pixel10Ptr = SourceDataPtr + (Y        * SourceSizeX + XPlusOne) * PIXEL_SIZE;
-					const uint8* Pixel01Ptr = SourceDataPtr + (YPlusOne * SourceSizeX + X)        * PIXEL_SIZE;
-					const uint8* Pixel11Ptr = SourceDataPtr + (YPlusOne * SourceSizeX + XPlusOne) * PIXEL_SIZE;
+						const FIntVector2 Coords = FIntVector2(CoordsF.X, CoordsF.Y);
+						const FIntVector2 CoordsPlusOne = FIntVector2( 
+								FMath::Min(Size.X - 1, Coords.X + 1),
+								FMath::Min(Size.Y - 1, Coords.Y + 1));
 
-					struct FPixelData
-					{
-						alignas(8) uint16 Data[PIXEL_SIZE];
-					};
+						uint8 const * const Pixel00Ptr = DataPtr + (Coords.Y        * Size.X + Coords.X)        * PIXEL_SIZE;
+						uint8 const * const Pixel10Ptr = DataPtr + (Coords.Y        * Size.X + CoordsPlusOne.X) * PIXEL_SIZE;
+						uint8 const * const Pixel01Ptr = DataPtr + (CoordsPlusOne.Y * Size.X + Coords.X)        * PIXEL_SIZE;
+						uint8 const * const Pixel11Ptr = DataPtr + (CoordsPlusOne.Y * Size.X + CoordsPlusOne.X) * PIXEL_SIZE;
 
-					auto LoadPixel = [](const uint8* Ptr) -> FPixelData
-					{
-						FPixelData Result;
+						auto LoadPixel = [](const uint8* Ptr) -> FPixelData
+						{
+							FPixelData Result;
+							if constexpr (PIXEL_SIZE == 4)
+							{
+								const uint32 PackedData = *reinterpret_cast<const uint32*>(Ptr);
+								
+								Result.Data[0] = static_cast<uint16>((PackedData >> (8 * 0)) & 0xFF);
+								Result.Data[1] = static_cast<uint16>((PackedData >> (8 * 1)) & 0xFF);
+								Result.Data[2] = static_cast<uint16>((PackedData >> (8 * 2)) & 0xFF);
+								Result.Data[3] = static_cast<uint16>((PackedData >> (8 * 3)) & 0xFF);
+							}
+							else
+							{
+								for (int32 C = 0; C < PIXEL_SIZE; ++C)
+								{
+									Result.Data[C] = static_cast<uint16>(Ptr[C]);
+								}
+							}
+							return Result;
+						};
+
+						FPixelData PixelData00 = LoadPixel(Pixel00Ptr);
+						FPixelData PixelData10 = LoadPixel(Pixel10Ptr);
+						FPixelData PixelData01 = LoadPixel(Pixel01Ptr);
+						FPixelData PixelData11 = LoadPixel(Pixel11Ptr);
+
+						FPixelData FilteredPixelData;
+					
 						for (int32 C = 0; C < PIXEL_SIZE; ++C)
 						{
-							Result.Data[C] = static_cast<uint16>(Ptr[C]);
-						}
+							const uint16 LerpY0 = ((PixelData10.Data[C] * Frac.X) + PixelData00.Data[C] * (255 - Frac.X)) / 255;
+							const uint16 LerpY1 = ((PixelData11.Data[C] * Frac.X) + PixelData01.Data[C] * (255 - Frac.X)) / 255;
+							FilteredPixelData.Data[C] = ((LerpY1 * Frac.Y) + LerpY0*(255 - Frac.Y)) / 255;
+						} 
 
-						return Result;
+						return FilteredPixelData;
 					};
 
-					FPixelData PixelData00 = LoadPixel(Pixel00Ptr);
-					FPixelData PixelData10 = LoadPixel(Pixel10Ptr);
-					FPixelData PixelData01 = LoadPixel(Pixel01Ptr);
-					FPixelData PixelData11 = LoadPixel(Pixel11Ptr);
+					FPixelData Sample0 = SampleImageBilinear(Uv, FUInt16Vector2(Source0SizeX, Source0SizeY), Source0DataPtr);
+					FPixelData Sample1 = SampleImageBilinear(Uv, FUInt16Vector2(Source1SizeX, Source1SizeY), Source1DataPtr);
 
-					uint16 FilteredPixelData[PIXEL_SIZE];
-					
+					const uint16 MipFactor = static_cast<uint16>(MipInterpolationFactor * 255.0f);
+
+					FPixelData& FilteredPixelData = Sample0;
 					for (int32 C = 0; C < PIXEL_SIZE; ++C)
 					{
-						const uint16 LerpY0 = ((PixelData10.Data[C] * FracX) + PixelData00.Data[C] * (255 - FracX)) / 255;
-						const uint16 LerpY1 = ((PixelData11.Data[C] * FracX) + PixelData01.Data[C] * (255 - FracX)) / 255;
-						FilteredPixelData[C] = ((LerpY1 * FracY) + LerpY0*(255 - FracY)) / 255;
-					} 
+						FilteredPixelData.Data[C] = ((Sample1.Data[C] * MipFactor) + Sample0.Data[C]*(255 - MipFactor)) / 255;
+					}
 
 					static_assert(PIXEL_SIZE <= 4);
 					const uint16 FadingMask[4] = { 
@@ -338,17 +384,21 @@ namespace
 					for (int32 I = 0; I < PIXEL_SIZE; ++I)
 					{
 						BufferPosPtr[I] = static_cast<uint8>(
-							(FilteredPixelData[I] * (MaskFactor * (1 - FadingMask[I]) + (Factor * FadingMask[I]))) / 255);
+							(FilteredPixelData.Data[I] * (MaskFactor * (1 - FadingMask[I]) + (Factor * FadingMask[I]))) / 255);
 					}				
                 }
             }
         }
 
-        const uint8* SourceDataPtr;
-        uint16 SourceSizeX, SourceSizeY;
+        const uint8* Source0DataPtr;
+        const uint8* Source1DataPtr;
+        uint16 Source0SizeX, Source0SizeY;
+        uint16 Source1SizeX, Source1SizeY;
 
         //! Cosine of the fading angle range
 		float FadeEndCos;
+		float MipInterpolationFactor;
+
 		// precomputation of 255 / (FadeStartCos - FadeEndCos) 
 		float OneOverFadeRangeTimes255;
 
@@ -731,7 +781,7 @@ void ImageRasterProjectedPlanar( const Mesh* pMesh,
 	const Image* pMask,
 	bool bIsRGBFadingEnabled, bool bIsAlphaFadingEnabled,
 	ESamplingMethod SamplingMethod,
-	float fadeStart, float fadeEnd,
+	float fadeStart, float fadeEnd, float MipInterpolationFactor,
 	int layout, int block,
 	UE::Math::TIntVector2<uint16> CropMin, UE::Math::TIntVector2<uint16> UncroppedSize,
 	SCRATCH_IMAGE_PROJECT* scratch )
@@ -763,7 +813,7 @@ void ImageRasterProjectedPlanar( const Mesh* pMesh,
 					pImage->GetData(),
 					pMask ? pMask->GetData() : nullptr,
 					bIsRGBFadingEnabled, bIsAlphaFadingEnabled,
-					fadeStart, fadeEnd);
+					fadeStart, fadeEnd, MipInterpolationFactor);
 				ImageRasterProjected_Optimised(pMesh, pImage, pixelProc, fadeEnd, CropMin, UncroppedSize, scratch);
 			}
 			else
@@ -792,7 +842,7 @@ void ImageRasterProjectedPlanar( const Mesh* pMesh,
 						pImage->GetData(),
 						pMask->GetData(),
 						bIsRGBFadingEnabled, bIsAlphaFadingEnabled,
-						fadeStart, fadeEnd);
+						fadeStart, fadeEnd, MipInterpolationFactor);
 					ImageRasterProjected_Optimised(pMesh, pImage, pixelProc, fadeEnd, CropMin, UncroppedSize, scratch);
 				}
 				else
@@ -817,7 +867,7 @@ void ImageRasterProjectedPlanar( const Mesh* pMesh,
 						pImage->GetData(),
 						nullptr,
 						bIsRGBFadingEnabled, bIsAlphaFadingEnabled,
-						fadeStart, fadeEnd);
+						fadeStart, fadeEnd, MipInterpolationFactor);
 					ImageRasterProjected_Optimised(pMesh, pImage, pixelProc, fadeEnd, CropMin, UncroppedSize, scratch);
 				}
 				else
@@ -843,7 +893,7 @@ void ImageRasterProjectedPlanar( const Mesh* pMesh,
 					pImage->GetData(),
 					pMask ? pMask->GetData() : nullptr,
 					bIsRGBFadingEnabled, bIsAlphaFadingEnabled,
-					fadeStart, fadeEnd);
+					fadeStart, fadeEnd, MipInterpolationFactor);
 				ImageRasterProjected_Optimised(pMesh, pImage, pixelProc, fadeEnd, CropMin, UncroppedSize, scratch);
 			}
 			else
@@ -871,7 +921,7 @@ void ImageRasterProjectedWrapping( const Mesh* pMesh, Image* pImage,
 	const Image* pSource, const Image* pMask,
 	bool bIsRGBFadingEnabled, bool bIsAlphaFadingEnabled,
 	ESamplingMethod SamplingMethod,
-	float fadeStart, float fadeEnd,
+	float fadeStart, float fadeEnd, float MipInterpolationFactor,
 	int layout, int block,
 	UE::Math::TIntVector2<uint16> CropMin, UE::Math::TIntVector2<uint16> UncroppedSize,
 	SCRATCH_IMAGE_PROJECT* scratch )
@@ -906,7 +956,7 @@ void ImageRasterProjectedWrapping( const Mesh* pMesh, Image* pImage,
 					pImage->GetData(),
 					pMask ? pMask->GetData() : 0,
 					bIsRGBFadingEnabled, bIsAlphaFadingEnabled,
-					fadeStart, fadeEnd);
+					fadeStart, fadeEnd, MipInterpolationFactor);
 
 				ImageRasterProjected_OptimisedWrapping(pMesh, pImage, pixelProc, fadeEnd, block, CropMin, UncroppedSize, scratch);
 			}
@@ -936,7 +986,7 @@ void ImageRasterProjectedWrapping( const Mesh* pMesh, Image* pImage,
 						pImage->GetData(),
 						pMask->GetData(),
 						bIsRGBFadingEnabled, bIsAlphaFadingEnabled,
-						fadeStart, fadeEnd);
+						fadeStart, fadeEnd, MipInterpolationFactor);
 					ImageRasterProjected_OptimisedWrapping(pMesh, pImage, pixelProc, fadeEnd, block, CropMin, UncroppedSize, scratch);
 				}
 				else
@@ -961,7 +1011,7 @@ void ImageRasterProjectedWrapping( const Mesh* pMesh, Image* pImage,
 						pImage->GetData(),
 						nullptr,
 						bIsRGBFadingEnabled, bIsAlphaFadingEnabled,
-						fadeStart, fadeEnd);
+						fadeStart, fadeEnd, MipInterpolationFactor);
 					ImageRasterProjected_OptimisedWrapping(pMesh, pImage, pixelProc, fadeEnd, block, CropMin, UncroppedSize, scratch);
 				}
 				else
@@ -989,7 +1039,7 @@ void ImageRasterProjectedWrapping( const Mesh* pMesh, Image* pImage,
 					pImage->GetData(),
 					pMask ? pMask->GetData() : 0,
 					bIsRGBFadingEnabled, bIsAlphaFadingEnabled,
-					fadeStart, fadeEnd);
+					fadeStart, fadeEnd, MipInterpolationFactor);
 
 				ImageRasterProjected_OptimisedWrapping(pMesh, pImage, pixelProc, fadeEnd, block, CropMin, UncroppedSize, scratch);
 			}
@@ -1012,7 +1062,7 @@ void ImageRasterProjectedWrapping( const Mesh* pMesh, Image* pImage,
 
 }
 
-int32 ComputeProjectedFootprintBestMip(
+float ComputeProjectedFootprintBestMip(
 	const Mesh* pMesh, const FProjector& Projector, const FVector2f& TargetSize, const FVector2f& SourceSize)
 {
 	// Compute projected mesh footprint on the traget image and extract a desired size for source image.
@@ -1049,10 +1099,13 @@ int32 ComputeProjectedFootprintBestMip(
 
 		const float TriangleTargetArea = ComputeTriangleArea(A.uv, B.uv, C.uv);
 
-		// Set weight to zero if source area is close to zero to prevent posible infinite scaling. 
-		float TriangleWeight = static_cast<float>(!FMath::IsNearlyZero(TriangleSourceArea));
+		// Set weight to zero if source or target area are close to zero to remove outliers.
+		float TriangleWeight = static_cast<float>(!FMath::IsNearlyZero(TriangleSourceArea)) *
+							   static_cast<float>(!FMath::IsNearlyZero(TriangleTargetArea));
 		if (bIsPlanarProjection)
 		{
+			// Scale contribution to total area based on the mean of the projected triangle fade factor, e.i.,
+			// the cosinus of the angle between the projection direction and the vertices normal.
 			TriangleWeight *= FMath::Clamp((A.nor.X + B.nor.X + C.nor.X) * (1.0f / 3.0f), 0.0f, 1.0f);
 		}
 
@@ -1065,33 +1118,14 @@ int32 ComputeProjectedFootprintBestMip(
 		SourceFootprintAreaOnTarget * (TargetSize.X * TargetSize.Y);
 
 	// Find the mip that better adapts to the footprint.
-	// Swipe from the smalles to the largest mip finding the best fitting one.
-
-	const FIntVector2 SourceSizeInt = FIntVector2(SourceSize.X, SourceSize.Y);
-
-	int32 Mip = FMath::CeilToInt(FMath::Log2(FMath::Max(SourceSize.GetMax(), 1.0f)));
-	int32 BestMip = Mip;
-	float BestMipAreaInPixelsSquared = 1.0f;
-	for (; Mip > 0; --Mip)
-	{
-		FIntVector2 SourceMipSize =
-			FIntVector2{ FMath::Max(SourceSizeInt.X >> Mip, 1), FMath::Max(SourceSizeInt.Y >> Mip, 1) };
+	float BestMip = FMath::Log2(FMath::Max(1.0f, SourceSize.X * SourceSize.Y)) * 0.5f - 
+				    FMath::Log2(FMath::Max(1.0f, SourceFootprintAreaOnTargetInPixelsSquared)) * 0.5f;
 
-		float AreaOfMipInPixelsSquared = static_cast<float>(SourceMipSize.X * SourceMipSize.Y);
-
-		const float CurrentMipAreaDifference = FMath::Abs(AreaOfMipInPixelsSquared - SourceFootprintAreaOnTargetInPixelsSquared);
-		const float BestMipAreaDifference = FMath::Abs(BestMipAreaInPixelsSquared - SourceFootprintAreaOnTargetInPixelsSquared);
-
-		if (CurrentMipAreaDifference < BestMipAreaDifference)
-		{
-			BestMip = Mip;
-			BestMipAreaInPixelsSquared = AreaOfMipInPixelsSquared;
-		}
-	}
-
-	return BestMip;
+	// TODO: We probably should allow negative mips to indicate that the image needs magnification.
+	return FMath::Max(0.0f, BestMip);
 }
 
+
 //-------------------------------------------------------------------------------------------------
 void ImageRasterProjectedCylindrical( const Mesh* pMesh, Image* pImage,
 	const Image* pSource, const Image* pMask,
diff --git a/Engine/Plugins/Experimental/Mutable/Source/MutableRuntime/Private/MuR/OpImageProject.h b/Engine/Plugins/Experimental/Mutable/Source/MutableRuntime/Private/MuR/OpImageProject.h
--- a/Engine/Plugins/Experimental/Mutable/Source/MutableRuntime/Private/MuR/OpImageProject.h
+++ b/Engine/Plugins/Experimental/Mutable/Source/MutableRuntime/Private/MuR/OpImageProject.h
@@ -21,7 +21,7 @@ namespace mu
 		const Image* pSource, const Image* pMask,
 		bool bIsRGBFadingEnabled, bool bIsAlphaFadingEnabled,
 		ESamplingMethod SamplingMethod,
-		float fadeStart, float fadeEnd,
+		float fadeStart, float fadeEnd, float MipInterpolationFactor,
 		int layout, int block,
 		UE::Math::TIntVector2<uint16> CropMin, UE::Math::TIntVector2<uint16> UncroppedSize,
 		SCRATCH_IMAGE_PROJECT* scratch );
@@ -39,12 +39,12 @@ namespace mu
 		const Image* pSource, const Image* pMask,
 		bool bIsRGBFadingEnabled, bool bIsAlphaFadingEnabled,
 		ESamplingMethod SamplingMethod,
-		float fadeStart, float fadeEnd,
+		float fadeStart, float fadeEnd, float MipInterpolationFactor,
 		int layout, int block,
 		UE::Math::TIntVector2<uint16> CropMin, UE::Math::TIntVector2<uint16> UncroppedSize,
 		SCRATCH_IMAGE_PROJECT* scratch );
 
-	extern int32 ComputeProjectedFootprintBestMip(
+	extern float ComputeProjectedFootprintBestMip(
 			const Mesh* pMesh, const FProjector& Projector, const FVector2f& TargetSize, const FVector2f& SourceSize);
 
     extern MeshPtr MeshProject( const Mesh* pMesh,
