[mutable] Support texture references in the core.
#rnx

diff --git a/Engine/Plugins/Experimental/Mutable/Source/CustomizableObject/Classes/MuCO/CustomizableObject.h b/Engine/Plugins/Experimental/Mutable/Source/CustomizableObject/Classes/MuCO/CustomizableObject.h
--- a/Engine/Plugins/Experimental/Mutable/Source/CustomizableObject/Classes/MuCO/CustomizableObject.h
+++ b/Engine/Plugins/Experimental/Mutable/Source/CustomizableObject/Classes/MuCO/CustomizableObject.h
@@ -1410,7 +1410,7 @@ private:
 	// This is a manual version number for the binary blobs in this asset.
 	// Increasing it invalidates all the previously compiled models.
 	// Warning: If while merging code both versions have changed, take the highest+1.
-	static const int32 CurrentSupportedVersion = 382;
+	static const int32 CurrentSupportedVersion = 383;
 
 public:
 
diff --git a/Engine/Plugins/Experimental/Mutable/Source/CustomizableObject/Private/MuCO/CustomizableObjectInstance.cpp b/Engine/Plugins/Experimental/Mutable/Source/CustomizableObject/Private/MuCO/CustomizableObjectInstance.cpp
--- a/Engine/Plugins/Experimental/Mutable/Source/CustomizableObject/Private/MuCO/CustomizableObjectInstance.cpp
+++ b/Engine/Plugins/Experimental/Mutable/Source/CustomizableObject/Private/MuCO/CustomizableObjectInstance.cpp
@@ -5714,13 +5714,31 @@ void UCustomizableInstancePrivateData::BuildMaterials(const TSharedPtr<FMutableO
 							check(Texture);
 						}
 
+						// Check if the image is a reference to an engine texture
+						if (MutableImage && MutableImage->IsReference())
+						{
+							uint32 ReferenceID = MutableImage->GetReferencedTexture();
+							if (CustomizableObject->ReferencedPassThroughTextures.IsValidIndex(ReferenceID))
+							{
+								TSoftObjectPtr<UTexture2D> Ref = CustomizableObject->ReferencedPassThroughTextures[ReferenceID];
+
+								// \TODO: This will force the load of the reference texture, potentially causing a hich. 
+								Texture = Ref.Get();
+							}
+							else
+							{
+								// internal error.
+								UE_LOG(LogMutable, Error, TEXT("Referenced image [%d] was not dtored in the resource array."), ReferenceID);
+							}
+						}
+
 						// Find the additional information for this image
 						int32 ImageKey = FCString::Atoi(*KeyName);
 						if (ImageKey >= 0 && ImageKey < CustomizableObject->ImageProperties.Num())
 						{
 							const FMutableModelImageProperties& Props = CustomizableObject->ImageProperties[ImageKey];
 
-							if (MutableImage)
+							if (!Texture && MutableImage)
 							{
 								TWeakObjectPtr<UTexture2D>* ReusedTexture = bReuseTextures ? TextureReuseCache.Find(TextureReuseCacheRef) : nullptr;
 
diff --git a/Engine/Plugins/Experimental/Mutable/Source/CustomizableObject/Private/MuCO/CustomizableObjectSystem.cpp b/Engine/Plugins/Experimental/Mutable/Source/CustomizableObject/Private/MuCO/CustomizableObjectSystem.cpp
--- a/Engine/Plugins/Experimental/Mutable/Source/CustomizableObject/Private/MuCO/CustomizableObjectSystem.cpp
+++ b/Engine/Plugins/Experimental/Mutable/Source/CustomizableObject/Private/MuCO/CustomizableObjectSystem.cpp
@@ -1106,10 +1106,6 @@ void UCustomizableObjectSystem::PurgePendingReleaseSkeletalMesh()
 				Info.SkeletalMesh->ReleaseResources();
 				Info.SkeletalMesh->ReleaseResourcesFence.Wait();
 
-				// TODO: This was done for version 4.18
-				//info.SkeletalMesh->GetImportedResource()->LODModels.Empty();
-				//info.SkeletalMesh->LODInfo.Empty();
-
 				PendingReleaseSkeletalMesh.RemoveAt(InfoIndex);
 			}
 		}
@@ -1512,8 +1508,15 @@ namespace impl
 					Image.Image = System->GetImage(OperationData->InstanceID, Image.ImageID, MipsToSkip, Image.LOD);
 				}
 
-				// We need one mip or the complete chain. Otherwise there was a bug.
 				check(Image.Image);
+
+				// If the image is a reference to an engine texture, we are done.
+				if (Image.Image->IsReference())
+				{
+					continue;
+				}
+
+				// We need one mip or the complete chain. Otherwise there was a bug.
 				int32 FullMipCount = Image.Image->GetMipmapCount(Image.Image->GetSizeX(), Image.Image->GetSizeY());
 				int32 RealMipCount = Image.Image->GetLODCount();
 				if ((RealMipCount != 1) && (RealMipCount != FullMipCount))
diff --git a/Engine/Plugins/Experimental/Mutable/Source/CustomizableObject/Private/MuCO/CustomizableObjectSystemPrivate.h b/Engine/Plugins/Experimental/Mutable/Source/CustomizableObject/Private/MuCO/CustomizableObjectSystemPrivate.h
--- a/Engine/Plugins/Experimental/Mutable/Source/CustomizableObject/Private/MuCO/CustomizableObjectSystemPrivate.h
+++ b/Engine/Plugins/Experimental/Mutable/Source/CustomizableObject/Private/MuCO/CustomizableObjectSystemPrivate.h
@@ -316,17 +316,16 @@ struct FInstanceGeneratedData
 
 // Mutable data generated during the update steps.
 // We keep it from begin to end update, and it is used in several steps.
-// TODO: Flatten this structure into 4 arrays and use indices in fixed-size structs instead of subarrays
 struct FInstanceUpdateData
 {
 	struct FImage
 	{
 		FString Name;
 		mu::RESOURCE_ID ImageID;
+		int32 LOD;
 		uint16 FullImageSizeX, FullImageSizeY;
-		mu::ImagePtrConst Image;
+		mu::Ptr<const mu::Image> Image;
 		TWeakObjectPtr<UTexture2D> Cached;
-		int32 LOD;
 	};
 
 	struct FVector
diff --git a/Engine/Plugins/Experimental/Mutable/Source/CustomizableObjectEditor/Private/MuCOE/GenerateMutableSource/GenerateMutableSourceImage.cpp b/Engine/Plugins/Experimental/Mutable/Source/CustomizableObjectEditor/Private/MuCOE/GenerateMutableSource/GenerateMutableSourceImage.cpp
--- a/Engine/Plugins/Experimental/Mutable/Source/CustomizableObjectEditor/Private/MuCOE/GenerateMutableSource/GenerateMutableSourceImage.cpp
+++ b/Engine/Plugins/Experimental/Mutable/Source/CustomizableObjectEditor/Private/MuCOE/GenerateMutableSource/GenerateMutableSourceImage.cpp
@@ -48,6 +48,7 @@
 #include "MuT/NodeImageTable.h"
 #include "MuT/NodeImageTransform.h"
 #include "MuT/NodeImageVariation.h"
+#include "MuT/NodeImageReference.h"
 #include "MuT/NodeScalarConstant.h"
 
 #define LOCTEXT_NAMESPACE "CustomizableObjectEditor"
@@ -101,7 +102,7 @@
 }
 
 
-mu::NodeImagePtr ResizeToMaxTextureSize(float MaxTextureSize, const UTexture2D* BaseTexture, mu::NodeImageConstantPtr ImageNode)
+mu::Ptr<mu::NodeImage> ResizeToMaxTextureSize(float MaxTextureSize, const UTexture2D* BaseTexture, mu::Ptr<mu::NodeImageConstant> ImageNode)
 {
 	// To scale when above MaxTextureSize if defined
 	if (MaxTextureSize > 0 && BaseTexture
@@ -137,7 +138,7 @@
 	const FGeneratedKey Key(reinterpret_cast<void*>(&GenerateMutableSourceImage), *Pin, *Node, GenerationContext, true);
 	if (const FGeneratedData* Generated = GenerationContext.Generated.Find(Key))
 	{
-		return static_cast<mu::NodeImage*>(Generated->Node.get());
+		return dynamic_cast<mu::NodeImage*>(Generated->Node.get());
 	}
 
 	mu::NodeImagePtr Result;
@@ -171,37 +172,6 @@
 		}
 	}
 
-	else if (const UCustomizableObjectNodePassThroughTexture* TypedNodePassThroughTex = Cast<UCustomizableObjectNodePassThroughTexture>(Node))
-	{
-		UTexture2D* BaseTexture = TypedNodePassThroughTex->Texture;
-		if (BaseTexture)
-		{
-			uint32* FoundIndex = GenerationContext.PassThroughTextureToIndexMap.Find(BaseTexture);
-			uint32 NewId;
-
-			if (!FoundIndex)
-			{
-				NewId = GenerationContext.PassThroughTextureToIndexMap.Num();
-				GenerationContext.PassThroughTextureToIndexMap.Add(BaseTexture, NewId);
-			}
-			else
-			{
-				NewId = *FoundIndex;
-			}
-
-			// TODO: Implement mutable core nodes
-			//mu::NodePassThroughImageConstantPtr PassThroughImageNode;
-			//PassThroughImageNode.SetIndex(NewId);
-
-			//Result = PassThroughImageNode;
-
-		}
-		else
-		{
-			GenerationContext.Compiler->CompilerLog(LOCTEXT("MissingImagePassThrough", "Missing image in pass-through texture node."), Node);
-		}
-	}
-
 	else if (const UCustomizableObjectNodeTextureParameter* TypedNodeParam = Cast<UCustomizableObjectNodeTextureParameter>(Node))
 	{
 		mu::Ptr<mu::NodeImageParameter> TextureNode = new mu::NodeImageParameter();
@@ -397,7 +367,7 @@
 
 				const int32 NumSwitchOptions = TypedNodeTextureSwitch->GetNumElements();
 
-				mu::NodeScalarEnumParameter* EnumParameter = static_cast<mu::NodeScalarEnumParameter*>(SwitchParam.get());
+				mu::NodeScalarEnumParameter* EnumParameter = dynamic_cast<mu::NodeScalarEnumParameter*>(SwitchParam.get());
 				if (NumSwitchOptions != EnumParameter->GetValueCount())
 				{
 					const FText Message = LOCTEXT("MismatchedSwitch", "Switch enum and switch node have different number of options. Please refresh the switch node to make sure the outcomes are labeled properly.");
@@ -433,73 +403,6 @@
 		}(); // invoke lambda.
 	}
 
-	else if (const UCustomizableObjectNodePassThroughTextureSwitch* TypedNodePassThroughTextureSwitch = Cast<UCustomizableObjectNodePassThroughTextureSwitch>(Node))
-	{
-		Result = [&]()
-		{
-			const UEdGraphPin* SwitchParameter = TypedNodePassThroughTextureSwitch->SwitchParameter();
-
-			// Check Switch Parameter arity preconditions.
-			if (const UEdGraphPin* ConnectedPin = FollowInputPin(*SwitchParameter))
-			{
-				mu::NodeScalarPtr SwitchParam = GenerateMutableSourceFloat(ConnectedPin, GenerationContext);
-				// Switch Param not generated
-				if (!SwitchParam)
-				{
-					const FText Message = LOCTEXT("FailedToGenerateSwitchParam", "Could not generate switch enum parameter. Please refesh the switch node and connect an enum.");
-					GenerationContext.Compiler->CompilerLog(Message, Node);
-
-					return Result;
-				}
-
-				if (SwitchParam->GetType() != mu::NodeScalarEnumParameter::GetStaticType())
-				{
-					const FText Message = LOCTEXT("WrongSwitchParamType", "Switch parameter of incorrect type.");
-					GenerationContext.Compiler->CompilerLog(Message, Node);
-
-					return Result;
-				}
-
-				const int32 NumSwitchOptions = TypedNodePassThroughTextureSwitch->GetNumElements();
-
-				mu::NodeScalarEnumParameter* EnumParameter = static_cast<mu::NodeScalarEnumParameter*>(SwitchParam.get());
-				if (NumSwitchOptions != EnumParameter->GetValueCount())
-				{
-					const FText Message = LOCTEXT("MismatchedSwitch", "Switch enum and switch node have different number of options. Please refresh the switch node to make sure the outcomes are labeled properly.");
-					GenerationContext.Compiler->CompilerLog(Message, Node);
-				}
-				
-				// TODO: Implement Mutable core pass-through switch nodes
-				//mu::NodePassThroughImageSwitchPtr SwitchNode = new mu::NodePassThroughImageSwitchPtr;
-				//SwitchNode->SetParameter(SwitchParam);
-				//SwitchNode->SetOptionCount(NumSwitchOptions);
-
-				for (int32 SelectorIndex = 0; SelectorIndex < NumSwitchOptions; ++SelectorIndex)
-				{
-					if (const UEdGraphPin* TexturePin = FollowInputPin(*TypedNodePassThroughTextureSwitch->GetElementPin(SelectorIndex)))
-					{
-						mu::NodeImagePtr PassThroughImage = GenerateMutableSourceImage(TexturePin, GenerationContext, MaxTextureSize);
-						//SwitchNode->SetOption(SelectorIndex, PassThroughImage);
-					}
-					else
-					{
-						const FText Message = LOCTEXT("MissingPassThroughTexture", "Unable to generate pass-through texture switch node. Required connection not found.");
-						GenerationContext.Compiler->CompilerLog(Message, Node);
-						return Result;
-					}
-				}
-
-				//Result = SwitchNode;
-				return Result;
-			}
-			else
-			{
-				GenerationContext.Compiler->CompilerLog(LOCTEXT("NoEnumParamInSwitch", "Switch nodes must have an enum switch parameter. Please connect an enum and refesh the switch node."), Node);
-				return Result;
-			}
-		}(); // invoke lambda.
-	}
-
 	else if (const UCustomizableObjectNodeTextureVariation* TypedNodeImageVar = Cast<const UCustomizableObjectNodeTextureVariation>(Node))
 	{
 		mu::NodeImageVariationPtr TextureNode = new mu::NodeImageVariation();
@@ -905,7 +808,103 @@
 			SaturateNode->SetFactor(FactorNode); 
 		}
 	}
-	
+
+	else if (const UCustomizableObjectNodePassThroughTexture* TypedNodePassThroughTex = Cast<UCustomizableObjectNodePassThroughTexture>(Node))
+	{
+		UTexture2D* BaseTexture = TypedNodePassThroughTex->Texture;
+		if (BaseTexture)
+		{
+			uint32* FoundIndex = GenerationContext.PassThroughTextureToIndexMap.Find(BaseTexture);
+			uint32 NewId;
+
+			if (!FoundIndex)
+			{
+				NewId = GenerationContext.PassThroughTextureToIndexMap.Num();
+				GenerationContext.PassThroughTextureToIndexMap.Add(BaseTexture, NewId);
+			}
+			else
+			{
+				NewId = *FoundIndex;
+			}
+
+			mu::Ptr<mu::NodeImageReference> ImageNode = new mu::NodeImageReference;
+			ImageNode->SetImageReference(NewId);
+
+			Result = ImageNode;
+		}
+		else
+		{
+			GenerationContext.Compiler->CompilerLog(LOCTEXT("MissingImagePassThrough", "Missing image in pass-through texture node."), Node);
+		}
+	}
+
+	else if (const UCustomizableObjectNodePassThroughTextureSwitch* TypedNodePassThroughTextureSwitch = Cast<UCustomizableObjectNodePassThroughTextureSwitch>(Node))
+	{
+		Result = [&]()
+		{
+			const UEdGraphPin* SwitchParameter = TypedNodePassThroughTextureSwitch->SwitchParameter();
+
+			// Check Switch Parameter arity preconditions.
+			if (const UEdGraphPin* ConnectedPin = FollowInputPin(*SwitchParameter))
+			{
+				mu::NodeScalarPtr SwitchParam = GenerateMutableSourceFloat(ConnectedPin, GenerationContext);
+				// Switch Param not generated
+				if (!SwitchParam)
+				{
+					const FText Message = LOCTEXT("FailedToGenerateSwitchParam", "Could not generate switch enum parameter. Please refesh the switch node and connect an enum.");
+					GenerationContext.Compiler->CompilerLog(Message, Node);
+
+					return Result;
+				}
+
+				if (SwitchParam->GetType() != mu::NodeScalarEnumParameter::GetStaticType())
+				{
+					const FText Message = LOCTEXT("WrongSwitchParamType", "Switch parameter of incorrect type.");
+					GenerationContext.Compiler->CompilerLog(Message, Node);
+
+					return Result;
+				}
+
+				const int32 NumSwitchOptions = TypedNodePassThroughTextureSwitch->GetNumElements();
+
+				mu::NodeScalarEnumParameter* EnumParameter = dynamic_cast<mu::NodeScalarEnumParameter*>(SwitchParam.get());
+				if (NumSwitchOptions != EnumParameter->GetValueCount())
+				{
+					const FText Message = LOCTEXT("MismatchedSwitch", "Switch enum and switch node have different number of options. Please refresh the switch node to make sure the outcomes are labeled properly.");
+					GenerationContext.Compiler->CompilerLog(Message, Node);
+				}
+
+				// TODO: Implement Mutable core pass-through switch nodes
+				mu::Ptr<mu::NodeImageSwitch> SwitchNode = new mu::NodeImageSwitch;
+				SwitchNode->SetParameter(SwitchParam);
+				SwitchNode->SetOptionCount(NumSwitchOptions);
+
+				for (int32 SelectorIndex = 0; SelectorIndex < NumSwitchOptions; ++SelectorIndex)
+				{
+					if (const UEdGraphPin* TexturePin = FollowInputPin(*TypedNodePassThroughTextureSwitch->GetElementPin(SelectorIndex)))
+					{
+						mu::Ptr<mu::NodeImage> PassThroughImage = GenerateMutableSourceImage(TexturePin, GenerationContext, MaxTextureSize);
+						SwitchNode->SetOption(SelectorIndex, PassThroughImage);
+					}
+					else
+					{
+						const FText Message = LOCTEXT("MissingPassThroughTexture", "Unable to generate pass-through texture switch node. Required connection not found.");
+						GenerationContext.Compiler->CompilerLog(Message, Node);
+						return Result;
+					}
+				}
+
+				Result = SwitchNode;
+				return Result;
+			}
+			else
+			{
+				GenerationContext.Compiler->CompilerLog(LOCTEXT("NoEnumParamInSwitch", "Switch nodes must have an enum switch parameter. Please connect an enum and refesh the switch node."), Node);
+				return Result;
+			}
+		}(); // invoke lambda.
+	}
+
 	// If the node is a plain colour node, generate an image out of it
 	else if (Pin->PinType.PinCategory == Schema->PC_Color)
 	{
@@ -994,5 +993,6 @@
 	return Result;
 }
 
+
 #undef LOCTEXT_NAMESPACE
 
diff --git a/Engine/Plugins/Experimental/Mutable/Source/CustomizableObjectEditor/Private/MuCOE/GenerateMutableSource/GenerateMutableSourceImage.h b/Engine/Plugins/Experimental/Mutable/Source/CustomizableObjectEditor/Private/MuCOE/GenerateMutableSource/GenerateMutableSourceImage.h
--- a/Engine/Plugins/Experimental/Mutable/Source/CustomizableObjectEditor/Private/MuCOE/GenerateMutableSource/GenerateMutableSourceImage.h
+++ b/Engine/Plugins/Experimental/Mutable/Source/CustomizableObjectEditor/Private/MuCOE/GenerateMutableSource/GenerateMutableSourceImage.h
@@ -3,8 +3,7 @@
 #pragma once
 
 
-#include "MuR/Image.h"
-#include "MuT/NodeImageConstant.h"
+#include "MuR/Ptr.h"
 
 class FCustomizableObjectCompiler;
 class UCustomizableObjectNode;
@@ -12,12 +11,19 @@ class UEdGraphPin;
 class UTexture2D;
 struct FMutableGraphGenerationContext;
 
+namespace mu
+{
+	class Image;
+	class NodeImage;
+	class NodeImageConstant;
+}
 
-mu::ImagePtr ConvertTextureUnrealToMutable(UTexture2D* Texture, const UCustomizableObjectNode* Node, FCustomizableObjectCompiler* Compiler, bool bIsNormalComposite);
 
+mu::Ptr<mu::Image> ConvertTextureUnrealToMutable(UTexture2D* Texture, const UCustomizableObjectNode* Node, FCustomizableObjectCompiler* Compiler, bool bIsNormalComposite);
 
-mu::NodeImagePtr ResizeToMaxTextureSize(float MaxTextureSize, const UTexture2D* BaseTexture, mu::NodeImageConstantPtr ImageNode);
 
+mu::Ptr<mu::NodeImage> ResizeToMaxTextureSize(float MaxTextureSize, const UTexture2D* BaseTexture, mu::Ptr<mu::NodeImageConstant> ImageNode);
 
-/** Convert a CustomizableObject Source Graph into a mutable source graph. */
-mu::NodeImagePtr GenerateMutableSourceImage(const UEdGraphPin* Pin, FMutableGraphGenerationContext& GenerationContext, float MaxTextureSize);
+
+/** Convert a CustomizableObject Source Graph from an Image pin into a mutable source graph. */
+mu::Ptr<mu::NodeImage> GenerateMutableSourceImage(const UEdGraphPin* Pin, FMutableGraphGenerationContext& GenerationContext, float MaxTextureSize);
diff --git a/Engine/Plugins/Experimental/Mutable/Source/CustomizableObjectEditor/Private/MuCOE/GenerateMutableSource/GenerateMutableSourceSurface.cpp b/Engine/Plugins/Experimental/Mutable/Source/CustomizableObjectEditor/Private/MuCOE/GenerateMutableSource/GenerateMutableSourceSurface.cpp
--- a/Engine/Plugins/Experimental/Mutable/Source/CustomizableObjectEditor/Private/MuCOE/GenerateMutableSource/GenerateMutableSourceSurface.cpp
+++ b/Engine/Plugins/Experimental/Mutable/Source/CustomizableObjectEditor/Private/MuCOE/GenerateMutableSource/GenerateMutableSourceSurface.cpp
@@ -442,8 +442,28 @@ mu::NodeSurfacePtr GenerateMutableSourceSurface(const UEdGraphPin * Pin, FMutabl
 
 				if (const UEdGraphPin* ConnectedPin = FollowInputPin(*ImagePin))
 				{
-					mu::NodeImagePtr PassThroughImagePtr = GenerateMutableSourceImage(ConnectedPin, GenerationContext, 0);
+					mu::Ptr<mu::NodeImage> PassThroughImagePtr = GenerateMutableSourceImage(ConnectedPin, GenerationContext, 0);
 					SurfNode->SetImage(ImageIndex, PassThroughImagePtr);
+
+					const FString ImageName = TypedNodeMat->GetParameterName(EMaterialParameterType::Texture, ImageIndex).ToString();
+					FString SurfNodeImageName = FString::Printf(TEXT("%d"), GenerationContext.ImageProperties.Num());
+					SurfNode->SetImageName(ImageIndex, StringCast<ANSICHAR>(*SurfNodeImageName).Get());
+					int32 UVLayout = TypedNodeMat->GetImageUVLayout(ImageIndex);
+					if (UVLayout >= 0)
+					{
+						FString msg = FString::Printf(TEXT("Passthrough texture [%s] will ignore layout despite set to use layout [%d]"), *ImageName, UVLayout);
+						GenerationContext.Compiler->CompilerLog(FText::FromString(msg), Node, EMessageSeverity::Warning);
+
+					}
+					SurfNode->SetImageLayoutIndex(ImageIndex, -1);
+					SurfNode->SetImageAdditionalNames(ImageIndex, StringCast<ANSICHAR>(*TypedNodeMat->Material->GetName()).Get(), StringCast<ANSICHAR>(*ImageName).Get());
+
+					// We don't need a reference texture or props here, but we do need the parameter name.
+					FGeneratedImageProperties Props;
+					Props.TextureParameterName = ImageName;
+					GenerationContext.ImageProperties.Add(Props);
+					SurfaceData.ImageProperties = Props;
+
 				}
 			}
 			else
diff --git a/Engine/Plugins/Experimental/Mutable/Source/MutableRuntime/Private/MuR/CodeRunner.cpp b/Engine/Plugins/Experimental/Mutable/Source/MutableRuntime/Private/MuR/CodeRunner.cpp
--- a/Engine/Plugins/Experimental/Mutable/Source/MutableRuntime/Private/MuR/CodeRunner.cpp
+++ b/Engine/Plugins/Experimental/Mutable/Source/MutableRuntime/Private/MuR/CodeRunner.cpp
@@ -3828,11 +3828,24 @@ namespace mu
             break;
         }
 
-        case OP_TYPE::IM_GPU:
-        {
-            check(false);
-            break;
-        }
+		case OP_TYPE::IM_REFERENCE:
+		{
+			OP::ResourceReferenceArgs Args = pModel->GetPrivate()->m_program.GetOpArgs<OP::ResourceReferenceArgs>(item.At);
+			switch (item.Stage)
+			{
+			case 0:
+			{
+				Ptr<Image> Result = Image::CreateAsReference(Args.ID);
+				GetMemory().SetImage(item, Result);
+				break;
+			}
+
+			default:
+				check(false);
+			}
+
+			break;
+		}
 
         case OP_TYPE::IM_CROP:
         {
@@ -5560,10 +5573,12 @@ namespace mu
 			check(item.Stage == 0);
 			OP::ResourceConstantArgs args = program.GetOpArgs<OP::ResourceConstantArgs>(item.At);
 			int32 ImageIndex = args.value;
-			m_heapImageDesc[item.CustomState].m_format = program.m_constantImages[ImageIndex].ImageFormat;
-			m_heapImageDesc[item.CustomState].m_size[0] = program.m_constantImages[ImageIndex].ImageSizeX;
-			m_heapImageDesc[item.CustomState].m_size[1] = program.m_constantImages[ImageIndex].ImageSizeY;
-			m_heapImageDesc[item.CustomState].m_lods = program.m_constantImages[ImageIndex].LODCount;
+
+			FImageDesc& Result = m_heapImageDesc[item.CustomState];
+			Result.m_format = program.m_constantImages[ImageIndex].ImageFormat;
+			Result.m_size[0] = program.m_constantImages[ImageIndex].ImageSizeX;
+			Result.m_size[1] = program.m_constantImages[ImageIndex].ImageSizeY;
+			Result.m_lods = program.m_constantImages[ImageIndex].LODCount;
 			GetMemory().SetValidDesc(item);
 			break;
 		}
@@ -5579,6 +5594,18 @@ namespace mu
 			break;
 		}
 
+		case OP_TYPE::IM_REFERENCE:
+		{
+			check(item.Stage == 0);
+			FImageDesc& Result = m_heapImageDesc[item.CustomState];
+			Result.m_format = EImageFormat::IF_NONE;
+			Result.m_size[0] = 0;
+			Result.m_size[1] = 0;
+			Result.m_lods = 0;
+			GetMemory().SetValidDesc(item);
+			break;
+		}
+
 		case OP_TYPE::IM_CONDITIONAL:
 		{
 			OP::ConditionalArgs args = program.GetOpArgs<OP::ConditionalArgs>(item.At);
diff --git a/Engine/Plugins/Experimental/Mutable/Source/MutableRuntime/Private/MuR/CodeRunnerTasked.cpp b/Engine/Plugins/Experimental/Mutable/Source/MutableRuntime/Private/MuR/CodeRunnerTasked.cpp
--- a/Engine/Plugins/Experimental/Mutable/Source/MutableRuntime/Private/MuR/CodeRunnerTasked.cpp
+++ b/Engine/Plugins/Experimental/Mutable/Source/MutableRuntime/Private/MuR/CodeRunnerTasked.cpp
@@ -219,6 +219,8 @@ namespace mu
 		{  20,   0,   0,   0 },	// IM_PARAMETER
 		{   0,   0,   0,   0 },	// ST_PARAMETER
 		    
+		{   0,   0,   0,   0 },	// IM_REFERENCE
+
 		{   0,   0,   0,   0 },	// NU_CONDITIONAL
 		{   0,   0,   0,   0 },	// SC_CONDITIONAL
 		{   0,   0,   0,   0 },	// CO_CONDITIONAL
@@ -272,7 +274,6 @@ namespace mu
 		{   0,   5,   0,   0 },	// IM_GRADIENT
 		{   0,   0,   0,   0 },	// IM_BINARISE
 		{   0,  20,   0,   0 },	// IM_PLAINCOLOUR
-		{   0,   0,   0,   0 },	// IM_GPU
 		{   0, -10,   0,   0 },	// IM_CROP
 		{   0, -10,   0,   0 },	// IM_PATCH
 		{   0,  10,  10,   0 },	// IM_RASTERMESH
diff --git a/Engine/Plugins/Experimental/Mutable/Source/MutableRuntime/Private/MuR/Image.cpp b/Engine/Plugins/Experimental/Mutable/Source/MutableRuntime/Private/MuR/Image.cpp
--- a/Engine/Plugins/Experimental/Mutable/Source/MutableRuntime/Private/MuR/Image.cpp
+++ b/Engine/Plugins/Experimental/Mutable/Source/MutableRuntime/Private/MuR/Image.cpp
@@ -104,7 +104,6 @@ namespace mu
         m_format = format;
         m_size = FImageSize( (uint16)sizeX, (uint16)sizeY );
         m_lods = (uint8_t)lods;
-        m_internalId = 0;
 
         const FImageFormatData& fdata = GetImageFormatData( format );
         int32 PixelsPerBlock = fdata.m_pixelsPerBlockX*fdata.m_pixelsPerBlockY;
@@ -150,6 +149,16 @@ namespace mu
     }
 
 
+	//---------------------------------------------------------------------------------------------
+	Ptr<Image> Image::CreateAsReference(uint32 ID)
+	{
+		Ptr<Image> Result = new Image;
+		Result->ReferenceID = ID;
+		Result->m_flags = EImageFlags::IF_IS_REFERENCE;
+		return Result;
+	}
+
+
     //---------------------------------------------------------------------------------------------
     void Image::Serialise( const Image* p, OutputArchive& arch )
     {
@@ -277,38 +286,44 @@ namespace mu
 
 
     //---------------------------------------------------------------------------------------------
-    const uint8_t* Image::GetData() const
+    const uint8* Image::GetData() const
     {
         return m_data.GetData();
     }
 
 
     //---------------------------------------------------------------------------------------------
-    int32_t Image::GetDataSize() const
+    int32 Image::GetDataSize() const
     {
         return m_data.Num();
     }
 
 
     //---------------------------------------------------------------------------------------------
-    int32_t Image::GetLODDataSize( int lod ) const
+    int32 Image::GetLODDataSize( int lod ) const
     {
         return CalculateDataSize( lod );
     }
 
 
     //---------------------------------------------------------------------------------------------
-    uint8_t* Image::GetData()
+    uint8* Image::GetData()
     {
 		return m_data.GetData();
     }
 
+	//---------------------------------------------------------------------------------------------
+	bool Image::IsReference() const
+	{
+		return m_flags & EImageFlags::IF_IS_REFERENCE;
+	}
 
-    //---------------------------------------------------------------------------------------------
-    uint32_t Image::GetId() const
-    {
-        return m_internalId;
-    }
+	//---------------------------------------------------------------------------------------------
+	uint32 Image::GetReferencedTexture() const
+	{
+		ensure(IsReference());
+		return ReferenceID;
+	}
 
 
     //---------------------------------------------------------------------------------------------
diff --git a/Engine/Plugins/Experimental/Mutable/Source/MutableRuntime/Private/MuR/Operations.cpp b/Engine/Plugins/Experimental/Mutable/Source/MutableRuntime/Private/MuR/Operations.cpp
--- a/Engine/Plugins/Experimental/Mutable/Source/MutableRuntime/Private/MuR/Operations.cpp
+++ b/Engine/Plugins/Experimental/Mutable/Source/MutableRuntime/Private/MuR/Operations.cpp
@@ -41,6 +41,8 @@ namespace mu
 		{ DT_IMAGE,			false,	{ 0, 0, 0, 0, 0, 0, 0, 0, 0, 0, 0, 0, 0, 0, 0, 0, 0, 0, 0, 0, 0, 0, 0, 0, 0, 0, 0, 0, 0 }		},	// IM_PARAMETER
 		{ DT_STRING,		false,	{ 0, 0, 0, 0, 0, 0, 0, 0, 0, 0, 0, 0, 0, 0, 0, 0, 0, 0, 0, 0, 0, 0, 0, 0, 0, 0, 0, 0, 0 }		},	// ST_PARAMETER
 
+		{ DT_IMAGE,			true,	{ 0, 0, 0, 0, 0, 0, 0, 0, 0, 0, 0, 0, 0, 0, 0, 0, 0, 0, 0, 0, 0, 0, 0, 0, 0, 0, 0, 0, 0 }		},	// IM_REFERENCE
+
 		{ DT_INT,			false,	{ 0, 0, 0, 0, 0, 0, 0, 0, 0, 0, 0, 0, 0, 0, 0, 0, 0, 0, 0, 0, 0, 0, 0, 0, 0, 0, 0, 0, 0 }		},	// NU_CONDITIONAL
 		{ DT_SCALAR,		false,	{ 0, 0, 0, 0, 0, 0, 0, 0, 0, 0, 0, 0, 0, 0, 0, 0, 0, 0, 0, 0, 0, 0, 0, 0, 0, 0, 0, 0, 0 }		},	// SC_CONDITIONAL
 		{ DT_COLOUR,		false,	{ 0, 0, 0, 0, 0, 0, 0, 0, 0, 0, 0, 0, 0, 0, 0, 0, 0, 0, 0, 0, 0, 0, 0, 0, 0, 0, 0, 0, 0 }		},	// CO_CONDITIONAL
@@ -94,7 +96,6 @@ namespace mu
 		{ DT_IMAGE,			true,	{ 0, 0, 0, 0, 0, 0, 0, 0, 0, 0, 0, 0, 0, 0, 0, 0, 0, 0, 0, 0, 0, 0, 0, 0, 0, 0, 0, 0, 0 }		},	// IM_GRADIENT
 		{ DT_IMAGE,			true,	{ 0, 1, 1, 1, 0, 0, 0, 0, 0, 0, 0, 0, 0, 0, 0, 0, 0, 0, 0, 0, 0, 0, 0, 0, 0, 0, 0, 0, 0 }		},	// IM_BINARISE
 		{ DT_IMAGE,			true,	{ 0, 0, 0, 0, 0, 0, 0, 0, 0, 0, 0, 0, 0, 0, 0, 0, 0, 0, 0, 0, 0, 0, 0, 0, 0, 0, 0, 0, 0 }		},	// IM_PLAINCOLOUR
-		{ DT_IMAGE,			true,	{ 0, 0, 0, 0, 0, 0, 0, 0, 0, 0, 0, 0, 0, 0, 0, 0, 0, 0, 0, 0, 0, 0, 0, 0, 0, 0, 0, 0, 0 }		},	// IM_GPU
 		{ DT_IMAGE,			true,	{ 0, 0, 0, 0, 0, 0, 0, 0, 0, 0, 0, 0, 0, 0, 0, 0, 0, 0, 0, 0, 0, 0, 0, 0, 0, 0, 0, 0, 0 }		},	// IM_CROP
 		{ DT_IMAGE,			true,	{ 0, 0, 0, 0, 0, 0, 0, 0, 0, 0, 0, 0, 0, 0, 0, 0, 0, 0, 0, 0, 0, 0, 0, 0, 0, 0, 0, 0, 0 }		},	// IM_PATCH
 		{ DT_IMAGE,			true,	{ 0, 0, 0, 0, 0, 0, 0, 0, 0, 0, 0, 0, 0, 0, 0, 0, 0, 0, 0, 0, 0, 0, 0, 0, 0, 0, 0, 0, 0 }		},	// IM_RASTERMESH
@@ -418,7 +419,8 @@ namespace mu
         case OP_TYPE::CO_PARAMETER:
         case OP_TYPE::PR_PARAMETER:
         case OP_TYPE::IM_PARAMETER:
-            break;
+		case OP_TYPE::IM_REFERENCE:
+			break;
 
         case OP_TYPE::SC_CURVE:
         {
diff --git a/Engine/Plugins/Experimental/Mutable/Source/MutableRuntime/Private/MuR/Operations.h b/Engine/Plugins/Experimental/Mutable/Source/MutableRuntime/Private/MuR/Operations.h
--- a/Engine/Plugins/Experimental/Mutable/Source/MutableRuntime/Private/MuR/Operations.h
+++ b/Engine/Plugins/Experimental/Mutable/Source/MutableRuntime/Private/MuR/Operations.h
@@ -51,6 +51,9 @@ namespace mu
         IM_PARAMETER,
         ST_PARAMETER,
 
+		//! A referenced, but opaque engine resource
+		IM_REFERENCE,
+
         //! Select one value or the other depending on a boolean input
         NU_CONDITIONAL,
         SC_CONDITIONAL,
@@ -198,10 +201,6 @@ namespace mu
         //! Generate a plain colour image
         IM_PLAINCOLOUR,
 
-        //! Image resulting from a GPU program
-		//! \TODO: Deprecated?
-		IM_GPU,
-
         //! Cut a rect from an image
         IM_CROP,
 
@@ -397,7 +396,7 @@ namespace mu
 
         struct IntConstantArgs
         {
-            int value;
+            int32 value;
         };
 
         struct ScalarConstantArgs
@@ -439,6 +438,11 @@ namespace mu
             ADDRESS condition, yes, no;
         };
 
+		struct ResourceReferenceArgs
+		{
+			int32 ID;
+		};
+
         //-------------------------------------------------------------------------------------
         struct BoolLessArgs
         {
diff --git a/Engine/Plugins/Experimental/Mutable/Source/MutableRuntime/Private/MuR/System.cpp b/Engine/Plugins/Experimental/Mutable/Source/MutableRuntime/Private/MuR/System.cpp
--- a/Engine/Plugins/Experimental/Mutable/Source/MutableRuntime/Private/MuR/System.cpp
+++ b/Engine/Plugins/Experimental/Mutable/Source/MutableRuntime/Private/MuR/System.cpp
@@ -307,8 +307,6 @@ namespace mu
 					pResult = new mu::Image(16, 16, 1, EImageFormat::IF_RGBA_UBYTE);
 				}
 
-				pResult->m_internalId = imageId;
-
 				break;
 			}
 		}
@@ -404,8 +402,6 @@ namespace mu
 					pResult = new Mesh();
 				}
 
-				pResult->m_internalId = meshId;
-
 				break;
 			}
 		}
diff --git a/Engine/Plugins/Experimental/Mutable/Source/MutableRuntime/Public/MuR/Image.h b/Engine/Plugins/Experimental/Mutable/Source/MutableRuntime/Public/MuR/Image.h
--- a/Engine/Plugins/Experimental/Mutable/Source/MutableRuntime/Public/MuR/Image.h
+++ b/Engine/Plugins/Experimental/Mutable/Source/MutableRuntime/Public/MuR/Image.h
@@ -262,6 +262,9 @@ namespace mu
         //!         depends on the image size.
         //! \param format Pixel format.
         Image( uint32 sizeX, uint32 sizeY, uint32 lods, EImageFormat format, EInitializationType Init=EInitializationType::Black);
+		
+		/** */
+		static Ptr<Image> CreateAsReference( uint32 ID );
 
 		/** */
 		ImagePtr ExtractMip(int32 Mip) const;
@@ -300,9 +303,11 @@ namespace mu
         //! Return the size in bytes of a specific LOD of the image.
         int32 GetLODDataSize( int lod ) const;
 
-		//! Get an internal identifier used to reference this image in operations like deferred
-		//! image building, or instance updating.
-        uint32 GetId() const;
+		/** Return true if this is a reference to an engine image. */
+		bool IsReference() const;
+
+		/** Return the id of the engine referenced texture. Only valid if IsReference. */
+		uint32 GetReferencedTexture() const;
 
 	protected:
 
@@ -317,14 +322,11 @@ namespace mu
 		//!
 		FImageSize m_size = FImageSize(0, 0);
 
-		//! Non-persistent internal id
-		mutable uint32 m_internalId = 0;
-
 		//!
-		EImageFormat m_format;
+		EImageFormat m_format = EImageFormat::IF_NONE;
 
 		//! Levels of detail (mipmaps)
-		uint8 m_lods;
+		uint8 m_lods = 0;
 
 		//! These set of flags are used to cache information for images at runtime.
 		typedef enum
@@ -340,6 +342,9 @@ namespace mu
 
 			// If this is set, the image has an updated relevancy map. This flag is not persisent.
 			IF_HAS_RELEVANCY_MAP = 1 << 3,
+
+			/** If this is set, this is a reference to an external image, and the ReferenceID is valid. */
+			IF_IS_REFERENCE = 1 << 4,
 		} EImageFlags;
 
 		//! Persistent flags with some image properties. The meaning will depend of every
@@ -353,6 +358,9 @@ namespace mu
 		//! Pixel data
 		TArray<uint8> m_data;
 
+		/** Only valid if the right flags are set, this identifies a referenced image. */
+		uint32 ReferenceID = 0;
+
 
 		// This used to be the methods in the private implementation of the image interface
 		//-----------------------------------------------------------------------------------------
@@ -366,13 +374,13 @@ namespace mu
 			ImagePtr pResult = new Image();
 
 			pResult->m_size = m_size;
-			pResult->m_internalId = m_internalId;
 			pResult->m_format = m_format;
 			pResult->m_lods = m_lods;
 			pResult->m_flags = m_flags;
 			pResult->RelevancyMinY = RelevancyMinY;
 			pResult->RelevancyMaxY = RelevancyMaxY;
 			pResult->m_data = m_data;
+			pResult->ReferenceID = ReferenceID;
 
 			return pResult;
 		}
@@ -390,7 +398,8 @@ namespace mu
 			return (m_lods == o.m_lods) &&
 				(m_size == o.m_size) &&
 				(m_format == o.m_format) &&
-				(m_data == o.m_data);
+				(m_data == o.m_data) &&
+				(ReferenceID == o.ReferenceID);
 		}
 
 		//-----------------------------------------------------------------------------------------
diff --git a/Engine/Plugins/Experimental/Mutable/Source/MutableRuntime/Public/MuR/Model.h b/Engine/Plugins/Experimental/Mutable/Source/MutableRuntime/Public/MuR/Model.h
--- a/Engine/Plugins/Experimental/Mutable/Source/MutableRuntime/Public/MuR/Model.h
+++ b/Engine/Plugins/Experimental/Mutable/Source/MutableRuntime/Public/MuR/Model.h
@@ -13,7 +13,7 @@
 //! data structures. Compiled models are not necessarily compatible when the runtime is updated,
 //! so this version number can be used externally to verify this. It is not used internally, and
 //! serializing models from different versions than this runtime will probably result in a crash.
-#define MUTABLE_COMPILED_MODEL_CODE_VERSION		uint32( 64 )
+#define MUTABLE_COMPILED_MODEL_CODE_VERSION		uint32( 65 )
 #define MUTABLE_PARAMETERS_VERSION              uint32( 2 )
 
 
diff --git a/Engine/Plugins/Experimental/Mutable/Source/MutableTools/Private/MuT/ASTOpConstantResource.cpp b/Engine/Plugins/Experimental/Mutable/Source/MutableTools/Private/MuT/ASTOpConstantResource.cpp
--- a/Engine/Plugins/Experimental/Mutable/Source/MutableTools/Private/MuT/ASTOpConstantResource.cpp
+++ b/Engine/Plugins/Experimental/Mutable/Source/MutableTools/Private/MuT/ASTOpConstantResource.cpp
@@ -34,7 +34,7 @@ namespace mu
 	//-------------------------------------------------------------------------------------------------
 	bool ASTOpConstantResource::IsEqual(const ASTOp& otherUntyped) const
 	{
-		if (auto other = dynamic_cast<const ASTOpConstantResource*>(&otherUntyped))
+		if (const ASTOpConstantResource* other = dynamic_cast<const ASTOpConstantResource*>(&otherUntyped))
 		{
 			return type == other->type && hash == other->hash &&
 				loadedValue == other->loadedValue && proxy == other->proxy;
@@ -72,7 +72,7 @@ namespace mu
 			if (type == OP_TYPE::ME_CONSTANT)
 			{
 				OP::MeshConstantArgs args;
-				memset(&args, 0, sizeof(args));
+				FMemory::Memset(&args, 0, sizeof(args));
 
 				Ptr<Mesh> pTyped = static_cast<const Mesh*>(GetValue().get())->Clone();
 				check(pTyped);
@@ -80,27 +80,6 @@ namespace mu
 				args.skeleton = -1;
 				if (Ptr<const Skeleton> pSkeleton = pTyped->GetSkeleton())
 				{
-					// See if it can be merged into another skeleton
-	//                for (const auto& pOtherSkeleton : program.m_constantSkeletons)
-	//                {
-	//                    if (true)
-	//                    //if (AreSkeletonsRelated(pOtherSkeleton, pSkeleton))
-	//                    {
-	//                        // \todo : ugly const hack
-	//                        ExtendSkeleton( const_cast<Skeleton*>(pOtherSkeleton.get()),
-	//                                        pSkeleton.get() );
-	//                        pSkeleton = pOtherSkeleton;
-
-	//                        Ptr<Mesh> pRemapped = MeshRemapSkeleton( pTyped.get(),
-	//                                                                 pOtherSkeleton.get() );
-	//                        if (pRemapped)
-	//                        {
-	//                            pTyped = pRemapped;
-	//                        }
-	//                        break;
-	//                    }
-	//                }
-
 					args.skeleton = program.AddConstant(pSkeleton.get());
 					pTyped->SetSkeleton(nullptr);
 				}
@@ -122,7 +101,7 @@ namespace mu
 			else
 			{
 				OP::ResourceConstantArgs args;
-				memset(&args, 0, sizeof(args));
+				FMemory::Memset(&args, 0, sizeof(args));
 
 				bool bValidData = true;
 
diff --git a/Engine/Plugins/Experimental/Mutable/Source/MutableTools/Private/MuT/ASTOpConstantResource.h b/Engine/Plugins/Experimental/Mutable/Source/MutableTools/Private/MuT/ASTOpConstantResource.h
--- a/Engine/Plugins/Experimental/Mutable/Source/MutableTools/Private/MuT/ASTOpConstantResource.h
+++ b/Engine/Plugins/Experimental/Mutable/Source/MutableTools/Private/MuT/ASTOpConstantResource.h
@@ -12,8 +12,8 @@
 
 namespace mu
 {
-struct FProgram;
-template <class SCALAR> class vec4;
+	struct FProgram;
+	template <class SCALAR> class vec4;
 
 
 	//---------------------------------------------------------------------------------------------
@@ -23,10 +23,6 @@ template <class SCALAR> class vec4;
 	{
 	private:
 
-		//!
-		//std::FILE* resourceFile = nullptr;
-		//size_t fileSize = 0;
-
 		//!
 		Ptr<const RefCounted> loadedValue;
 		Ptr<RefCounted> proxy;
diff --git a/Engine/Plugins/Experimental/Mutable/Source/MutableTools/Private/MuT/ASTOpReferenceResource.cpp b/Engine/Plugins/Experimental/Mutable/Source/MutableTools/Private/MuT/ASTOpReferenceResource.cpp
new file mode 100644
--- /dev/null
+++ b/Engine/Plugins/Experimental/Mutable/Source/MutableTools/Private/MuT/ASTOpReferenceResource.cpp
@@ -0,0 +1,144 @@
+// Copyright Epic Games, Inc. All Rights Reserved.
+
+#include "MuT/ASTOpReferenceResource.h"
+
+#include "Containers/Array.h"
+#include "HAL/PlatformMath.h"
+#include "Hash/CityHash.h"
+#include "Misc/AssertionMacros.h"
+#include "MuR/Layout.h"
+#include "MuR/Mesh.h"
+#include "MuR/ModelPrivate.h"
+#include "MuR/MutableMath.h"
+#include "MuR/PhysicsBody.h"
+#include "MuR/Serialisation.h"
+#include "MuR/Skeleton.h"
+#include "MuR/Types.h"
+#include "MuT/StreamsPrivate.h"
+
+
+namespace mu
+{
+
+	void ASTOpReferenceResource::ForEachChild(const TFunctionRef<void(ASTChild&)>)
+	{
+	}
+
+
+	//-------------------------------------------------------------------------------------------------
+	bool ASTOpReferenceResource::IsEqual(const ASTOp& otherUntyped) const
+	{
+		if (const ASTOpReferenceResource* other = dynamic_cast<const ASTOpReferenceResource*>(&otherUntyped))
+		{
+			return type == other->type && ID == other->ID;
+		}
+		return false;
+	}
+
+
+	//-------------------------------------------------------------------------------------------------
+	mu::Ptr<ASTOp> ASTOpReferenceResource::Clone(MapChildFuncRef) const
+	{
+		Ptr<ASTOpReferenceResource> n = new ASTOpReferenceResource();
+		n->type = type;
+		n->ID = ID;
+		return n;
+	}
+
+
+	//-------------------------------------------------------------------------------------------------
+	uint64 ASTOpReferenceResource::Hash() const
+	{
+		uint64 res = std::hash<uint64>()(uint64(type));
+		hash_combine(res, ID);
+		return res;
+	}
+
+
+	//-------------------------------------------------------------------------------------------------
+	void ASTOpReferenceResource::Link(FProgram& program, const FLinkerOptions* Options)
+	{
+		if (!linkedAddress)
+		{			
+			OP::ResourceReferenceArgs Args;
+			FMemory::Memset(&Args, 0, sizeof(Args));
+			Args.ID = ID;
+
+			linkedAddress = (OP::ADDRESS)program.m_opAddress.Num();
+			program.m_opAddress.Add((uint32)program.m_byteCode.Num());
+			AppendCode(program.m_byteCode, type);
+			AppendCode(program.m_byteCode, Args);
+		}
+	}
+
+
+	//-------------------------------------------------------------------------------------------------
+	FImageDesc ASTOpReferenceResource::GetImageDesc(bool, class FGetImageDescContext*) const
+	{
+		FImageDesc res;
+		return res;
+	}
+
+
+	//-------------------------------------------------------------------------------------------------
+	void ASTOpReferenceResource::GetBlockLayoutSize(int blockIndex, int* pBlockX, int* pBlockY, FBlockLayoutSizeCache*)
+	{
+		switch (type)
+		{
+
+		case OP_TYPE::IM_REFERENCE:
+		{
+			*pBlockX = 0;
+			*pBlockY = 0;
+			break;
+		}
+
+		default:
+			checkf(false, TEXT("Instruction not supported"));
+		}
+	}
+
+
+	//-------------------------------------------------------------------------------------------------
+	void ASTOpReferenceResource::GetLayoutBlockSize(int* pBlockX, int* pBlockY)
+	{
+		switch (type)
+		{
+
+		case OP_TYPE::IM_REFERENCE:
+		{
+			// We didn't find any layout.
+			*pBlockX = 0;
+			*pBlockY = 0;
+			break;
+		}
+
+		default:
+			checkf(false, TEXT("Instruction not supported"));
+		}
+	}
+
+
+	//-------------------------------------------------------------------------------------------------
+	bool ASTOpReferenceResource::GetNonBlackRect(FImageRect&) const
+	{
+		return false;
+	}
+
+
+	//-------------------------------------------------------------------------------------------------
+	ASTOpReferenceResource::~ASTOpReferenceResource()
+	{
+	}
+
+
+	//-------------------------------------------------------------------------------------------------
+	mu::Ptr<ImageSizeExpression> ASTOpReferenceResource::GetImageSizeExpression() const
+	{
+		Ptr<ImageSizeExpression> pRes = new ImageSizeExpression;
+		pRes->type = ImageSizeExpression::ISET_UNKNOWN;
+
+		return pRes;
+	}
+
+}
\ No newline at end of file
diff --git a/Engine/Plugins/Experimental/Mutable/Source/MutableTools/Private/MuT/ASTOpReferenceResource.h b/Engine/Plugins/Experimental/Mutable/Source/MutableTools/Private/MuT/ASTOpReferenceResource.h
new file mode 100644
--- /dev/null
+++ b/Engine/Plugins/Experimental/Mutable/Source/MutableTools/Private/MuT/ASTOpReferenceResource.h
@@ -0,0 +1,57 @@
+// Copyright Epic Games, Inc. All Rights Reserved.
+
+#pragma once
+
+#include "HAL/Platform.h"
+#include "MuR/Image.h"
+#include "MuR/Operations.h"
+#include "MuR/Ptr.h"
+#include "MuR/RefCounted.h"
+#include "MuT/AST.h"
+
+
+namespace mu
+{
+	struct FProgram;
+	template <class SCALAR> class vec4;
+
+
+	//---------------------------------------------------------------------------------------------
+	//! A reference to an engine image (or other resources in the future)
+	//---------------------------------------------------------------------------------------------
+	class ASTOpReferenceResource : public ASTOp
+	{
+	public:
+
+		//! Type of switch
+		OP_TYPE type = OP_TYPE::NONE;
+
+		//!
+		uint32 ID = 0;
+
+	public:
+
+		~ASTOpReferenceResource() override;
+
+		// Own interface
+
+		//! Get a hash of the stored value.
+		uint64 GetValueHash() const;
+
+		// ASTOp interface
+		OP_TYPE GetOpType() const override { return type; }
+		void ForEachChild(const TFunctionRef<void(ASTChild&)>) override;
+		bool IsEqual(const ASTOp& otherUntyped) const override;
+		Ptr<ASTOp> Clone(MapChildFuncRef mapChild) const override;
+		uint64 Hash() const override;
+		void Link(FProgram& program, const FLinkerOptions*) override;
+		FImageDesc GetImageDesc(bool, class FGetImageDescContext*) const override;
+		void GetBlockLayoutSize(int blockIndex, int* pBlockX, int* pBlockY, FBlockLayoutSizeCache* cache) override;
+		void GetLayoutBlockSize(int* pBlockX, int* pBlockY) override;
+		bool GetNonBlackRect(FImageRect& maskUsage) const override;
+		Ptr<ImageSizeExpression> GetImageSizeExpression() const override;
+	};
+
+
+}
+
diff --git a/Engine/Plugins/Experimental/Mutable/Source/MutableTools/Private/MuT/CodeGenerator.h b/Engine/Plugins/Experimental/Mutable/Source/MutableTools/Private/MuT/CodeGenerator.h
--- a/Engine/Plugins/Experimental/Mutable/Source/MutableTools/Private/MuT/CodeGenerator.h
+++ b/Engine/Plugins/Experimental/Mutable/Source/MutableTools/Private/MuT/CodeGenerator.h
@@ -59,6 +59,7 @@ namespace mu
 	class NodeImageColourMap;
 	class NodeImageConditional;
 	class NodeImageConstant;
+	class NodeImageReference;
 	class NodeImageDifference;
 	class NodeImageFormat;
 	class NodeImageGradient;
@@ -386,6 +387,7 @@ namespace mu
 
 		void GenerateImage(FImageGenerationResult& result, const NodeImagePtrConst& node);
 		void GenerateImage_Constant(FImageGenerationResult&, const NodeImageConstant*);
+		void GenerateImage_Reference(FImageGenerationResult&, const NodeImageReference*);
 		void GenerateImage_Difference(FImageGenerationResult&, const NodeImageDifference*);
 		void GenerateImage_Interpolate(FImageGenerationResult&, const NodeImageInterpolate*);
 		void GenerateImage_Saturate(FImageGenerationResult&, const NodeImageSaturate*);
diff --git a/Engine/Plugins/Experimental/Mutable/Source/MutableTools/Private/MuT/CodeGenerator_Image.cpp b/Engine/Plugins/Experimental/Mutable/Source/MutableTools/Private/MuT/CodeGenerator_Image.cpp
--- a/Engine/Plugins/Experimental/Mutable/Source/MutableTools/Private/MuT/CodeGenerator_Image.cpp
+++ b/Engine/Plugins/Experimental/Mutable/Source/MutableTools/Private/MuT/CodeGenerator_Image.cpp
@@ -20,6 +20,7 @@
 #include "MuT/AST.h"
 #include "MuT/ASTOpConditional.h"
 #include "MuT/ASTOpConstantResource.h"
+#include "MuT/ASTOpReferenceResource.h"
 #include "MuT/ASTOpImageMipmap.h"
 #include "MuT/ASTOpImageLayer.h"
 #include "MuT/ASTOpImageLayerColor.h"
@@ -55,6 +56,7 @@
 #include "MuT/NodeImageConditionalPrivate.h"
 #include "MuT/NodeImageConstant.h"
 #include "MuT/NodeImageConstantPrivate.h"
+#include "MuT/NodeImageReferencePrivate.h"
 #include "MuT/NodeImageDifference.h"
 #include "MuT/NodeImageDifferencePrivate.h"
 #include "MuT/NodeImageFormat.h"
@@ -110,9 +112,6 @@
 #include "MuT/Table.h"
 #include "MuT/TablePrivate.h"
 
-#include <memory>
-#include <utility>
-
 
 namespace mu
 {
@@ -185,6 +184,7 @@ namespace mu
 			case NodeImage::EType::Variation: GenerateImage_Variation(result, static_cast<const NodeImageVariation*>(Node)); break;
 			case NodeImage::EType::NormalComposite: GenerateImage_NormalComposite(result, static_cast<const NodeImageNormalComposite*>(Node)); break;
 			case NodeImage::EType::Transform: GenerateImage_Transform(result, static_cast<const NodeImageTransform*>(Node)); break;
+			case NodeImage::EType::Reference: GenerateImage_Reference(result, static_cast<const NodeImageReference*>(Node)); break;
 			case NodeImage::EType::None: check(false);
 			}
 
@@ -259,6 +259,19 @@ namespace mu
 		result.op = op;
     }
 
+	//---------------------------------------------------------------------------------------------
+	void CodeGenerator::GenerateImage_Reference(FImageGenerationResult& result, const NodeImageReference* InNode)
+	{
+		const NodeImageReference::Private& node = *InNode->GetPrivate();
+
+		Ptr<ASTOpReferenceResource> op = new ASTOpReferenceResource();
+		op->type = OP_TYPE::IM_REFERENCE;
+		op->ID = node.ImageReferenceID;
+
+		// TODO: check no crop
+		result.op = op;
+	}
+
 
     //---------------------------------------------------------------------------------------------
 	void CodeGenerator::GenerateImage_Parameter(FImageGenerationResult& result, const NodeImageParameter* InNode )
diff --git a/Engine/Plugins/Experimental/Mutable/Source/MutableTools/Private/MuT/CodeOptimiser.cpp b/Engine/Plugins/Experimental/Mutable/Source/MutableTools/Private/MuT/CodeOptimiser.cpp
--- a/Engine/Plugins/Experimental/Mutable/Source/MutableTools/Private/MuT/CodeOptimiser.cpp
+++ b/Engine/Plugins/Experimental/Mutable/Source/MutableTools/Private/MuT/CodeOptimiser.cpp
@@ -632,6 +632,7 @@ namespace mu
 				switch (n->GetOpType())
 				{
 				case OP_TYPE::IM_BLANKLAYOUT:
+				case OP_TYPE::IM_REFERENCE:
 				case OP_TYPE::IM_COMPOSE:
 				//case OP_TYPE::IM_RASTERMESH:            // TODO review this one
 				case OP_TYPE::ME_MERGE:
diff --git a/Engine/Plugins/Experimental/Mutable/Source/MutableTools/Private/MuT/ErrorLog.cpp b/Engine/Plugins/Experimental/Mutable/Source/MutableTools/Private/MuT/ErrorLog.cpp
--- a/Engine/Plugins/Experimental/Mutable/Source/MutableTools/Private/MuT/ErrorLog.cpp
+++ b/Engine/Plugins/Experimental/Mutable/Source/MutableTools/Private/MuT/ErrorLog.cpp
@@ -204,6 +204,7 @@ namespace mu
 		TEXT("LA_CONSTANT      "),
 		TEXT("PR_CONSTANT      "),
 		TEXT("ST_CONSTANT      "),
+		TEXT("ED_CONSTANT      "),
 
 		TEXT("BO_PARAMETER     "),
 		TEXT("NU_PARAMETER     "),
@@ -213,6 +214,8 @@ namespace mu
 		TEXT("IM_PARAMETER     "),
 		TEXT("ST_PARAMETER     "),
 		
+		TEXT("IM_REFERENCE     "),
+
 		TEXT("NU_CONDITIONAL   "),
 		TEXT("SC_CONDITIONAL   "),
 		TEXT("CO_CONDITIONAL   "),
@@ -266,7 +269,6 @@ namespace mu
 		TEXT("IM_GRADIENT      "),
 		TEXT("IM_BINARISE      "),
 		TEXT("IM_PLAINCOLOUR   "),
-		TEXT("IM_GPU           "),
 		TEXT("IM_CROP          "),
 		TEXT("IM_PATCH         "),
 		TEXT("IM_RASTERMESH    "),
@@ -326,302 +328,4 @@ namespace mu
         return s_opNames[(int)type];
     }
 
-
-	//---------------------------------------------------------------------------------------------
-    extern FString GetOpDesc( const FProgram& program, OP::ADDRESS at )
-	{
-		char temp[1024];
-		FMemory::Memzero( temp, 1024 );
-
-		size_t done = 0;
-
-        OP_TYPE type = program.GetOpType(at);
-        switch (type)
-		{
-		case OP_TYPE::BO_CONSTANT:
-        {
-            auto args = program.GetOpArgs<OP::BoolConstantArgs>(at);
-			mutable_snprintf( temp, 1024,
-                              "bool : %s\t", args.value?"true":"false" );
-			break;
-        }
-        case OP_TYPE::NU_CONSTANT:
-        {
-            auto args = program.GetOpArgs<OP::IntConstantArgs>(at);
-            mutable_snprintf( temp, 1024,
-                              "int : %d\t", args.value );
-            break;
-        }
-        case OP_TYPE::SC_CONSTANT:
-        {
-            auto args = program.GetOpArgs<OP::ScalarConstantArgs>(at);
-            mutable_snprintf( temp, 1024,
-                              "float : %f\t", args.value );
-            break;
-        }
-        case OP_TYPE::CO_CONSTANT:
-        {
-            auto args = program.GetOpArgs<OP::ColourConstantArgs>(at);
-            mutable_snprintf( temp, 1024,
-							  "colour : %f %f %f\t",
-                              args.value[0],
-                              args.value[1],
-                              args.value[2] );
-			break;
-        }
-        case OP_TYPE::IM_CONSTANT:
-        case OP_TYPE::LA_CONSTANT:
-        {
-            auto args = program.GetOpArgs<OP::ResourceConstantArgs>(at);
-            mutable_snprintf( temp, 1024,
-                              "res  : %d\t", (int)args.value );
-            break;
-        }
-        case OP_TYPE::ME_CONSTANT:
-        {
-            auto args = program.GetOpArgs<OP::MeshConstantArgs>(at);
-            mutable_snprintf( temp, 1024,
-                              "res  : %d\t", (int)args.value );
-            break;
-        }
-        case OP_TYPE::BO_PARAMETER:
-        case OP_TYPE::NU_PARAMETER:
-        case OP_TYPE::SC_PARAMETER:
-        case OP_TYPE::CO_PARAMETER:
-        case OP_TYPE::PR_PARAMETER:
-        case OP_TYPE::IM_PARAMETER:
-        {
-            auto args = program.GetOpArgs<OP::ParameterArgs>(at);
-            mutable_snprintf( temp, 1024,
-                              "variable : %2d", (int)args.variable );
-			break;
-        }
-        case OP_TYPE::NU_CONDITIONAL:
-        case OP_TYPE::SC_CONDITIONAL:
-        case OP_TYPE::CO_CONDITIONAL:
-		case OP_TYPE::IM_CONDITIONAL:
-		case OP_TYPE::ME_CONDITIONAL:
-		case OP_TYPE::LA_CONDITIONAL:
-		case OP_TYPE::IN_CONDITIONAL:
-        {
-            auto args = program.GetOpArgs<OP::ConditionalArgs>(at);
-            mutable_snprintf( temp, 1024,
-							  "yes : %3d\tno : %3d\t",
-                              (int)args.yes,
-                              (int)args.no );
-			break;
-        }
-
-        case OP_TYPE::NU_SWITCH:
-        case OP_TYPE::SC_SWITCH:
-        case OP_TYPE::CO_SWITCH:
-        case OP_TYPE::IM_SWITCH:
-		case OP_TYPE::ME_SWITCH:
-        case OP_TYPE::LA_SWITCH:
-        case OP_TYPE::IN_SWITCH:
-        {	
-			const uint8_t* data = program.GetOpArgsPointer(at);
-		
-			OP::ADDRESS VarAddress;
-			FMemory::Memcpy( &VarAddress, data, sizeof(OP::ADDRESS) );
-			data += sizeof(OP::ADDRESS);
-			
-			OP::ADDRESS DefAddress;
-			FMemory::Memcpy( &DefAddress, data, sizeof(OP::ADDRESS));
-			data += sizeof(OP::ADDRESS);
-	
-			uint32_t CaseCount;
-			FMemory::Memcpy( &CaseCount, data,  sizeof(uint32_t));
-			data += sizeof(uint32_t);
-
-            done = mutable_snprintf( temp, 1024, "variable : %3d", (int)VarAddress );
-
-            for ( uint32_t C = 0; C < CaseCount; ++C )
-            {
-				int32_t Condition;
-				FMemory::Memcpy( &Condition, data, sizeof(int32_t));
-				data += sizeof(int32_t);
-				
-				OP::ADDRESS At;
-				FMemory::Memcpy( &At, data, sizeof(OP::ADDRESS));
-				data += sizeof(OP::ADDRESS);
-
-                done += mutable_snprintf( temp+done, 1024-done, "\t[ %3d : %3d]", (int)Condition, (int)At);
-			}
-
-			done += mutable_snprintf( temp+done, 1024-done, "\tdefault : %3d", (int)DefAddress );
-			break;
-        }
-
-		//-----------------------------------------------------------------------------------------
-		case OP_TYPE::IM_RESIZE:
-        {
-            auto args = program.GetOpArgs<OP::ImageResizeArgs>(at);
-            mutable_snprintf( temp, 1024,
-							  "base : %3d\tsize  : %4d %4d",
-                              (int)args.source,
-                              (int)args.size[0],
-                              (int)args.size[1] );
-			break;
-        }
-		case OP_TYPE::IM_RESIZELIKE:
-        {
-            auto args = program.GetOpArgs<OP::ImageResizeLikeArgs>(at);
-            mutable_snprintf( temp, 1024,
-							  "base : %3d\tsize  : %3d",
-                              (int)args.source,
-                              (int)args.sizeSource );
-			break;
-        }
-        case OP_TYPE::IM_PIXELFORMAT:
-        {
-            auto args = program.GetOpArgs<OP::ImagePixelFormatArgs>(at);
-            mutable_snprintf( temp, 1024,
-                              "base : %3d format : %2d  format-if-alpha : %2d",
-                              (int)args.source,
-                              (int)args.format,
-                              (int)args.formatIfAlpha );
-            break;
-        }
-        case OP_TYPE::IM_MIPMAP:
-        {
-            auto args = program.GetOpArgs<OP::ImageMipmapArgs>(at);
-            mutable_snprintf( temp, 1024,
-                              "base : %3d\tlevels : %2d\tblockLevels : %2d",
-                              (int)args.source,
-                              (int)args.levels,
-                              (int)args.blockLevels );
-            break;
-        }
-		case OP_TYPE::IM_BLANKLAYOUT:
-        {
-            auto args = program.GetOpArgs<OP::ImageBlankLayoutArgs>(at);
-            mutable_snprintf( temp, 1024,
-							  "block : %4d %4d  fmt  : %2d",
-                              (int)args.blockSize[0],
-                              (int)args.blockSize[1],
-                              (int)args.format );
-			break;
-        }
-		case OP_TYPE::IM_COMPOSE:
-        {
-            auto args = program.GetOpArgs<OP::ImageComposeArgs>(at);
-            mutable_snprintf( temp, 1024,
-							  "base : %3d\tblock image : %3d\tblock index : %3d\tlayout: %3d\tmask: %3d",
-                              (int)args.base,
-                              (int)args.blockImage,
-                              (int)args.blockIndex,
-                              (int)args.layout,
-                              (int)args.mask );
-			break;
-        }
-		case OP_TYPE::IM_INTERPOLATE:
-        {
-            auto args = program.GetOpArgs<OP::ImageInterpolateArgs>(at);
-            done += mutable_snprintf( temp, 1024,
-                                      "factor : %3d", args.factor);
-			for (int t=0; t<MUTABLE_OP_MAX_INTERPOLATE_COUNT; ++t)
-			{
-				done += mutable_snprintf( temp+done, 1024-done,
-                                          "\tt%d : %3d", t, args.targets[t] );
-			}
-			break;
-        }
-		case OP_TYPE::IM_SWIZZLE:
-        {
-            auto args = program.GetOpArgs<OP::ImageSwizzleArgs>(at);
-            done += mutable_snprintf( temp, 1024,
-                              "format : %3d", args.format);
-			for (int t=0; t<MUTABLE_OP_MAX_SWIZZLE_CHANNELS; ++t)
-			{
-				done += mutable_snprintf( temp+done, 1024-done,
-							  "\tsource %d, channel %3d",
-                          args.sources[t],
-                          args.sourceChannels[t] );
-			}
-			break;
-        }
-		case OP_TYPE::IM_SELECTCOLOUR:
-        {
-            auto args = program.GetOpArgs<OP::ImageSelectColourArgs>(at);
-            mutable_snprintf( temp, 1024,
-							  "base : %3d\tcolour  : %3d",
-                    (int)args.base,
-                    (int)args.colour );
-			break;
-        }
-		case OP_TYPE::IM_COLOURMAP:
-        {
-            auto args = program.GetOpArgs<OP::ImageColourMapArgs>(at);
-            mutable_snprintf( temp, 1024,
-							  "base : %3d\tmask  : %3d\tmap  : %3d",
-                      (int)args.base,
-                      (int)args.mask,
-                      (int)args.map );
-			break;
-        }
-		case OP_TYPE::IM_GRADIENT:
-        {
-            auto args = program.GetOpArgs<OP::ImageGradientArgs>(at);
-            mutable_snprintf( temp, 1024,
-							  "colour0 : %3d\tcolour1  : %3d\tsize  : %3d x %3d",
-                      (int)args.colour0,
-                      (int)args.colour1,
-                      (int)args.size[0],
-                      (int)args.size[1]
-					  );
-			break;
-        }
-		case OP_TYPE::IM_GPU:
-		{
-            auto args = program.GetOpArgs<OP::ImageGPUArgs>(at);
-            mutable_snprintf( temp, 1024,
-							  "program : %3d",
-                              (int)args.program
-							  );
-			break;
-		}
-
-        case OP_TYPE::ME_MERGE:
-        {
-            auto args = program.GetOpArgs<OP::MeshMergeArgs>(at);
-            mutable_snprintf( temp, 1024,
-                              "base : %3d    added : %3d",
-                    (int)args.base,
-                    (int)args.added );
-			break;
-        }
-		case OP_TYPE::ME_INTERPOLATE:
-        {
-            auto args = program.GetOpArgs<OP::MeshInterpolateArgs>(at);
-            done += mutable_snprintf( temp, 1024,
-							  "factor : %3d\tbase : %3d",
-                              args.factor,
-                              args.base);
-			for (int t=0; t<MUTABLE_OP_MAX_INTERPOLATE_COUNT-1; ++t)
-			{
-				done += mutable_snprintf( temp+done, 1024-done,
-                              "\tt%d : %3d", t, args.targets[t] );
-			}
-			break;
-        }
-		case OP_TYPE::ME_OPTIMIZESKINNING:
-		{
-			auto args = program.GetOpArgs<OP::MeshOptimizeSkinningArgs>(at);
-			done += mutable_snprintf(temp, 1024,
-							"source : %d",
-							args.source);
-			break;
-		}
-
-		//-----------------------------------------------------------------------------------------
-		default:
-			//check( false );
-			break;
-		}
-
-		temp[1023] = 0;
-		return temp;
-	}
-
 }
diff --git a/Engine/Plugins/Experimental/Mutable/Source/MutableTools/Private/MuT/ErrorLogPrivate.h b/Engine/Plugins/Experimental/Mutable/Source/MutableTools/Private/MuT/ErrorLogPrivate.h
--- a/Engine/Plugins/Experimental/Mutable/Source/MutableTools/Private/MuT/ErrorLogPrivate.h
+++ b/Engine/Plugins/Experimental/Mutable/Source/MutableTools/Private/MuT/ErrorLogPrivate.h
@@ -64,7 +64,6 @@ namespace mu
 	typedef Ptr<const Model> ModelPtrConst;
 
     extern const TCHAR* GetOpName( OP_TYPE type );
-    extern FString GetOpDesc( const FProgram& program, OP::ADDRESS at );
 
 }
 
diff --git a/Engine/Plugins/Experimental/Mutable/Source/MutableTools/Private/MuT/ImageDescGenerator.cpp b/Engine/Plugins/Experimental/Mutable/Source/MutableTools/Private/MuT/ImageDescGenerator.cpp
--- a/Engine/Plugins/Experimental/Mutable/Source/MutableTools/Private/MuT/ImageDescGenerator.cpp
+++ b/Engine/Plugins/Experimental/Mutable/Source/MutableTools/Private/MuT/ImageDescGenerator.cpp
@@ -99,6 +99,14 @@ namespace mu
     }
 
 
+	//---------------------------------------------------------------------------------------------
+	Ptr<ASTOp> ImageDescGenerator::Visit( const NodeImageReference::Private& node )
+	{
+		m_desc = MUTABLE_MISSING_IMAGE_DESC;
+		return 0;
+	}
+
+
     //---------------------------------------------------------------------------------------------
     Ptr<ASTOp> ImageDescGenerator::Visit( const NodeImageParameter::Private& )
     {
diff --git a/Engine/Plugins/Experimental/Mutable/Source/MutableTools/Private/MuT/ImageDescGenerator.h b/Engine/Plugins/Experimental/Mutable/Source/MutableTools/Private/MuT/ImageDescGenerator.h
--- a/Engine/Plugins/Experimental/Mutable/Source/MutableTools/Private/MuT/ImageDescGenerator.h
+++ b/Engine/Plugins/Experimental/Mutable/Source/MutableTools/Private/MuT/ImageDescGenerator.h
@@ -33,6 +33,7 @@
 #include "MuT/NodeImageTable.h"
 #include "MuT/NodeImageTransform.h"
 #include "MuT/NodeImageVariation.h"
+#include "MuT/NodeImageReference.h"
 #include "MuT/Visitor.h"
 
 namespace mu { class ASTOp; }
@@ -77,7 +78,9 @@ namespace mu
                                public Visitor<NodeImageMultiLayer::Private, Ptr<ASTOp>, true>,
                                public Visitor<NodeImageInvert::Private, Ptr<ASTOp>, true>,
                                public Visitor<NodeImageNormalComposite::Private, Ptr<ASTOp>, true>,
-                               public Visitor<NodeImageTransform::Private, Ptr<ASTOp>, true>
+                               public Visitor<NodeImageTransform::Private, Ptr<ASTOp>, true>, 
+							   public Visitor<NodeImageReference::Private, Ptr<ASTOp>, true>
+
     {
     public:
 
@@ -109,7 +112,8 @@ namespace mu
         Ptr<ASTOp> Visit( const NodeImageMultiLayer::Private& );
         Ptr<ASTOp> Visit( const NodeImageInvert::Private& );
         Ptr<ASTOp> Visit( const NodeImageNormalComposite::Private& );
-        Ptr<ASTOp> Visit( const NodeImageTransform::Private& );
+		Ptr<ASTOp> Visit( const NodeImageTransform::Private& );
+		Ptr<ASTOp> Visit( const NodeImageReference::Private& );
 
     public:
 
diff --git a/Engine/Plugins/Experimental/Mutable/Source/MutableTools/Private/MuT/NodeImage.cpp b/Engine/Plugins/Experimental/Mutable/Source/MutableTools/Private/MuT/NodeImage.cpp
--- a/Engine/Plugins/Experimental/Mutable/Source/MutableTools/Private/MuT/NodeImage.cpp
+++ b/Engine/Plugins/Experimental/Mutable/Source/MutableTools/Private/MuT/NodeImage.cpp
@@ -33,8 +33,7 @@
 #include "MuT/NodeImageTable.h"
 #include "MuT/NodeImageTransform.h"
 #include "MuT/NodeImageVariation.h"
-
-#include <stdint.h>
+#include "MuT/NodeImageReference.h"
 
 
 namespace mu
@@ -49,10 +48,10 @@ namespace mu
 	//---------------------------------------------------------------------------------------------
 	void NodeImage::Serialise( const NodeImage* p, OutputArchive& arch )
 	{
-        uint32_t ver = 0;
+        uint32 ver = 0;
 		arch << ver;
 
-		arch << uint32_t(p->Type);
+		arch << uint32(p->Type);
 		p->SerialiseWrapper(arch);
     }
 
@@ -60,11 +59,11 @@ namespace mu
 	//---------------------------------------------------------------------------------------------
 	NodeImagePtr NodeImage::StaticUnserialise( InputArchive& arch )
 	{
-        uint32_t ver;
+        uint32 ver;
 		arch >> ver;
 		check( ver == 0 );
 
-        uint32_t id;
+        uint32 id;
 		arch >> id;
 
 		switch (id)
@@ -98,8 +97,9 @@ namespace mu
         case 31 :  return NodeImageInvert::StaticUnserialise( arch ); break;
         case 32 :  return NodeImageVariation::StaticUnserialise( arch ); break;
         case 33 :  return NodeImageNormalComposite::StaticUnserialise( arch ); break;
-        case 34 :  return NodeImageTransform::StaticUnserialise( arch ); break;
-        default : check(false);
+		case 34:  return NodeImageTransform::StaticUnserialise(arch); break;
+		case 35:  return NodeImageReference::StaticUnserialise(arch); break;
+		default : check(false);
 		}
 
 		return 0;
diff --git a/Engine/Plugins/Experimental/Mutable/Source/MutableTools/Private/MuT/NodeImageConstantPrivate.h b/Engine/Plugins/Experimental/Mutable/Source/MutableTools/Private/MuT/NodeImageConstantPrivate.h
--- a/Engine/Plugins/Experimental/Mutable/Source/MutableTools/Private/MuT/NodeImageConstantPrivate.h
+++ b/Engine/Plugins/Experimental/Mutable/Source/MutableTools/Private/MuT/NodeImageConstantPrivate.h
@@ -10,7 +10,6 @@
 
 namespace mu
 {
-
 	class NodeImageConstant::Private : public NodeImage::Private
 	{
 	public:
@@ -26,34 +25,34 @@ namespace mu
 		//!
 		void Serialise( OutputArchive& arch ) const
 		{
-            uint32_t ver = 0;
-			arch << ver;
+			uint32 Ver = 0;
+			arch << Ver;
 
-            Ptr<const Image> image;
-            if (m_pProxy)
-            {
-                image = m_pProxy->Get();
-            }
+			Ptr<const Image> ActualImage;
+			if (m_pProxy)
+			{
+				ActualImage = m_pProxy->Get();
+			}
 
-            arch << image;
+			arch << ActualImage;
 		}
 
 		//!
 		void Unserialise( InputArchive& arch )
 		{
-            uint32_t ver;
-			arch >> ver;
-			check(ver==0);
-
-            // Are we using proxies?
-            m_pProxy = arch.NewImageProxy();
-            if (!m_pProxy)
-            {
-                // Normal serialisation
-                ImagePtr image;
-                arch >> image;
-                m_pProxy = new ResourceProxyMemory<Image>( image.get() );
-            }
+            uint32 Ver;
+			arch >> Ver;
+			check(Ver==0);
+
+
+			m_pProxy = arch.NewImageProxy();
+			if (!m_pProxy)
+			{
+				// Normal serialisation
+				Ptr<Image> ActualImage;
+				arch >> ActualImage;
+				m_pProxy = new ResourceProxyMemory<Image>(ActualImage.get());
+			}
 		}
 	};
 
diff --git a/Engine/Plugins/Experimental/Mutable/Source/MutableTools/Private/MuT/NodeImageReference.cpp b/Engine/Plugins/Experimental/Mutable/Source/MutableTools/Private/MuT/NodeImageReference.cpp
new file mode 100644
--- /dev/null
+++ b/Engine/Plugins/Experimental/Mutable/Source/MutableTools/Private/MuT/NodeImageReference.cpp
@@ -0,0 +1,55 @@
+// Copyright Epic Games, Inc. All Rights Reserved.
+
+
+#include "MuT/NodeImageReference.h"
+
+#include "Misc/AssertionMacros.h"
+#include "MuR/Image.h"
+#include "MuR/Serialisation.h"
+#include "MuT/NodeImageReferencePrivate.h"
+#include "MuT/NodePrivate.h"
+
+
+namespace mu
+{
+
+	//---------------------------------------------------------------------------------------------
+	// Static initialisation
+	//---------------------------------------------------------------------------------------------
+	NODE_TYPE NodeImageReference::Private::s_type = NODE_TYPE( "ImageReference", NodeImage::GetStaticType() );
+
+
+	//---------------------------------------------------------------------------------------------
+	MUTABLE_IMPLEMENT_NODE( NodeImageReference, EType::Reference, Node, Node::EType::Image);
+
+
+	//---------------------------------------------------------------------------------------------
+	int NodeImageReference::GetInputCount() const
+	{
+		return 0;
+	}
+
+
+	//---------------------------------------------------------------------------------------------
+    Node* NodeImageReference::GetInputNode( int ) const
+	{
+		check( false );
+		return 0;
+	}
+
+
+	//---------------------------------------------------------------------------------------------
+    void NodeImageReference::SetInputNode( int, NodePtr )
+	{
+		check( false );
+	}
+
+
+	//---------------------------------------------------------------------------------------------
+	void NodeImageReference::SetImageReference( uint32 ID)
+	{
+		m_pD->ImageReferenceID = ID;
+	}
+
+
+}
diff --git a/Engine/Plugins/Experimental/Mutable/Source/MutableTools/Private/MuT/NodeImageReferencePrivate.h b/Engine/Plugins/Experimental/Mutable/Source/MutableTools/Private/MuT/NodeImageReferencePrivate.h
new file mode 100644
--- /dev/null
+++ b/Engine/Plugins/Experimental/Mutable/Source/MutableTools/Private/MuT/NodeImageReferencePrivate.h
@@ -0,0 +1,45 @@
+// Copyright Epic Games, Inc. All Rights Reserved.
+
+#pragma once
+
+#include "MuT/NodeImagePrivate.h"
+#include "MuT/NodeImageReference.h"
+#include "MuT/StreamsPrivate.h"
+
+
+namespace mu
+{
+
+	class NodeImageReference::Private : public NodeImage::Private
+	{
+	public:
+
+		MUTABLE_DEFINE_CONST_VISITABLE();
+
+	public:
+
+		static NODE_TYPE s_type;
+
+		uint32 ImageReferenceID = 0;
+
+		//!
+		void Serialise(OutputArchive& arch) const
+		{
+			uint32 Ver = 0;
+			arch << Ver;
+
+			arch << ImageReferenceID;
+		}
+
+		//!
+		void Unserialise(InputArchive& arch)
+		{
+			uint32 Ver;
+			arch >> Ver;
+			check(Ver == 0);
+
+			arch >> ImageReferenceID;
+		}
+	};
+
+}
diff --git a/Engine/Plugins/Experimental/Mutable/Source/MutableTools/Private/MuT/NodeSurfaceNewPrivate.h b/Engine/Plugins/Experimental/Mutable/Source/MutableTools/Private/MuT/NodeSurfaceNewPrivate.h
--- a/Engine/Plugins/Experimental/Mutable/Source/MutableTools/Private/MuT/NodeSurfaceNewPrivate.h
+++ b/Engine/Plugins/Experimental/Mutable/Source/MutableTools/Private/MuT/NodeSurfaceNewPrivate.h
@@ -62,7 +62,7 @@ namespace mu
 			NodeImagePtr m_pImage;
 
 			// It could be negative, to indicate no layout.
-            int8_t m_layoutIndex = 0;
+            int8 m_layoutIndex = 0;
 
             //!
 			void Serialise(OutputArchive& arch) const
diff --git a/Engine/Plugins/Experimental/Mutable/Source/MutableTools/Public/MuT/NodeImage.h b/Engine/Plugins/Experimental/Mutable/Source/MutableTools/Public/MuT/NodeImage.h
--- a/Engine/Plugins/Experimental/Mutable/Source/MutableTools/Public/MuT/NodeImage.h
+++ b/Engine/Plugins/Experimental/Mutable/Source/MutableTools/Public/MuT/NodeImage.h
@@ -53,6 +53,7 @@ namespace mu
 			Variation = 32,
 			NormalComposite = 33,
 			Transform = 34,
+			Reference = 35,
 			None
 		};
 
diff --git a/Engine/Plugins/Experimental/Mutable/Source/MutableTools/Public/MuT/NodeImageConstant.h b/Engine/Plugins/Experimental/Mutable/Source/MutableTools/Public/MuT/NodeImageConstant.h
--- a/Engine/Plugins/Experimental/Mutable/Source/MutableTools/Public/MuT/NodeImageConstant.h
+++ b/Engine/Plugins/Experimental/Mutable/Source/MutableTools/Public/MuT/NodeImageConstant.h
@@ -21,8 +21,9 @@ namespace mu
 	class OutputArchive;
 
 
-	//! Node that outputs a constant image.
-	//! \ingroup model
+	/** Node that outputs a constant image.
+	* This node also supports "image references".
+	*/
 	class MUTABLETOOLS_API NodeImageConstant : public NodeImage
 	{
 	public:
@@ -42,8 +43,6 @@ namespace mu
 		// Node Interface
 		//-----------------------------------------------------------------------------------------
 
-        
-
         const NODE_TYPE* GetType() const override;
 		static const NODE_TYPE* GetStaticType();
 
diff --git a/Engine/Plugins/Experimental/Mutable/Source/MutableTools/Public/MuT/NodeImageReference.h b/Engine/Plugins/Experimental/Mutable/Source/MutableTools/Public/MuT/NodeImageReference.h
new file mode 100644
--- /dev/null
+++ b/Engine/Plugins/Experimental/Mutable/Source/MutableTools/Public/MuT/NodeImageReference.h
@@ -0,0 +1,72 @@
+// Copyright Epic Games, Inc. All Rights Reserved.
+
+#pragma once
+
+#include "MuR/Image.h"
+#include "MuR/Ptr.h"
+#include "MuR/RefCounted.h"
+#include "MuT/Node.h"
+#include "MuT/NodeImage.h"
+
+
+namespace mu
+{
+
+	// Forward definitions
+	class InputArchive;
+	class OutputArchive;
+
+
+	/** Node that outputs an "image reference".
+	*/
+	class MUTABLETOOLS_API NodeImageReference : public NodeImage
+	{
+	public:
+
+		//-----------------------------------------------------------------------------------------
+		// Life cycle
+		//-----------------------------------------------------------------------------------------
+
+		NodeImageReference();
+
+		void SerialiseWrapper(OutputArchive& arch) const override;
+		static void Serialise(const NodeImageReference* pNode, OutputArchive& arch);
+		static Ptr<NodeImageReference> StaticUnserialise(InputArchive& arch);
+
+		//-----------------------------------------------------------------------------------------
+		// Node Interface
+		//-----------------------------------------------------------------------------------------
+
+		const NODE_TYPE* GetType() const override;
+		static const NODE_TYPE* GetStaticType();
+
+		virtual int GetInputCount() const override;
+		virtual Node* GetInputNode(int i) const override;
+		void SetInputNode(int i, Ptr<Node> pNode) override;
+
+		//-----------------------------------------------------------------------------------------
+		// Own Interface
+		//-----------------------------------------------------------------------------------------
+
+		/** Set this node value to be an "image reference" (to point to an unreal engine image). */
+		void SetImageReference(uint32 ID);
+
+		//-----------------------------------------------------------------------------------------
+		// Interface pattern
+		//-----------------------------------------------------------------------------------------
+		class Private;
+		Private* GetPrivate() const;
+		Node::Private* GetBasePrivate() const override;
+
+	protected:
+
+		//! Forbidden. Manage with the Ptr<> template.
+		~NodeImageReference();
+
+	private:
+
+		Private* m_pD;
+
+	};
+
+}
