[mutable] fix compilation errors related to rtti in some projects.
#rnx

diff --git a/Engine/Plugins/Experimental/Mutable/Source/CustomizableObjectEditor/Private/MuCOE/GenerateMutableSource/GenerateMutableSourceImage.cpp b/Engine/Plugins/Experimental/Mutable/Source/CustomizableObjectEditor/Private/MuCOE/GenerateMutableSource/GenerateMutableSourceImage.cpp
--- a/Engine/Plugins/Experimental/Mutable/Source/CustomizableObjectEditor/Private/MuCOE/GenerateMutableSource/GenerateMutableSourceImage.cpp
+++ b/Engine/Plugins/Experimental/Mutable/Source/CustomizableObjectEditor/Private/MuCOE/GenerateMutableSource/GenerateMutableSourceImage.cpp
@@ -138,7 +138,7 @@ mu::NodeImagePtr GenerateMutableSourceImage(const UEdGraphPin* Pin, FMutableGrap
 	const FGeneratedKey Key(reinterpret_cast<void*>(&GenerateMutableSourceImage), *Pin, *Node, GenerationContext, true);
 	if (const FGeneratedData* Generated = GenerationContext.Generated.Find(Key))
 	{
-		return dynamic_cast<mu::NodeImage*>(Generated->Node.get());
+		return static_cast<mu::NodeImage*>(Generated->Node.get());
 	}
 
 	mu::NodeImagePtr Result;
@@ -367,7 +367,7 @@ mu::NodeImagePtr GenerateMutableSourceImage(const UEdGraphPin* Pin, FMutableGrap
 
 				const int32 NumSwitchOptions = TypedNodeTextureSwitch->GetNumElements();
 
-				mu::NodeScalarEnumParameter* EnumParameter = dynamic_cast<mu::NodeScalarEnumParameter*>(SwitchParam.get());
+				mu::NodeScalarEnumParameter* EnumParameter = static_cast<mu::NodeScalarEnumParameter*>(SwitchParam.get());
 				if (NumSwitchOptions != EnumParameter->GetValueCount())
 				{
 					const FText Message = LOCTEXT("MismatchedSwitch", "Switch enum and switch node have different number of options. Please refresh the switch node to make sure the outcomes are labeled properly.");
@@ -867,7 +867,7 @@ mu::NodeImagePtr GenerateMutableSourceImage(const UEdGraphPin* Pin, FMutableGrap
 
 				const int32 NumSwitchOptions = TypedNodePassThroughTextureSwitch->GetNumElements();
 
-				mu::NodeScalarEnumParameter* EnumParameter = dynamic_cast<mu::NodeScalarEnumParameter*>(SwitchParam.get());
+				mu::NodeScalarEnumParameter* EnumParameter = static_cast<mu::NodeScalarEnumParameter*>(SwitchParam.get());
 				if (NumSwitchOptions != EnumParameter->GetValueCount())
 				{
 					const FText Message = LOCTEXT("MismatchedSwitch", "Switch enum and switch node have different number of options. Please refresh the switch node to make sure the outcomes are labeled properly.");
