[mutable] Added support for Passthrough textures to Table nodes.

#rnx

diff --git a/Engine/Plugins/Experimental/Mutable/Source/CustomizableObjectEditor/Classes/MuCOE/Nodes/CustomizableObjectNodeTable.h b/Engine/Plugins/Experimental/Mutable/Source/CustomizableObjectEditor/Classes/MuCOE/Nodes/CustomizableObjectNodeTable.h
--- a/Engine/Plugins/Experimental/Mutable/Source/CustomizableObjectEditor/Classes/MuCOE/Nodes/CustomizableObjectNodeTable.h
+++ b/Engine/Plugins/Experimental/Mutable/Source/CustomizableObjectEditor/Classes/MuCOE/Nodes/CustomizableObjectNodeTable.h
@@ -291,6 +291,9 @@ public:
 	// Returns true if the pin is in the default mode (same as node)
 	bool IsImagePinDefault(UEdGraphPin* Pin);
 
+	// Returns the image mode of the column
+	ETableTextureType GetColumnImageMode(FString ColumnName) const;
+
 private:
 
 	/** Number of properties to know when the node needs an update */
diff --git a/Engine/Plugins/Experimental/Mutable/Source/CustomizableObjectEditor/Private/MuCOE/GenerateMutableSource/GenerateMutableSourceTable.cpp b/Engine/Plugins/Experimental/Mutable/Source/CustomizableObjectEditor/Private/MuCOE/GenerateMutableSource/GenerateMutableSourceTable.cpp
--- a/Engine/Plugins/Experimental/Mutable/Source/CustomizableObjectEditor/Private/MuCOE/GenerateMutableSource/GenerateMutableSourceTable.cpp
+++ b/Engine/Plugins/Experimental/Mutable/Source/CustomizableObjectEditor/Private/MuCOE/GenerateMutableSource/GenerateMutableSourceTable.cpp
@@ -334,7 +334,27 @@ void FillTableColumn(const UCustomizableObjectNodeTable* TableNode,	mu::TablePtr
 				CurrentColumn = MutableTable->AddColumn(StringCast<ANSICHAR>(*ColumnName).Get(), mu::TABLE_COLUMN_TYPE::TCT_IMAGE);
 			}
 
-			GenerationContext.ArrayTextureUnrealToMutableTask.Add(FTextureUnrealToMutableTask(MutableTable, Texture, TableNode, CurrentColumn, RowIdx));
+			if (TableNode->GetColumnImageMode(ColumnName) == ETableTextureType::PASSTHROUGH_TEXTURE)
+			{
+				uint32* FoundIndex = GenerationContext.PassThroughTextureToIndexMap.Find(Texture);
+				uint32 ImageReferenceID;
+
+				if (!FoundIndex)
+				{
+					ImageReferenceID = GenerationContext.PassThroughTextureToIndexMap.Num();
+					GenerationContext.PassThroughTextureToIndexMap.Add(Texture, ImageReferenceID);
+				}
+				else
+				{
+					ImageReferenceID = *FoundIndex;
+				}
+
+				MutableTable->SetCell(CurrentColumn, RowIdx, mu::Image::CreateAsReference(ImageReferenceID).get());
+			}
+			else
+			{
+				GenerationContext.ArrayTextureUnrealToMutableTask.Add(FTextureUnrealToMutableTask(MutableTable, Texture, TableNode, CurrentColumn, RowIdx));
+			}
 		}
 
 		else if (UMaterialInstance* Material = Cast<UMaterialInstance>(Object))
diff --git a/Engine/Plugins/Experimental/Mutable/Source/CustomizableObjectEditor/Private/MuCOE/Nodes/CustomizableObjectNodeTable.cpp b/Engine/Plugins/Experimental/Mutable/Source/CustomizableObjectEditor/Private/MuCOE/Nodes/CustomizableObjectNodeTable.cpp
--- a/Engine/Plugins/Experimental/Mutable/Source/CustomizableObjectEditor/Private/MuCOE/Nodes/CustomizableObjectNodeTable.cpp
+++ b/Engine/Plugins/Experimental/Mutable/Source/CustomizableObjectEditor/Private/MuCOE/Nodes/CustomizableObjectNodeTable.cpp
@@ -1173,4 +1173,21 @@ bool UCustomizableObjectNodeTable::IsImagePinDefault(UEdGraphPin* Pin)
 	return true;
 }
 
+
+ETableTextureType UCustomizableObjectNodeTable::GetColumnImageMode(FString ColumnName) const
+{
+	for (UEdGraphPin* Pin : Pins)
+	{
+		if (UCustomizableObjectNodeTableImagePinData* PinData = Cast<UCustomizableObjectNodeTableImagePinData>(GetPinData(*(Pin))))
+		{
+			if (PinData->ColumnName == ColumnName)
+			{
+				return PinData->ImageMode;
+			}
+		}
+	}
+
+	return DefaultImageMode;
+}
+
 #undef LOCTEXT_NAMESPACE
diff --git a/Engine/Plugins/Experimental/Mutable/Source/MutableTools/Private/MuT/CodeGenerator_Image.cpp b/Engine/Plugins/Experimental/Mutable/Source/MutableTools/Private/MuT/CodeGenerator_Image.cpp
--- a/Engine/Plugins/Experimental/Mutable/Source/MutableTools/Private/MuT/CodeGenerator_Image.cpp
+++ b/Engine/Plugins/Experimental/Mutable/Source/MutableTools/Private/MuT/CodeGenerator_Image.cpp
@@ -259,6 +259,7 @@ namespace mu
 		result.op = op;
     }
 
+
 	//---------------------------------------------------------------------------------------------
 	void CodeGenerator::GenerateImage_Reference(FImageGenerationResult& result, const NodeImageReference* InNode)
 	{
@@ -1856,10 +1857,11 @@ namespace mu
 		result.op = GenerateTableSwitch<NodeImageTable::Private, TCT_IMAGE, OP_TYPE::IM_SWITCH>(node,
 			[this](const NodeImageTable::Private& node, int colIndex, int row, ErrorLog* pErrorLog)
 			{
-				NodeImageConstantPtr pCell = new NodeImageConstant();
+				TABLE_VALUE CellData = node.m_pTable->GetPrivate()->m_rows[row].m_values[colIndex];
 				ImagePtrConst pImage = nullptr;
-				
-				if (Ptr<ResourceProxy<Image>> pProxyImage = node.m_pTable->GetPrivate()->m_rows[row].m_values[colIndex].m_pProxyImage)
+				NodeImagePtr CellImage = nullptr;
+
+				if (Ptr<ResourceProxy<Image>> pProxyImage = CellData.m_pProxyImage)
 				{
 					pImage = pProxyImage->Get();
 				}
@@ -1871,9 +1873,25 @@ namespace mu
 						"Table has a missing image in column %d, row %d.", colIndex, row);
 					pErrorLog->GetPrivate()->Add(temp, ELMT_ERROR, node.m_errorContext);
 				}
+				else
+				{
+					if (pImage->IsReference())
+					{
+						Ptr<NodeImageReference> ImageRef = new NodeImageReference();
+						ImageRef->SetImageReference(pImage->GetReferencedTexture());
+
+						CellImage = ImageRef;
+					}
+					else
+					{
+						NodeImageConstantPtr ImageConst = new NodeImageConstant();
+						ImageConst->SetValue(pImage);
+
+						CellImage = ImageConst;
+					}
+				}
 
-				pCell->SetValue(pImage);
-				return Generate(pCell);
+				return Generate(CellImage);
 			});
 	}
 
diff --git a/Engine/Plugins/Experimental/Mutable/Source/MutableTools/Private/MuT/ImageDescGenerator.cpp b/Engine/Plugins/Experimental/Mutable/Source/MutableTools/Private/MuT/ImageDescGenerator.cpp
--- a/Engine/Plugins/Experimental/Mutable/Source/MutableTools/Private/MuT/ImageDescGenerator.cpp
+++ b/Engine/Plugins/Experimental/Mutable/Source/MutableTools/Private/MuT/ImageDescGenerator.cpp
@@ -135,6 +135,12 @@ namespace mu
         std::size_t i = 0;
         while ( !pImage && i<node.m_pTable->GetPrivate()->m_rows.Num() )
         {
+			if (node.m_pTable->GetPrivate()->m_rows[i].m_values[colIndex].m_pProxyImage->Get()->IsReference())
+			{
+				// Image References do not need an image desc
+				break;
+			}
+
             pImage = node.m_pTable->GetPrivate()->m_rows[i].m_values[ colIndex ].m_pProxyImage->Get();
             ++i;
         }
