#rnx

diff --git a/Engine/Plugins/Experimental/Mutable/Source/CustomizableObjectEditor/Private/MuCOE/AssetDefinition_CustomizableObjectInstance.cpp b/Engine/Plugins/Experimental/Mutable/Source/CustomizableObjectEditor/Private/MuCOE/AssetDefinition_CustomizableObjectInstance.cpp
--- a/Engine/Plugins/Experimental/Mutable/Source/CustomizableObjectEditor/Private/MuCOE/AssetDefinition_CustomizableObjectInstance.cpp
+++ b/Engine/Plugins/Experimental/Mutable/Source/CustomizableObjectEditor/Private/MuCOE/AssetDefinition_CustomizableObjectInstance.cpp
@@ -26,7 +26,13 @@ EAssetCommandResult UAssetDefinition_CustomizableObjectInstance::OpenAssets(cons
 
 TConstArrayView<FAssetCategoryPath> UAssetDefinition_CustomizableObjectInstance::GetAssetCategories() const
 {
-	return TConstArrayView<FAssetCategoryPath>();
+	static const std::initializer_list<FAssetCategoryPath> Categories =
+	{
+		// Asset can be found inside the Mutable submenu 
+		NSLOCTEXT("AssetTypeActions", "Mutable", "Mutable")
+	};
+
+	return Categories;
 }
 
 EAssetCommandResult UAssetDefinition_CustomizableObjectInstance::ActivateAssets(
