[Mutable] Add Projection lodbias and mode cvars. Reorder memory allocation to minimize IM_RASTERMESH operation memory peak.

- Added mutable.ForceProjectionMode, -1 uses projector settings, 0 point sampling + none min filter method, 1 bilinear sampling  + TotalAreaHeuristic min filter method.
- Added mutable.GlobalProjectionLodBias, adds the float value specified to min filter lod selection. Does not apply if the method is none.
- Deprecate mutable.ForceBestQualityProjectionSampling cvar
- Fix ME_MERGE operation cpu scope name.

#rnx

diff --git a/Engine/Plugins/Experimental/Mutable/Source/MutableRuntime/Private/MuR/CodeRunner.cpp b/Engine/Plugins/Experimental/Mutable/Source/MutableRuntime/Private/MuR/CodeRunner.cpp
--- a/Engine/Plugins/Experimental/Mutable/Source/MutableRuntime/Private/MuR/CodeRunner.cpp
+++ b/Engine/Plugins/Experimental/Mutable/Source/MutableRuntime/Private/MuR/CodeRunner.cpp
@@ -74,13 +74,29 @@
 
 namespace
 {
-	
+
+// This cvar will be deprecated, use CVarForceProjectionSamplingMode with value 1 to achieve the same result.
 bool bForceBestQualityProjectionSampling = false;
 static FAutoConsoleVariableRef CVarForceBestQualityProjectionSampling (
 	TEXT("mutable.ForceBestQualityProjectionSampling"),
 	bForceBestQualityProjectionSampling,
-	TEXT("force mutable to use the projection algorithm with best projection quality, BiLinear with TotalAreaHeuristic"),
+	TEXT("DEPRECATED, please use mutable.ForceProjectionMode with value 1 to achieve the same result. Force mutable to use projection mode BiLinear + TotalAreaHeuristic."),
+	ECVF_Default);
+
+int32 ForcedProjectionMode = -1;
+static FAutoConsoleVariableRef CVarForceProjectionSamplingMode (
+	TEXT("mutable.ForceProjectionMode"),
+	ForcedProjectionMode,
+	TEXT("force mutable to use an specific projection mode, 0 = Point + None, 1 = Bilinear + TotalAreaHeuristic, -1 uses the values provided by the projector."),
 	ECVF_Default);
+
+float GlobalProjectionLodBias = 0.0f;
+static FAutoConsoleVariableRef CVarGlobalProjectionLodBias (
+	TEXT("mutable.GlobalProjectionLodBias"),
+	GlobalProjectionLodBias,
+	TEXT("Lod bias applied to the lod resulting form the best mip computation, only used if a min filter method different than None is used."),
+	ECVF_Default);
+
 }
 
 
@@ -968,7 +984,7 @@ namespace mu
     }
 
     //---------------------------------------------------------------------------------------------
-    void CodeRunner::RunCode_Mesh(const FScheduledOp& item, const Model* pModel )
+	void CodeRunner::RunCode_Mesh(const FScheduledOp& item, const Model* pModel )
     {
 		MUTABLE_CPUPROFILER_SCOPE(RunCode_Mesh);
 
@@ -1276,7 +1292,7 @@ namespace mu
 
             case 1:
             {
-				MUTABLE_CPUPROFILER_SCOPE(ME_MORPH2_1)
+				MUTABLE_CPUPROFILER_SCOPE(ME_MERGE_1)
 
                 MeshPtrConst pA = GetMemory().GetMesh( FCacheAddress(args.base,item) );
                 MeshPtrConst pB = GetMemory().GetMesh( FCacheAddress(args.added,item) );
@@ -1301,7 +1317,7 @@ namespace mu
                 }
                 else if (pA && pA->GetVertexCount())
                 {
-                    pResult = pA->Clone();
+					pResult = pA->Clone();
                 }
                 else if (pB && pB->GetVertexCount())
                 {
@@ -4065,9 +4081,23 @@ namespace mu
 				Ptr<const Projector> pProjector = GetMemory().GetProjector(FScheduledOp::FromOpAndOptions(args.projector, item, 0));
 				if (pProjector)
 				{
-					EMinFilterMethod MinFilterMethod = bForceBestQualityProjectionSampling
-						? EMinFilterMethod::TotalAreaHeuristic
-						: static_cast<EMinFilterMethod>(args.MinFilterMethod);
+					EMinFilterMethod MinFilterMethod = Invoke([&]() -> EMinFilterMethod
+					{
+						if (bForceBestQualityProjectionSampling)
+						{
+							return EMinFilterMethod::TotalAreaHeuristic;
+						}
+						else if (ForcedProjectionMode == 0)
+						{
+							return EMinFilterMethod::None;
+						}
+						else if (ForcedProjectionMode == 1)
+						{
+							return EMinFilterMethod::TotalAreaHeuristic;
+						}
+						
+						return static_cast<EMinFilterMethod>(args.MinFilterMethod);
+					});
 
 					if (MinFilterMethod == EMinFilterMethod::TotalAreaHeuristic)
 					{
@@ -4078,9 +4108,9 @@ namespace mu
 						
 						if (pMesh)
 						{ 
-							Data.RasterMesh.MipValue =  
-								ComputeProjectedFootprintBestMip(pMesh.get(), pProjector->m_value, TargetImageSizeF, SourceImageSizeF);
+							const float ComputedMip = ComputeProjectedFootprintBestMip(pMesh.get(), pProjector->m_value, TargetImageSizeF, SourceImageSizeF);
 
+							Data.RasterMesh.MipValue = FMath::Max(0.0f, ComputedMip + GlobalProjectionLodBias);
 							Data.RasterMesh.Mip = static_cast<uint8>(FMath::FloorToInt32(Data.RasterMesh.MipValue));
 						}
 					}
@@ -4180,7 +4210,6 @@ namespace mu
 				const float FadeEndRad = FMath::DegreesToRadians(fadeEnd);
 
 				EImageFormat Format = pSource ? GetUncompressedFormat(pSource->GetFormat()) : EImageFormat::IF_L_UBYTE;
-                Ptr<Image> pNew = new Image( SizeX, SizeY, 1, Format );
 
 				if (pSource && pSource->GetFormat()!=Format)
 				{
@@ -4197,9 +4226,23 @@ namespace mu
 				// moving the right layout channel to the 0.
 				int layout = 0;
 
-				ESamplingMethod SamplingMethod = bForceBestQualityProjectionSampling
-						? ESamplingMethod::BiLinear
-						: static_cast<ESamplingMethod>(args.SamplingMethod);
+				ESamplingMethod SamplingMethod = Invoke([&]() -> ESamplingMethod
+				{
+					if (bForceBestQualityProjectionSampling)
+					{
+						return ESamplingMethod::BiLinear;
+					}
+					else if (ForcedProjectionMode == 0)
+					{
+						return ESamplingMethod::Point;
+					}
+					else if (ForcedProjectionMode == 1)
+					{
+						return ESamplingMethod::BiLinear;
+					}
+					
+					return static_cast<ESamplingMethod>(args.SamplingMethod);
+				});
 
 				if (SamplingMethod == ESamplingMethod::BiLinear)
 				{
@@ -4218,6 +4261,9 @@ namespace mu
 					}
 				}
 
+				// Allocate new image after bilinear mip generation to reduce operation memory peak.
+                Ptr<Image> pNew = new Image(SizeX, SizeY, 1, Format);
+
 				if (args.projector && pSource && pSource->GetSizeX() > 0 && pSource->GetSizeY() > 0)
 				{
 					Ptr<const Projector> pProjector = GetMemory().GetProjector(FScheduledOp::FromOpAndOptions(args.projector, item, 0));
diff --git a/Engine/Plugins/Experimental/Mutable/Source/MutableRuntime/Private/MuR/OpImageProject.cpp b/Engine/Plugins/Experimental/Mutable/Source/MutableRuntime/Private/MuR/OpImageProject.cpp
--- a/Engine/Plugins/Experimental/Mutable/Source/MutableRuntime/Private/MuR/OpImageProject.cpp
+++ b/Engine/Plugins/Experimental/Mutable/Source/MutableRuntime/Private/MuR/OpImageProject.cpp
@@ -1121,11 +1121,9 @@ float ComputeProjectedFootprintBestMip(
 	float BestMip = FMath::Log2(FMath::Max(1.0f, SourceSize.X * SourceSize.Y)) * 0.5f - 
 				    FMath::Log2(FMath::Max(1.0f, SourceFootprintAreaOnTargetInPixelsSquared)) * 0.5f;
 
-	// TODO: We probably should allow negative mips to indicate that the image needs magnification.
-	return FMath::Max(0.0f, BestMip);
+	return BestMip;
 }
 
-
 //-------------------------------------------------------------------------------------------------
 void ImageRasterProjectedCylindrical( const Mesh* pMesh, Image* pImage,
 	const Image* pSource, const Image* pMask,
