[Mutable] Remap Pins minor improvements.

- Removed node from context. Now is passed through function calls.

#rnx

diff --git a/Engine/Plugins/Experimental/Mutable/Source/CustomizableObjectEditor/Classes/MuCOE/Nodes/CustomizableObjectNodeExternalPin.h b/Engine/Plugins/Experimental/Mutable/Source/CustomizableObjectEditor/Classes/MuCOE/Nodes/CustomizableObjectNodeExternalPin.h
--- a/Engine/Plugins/Experimental/Mutable/Source/CustomizableObjectEditor/Classes/MuCOE/Nodes/CustomizableObjectNodeExternalPin.h
+++ b/Engine/Plugins/Experimental/Mutable/Source/CustomizableObjectEditor/Classes/MuCOE/Nodes/CustomizableObjectNodeExternalPin.h
@@ -28,9 +28,7 @@ class UCustomizableObjectNodeRemapPinsCustomExternalPin : public UCustomizableOb
 	GENERATED_BODY()
 
 public:
-	class UCustomizableObjectNodeExternalPin* Node = nullptr;
-
-	virtual void RemapPins(const TArray<UEdGraphPin*>& OldPins, const TArray<UEdGraphPin*>& NewPins, TMap<UEdGraphPin*, UEdGraphPin*>& PinsToRemap, TArray<UEdGraphPin*>& PinsToOrphan) override;
+	virtual void RemapPins(const UCustomizableObjectNode& Node, const TArray<UEdGraphPin*>& OldPins, const TArray<UEdGraphPin*>& NewPins, TMap<UEdGraphPin*, UEdGraphPin*>& PinsToRemap, TArray<UEdGraphPin*>& PinsToOrphan) override;
 };
 
 UCLASS()
diff --git a/Engine/Plugins/Experimental/Mutable/Source/CustomizableObjectEditor/Classes/MuCOE/Nodes/CustomizableObjectNodeMaterial.h b/Engine/Plugins/Experimental/Mutable/Source/CustomizableObjectEditor/Classes/MuCOE/Nodes/CustomizableObjectNodeMaterial.h
--- a/Engine/Plugins/Experimental/Mutable/Source/CustomizableObjectEditor/Classes/MuCOE/Nodes/CustomizableObjectNodeMaterial.h
+++ b/Engine/Plugins/Experimental/Mutable/Source/CustomizableObjectEditor/Classes/MuCOE/Nodes/CustomizableObjectNodeMaterial.h
@@ -38,13 +38,11 @@ class UCustomizableObjectNodeMaterialRemapPinsByName : public UCustomizableObjec
 {
 	GENERATED_BODY()
 public:
-	class UCustomizableObjectNodeMaterial* Node = nullptr;
+	virtual bool Equal(const UCustomizableObjectNode& Node, const UEdGraphPin& OldPin, const UEdGraphPin& NewPin) const override;
 
-	virtual bool Equal(const UEdGraphPin& OldPin, const UEdGraphPin& NewPin) const override;
+	virtual void RemapPins(const UCustomizableObjectNode& Node, const TArray<UEdGraphPin*>& OldPins, const TArray<UEdGraphPin*>& NewPins, TMap<UEdGraphPin*, UEdGraphPin*>& PinsToRemap, TArray<UEdGraphPin*>& PinsToOrphan) override;
 
-	virtual void RemapPins(const TArray<UEdGraphPin*>& OldPins, const TArray<UEdGraphPin*>& NewPins, TMap<UEdGraphPin*, UEdGraphPin*>& PinsToRemap, TArray<UEdGraphPin*>& PinsToOrphan) override;
-
-	bool HasSavedPinData(const UEdGraphPin &Pin) const;
+	bool HasSavedPinData(const UCustomizableObjectNode& Node, const UEdGraphPin &Pin) const;
 };
 
 
diff --git a/Engine/Plugins/Experimental/Mutable/Source/CustomizableObjectEditor/Classes/MuCOE/Nodes/CustomizableObjectNodeTable.h b/Engine/Plugins/Experimental/Mutable/Source/CustomizableObjectEditor/Classes/MuCOE/Nodes/CustomizableObjectNodeTable.h
--- a/Engine/Plugins/Experimental/Mutable/Source/CustomizableObjectEditor/Classes/MuCOE/Nodes/CustomizableObjectNodeTable.h
+++ b/Engine/Plugins/Experimental/Mutable/Source/CustomizableObjectEditor/Classes/MuCOE/Nodes/CustomizableObjectNodeTable.h
@@ -111,10 +111,10 @@ class UCustomizableObjectNodeTableRemapPins : public UCustomizableObjectNodeRema
 public:
 
 	// Specific method to decide when two pins are equal
-	virtual bool Equal(const UEdGraphPin& OldPin, const UEdGraphPin& NewPin) const override;
+	virtual bool Equal(const UCustomizableObjectNode& Node, const UEdGraphPin& OldPin, const UEdGraphPin& NewPin) const override;
 
 	// Method to use in the RemapPins step of the node reconstruction process
-	virtual void RemapPins(const TArray<UEdGraphPin*>& OldPins, const TArray<UEdGraphPin*>& NewPins, TMap<UEdGraphPin*, UEdGraphPin*>& PinsToRemap, TArray<UEdGraphPin*>& PinsToOrphan) override;
+	virtual void RemapPins(const UCustomizableObjectNode& Node, const TArray<UEdGraphPin*>& OldPins, const TArray<UEdGraphPin*>& NewPins, TMap<UEdGraphPin*, UEdGraphPin*>& PinsToRemap, TArray<UEdGraphPin*>& PinsToOrphan) override;
 };
 
 
diff --git a/Engine/Plugins/Experimental/Mutable/Source/CustomizableObjectEditor/Private/MuCOE/Nodes/CustomizableObjectNode.cpp b/Engine/Plugins/Experimental/Mutable/Source/CustomizableObjectEditor/Private/MuCOE/Nodes/CustomizableObjectNode.cpp
--- a/Engine/Plugins/Experimental/Mutable/Source/CustomizableObjectEditor/Private/MuCOE/Nodes/CustomizableObjectNode.cpp
+++ b/Engine/Plugins/Experimental/Mutable/Source/CustomizableObjectEditor/Private/MuCOE/Nodes/CustomizableObjectNode.cpp
@@ -231,7 +231,7 @@ void UCustomizableObjectNode::ReconstructNode(UCustomizableObjectNodeRemapPins*
 
 	TMap<UEdGraphPin*, UEdGraphPin*> PinsToRemap;
 	TArray<UEdGraphPin*> PinsToOrphan;
-	RemapPinsAction->RemapPins(OldPins, NewPins, PinsToRemap, PinsToOrphan);
+	RemapPinsAction->RemapPins(*this, OldPins, NewPins, PinsToRemap, PinsToOrphan);
 
 	// Check only.
 	for (const TTuple<UEdGraphPin*, UEdGraphPin*>& Pair : PinsToRemap)
diff --git a/Engine/Plugins/Experimental/Mutable/Source/CustomizableObjectEditor/Private/MuCOE/Nodes/CustomizableObjectNodeCopyMaterial.cpp b/Engine/Plugins/Experimental/Mutable/Source/CustomizableObjectEditor/Private/MuCOE/Nodes/CustomizableObjectNodeCopyMaterial.cpp
--- a/Engine/Plugins/Experimental/Mutable/Source/CustomizableObjectEditor/Private/MuCOE/Nodes/CustomizableObjectNodeCopyMaterial.cpp
+++ b/Engine/Plugins/Experimental/Mutable/Source/CustomizableObjectEditor/Private/MuCOE/Nodes/CustomizableObjectNodeCopyMaterial.cpp
@@ -25,11 +25,6 @@ FText UCustomizableObjectNodeCopyMaterial::GetNodeTitle(ENodeTitleType::Type Tit
 
 void UCustomizableObjectNodeCopyMaterial::AllocateDefaultPins(UCustomizableObjectNodeRemapPins* RemapPins)
 {
-	if (UCustomizableObjectNodeMaterialRemapPinsByName* RemapPinsByNameCustom = Cast<UCustomizableObjectNodeMaterialRemapPinsByName>(RemapPins))
-	{
-		RemapPinsByNameCustom->Node = this;
-	}
-	
 	const UEdGraphSchema_CustomizableObject* Schema = GetDefault<UEdGraphSchema_CustomizableObject>();
 
 	// Input pins
diff --git a/Engine/Plugins/Experimental/Mutable/Source/CustomizableObjectEditor/Private/MuCOE/Nodes/CustomizableObjectNodeExternalPin.cpp b/Engine/Plugins/Experimental/Mutable/Source/CustomizableObjectEditor/Private/MuCOE/Nodes/CustomizableObjectNodeExternalPin.cpp
--- a/Engine/Plugins/Experimental/Mutable/Source/CustomizableObjectEditor/Private/MuCOE/Nodes/CustomizableObjectNodeExternalPin.cpp
+++ b/Engine/Plugins/Experimental/Mutable/Source/CustomizableObjectEditor/Private/MuCOE/Nodes/CustomizableObjectNodeExternalPin.cpp
@@ -10,8 +10,10 @@
 #define LOCTEXT_NAMESPACE "CustomizableObjectEditor"
 
 
-void UCustomizableObjectNodeRemapPinsCustomExternalPin::RemapPins(const TArray<UEdGraphPin*>& OldPins, const TArray<UEdGraphPin*>& NewPins, TMap<UEdGraphPin*, UEdGraphPin*>& PinsToRemap, TArray<UEdGraphPin*>& PinsToOrphan)
+void UCustomizableObjectNodeRemapPinsCustomExternalPin::RemapPins(const UCustomizableObjectNode& Node, const TArray<UEdGraphPin*>& OldPins, const TArray<UEdGraphPin*>& NewPins, TMap<UEdGraphPin*, UEdGraphPin*>& PinsToRemap, TArray<UEdGraphPin*>& PinsToOrphan)
 {
+	const UCustomizableObjectNodeExternalPin* ExternalPinNode = Cast<UCustomizableObjectNodeExternalPin>(&Node);
+	
 	if (OldPins.Num()) 
 	{
 		UEdGraphPin* const OldPin = OldPins[0];
@@ -20,7 +22,7 @@ void UCustomizableObjectNodeRemapPinsCustomExternalPin::RemapPins(const TArray<U
 			PinsToRemap.Add(OldPin, NewPins[0]);
 		}
 
-		if (!Node->GetNodeExposePin())
+		if (!ExternalPinNode->GetNodeExposePin())
 		{
 			PinsToOrphan.Add(OldPin);
 		}
@@ -125,14 +127,6 @@ void UCustomizableObjectNodeExternalPin::BeginConstruct()
 
 void UCustomizableObjectNodeExternalPin::AllocateDefaultPins(UCustomizableObjectNodeRemapPins* RemapPins)
 {
-	const UEdGraphSchema_CustomizableObject* Schema = GetDefault<UEdGraphSchema_CustomizableObject>();
-
-	// Pass information to the remap pins action context.
-	if (UCustomizableObjectNodeRemapPinsCustomExternalPin* RemapPinsCustom = Cast<UCustomizableObjectNodeRemapPinsCustomExternalPin>(RemapPins))
-	{
-		RemapPinsCustom->Node = this;
-	}
-
 	if (UCustomizableObjectNodeExposePin* NodeExposePin = GetNodeExposePin())
 	{
 		CreateExternalPin(NodeExposePin);
diff --git a/Engine/Plugins/Experimental/Mutable/Source/CustomizableObjectEditor/Private/MuCOE/Nodes/CustomizableObjectNodeMaterial.cpp b/Engine/Plugins/Experimental/Mutable/Source/CustomizableObjectEditor/Private/MuCOE/Nodes/CustomizableObjectNodeMaterial.cpp
--- a/Engine/Plugins/Experimental/Mutable/Source/CustomizableObjectEditor/Private/MuCOE/Nodes/CustomizableObjectNodeMaterial.cpp
+++ b/Engine/Plugins/Experimental/Mutable/Source/CustomizableObjectEditor/Private/MuCOE/Nodes/CustomizableObjectNodeMaterial.cpp
@@ -33,22 +33,22 @@ const TArray<EMaterialParameterType> UCustomizableObjectNodeMaterial::ParameterT
 	EMaterialParameterType::Scalar };
 
 
-bool UCustomizableObjectNodeMaterialRemapPinsByName::Equal(const UEdGraphPin& OldPin, const UEdGraphPin& NewPin) const
+bool UCustomizableObjectNodeMaterialRemapPinsByName::Equal(const UCustomizableObjectNode& Node, const UEdGraphPin& OldPin, const UEdGraphPin& NewPin) const
 {
-	const UCustomizableObjectNodeMaterialPinDataParameter* PinDataOldPin = Cast<UCustomizableObjectNodeMaterialPinDataParameter>(Node->GetPinData(OldPin));
-	const UCustomizableObjectNodeMaterialPinDataParameter* PinDataNewPin = Cast<UCustomizableObjectNodeMaterialPinDataParameter>(Node->GetPinData(NewPin));
+	const UCustomizableObjectNodeMaterialPinDataParameter* PinDataOldPin = Cast<UCustomizableObjectNodeMaterialPinDataParameter>(Node.GetPinData(OldPin));
+	const UCustomizableObjectNodeMaterialPinDataParameter* PinDataNewPin = Cast<UCustomizableObjectNodeMaterialPinDataParameter>(Node.GetPinData(NewPin));
 	if (PinDataOldPin && PinDataNewPin)
 	{
 		return PinDataOldPin->ParameterId == PinDataNewPin->ParameterId;
 	}
 	else
 	{
-		return Super::Equal(OldPin, NewPin);
+		return Super::Equal(Node, OldPin, NewPin);
 	}
 }
 
 
-void UCustomizableObjectNodeMaterialRemapPinsByName::RemapPins(const TArray<UEdGraphPin*>& OldPins, const TArray<UEdGraphPin*>& NewPins, TMap<UEdGraphPin*, UEdGraphPin*>& PinsToRemap, TArray<UEdGraphPin*>& PinsToOrphan)
+void UCustomizableObjectNodeMaterialRemapPinsByName::RemapPins(const UCustomizableObjectNode& Node, const TArray<UEdGraphPin*>& OldPins, const TArray<UEdGraphPin*>& NewPins, TMap<UEdGraphPin*, UEdGraphPin*>& PinsToRemap, TArray<UEdGraphPin*>& PinsToOrphan)
 {
 	for (UEdGraphPin* OldPin : OldPins)
 	{
@@ -56,7 +56,7 @@ void UCustomizableObjectNodeMaterialRemapPinsByName::RemapPins(const TArray<UEdG
 
 		for (UEdGraphPin* NewPin : NewPins)
 		{
-			if (Equal(*OldPin, *NewPin))
+			if (Equal(Node, *OldPin, *NewPin))
 			{
 				bFound = true;
 
@@ -65,7 +65,7 @@ void UCustomizableObjectNodeMaterialRemapPinsByName::RemapPins(const TArray<UEdG
 			}
 		}
 
-		if (!bFound && (OldPin->LinkedTo.Num() || HasSavedPinData(*OldPin)))
+		if (!bFound && (OldPin->LinkedTo.Num() || HasSavedPinData(Node, *OldPin)))
 		{
 			PinsToOrphan.Add(OldPin);
 		}
@@ -73,9 +73,9 @@ void UCustomizableObjectNodeMaterialRemapPinsByName::RemapPins(const TArray<UEdG
 }
 
 
-bool UCustomizableObjectNodeMaterialRemapPinsByName::HasSavedPinData(const UEdGraphPin &Pin) const
+bool UCustomizableObjectNodeMaterialRemapPinsByName::HasSavedPinData(const UCustomizableObjectNode& Node, const UEdGraphPin &Pin) const
 {
-	if (const UCustomizableObjectNodeMaterialPinDataParameter* PinData = Cast<UCustomizableObjectNodeMaterialPinDataParameter>(Node->GetPinData(Pin)))
+	if (const UCustomizableObjectNodeMaterialPinDataParameter* PinData = Cast<UCustomizableObjectNodeMaterialPinDataParameter>(Node.GetPinData(Pin)))
 	{
 		return !PinData->IsDefault();
 	}
@@ -238,11 +238,6 @@ void UCustomizableObjectNodeMaterialPinDataImage::PostEditChangeProperty(FProper
 
 void UCustomizableObjectNodeMaterial::AllocateDefaultPins(UCustomizableObjectNodeRemapPins* RemapPins)
 {
-	if (UCustomizableObjectNodeMaterialRemapPinsByName* RemapPinsByNameCustom = Cast<UCustomizableObjectNodeMaterialRemapPinsByName>(RemapPins))
-	{
-		RemapPinsByNameCustom->Node = this;
-	}
-	
 	const UEdGraphSchema_CustomizableObject* Schema = GetDefault<UEdGraphSchema_CustomizableObject>();
 
 	{
diff --git a/Engine/Plugins/Experimental/Mutable/Source/CustomizableObjectEditor/Private/MuCOE/Nodes/CustomizableObjectNodeTable.cpp b/Engine/Plugins/Experimental/Mutable/Source/CustomizableObjectEditor/Private/MuCOE/Nodes/CustomizableObjectNodeTable.cpp
--- a/Engine/Plugins/Experimental/Mutable/Source/CustomizableObjectEditor/Private/MuCOE/Nodes/CustomizableObjectNodeTable.cpp
+++ b/Engine/Plugins/Experimental/Mutable/Source/CustomizableObjectEditor/Private/MuCOE/Nodes/CustomizableObjectNodeTable.cpp
@@ -24,7 +24,7 @@ class UCustomizableObjectNodeRemapPins;
 #define LOCTEXT_NAMESPACE "CustomizableObjectEditor"
 
 
-bool UCustomizableObjectNodeTableRemapPins::Equal(const UEdGraphPin& OldPin, const UEdGraphPin& NewPin) const
+bool UCustomizableObjectNodeTableRemapPins::Equal(const UCustomizableObjectNode& Node, const UEdGraphPin& OldPin, const UEdGraphPin& NewPin) const
 {
 	if (OldPin.PinType.PinCategory == UEdGraphSchema_CustomizableObject::PC_Image ||
 		OldPin.PinType.PinCategory == UEdGraphSchema_CustomizableObject::PC_PassThroughImage)
@@ -40,7 +40,7 @@ bool UCustomizableObjectNodeTableRemapPins::Equal(const UEdGraphPin& OldPin, con
 }
 
 
-void UCustomizableObjectNodeTableRemapPins::RemapPins(const TArray<UEdGraphPin*>& OldPins, const TArray<UEdGraphPin*>& NewPins, TMap<UEdGraphPin*, UEdGraphPin*>& PinsToRemap, TArray<UEdGraphPin*>& PinsToOrphan)
+void UCustomizableObjectNodeTableRemapPins::RemapPins(const UCustomizableObjectNode& Node, const TArray<UEdGraphPin*>& OldPins, const TArray<UEdGraphPin*>& NewPins, TMap<UEdGraphPin*, UEdGraphPin*>& PinsToRemap, TArray<UEdGraphPin*>& PinsToOrphan)
 {
 	for (UEdGraphPin* OldPin : OldPins)
 	{
@@ -48,7 +48,7 @@ void UCustomizableObjectNodeTableRemapPins::RemapPins(const TArray<UEdGraphPin*>
 
 		for (UEdGraphPin* NewPin : NewPins)
 		{
-			if (Equal(*OldPin, *NewPin))
+			if (Equal(Node, *OldPin, *NewPin))
 			{
 				bFound = true;
 
diff --git a/Engine/Plugins/Experimental/Mutable/Source/CustomizableObjectEditor/Private/MuCOE/RemapPins/CustomizableObjectNodeRemapPinsByName.cpp b/Engine/Plugins/Experimental/Mutable/Source/CustomizableObjectEditor/Private/MuCOE/RemapPins/CustomizableObjectNodeRemapPinsByName.cpp
--- a/Engine/Plugins/Experimental/Mutable/Source/CustomizableObjectEditor/Private/MuCOE/RemapPins/CustomizableObjectNodeRemapPinsByName.cpp
+++ b/Engine/Plugins/Experimental/Mutable/Source/CustomizableObjectEditor/Private/MuCOE/RemapPins/CustomizableObjectNodeRemapPinsByName.cpp
@@ -3,9 +3,10 @@
 #include "MuCOE/RemapPins/CustomizableObjectNodeRemapPinsByName.h"
 
 #include "MuCOE/UnrealEditorPortabilityHelpers.h"
+#include "MuCOE/Nodes/CustomizableObjectNode.h"
 
 
-bool UCustomizableObjectNodeRemapPinsByName::Equal(const UEdGraphPin& OldPin, const UEdGraphPin& NewPin) const
+bool UCustomizableObjectNodeRemapPinsByName::Equal(const UCustomizableObjectNode& Node, const UEdGraphPin& OldPin, const UEdGraphPin& NewPin) const
 {
 	return Helper_GetPinName(&OldPin) == Helper_GetPinName(&NewPin) &&
 				OldPin.PinType == NewPin.PinType &&
@@ -13,7 +14,7 @@ bool UCustomizableObjectNodeRemapPinsByName::Equal(const UEdGraphPin& OldPin, co
 }
 
 
-void UCustomizableObjectNodeRemapPinsByName::RemapPins(const TArray<UEdGraphPin*>& OldPins, const TArray<UEdGraphPin*>& NewPins, TMap<UEdGraphPin*, UEdGraphPin*>& PinsToRemap, TArray<UEdGraphPin*>& PinsToOrphan)
+void UCustomizableObjectNodeRemapPinsByName::RemapPins(const UCustomizableObjectNode& Node, const TArray<UEdGraphPin*>& OldPins, const TArray<UEdGraphPin*>& NewPins, TMap<UEdGraphPin*, UEdGraphPin*>& PinsToRemap, TArray<UEdGraphPin*>& PinsToOrphan)
 {
 	for (UEdGraphPin* OldPin : OldPins)
 	{
@@ -21,7 +22,7 @@ void UCustomizableObjectNodeRemapPinsByName::RemapPins(const TArray<UEdGraphPin*
 
 		for (UEdGraphPin* NewPin : NewPins)
 		{
-			if (Equal(*OldPin, *NewPin))
+			if (Equal(Node, *OldPin, *NewPin))
 			{
 				bFound = true;
 
diff --git a/Engine/Plugins/Experimental/Mutable/Source/CustomizableObjectEditor/Private/MuCOE/RemapPins/CustomizableObjectNodeRemapPinsByNameDefaultPin.cpp b/Engine/Plugins/Experimental/Mutable/Source/CustomizableObjectEditor/Private/MuCOE/RemapPins/CustomizableObjectNodeRemapPinsByNameDefaultPin.cpp
--- a/Engine/Plugins/Experimental/Mutable/Source/CustomizableObjectEditor/Private/MuCOE/RemapPins/CustomizableObjectNodeRemapPinsByNameDefaultPin.cpp
+++ b/Engine/Plugins/Experimental/Mutable/Source/CustomizableObjectEditor/Private/MuCOE/RemapPins/CustomizableObjectNodeRemapPinsByNameDefaultPin.cpp
@@ -5,7 +5,7 @@
 #include "EdGraph/EdGraphPin.h"
 
 
-void UCustomizableObjectNodeRemapPinsByNameDefaultPin::RemapPins(const TArray<UEdGraphPin*>& OldPins, const TArray<UEdGraphPin*>& NewPins, TMap<UEdGraphPin*, UEdGraphPin*>& PinsToRemap, TArray<UEdGraphPin*>& PinsToOrphan)
+void UCustomizableObjectNodeRemapPinsByNameDefaultPin::RemapPins(const UCustomizableObjectNode& Node, const TArray<UEdGraphPin*>& OldPins, const TArray<UEdGraphPin*>& NewPins, TMap<UEdGraphPin*, UEdGraphPin*>& PinsToRemap, TArray<UEdGraphPin*>& PinsToOrphan)
 {
 	if (DefaultPin)
 	{
@@ -39,7 +39,7 @@ void UCustomizableObjectNodeRemapPinsByNameDefaultPin::RemapPins(const TArray<UE
 	}
 	else
 	{
-		Super::RemapPins(OldPins, NewPins, PinsToRemap, PinsToOrphan);
+		Super::RemapPins(Node, OldPins, NewPins, PinsToRemap, PinsToOrphan);
 	}
 }
 
diff --git a/Engine/Plugins/Experimental/Mutable/Source/CustomizableObjectEditor/Private/MuCOE/RemapPins/CustomizableObjectNodeRemapPinsByPosition.cpp b/Engine/Plugins/Experimental/Mutable/Source/CustomizableObjectEditor/Private/MuCOE/RemapPins/CustomizableObjectNodeRemapPinsByPosition.cpp
--- a/Engine/Plugins/Experimental/Mutable/Source/CustomizableObjectEditor/Private/MuCOE/RemapPins/CustomizableObjectNodeRemapPinsByPosition.cpp
+++ b/Engine/Plugins/Experimental/Mutable/Source/CustomizableObjectEditor/Private/MuCOE/RemapPins/CustomizableObjectNodeRemapPinsByPosition.cpp
@@ -45,7 +45,7 @@ void RemapPinsInternal(const TArray<UEdGraphPin*>& OldPins, const TArray<UEdGrap
 }
 
 
-void UCustomizableObjectNodeRemapPinsByPosition::RemapPins(const TArray<UEdGraphPin*>& OldPins, const TArray<UEdGraphPin*>& NewPins, TMap<UEdGraphPin*, UEdGraphPin*>& PinsToRemap, TArray<UEdGraphPin*>& PinsToOrphan)
+void UCustomizableObjectNodeRemapPinsByPosition::RemapPins(const UCustomizableObjectNode& Node, const TArray<UEdGraphPin*>& OldPins, const TArray<UEdGraphPin*>& NewPins, TMap<UEdGraphPin*, UEdGraphPin*>& PinsToRemap, TArray<UEdGraphPin*>& PinsToOrphan)
 {
 	RemapPinsInternal(OldPins, NewPins, PinsToRemap, PinsToOrphan, EGPD_Input);
 	RemapPinsInternal(OldPins, NewPins, PinsToRemap, PinsToOrphan, EGPD_Output);
diff --git a/Engine/Plugins/Experimental/Mutable/Source/CustomizableObjectEditor/Public/MuCOE/RemapPins/CustomizableObjectNodeRemapPins.h b/Engine/Plugins/Experimental/Mutable/Source/CustomizableObjectEditor/Public/MuCOE/RemapPins/CustomizableObjectNodeRemapPins.h
--- a/Engine/Plugins/Experimental/Mutable/Source/CustomizableObjectEditor/Public/MuCOE/RemapPins/CustomizableObjectNodeRemapPins.h
+++ b/Engine/Plugins/Experimental/Mutable/Source/CustomizableObjectEditor/Public/MuCOE/RemapPins/CustomizableObjectNodeRemapPins.h
@@ -6,6 +6,7 @@
 #include "CustomizableObjectNodeRemapPins.generated.h"
 
 class UEdGraphPin;
+class UCustomizableObjectNode;
 
 
 /**
@@ -27,11 +28,11 @@ public:
 	 * 
 	 * @param OldPins Old pins which are no longer present in the node.
 	 * @param NewPins Current node pins.
-	 * @param PinsToRemap Out parameter. Map of remapping pairs. Key is the old pin while value is the new pin.
+	 * @param PinsToRemap Out parameter. Remapping pairs. Key is the old pin while value is the new pin.
 	 * From pin has to be a pin from the NewPins, while FromPin has to be a pin from OldPins.
 	 * @param PinsToOrphan Out Parameter. List of pins which will be marked as orphaned. Only pins from OldPins must be added.
 	 */
-	virtual void RemapPins(const TArray<UEdGraphPin*>& OldPins, const TArray<UEdGraphPin*>& NewPins, TMap<UEdGraphPin*, UEdGraphPin*>& PinsToRemap, TArray<UEdGraphPin*>& PinsToOrphan) { check(false) };
+	virtual void RemapPins(const UCustomizableObjectNode& Node, const TArray<UEdGraphPin*>& OldPins, const TArray<UEdGraphPin*>& NewPins, TMap<UEdGraphPin*, UEdGraphPin*>& PinsToRemap, TArray<UEdGraphPin*>& PinsToOrphan) PURE_VIRTUAL(UCustomizableObjectNodeRemapPins::RemapPins, );
 };
 
 
diff --git a/Engine/Plugins/Experimental/Mutable/Source/CustomizableObjectEditor/Public/MuCOE/RemapPins/CustomizableObjectNodeRemapPinsByName.h b/Engine/Plugins/Experimental/Mutable/Source/CustomizableObjectEditor/Public/MuCOE/RemapPins/CustomizableObjectNodeRemapPinsByName.h
--- a/Engine/Plugins/Experimental/Mutable/Source/CustomizableObjectEditor/Public/MuCOE/RemapPins/CustomizableObjectNodeRemapPinsByName.h
+++ b/Engine/Plugins/Experimental/Mutable/Source/CustomizableObjectEditor/Public/MuCOE/RemapPins/CustomizableObjectNodeRemapPinsByName.h
@@ -23,7 +23,7 @@ class UCustomizableObjectNodeRemapPinsByName : public UCustomizableObjectNodeRem
 public:
 	GENERATED_BODY()
 
-	virtual bool Equal(const UEdGraphPin& OldPin, const UEdGraphPin& NewPin) const;
+	virtual bool Equal(const UCustomizableObjectNode& Node, const UEdGraphPin& OldPin, const UEdGraphPin& NewPin) const;
 	
-	virtual void RemapPins(const TArray<UEdGraphPin*>& OldPins, const TArray<UEdGraphPin*>& NewPins, TMap<UEdGraphPin*, UEdGraphPin*>& PinsToRemap, TArray<UEdGraphPin*>& PinsToOrphan) override;
+	virtual void RemapPins(const UCustomizableObjectNode& Node, const TArray<UEdGraphPin*>& OldPins, const TArray<UEdGraphPin*>& NewPins, TMap<UEdGraphPin*, UEdGraphPin*>& PinsToRemap, TArray<UEdGraphPin*>& PinsToOrphan) override;
 };
diff --git a/Engine/Plugins/Experimental/Mutable/Source/CustomizableObjectEditor/Public/MuCOE/RemapPins/CustomizableObjectNodeRemapPinsByNameDefaultPin.h b/Engine/Plugins/Experimental/Mutable/Source/CustomizableObjectEditor/Public/MuCOE/RemapPins/CustomizableObjectNodeRemapPinsByNameDefaultPin.h
--- a/Engine/Plugins/Experimental/Mutable/Source/CustomizableObjectEditor/Public/MuCOE/RemapPins/CustomizableObjectNodeRemapPinsByNameDefaultPin.h
+++ b/Engine/Plugins/Experimental/Mutable/Source/CustomizableObjectEditor/Public/MuCOE/RemapPins/CustomizableObjectNodeRemapPinsByNameDefaultPin.h
@@ -6,6 +6,7 @@
 
 #include "CustomizableObjectNodeRemapPinsByNameDefaultPin.generated.h"
 
+class UCustomizableObjectNode;
 class UEdGraphPin;
 class UObject;
 
@@ -20,6 +21,6 @@ public:
 	/** Default pin to remap with. */
 	UEdGraphPin* DefaultPin = nullptr;
 
-	virtual void RemapPins(const TArray<UEdGraphPin*>& OldPins, const TArray<UEdGraphPin*>& NewPins, TMap<UEdGraphPin*, UEdGraphPin*>& PinsToRemap, TArray<UEdGraphPin*>& PinsToOrphan) override;
+	virtual void RemapPins(const UCustomizableObjectNode& Node, const TArray<UEdGraphPin*>& OldPins, const TArray<UEdGraphPin*>& NewPins, TMap<UEdGraphPin*, UEdGraphPin*>& PinsToRemap, TArray<UEdGraphPin*>& PinsToOrphan) override;
 };
 
diff --git a/Engine/Plugins/Experimental/Mutable/Source/CustomizableObjectEditor/Public/MuCOE/RemapPins/CustomizableObjectNodeRemapPinsByPosition.h b/Engine/Plugins/Experimental/Mutable/Source/CustomizableObjectEditor/Public/MuCOE/RemapPins/CustomizableObjectNodeRemapPinsByPosition.h
--- a/Engine/Plugins/Experimental/Mutable/Source/CustomizableObjectEditor/Public/MuCOE/RemapPins/CustomizableObjectNodeRemapPinsByPosition.h
+++ b/Engine/Plugins/Experimental/Mutable/Source/CustomizableObjectEditor/Public/MuCOE/RemapPins/CustomizableObjectNodeRemapPinsByPosition.h
@@ -23,5 +23,5 @@ class UCustomizableObjectNodeRemapPinsByPosition : public UCustomizableObjectNod
 public:
 	GENERATED_BODY()
 
-	virtual void RemapPins(const TArray<UEdGraphPin*>& OldPins, const TArray<UEdGraphPin*>& NewPins, TMap<UEdGraphPin*, UEdGraphPin*>& PinsToRemap, TArray<UEdGraphPin*>& PinsToOrphan) override;
+	virtual void RemapPins(const UCustomizableObjectNode& Node, const TArray<UEdGraphPin*>& OldPins, const TArray<UEdGraphPin*>& NewPins, TMap<UEdGraphPin*, UEdGraphPin*>& PinsToRemap, TArray<UEdGraphPin*>& PinsToOrphan) override;
 };
