[Mutable] Fix linux CustomizableObjectNodeRemapPins compilation issues.

#rnx

diff --git a/Engine/Plugins/Experimental/Mutable/Source/CustomizableObjectEditor/Public/MuCOE/RemapPins/CustomizableObjectNodeRemapPins.h b/Engine/Plugins/Experimental/Mutable/Source/CustomizableObjectEditor/Public/MuCOE/RemapPins/CustomizableObjectNodeRemapPins.h
--- a/Engine/Plugins/Experimental/Mutable/Source/CustomizableObjectEditor/Public/MuCOE/RemapPins/CustomizableObjectNodeRemapPins.h
+++ b/Engine/Plugins/Experimental/Mutable/Source/CustomizableObjectEditor/Public/MuCOE/RemapPins/CustomizableObjectNodeRemapPins.h
@@ -32,7 +32,7 @@ public:
 	 * From pin has to be a pin from the NewPins, while FromPin has to be a pin from OldPins.
 	 * @param PinsToOrphan Out Parameter. List of pins which will be marked as orphaned. Only pins from OldPins must be added.
 	 */
-	virtual void RemapPins(const UCustomizableObjectNode& Node, const TArray<UEdGraphPin*>& OldPins, const TArray<UEdGraphPin*>& NewPins, TMap<UEdGraphPin*, UEdGraphPin*>& PinsToRemap, TArray<UEdGraphPin*>& PinsToOrphan) PURE_VIRTUAL(UCustomizableObjectNodeRemapPins::RemapPins, );
+	virtual void RemapPins(const UCustomizableObjectNode& Node, const TArray<UEdGraphPin*>& OldPins, const TArray<UEdGraphPin*>& NewPins, TMap<UEdGraphPin*, UEdGraphPin*>& PinsToRemap, TArray<UEdGraphPin*>& PinsToOrphan) { check(false); }
 };
 
 
