#rnx

diff --git a/Engine/Plugins/Experimental/Mutable/Source/CustomizableObjectEditor/Private/MuCOE/Nodes/CustomizableObjectNodeTableDetails.cpp b/Engine/Plugins/Experimental/Mutable/Source/CustomizableObjectEditor/Private/MuCOE/Nodes/CustomizableObjectNodeTableDetails.cpp
--- a/Engine/Plugins/Experimental/Mutable/Source/CustomizableObjectEditor/Private/MuCOE/Nodes/CustomizableObjectNodeTableDetails.cpp
+++ b/Engine/Plugins/Experimental/Mutable/Source/CustomizableObjectEditor/Private/MuCOE/Nodes/CustomizableObjectNodeTableDetails.cpp
@@ -347,7 +347,7 @@ void FCustomizableObjectNodeTableDetails::GenerateAnimInstanceComboBoxOptions()
 
 	const UScriptStruct* TableStruct = Node->Table->GetRowStruct();
 
-	// Fillins name options arrays and setting the selected item if any
+	// Fill in name option arrays and set the selected item if any
 	for (TFieldIterator<FProperty> It(TableStruct); It; ++It)
 	{
 		if (FProperty* ColumnProperty = *It)
@@ -371,7 +371,7 @@ void FCustomizableObjectNodeTableDetails::GenerateAnimInstanceComboBoxOptions()
 					}
 				}
 			}
-			else if (const FIntProperty* NumProperty = CastField<FIntProperty>(ColumnProperty))
+			else if (CastField<FIntProperty>(ColumnProperty) || CastField<FNameProperty>(ColumnProperty))
 			{
 				TSharedPtr<FString> Option = MakeShareable(new FString(DataTableUtils::GetPropertyExportName(ColumnProperty)));
 				AnimSlotOptionNames.Add(Option);
