[Mutable] Fix incorrect Node Project Texture size properties description.

#rnx

diff --git a/Engine/Plugins/Experimental/Mutable/Source/CustomizableObjectEditor/Classes/MuCOE/Nodes/CustomizableObjectNodeTextureProject.h b/Engine/Plugins/Experimental/Mutable/Source/CustomizableObjectEditor/Classes/MuCOE/Nodes/CustomizableObjectNodeTextureProject.h
--- a/Engine/Plugins/Experimental/Mutable/Source/CustomizableObjectEditor/Classes/MuCOE/Nodes/CustomizableObjectNodeTextureProject.h
+++ b/Engine/Plugins/Experimental/Mutable/Source/CustomizableObjectEditor/Classes/MuCOE/Nodes/CustomizableObjectNodeTextureProject.h
@@ -54,11 +54,11 @@ public:
 	UPROPERTY(EditAnywhere, Category = CustomizableObject)
 	ETextureProjectMinFilterMethod MinFilterMethod = ETextureProjectMinFilterMethod::None;
 
-	/** Set the width of the Texture. If the reference texture is not null, it overrides the Size X property */
+	/** Set the width of the Texture. If greater than zero, it overrides the Reference Texture width. */
 	UPROPERTY(EditAnywhere, Category = CustomizableObject)
 	uint32 TextureSizeX = 0;
 
-	/** Set the height of the Texture. If the reference texture is not null, it overrides the Size Y property */
+	/** Set the height of the Texture. If greater than zero, it overrides the Reference Texture height. */
 	UPROPERTY(EditAnywhere, Category = CustomizableObject)
 	uint32 TextureSizeY = 0;
 
