[mutable] Now, Vertices from a SKM Node linked to an Extend Materia that do not belong to any layout block will be automatically removed form the mesh.
+ Fixed case where pressing 'Select All" in a layout block selector would add duplicated blocks.
+ Added warning indicating that a SKM node without any Layout node and linked to an Extend Material will add a 4x4 layout.

#rnx

diff --git a/Engine/Plugins/Experimental/Mutable/Source/CustomizableObjectEditor/Classes/MuCOE/Nodes/CustomizableObjectNodeExtendMaterial.h b/Engine/Plugins/Experimental/Mutable/Source/CustomizableObjectEditor/Classes/MuCOE/Nodes/CustomizableObjectNodeExtendMaterial.h
--- a/Engine/Plugins/Experimental/Mutable/Source/CustomizableObjectEditor/Classes/MuCOE/Nodes/CustomizableObjectNodeExtendMaterial.h
+++ b/Engine/Plugins/Experimental/Mutable/Source/CustomizableObjectEditor/Classes/MuCOE/Nodes/CustomizableObjectNodeExtendMaterial.h
@@ -55,6 +55,7 @@ public:
 	// UCustomizableObjectNodeMaterialBase interface
 	virtual UEdGraphPin* OutputPin() const override;
 	virtual void SetParentNode(UCustomizableObject* Object, FGuid NodeId) override;
+	TArray<UCustomizableObjectLayout*> GetLayouts() override;
 	
 	// FCustomizableObjectNodeParentMaterial interface
 	virtual void SaveParentNode(UCustomizableObject* Object, FGuid NodeId) override;
diff --git a/Engine/Plugins/Experimental/Mutable/Source/CustomizableObjectEditor/Private/MuCOE/GenerateMutableSource/GenerateMutableSourceLayout.cpp b/Engine/Plugins/Experimental/Mutable/Source/CustomizableObjectEditor/Private/MuCOE/GenerateMutableSource/GenerateMutableSourceLayout.cpp
--- a/Engine/Plugins/Experimental/Mutable/Source/CustomizableObjectEditor/Private/MuCOE/GenerateMutableSource/GenerateMutableSourceLayout.cpp
+++ b/Engine/Plugins/Experimental/Mutable/Source/CustomizableObjectEditor/Private/MuCOE/GenerateMutableSource/GenerateMutableSourceLayout.cpp
@@ -11,7 +11,7 @@
 #define LOCTEXT_NAMESPACE "CustomizableObjectEditor"
 
 
-mu::NodeLayoutPtr GenerateMutableSourceLayout(const UEdGraphPin * Pin, FMutableGraphGenerationContext & GenerationContext)
+mu::NodeLayoutPtr GenerateMutableSourceLayout(const UEdGraphPin * Pin, FMutableGraphGenerationContext & GenerationContext, bool bLinkedToExtendMaterial)
 {
 	check(Pin)
 	RETURN_ON_CYCLE(*Pin, GenerationContext)
@@ -58,10 +58,19 @@ mu::NodeLayoutPtr GenerateMutableSourceLayout(const UEdGraphPin * Pin, FMutableG
 		LayoutNode->SetGridSize(TypedNodeBlocks->Layout->GetGridSize().X, TypedNodeBlocks->Layout->GetGridSize().Y);
 		LayoutNode->SetMaxGridSize(TypedNodeBlocks->Layout->GetMaxGridSize().X, TypedNodeBlocks->Layout->GetMaxGridSize().Y);
 		LayoutNode->SetBlockCount(TypedNodeBlocks->Layout->Blocks.Num() ? TypedNodeBlocks->Layout->Blocks.Num() : 1);
-		LayoutNode->SetIgnoreWarningsLOD(TypedNodeBlocks->Layout->GetIgnoreVertexLayoutWarnings() ? TypedNodeBlocks->Layout->GetFirstLODToIgnoreWarnings() : -1);
 		LayoutNode->SetLayoutPackingStrategy(TypedNodeBlocks->Layout->GetPackingStrategy() == ECustomizableObjectTextureLayoutPackingStrategy::Fixed ? mu::EPackStrategy::FIXED_LAYOUT : mu::EPackStrategy::RESIZABLE_LAYOUT);
 		LayoutNode->SetBlockReductionMethod(TypedNodeBlocks->Layout->GetBlockReductionMethod() == ECustomizableObjectLayoutBlockReductionMethod::Halve ? mu::EReductionMethod::HALVE_REDUCTION : mu::EReductionMethod::UNITARY_REDUCTION);
 
+		if (bLinkedToExtendMaterial)
+		{
+			// Layout warnings can be safely ignored in this case. Vertices that do not belong to any layout block will be removed (Extend Materials only)
+			LayoutNode->SetIgnoreWarningsLOD(0);
+		}
+		else
+		{
+			LayoutNode->SetIgnoreWarningsLOD(TypedNodeBlocks->Layout->GetIgnoreVertexLayoutWarnings() ? TypedNodeBlocks->Layout->GetFirstLODToIgnoreWarnings() : -1);
+		}
+
 		if (TypedNodeBlocks->Layout->Blocks.Num())
 		{
 			for (int BlockIndex = 0; BlockIndex < TypedNodeBlocks->Layout->Blocks.Num(); ++BlockIndex)
diff --git a/Engine/Plugins/Experimental/Mutable/Source/CustomizableObjectEditor/Private/MuCOE/GenerateMutableSource/GenerateMutableSourceLayout.h b/Engine/Plugins/Experimental/Mutable/Source/CustomizableObjectEditor/Private/MuCOE/GenerateMutableSource/GenerateMutableSourceLayout.h
--- a/Engine/Plugins/Experimental/Mutable/Source/CustomizableObjectEditor/Private/MuCOE/GenerateMutableSource/GenerateMutableSourceLayout.h
+++ b/Engine/Plugins/Experimental/Mutable/Source/CustomizableObjectEditor/Private/MuCOE/GenerateMutableSource/GenerateMutableSourceLayout.h
@@ -8,4 +8,4 @@ class UEdGraphPin;
 struct FMutableGraphGenerationContext;
 
 
-mu::NodeLayoutPtr GenerateMutableSourceLayout(const UEdGraphPin* Pin, FMutableGraphGenerationContext& GenerationContext);
+mu::NodeLayoutPtr GenerateMutableSourceLayout(const UEdGraphPin* Pin, FMutableGraphGenerationContext& GenerationContext, bool bLinkedToExtendMaterial = false);
diff --git a/Engine/Plugins/Experimental/Mutable/Source/CustomizableObjectEditor/Private/MuCOE/GenerateMutableSource/GenerateMutableSourceMesh.cpp b/Engine/Plugins/Experimental/Mutable/Source/CustomizableObjectEditor/Private/MuCOE/GenerateMutableSource/GenerateMutableSourceMesh.cpp
--- a/Engine/Plugins/Experimental/Mutable/Source/CustomizableObjectEditor/Private/MuCOE/GenerateMutableSource/GenerateMutableSourceMesh.cpp
+++ b/Engine/Plugins/Experimental/Mutable/Source/CustomizableObjectEditor/Private/MuCOE/GenerateMutableSource/GenerateMutableSourceMesh.cpp
@@ -2719,7 +2719,8 @@ void GenerateMorphTarget(const UCustomizableObjectNode* Node, const UEdGraphPin*
 /** Convert a CustomizableObject Source Graph into a mutable source graph  */
 mu::NodeMeshPtr GenerateMutableSourceMesh(const UEdGraphPin * Pin,
 	FMutableGraphGenerationContext & GenerationContext,
-	FMutableGraphMeshGenerationData & MeshData)
+	FMutableGraphMeshGenerationData & MeshData,
+	bool bLinkedToExtendMaterial)
 {
 	check(Pin)
 	RETURN_ON_CYCLE(*Pin, GenerationContext)
@@ -2923,7 +2924,7 @@ mu::NodeMeshPtr GenerateMutableSourceMesh(const UEdGraphPin * Pin,
 						{
 							if (const UEdGraphPin* ConnectedPin = FollowInputPin(*TypedNodeSkel->LODs[LayoutLOD].Materials[SectionIndex].LayoutPinsRef[LayoutIndex].Get()))
 							{
-								mu::NodeLayoutPtr LayoutNode = GenerateMutableSourceLayout(ConnectedPin, GenerationContext);
+								mu::NodeLayoutPtr LayoutNode = GenerateMutableSourceLayout(ConnectedPin, GenerationContext, bLinkedToExtendMaterial);
 								MeshNode->SetLayout(LayoutIndex, LayoutNode);
 								LayoutFound = true;
 							}
@@ -3568,6 +3569,12 @@ mu::NodeMeshPtr GenerateMutableSourceMesh(const UEdGraphPin * Pin,
 					LayoutNode->SetLayoutPackingStrategy(Layouts[i]->GetPackingStrategy() == ECustomizableObjectTextureLayoutPackingStrategy::Fixed ? mu::EPackStrategy::FIXED_LAYOUT : mu::EPackStrategy::RESIZABLE_LAYOUT);
 					LayoutNode->SetBlockReductionMethod(Layouts[i]->GetBlockReductionMethod() == ECustomizableObjectLayoutBlockReductionMethod::Halve ? mu::EReductionMethod::HALVE_REDUCTION : mu::EReductionMethod::UNITARY_REDUCTION);
 
+					if (bLinkedToExtendMaterial)
+					{
+						// Layout warnings can be safely ignored in this case. Vertices that do not belong to any layout block will be removed (Extend Materials only)
+						LayoutNode->SetIgnoreWarningsLOD(0);
+					}
+
 					if (Layouts[i]->Blocks.Num())
 					{
 						for (int BlockIndex = 0; BlockIndex < Layouts[i]->Blocks.Num(); ++BlockIndex)
diff --git a/Engine/Plugins/Experimental/Mutable/Source/CustomizableObjectEditor/Private/MuCOE/GenerateMutableSource/GenerateMutableSourceMesh.h b/Engine/Plugins/Experimental/Mutable/Source/CustomizableObjectEditor/Private/MuCOE/GenerateMutableSource/GenerateMutableSourceMesh.h
--- a/Engine/Plugins/Experimental/Mutable/Source/CustomizableObjectEditor/Private/MuCOE/GenerateMutableSource/GenerateMutableSourceMesh.h
+++ b/Engine/Plugins/Experimental/Mutable/Source/CustomizableObjectEditor/Private/MuCOE/GenerateMutableSource/GenerateMutableSourceMesh.h
@@ -30,4 +30,4 @@ mu::MeshPtr GenerateMutableMesh(UObject * Mesh, const TSoftClassPtr<class UAnimI
 mu::MeshPtr BuildMorphedMutableMesh(const UEdGraphPin * BaseSourcePin, const FString& MorphTargetName, FMutableGraphGenerationContext & GenerationContext, const FName& RowName = "");
 
 /** Convert a CustomizableObject Source Graph into a mutable source graph. */
-mu::NodeMeshPtr GenerateMutableSourceMesh(const UEdGraphPin* Pin, FMutableGraphGenerationContext& GenerationContext, FMutableGraphMeshGenerationData& MeshData);
+mu::NodeMeshPtr GenerateMutableSourceMesh(const UEdGraphPin* Pin, FMutableGraphGenerationContext& GenerationContext, FMutableGraphMeshGenerationData& MeshData, bool bLinkedToExtendMaterial = false);
diff --git a/Engine/Plugins/Experimental/Mutable/Source/CustomizableObjectEditor/Private/MuCOE/GenerateMutableSource/GenerateMutableSourceSurface.cpp b/Engine/Plugins/Experimental/Mutable/Source/CustomizableObjectEditor/Private/MuCOE/GenerateMutableSource/GenerateMutableSourceSurface.cpp
--- a/Engine/Plugins/Experimental/Mutable/Source/CustomizableObjectEditor/Private/MuCOE/GenerateMutableSource/GenerateMutableSourceSurface.cpp
+++ b/Engine/Plugins/Experimental/Mutable/Source/CustomizableObjectEditor/Private/MuCOE/GenerateMutableSource/GenerateMutableSourceSurface.cpp
@@ -83,6 +83,7 @@ void LayoutOperationUVNormalizedWarning(FMutableGraphGenerationContext& Generati
 	}
 }
 
+
 void SetSurfaceFormat( FMutableGraphGenerationContext& GenerationContext,
 					   mu::FMeshBufferSet& OutVertexBufferFormat, mu::FMeshBufferSet& OutIndexBufferFormat, const FMutableGraphMeshGenerationData& MeshData, 
 					   bool bWithExtraBoneInfluences, bool bWithRealTimeMorphs, bool bWithClothing, bool bWith16BitWeights )
@@ -1019,11 +1020,37 @@ mu::NodeSurfacePtr GenerateMutableSourceSurface(const UEdGraphPin * Pin, FMutabl
 			FMutableGraphMeshGenerationData MeshData;
 			if (const UEdGraphPin* ConnectedPin = FollowInputPin(*TypedNodeExt->AddMeshPin()))
 			{
-				AddMeshNode = GenerateMutableSourceMesh(ConnectedPin, GenerationContext, MeshData);
+				AddMeshNode = GenerateMutableSourceMesh(ConnectedPin, GenerationContext, MeshData, true);
 			}
 
 			if (AddMeshNode)
 			{
+				mu::NodeMeshPtr MeshPtr = AddMeshNode;
+
+				const TArray<UCustomizableObjectLayout*> Layouts = TypedNodeExt->GetLayouts();
+
+				if (Layouts.Num())
+				{
+					//TODO: Implement support for multiple UV channels (e.g. Add warning for vertices which have a block in a layout but not in the other)
+					mu::NodeMeshFragmentPtr MeshFrag = new mu::NodeMeshFragment();
+				
+					MeshFrag->SetMesh(MeshPtr);
+					MeshFrag->SetLayoutOrGroup(0);
+					MeshFrag->SetFragmentType(mu::NodeMeshFragment::FT_LAYOUT_BLOCKS);
+					MeshFrag->SetBlockCount(Layouts[0]->Blocks.Num());
+				
+					for (int i = 0; i < Layouts[0]->Blocks.Num(); ++i)
+					{
+						MeshFrag->SetBlock(i, i);
+					}
+				
+					MeshPtr = MeshFrag;
+				}
+				else
+				{
+					GenerationContext.Compiler->CompilerLog(LOCTEXT("ExtendMaterialLayoutMissing","Skeletal Mesh without Layout Node linked to an Extend Material. A 4x4 layout will be added as default layout."), Node);
+				}
+
 				mu::NodeMeshFormatPtr MeshFormat = new mu::NodeMeshFormat();
 				SetSurfaceFormat( GenerationContext,
 						MeshFormat->GetVertexBuffers(), MeshFormat->GetIndexBuffers(), MeshData,
@@ -1032,7 +1059,7 @@ mu::NodeSurfacePtr GenerateMutableSourceSurface(const UEdGraphPin * Pin, FMutabl
 						GenerationContext.Options.bClothingEnabled,
 						GenerationContext.Options.b16BitBoneWeightsEnabled);
 
-				MeshFormat->SetSource(AddMeshNode.get());
+				MeshFormat->SetSource(MeshPtr.get());
 
 				mu::NodePatchMeshPtr MeshPatch = new mu::NodePatchMesh();
 				MeshPatch->SetAdd(MeshFormat.get());
diff --git a/Engine/Plugins/Experimental/Mutable/Source/CustomizableObjectEditor/Private/MuCOE/Nodes/CustomizableObjectNodeExtendMaterial.cpp b/Engine/Plugins/Experimental/Mutable/Source/CustomizableObjectEditor/Private/MuCOE/Nodes/CustomizableObjectNodeExtendMaterial.cpp
--- a/Engine/Plugins/Experimental/Mutable/Source/CustomizableObjectEditor/Private/MuCOE/Nodes/CustomizableObjectNodeExtendMaterial.cpp
+++ b/Engine/Plugins/Experimental/Mutable/Source/CustomizableObjectEditor/Private/MuCOE/Nodes/CustomizableObjectNodeExtendMaterial.cpp
@@ -5,7 +5,10 @@
 #include "MaterialTypes.h"
 #include "MuCOE/CustomizableObjectEditor_Deprecated.h"
 #include "MuCOE/EdGraphSchema_CustomizableObject.h"
+#include "MuCOE/GraphTraversal.h"
 #include "MuCOE/Nodes/CustomizableObjectNodeCopyMaterial.h"
+#include "MuCOE/Nodes/CustomizableObjectNodeSkeletalMesh.h"
+#include "MuCOE/Nodes/CustomizableObjectNodeTable.h"
 
 class UCustomizableObjectNode;
 class UCustomizableObjectNodeRemapPins;
@@ -159,6 +162,7 @@ UEdGraphPin* UCustomizableObjectNodeExtendMaterial::OutputPin() const
 	return FindPin(TEXT("Material"));
 }
 
+
 void UCustomizableObjectNodeExtendMaterial::SetParentNode(UCustomizableObject* Object, FGuid NodeId)
 {
 	PreSetParentNodeWork(Object, NodeId);
@@ -169,6 +173,32 @@ void UCustomizableObjectNodeExtendMaterial::SetParentNode(UCustomizableObject* O
 }
 
 
+TArray<UCustomizableObjectLayout*> UCustomizableObjectNodeExtendMaterial::GetLayouts()
+{
+	TArray<UCustomizableObjectLayout*> Result;
+
+	if (UEdGraphPin* MeshPin = AddMeshPin())
+	{
+		if (const UEdGraphPin* ConnectedPin = FollowInputPin(*MeshPin))
+		{
+			if (const UEdGraphPin* SourceMeshPin = FindMeshBaseSource(*ConnectedPin, false))
+			{
+				if (const UCustomizableObjectNodeSkeletalMesh* MeshNode = Cast<UCustomizableObjectNodeSkeletalMesh>(SourceMeshPin->GetOwningNode()))
+				{
+					Result = MeshNode->GetLayouts(SourceMeshPin);
+				}
+				else if (const UCustomizableObjectNodeTable* TableNode = Cast<UCustomizableObjectNodeTable>(SourceMeshPin->GetOwningNode()))
+				{
+					Result = TableNode->GetLayouts(SourceMeshPin);
+				}
+			}
+		}
+	}
+
+	return Result;
+}
+
+
 FText UCustomizableObjectNodeExtendMaterial::GetTooltipText() const
 {
 	return LOCTEXT("Extend_Material_Tooltip", "Extend an ancestor's material: add a new mesh section, and add its corresponding texture to the ancestor's material texture parameters.");
diff --git a/Engine/Plugins/Experimental/Mutable/Source/CustomizableObjectEditor/Private/MuCOE/SCustomizableObjectNodeLayoutBlocksSelector.cpp b/Engine/Plugins/Experimental/Mutable/Source/CustomizableObjectEditor/Private/MuCOE/SCustomizableObjectNodeLayoutBlocksSelector.cpp
--- a/Engine/Plugins/Experimental/Mutable/Source/CustomizableObjectEditor/Private/MuCOE/SCustomizableObjectNodeLayoutBlocksSelector.cpp
+++ b/Engine/Plugins/Experimental/Mutable/Source/CustomizableObjectEditor/Private/MuCOE/SCustomizableObjectNodeLayoutBlocksSelector.cpp
@@ -203,7 +203,7 @@ void SCustomizableObjectNodeLayoutBlocksSelector::OnSelectAll()
 
 			for( const FCustomizableObjectLayoutBlock& Block : Layout->Blocks )
 			{
-				CurrentNode->BlockIds.Add(Block.Id);
+				CurrentNode->BlockIds.AddUnique(Block.Id);
 			}
 
 			BlocksLabel->SetText(FText::FromString(FString::Printf(TEXT("%d blocks selected"), CurrentNode->BlockIds.Num())));
diff --git a/Engine/Plugins/Experimental/Mutable/Source/MutableTools/Private/MuT/CodeGenerator_Mesh.cpp b/Engine/Plugins/Experimental/Mutable/Source/MutableTools/Private/MuT/CodeGenerator_Mesh.cpp
--- a/Engine/Plugins/Experimental/Mutable/Source/MutableTools/Private/MuT/CodeGenerator_Mesh.cpp
+++ b/Engine/Plugins/Experimental/Mutable/Source/MutableTools/Private/MuT/CodeGenerator_Mesh.cpp
@@ -268,7 +268,7 @@ class Node;
             attachedDataView.m_unassignedUVs = unassignedUVs.GetData();
             attachedDataView.m_unassignedUVsSize = (size_t)unassignedUVs.Num();
 
-            m_pErrorLog->GetPrivate()->Add( buf, attachedDataView, blockCount==1?ELMT_INFO:ELMT_WARNING, errorContext );
+            m_pErrorLog->GetPrivate()->Add( buf, attachedDataView, ELMT_WARNING, errorContext );
         }
         
         // Assign broken vertices to the first block
