[mutable] Added alpha channel support to color constant nodes.

#rnx

diff --git a/Engine/Plugins/Experimental/Mutable/Source/CustomizableObjectEditor/Private/MuCOE/GenerateMutableSource/GenerateMutableSourceColor.cpp b/Engine/Plugins/Experimental/Mutable/Source/CustomizableObjectEditor/Private/MuCOE/GenerateMutableSource/GenerateMutableSourceColor.cpp
--- a/Engine/Plugins/Experimental/Mutable/Source/CustomizableObjectEditor/Private/MuCOE/GenerateMutableSource/GenerateMutableSourceColor.cpp
+++ b/Engine/Plugins/Experimental/Mutable/Source/CustomizableObjectEditor/Private/MuCOE/GenerateMutableSource/GenerateMutableSourceColor.cpp
@@ -54,7 +54,7 @@ mu::NodeColourPtr GenerateMutableSourceColor(const UEdGraphPin* Pin, FMutableGra
 		mu::NodeColourConstantPtr ColorNode = new mu::NodeColourConstant();
 		Result = ColorNode;
 
-		ColorNode->SetValue(TypedNodeColorConst->Value.R, TypedNodeColorConst->Value.G, TypedNodeColorConst->Value.B);
+		ColorNode->SetValue(TypedNodeColorConst->Value.R, TypedNodeColorConst->Value.G, TypedNodeColorConst->Value.B, TypedNodeColorConst->Value.A);
 	}
 
 	else if (const UCustomizableObjectNodeColorParameter* TypedNodeColorParam = Cast<UCustomizableObjectNodeColorParameter>(Node))
diff --git a/Engine/Plugins/Experimental/Mutable/Source/CustomizableObjectEditor/Private/MuCOE/GenerateMutableSource/GenerateMutableSourceGroupProjector.cpp b/Engine/Plugins/Experimental/Mutable/Source/CustomizableObjectEditor/Private/MuCOE/GenerateMutableSource/GenerateMutableSourceGroupProjector.cpp
--- a/Engine/Plugins/Experimental/Mutable/Source/CustomizableObjectEditor/Private/MuCOE/GenerateMutableSource/GenerateMutableSourceGroupProjector.cpp
+++ b/Engine/Plugins/Experimental/Mutable/Source/CustomizableObjectEditor/Private/MuCOE/GenerateMutableSource/GenerateMutableSourceGroupProjector.cpp
@@ -276,7 +276,7 @@ mu::NodeImagePtr GenerateMutableGroupProjection(const int32 NodeLOD, const int32
 	}
 	
 	mu::NodeColourConstantPtr ZeroColorNode = new mu::NodeColourConstant();
-	ZeroColorNode->SetValue(0.f, 0.f, 0.f);
+	ZeroColorNode->SetValue(0.f, 0.f, 0.f, 1.0f);
 
 	if (TextureSize <= 0 || !FMath::IsPowerOfTwo(TextureSize))
 	{
diff --git a/Engine/Plugins/Experimental/Mutable/Source/CustomizableObjectEditor/Private/MuCOE/GenerateMutableSource/GenerateMutableSourceTable.cpp b/Engine/Plugins/Experimental/Mutable/Source/CustomizableObjectEditor/Private/MuCOE/GenerateMutableSource/GenerateMutableSourceTable.cpp
--- a/Engine/Plugins/Experimental/Mutable/Source/CustomizableObjectEditor/Private/MuCOE/GenerateMutableSource/GenerateMutableSourceTable.cpp
+++ b/Engine/Plugins/Experimental/Mutable/Source/CustomizableObjectEditor/Private/MuCOE/GenerateMutableSource/GenerateMutableSourceTable.cpp
@@ -517,7 +517,7 @@ void FillTableColumn(const UCustomizableObjectNodeTable* TableNode,	mu::TablePtr
 					}
 
 					// Setting cell value
-					MutableTable->SetCell(ColumnIndex, RowIdx, VectorValue.R, VectorValue.G, VectorValue.B);
+					MutableTable->SetCell(ColumnIndex, RowIdx, VectorValue.R, VectorValue.G, VectorValue.B, VectorValue.A);
 				}
 
 				// Checking if all modifiable vectors of the material instance have been modified
@@ -617,7 +617,7 @@ void FillTableColumn(const UCustomizableObjectNodeTable* TableNode,	mu::TablePtr
 
 			// Setting cell value
 			FLinearColor Value = *(FLinearColor*)CellData;
-			MutableTable->SetCell(CurrentColumn, RowIdx, Value.R, Value.G, Value.B);
+			MutableTable->SetCell(CurrentColumn, RowIdx, Value.R, Value.G, Value.B, Value.A);
 		}
 	}
 
diff --git a/Engine/Plugins/Experimental/Mutable/Source/MutableRuntime/Private/MuR/SerialisationPrivate.h b/Engine/Plugins/Experimental/Mutable/Source/MutableRuntime/Private/MuR/SerialisationPrivate.h
--- a/Engine/Plugins/Experimental/Mutable/Source/MutableRuntime/Private/MuR/SerialisationPrivate.h
+++ b/Engine/Plugins/Experimental/Mutable/Source/MutableRuntime/Private/MuR/SerialisationPrivate.h
@@ -6,9 +6,11 @@
 #include "MuR/Serialisation.h"
 
 #include "Misc/TVariant.h"
+#include "Math/Color.h"
 #include "Math/Quat.h"
 #include "Math/Vector.h"
 #include "Math/IntVector.h"
+#include "Math/Vector4.h"
 #include "Math/Transform.h"
 #include "MuR/Ptr.h"
 #include "MuR/RefCounted.h"
@@ -365,6 +367,11 @@ namespace mu
 	MUTABLE_DEFINE_POD_VECTOR_SERIALISABLE(UE::Math::TIntVector2<uint16>);
 	MUTABLE_DEFINE_POD_VECTOR_SERIALISABLE(UE::Math::TIntVector2<int16>);
 
+	MUTABLE_DEFINE_POD_SERIALISABLE(FVector4f);
+	MUTABLE_DEFINE_POD_SERIALISABLE(UE::Math::TVector4<float>);
+	MUTABLE_DEFINE_POD_VECTOR_SERIALISABLE(FVector4f);
+	MUTABLE_DEFINE_POD_VECTOR_SERIALISABLE(UE::Math::TVector4<float>);
+
 
 	//---------------------------------------------------------------------------------------------
 	template<>
@@ -546,4 +553,24 @@ namespace mu
 		v.SetComponents(Rot, Trans, Scale);
 	}
 
+
+	//---------------------------------------------------------------------------------------------
+	template< typename T >
+	inline void operator<< (OutputArchive& arch, const UE::Math::TVector4<T>& v)
+	{
+		arch << v.X;
+		arch << v.Y;
+		arch << v.Z;
+		arch << v.W;
+	}
+
+	template< typename T >
+	inline void operator>> (InputArchive& arch, UE::Math::TVector4<T>& v)
+	{
+		arch >> v.X;
+		arch >> v.Y;
+		arch >> v.Z;
+		arch >> v.W;
+	}
+
 }
diff --git a/Engine/Plugins/Experimental/Mutable/Source/MutableTools/Private/MuT/CodeGenerator_Color.cpp b/Engine/Plugins/Experimental/Mutable/Source/MutableTools/Private/MuT/CodeGenerator_Color.cpp
--- a/Engine/Plugins/Experimental/Mutable/Source/MutableTools/Private/MuT/CodeGenerator_Color.cpp
+++ b/Engine/Plugins/Experimental/Mutable/Source/MutableTools/Private/MuT/CodeGenerator_Color.cpp
@@ -124,10 +124,10 @@ namespace mu
 
 		Ptr<ASTOpFixed> op = new ASTOpFixed();
 		op->op.type = OP_TYPE::CO_CONSTANT;
-		op->op.args.ColourConstant.value[0] = node.m_value[0];
-		op->op.args.ColourConstant.value[1] = node.m_value[1];
-		op->op.args.ColourConstant.value[2] = node.m_value[2];
-		op->op.args.ColourConstant.value[3] = 1.0f;
+		op->op.args.ColourConstant.value[0] = node.m_value.X;
+		op->op.args.ColourConstant.value[1] = node.m_value.Y;
+		op->op.args.ColourConstant.value[2] = node.m_value.Z;
+		op->op.args.ColourConstant.value[3] = node.m_value.W;
 
 		result.op = op;
 	}
@@ -493,10 +493,10 @@ namespace mu
 		result.op = GenerateTableSwitch<NodeColourTable::Private, TCT_COLOUR, OP_TYPE::CO_SWITCH>(node,
 			[this](const NodeColourTable::Private& node, int colIndex, int row, ErrorLog* pErrorLog)
 			{
-				NodeColourConstantPtr pCell = new NodeColourConstant();
-				vec3<float> colour = node.m_pTable->GetPrivate()->m_rows[row].m_values[colIndex].m_colour;
-				pCell->SetValue(colour[0], colour[1], colour[2]);
-				return Generate(pCell);
+				NodeColourConstantPtr CellData = new NodeColourConstant();
+				FVector4f Colour = node.m_pTable->GetPrivate()->m_rows[row].m_values[colIndex].m_colour;
+				CellData->SetValue(Colour.X, Colour.Y, Colour.Z, Colour.W);
+				return Generate(CellData);
 			});
 	}
 
@@ -511,7 +511,7 @@ namespace mu
 
 		// Create a constant colour node
 		NodeColourConstantPtr pNode = new NodeColourConstant();
-		pNode->SetValue(1, 1, 0);
+		pNode->SetValue(1, 1, 0, 1);
 
 		FColorGenerationResult Result;
 		GenerateColor(Result, pNode);
diff --git a/Engine/Plugins/Experimental/Mutable/Source/MutableTools/Private/MuT/NodeColour.cpp b/Engine/Plugins/Experimental/Mutable/Source/MutableTools/Private/MuT/NodeColour.cpp
--- a/Engine/Plugins/Experimental/Mutable/Source/MutableTools/Private/MuT/NodeColour.cpp
+++ b/Engine/Plugins/Experimental/Mutable/Source/MutableTools/Private/MuT/NodeColour.cpp
@@ -30,7 +30,7 @@ namespace mu
 	//---------------------------------------------------------------------------------------------
 	void NodeColour::Serialise( const NodeColour* p, OutputArchive& arch )
 	{
-        uint32_t ver = 0;
+        uint32_t ver = 1;
 		arch << ver;
 
 		arch << uint32_t(p->Type);
@@ -43,14 +43,20 @@ namespace mu
 	{
         uint32_t ver;
 		arch >> ver;
-		check( ver == 0 );
+		check( ver <= 1 );
 
         uint32_t id;
 		arch >> id;
 
 		switch (id)
 		{
-		case 0 :  return NodeColourConstant::StaticUnserialise( arch ); break;
+		case 0 :
+			if (ver == 0)
+			{
+				return NodeColourConstant::OldStaticUnserialise(arch);
+			}
+			return NodeColourConstant::StaticUnserialise(arch);
+			break;
 		case 1 :  return NodeColourParameter::StaticUnserialise( arch ); break;
 		case 2 :  return NodeColourSampleImage::StaticUnserialise( arch ); break;
 		case 3 :  return NodeColourTable::StaticUnserialise( arch ); break;
diff --git a/Engine/Plugins/Experimental/Mutable/Source/MutableTools/Private/MuT/NodeColourConstant.cpp b/Engine/Plugins/Experimental/Mutable/Source/MutableTools/Private/MuT/NodeColourConstant.cpp
--- a/Engine/Plugins/Experimental/Mutable/Source/MutableTools/Private/MuT/NodeColourConstant.cpp
+++ b/Engine/Plugins/Experimental/Mutable/Source/MutableTools/Private/MuT/NodeColourConstant.cpp
@@ -25,6 +25,17 @@ namespace mu
 
 	MUTABLE_IMPLEMENT_NODE( NodeColourConstant, EType::Constant, Node, Node::EType::Colour)
 
+	NodeColourConstantPtr NodeColourConstant::OldStaticUnserialise(InputArchive& arch)
+	{
+		NodeColourConstantPtr pResult = new NodeColourConstant();
+		vec3<float> Value;
+
+		arch >> Value;
+		pResult->GetPrivate()->m_value = FVector4f(Value[0], Value[1], Value[2], 1.0f);
+
+		return pResult;
+	}
+
 
 	//---------------------------------------------------------------------------------------------
 	// Node Interface
@@ -53,7 +64,7 @@ namespace mu
 	//---------------------------------------------------------------------------------------------
 	// Own Interface
 	//---------------------------------------------------------------------------------------------
-	void NodeColourConstant::GetValue( float* pR, float *pG, float* pB ) const
+	void NodeColourConstant::GetValue( float* pR, float *pG, float* pB, float* pA) const
 	{
 		if (pR)
 		{
@@ -69,13 +80,18 @@ namespace mu
 		{
 			*pB = m_pD->m_value[2];
 		}
+
+		if (pA)
+		{
+			*pA = m_pD->m_value[3];
+		}
 	}
 
 
 	//---------------------------------------------------------------------------------------------
-	void NodeColourConstant::SetValue( float r, float g, float b )
+	void NodeColourConstant::SetValue( float r, float g, float b, float a)
 	{
-		m_pD->m_value = vec3<float>(r,g,b);
+		m_pD->m_value = FVector4f(r, g, b, a);
 	}
 
 
diff --git a/Engine/Plugins/Experimental/Mutable/Source/MutableTools/Private/MuT/NodeColourConstantPrivate.h b/Engine/Plugins/Experimental/Mutable/Source/MutableTools/Private/MuT/NodeColourConstantPrivate.h
--- a/Engine/Plugins/Experimental/Mutable/Source/MutableTools/Private/MuT/NodeColourConstantPrivate.h
+++ b/Engine/Plugins/Experimental/Mutable/Source/MutableTools/Private/MuT/NodeColourConstantPrivate.h
@@ -7,7 +7,6 @@
 
 #include "MuR/MutableMath.h"
 
-
 namespace mu
 {
 
@@ -21,20 +20,25 @@ namespace mu
 
 		static NODE_TYPE s_type;
 
-		vec3<float> m_value;
+		FVector4f m_value;
 
 		//!
 		void Serialise( OutputArchive& arch ) const
 		{
+			const uint32 Version = 0;
+			arch << Version;
 			arch << m_value;
 		}
 
 		//!
 		void Unserialise( InputArchive& arch )
 		{
+			uint32 Version;
+			arch >> Version;
+			check(Version==0)
+
 			arch >> m_value;
 		}
-
 	};
 
 }
diff --git a/Engine/Plugins/Experimental/Mutable/Source/MutableTools/Private/MuT/Table.cpp b/Engine/Plugins/Experimental/Mutable/Source/MutableTools/Private/MuT/Table.cpp
--- a/Engine/Plugins/Experimental/Mutable/Source/MutableTools/Private/MuT/Table.cpp
+++ b/Engine/Plugins/Experimental/Mutable/Source/MutableTools/Private/MuT/Table.cpp
@@ -150,7 +150,7 @@ namespace mu
 
 
 	//---------------------------------------------------------------------------------------------
-    void Table::SetCell( int column, uint32_t rowId, float r, float g, float b )
+    void Table::SetCell( int column, uint32_t rowId, float r, float g, float b, float a)
 	{
 		int row = m_pD->FindRow(rowId);
 		check( row>=0 );
@@ -158,7 +158,7 @@ namespace mu
 		check( column < (int)m_pD->m_rows[row].m_values.Num() );
 		check( m_pD->m_columns[column].m_type == TCT_COLOUR );
 
-		m_pD->m_rows[ row ].m_values[column].m_colour = vec3<float>( r, g, b );
+		m_pD->m_rows[ row ].m_values[column].m_colour = FVector4f(r, g, b, a);
 	}
 
 
diff --git a/Engine/Plugins/Experimental/Mutable/Source/MutableTools/Private/MuT/TablePrivate.h b/Engine/Plugins/Experimental/Mutable/Source/MutableTools/Private/MuT/TablePrivate.h
--- a/Engine/Plugins/Experimental/Mutable/Source/MutableTools/Private/MuT/TablePrivate.h
+++ b/Engine/Plugins/Experimental/Mutable/Source/MutableTools/Private/MuT/TablePrivate.h
@@ -28,7 +28,7 @@ namespace mu
 	{
 		// TODO: Union
 		float m_scalar;
-		vec3<float> m_colour;
+		FVector4f m_colour;
 		Ptr<ResourceProxy<Image>> m_pProxyImage;
 		MeshPtr m_pMesh;
 		string m_string;
@@ -52,11 +52,10 @@ namespace mu
 		TArray<TABLE_ROW> m_rows;
 		bool m_NoneOption = false;
 
-
 		//!
 		void Serialise( OutputArchive& arch ) const
 		{
-            uint32_t ver = 1;
+            uint32_t ver = 2;
 			arch << ver;
 
             arch << (uint32_t)m_columns.Num();
@@ -97,10 +96,8 @@ namespace mu
 			}
 
 			arch << m_NoneOption;
-
 		}
 
-
 		//!
 		void Unserialise( InputArchive& arch )
 		{
@@ -132,7 +129,21 @@ namespace mu
 					switch (m_columns[c].m_type)
 					{
 					case TCT_SCALAR:	arch >> v.m_scalar; break;
-					case TCT_COLOUR:	arch >> v.m_colour; break;
+					case TCT_COLOUR:
+					{
+						if (ver <= 1)
+						{
+							vec3<float> Value;
+							arch >> Value;
+
+							v.m_colour = FVector4f(Value[0], Value[1], Value[2], 1.0f);
+						}
+						else
+						{
+							arch >> v.m_colour;
+						}
+						break;
+					}
 					case TCT_MESH:		arch >> v.m_pMesh; break;
 					case TCT_IMAGE:		
 					{
diff --git a/Engine/Plugins/Experimental/Mutable/Source/MutableTools/Public/MuT/NodeColourConstant.h b/Engine/Plugins/Experimental/Mutable/Source/MutableTools/Public/MuT/NodeColourConstant.h
--- a/Engine/Plugins/Experimental/Mutable/Source/MutableTools/Public/MuT/NodeColourConstant.h
+++ b/Engine/Plugins/Experimental/Mutable/Source/MutableTools/Public/MuT/NodeColourConstant.h
@@ -31,6 +31,7 @@ namespace mu
 		void SerialiseWrapper(OutputArchive& arch) const override;
 		static void Serialise( const NodeColourConstant* pNode, OutputArchive& arch );
 		static NodeColourConstantPtr StaticUnserialise( InputArchive& arch );
+		static NodeColourConstantPtr OldStaticUnserialise( InputArchive& arch );
 
 		//-----------------------------------------------------------------------------------------
 		// Node Interface
@@ -50,10 +51,10 @@ namespace mu
 		//-----------------------------------------------------------------------------------------
 
 		//! Get the value that this node returns
-		void GetValue( float* pR, float* pG, float* pB ) const;
+		void GetValue( float* pR, float* pG, float* pB, float* pA) const;
 
 		//! Set the value to be returned by this node
-		void SetValue( float r, float g, float b );
+		void SetValue( float r, float g, float b, float a);
 
 		//-----------------------------------------------------------------------------------------
 		// Interface pattern
diff --git a/Engine/Plugins/Experimental/Mutable/Source/MutableTools/Public/MuT/Table.h b/Engine/Plugins/Experimental/Mutable/Source/MutableTools/Public/MuT/Table.h
--- a/Engine/Plugins/Experimental/Mutable/Source/MutableTools/Public/MuT/Table.h
+++ b/Engine/Plugins/Experimental/Mutable/Source/MutableTools/Public/MuT/Table.h
@@ -78,7 +78,7 @@ namespace mu
 
 		//!
         void SetCell( int column, uint32 rowId, float value );
-        void SetCell( int column, uint32 rowId, float r, float g, float b );
+        void SetCell( int column, uint32 rowId, float r, float g, float b, float a);
         void SetCell( int column, uint32 rowId, Image* pImage );
         void SetCell( int column, uint32 rowId, Mesh* pMesh );
         void SetCell( int column, uint32 rowId, const char* strValue );
