[Mutable] Add vectorized implementation of image raster operation bilinear pixel processor.

- This optimized version of the bilinear sampling operates with floating point values using generic vector operations. Initial benchmarking shows a 2x in performance.

- Added cvar mutable.UseProjectionVectorImpl to toggle the vectorized implementation.
- Removed smooth interpolation option form projector pixel processor since it is not used.

#rnx

diff --git a/Engine/Plugins/Experimental/Mutable/Source/MutableRuntime/Private/MuR/CodeRunner.cpp b/Engine/Plugins/Experimental/Mutable/Source/MutableRuntime/Private/MuR/CodeRunner.cpp
--- a/Engine/Plugins/Experimental/Mutable/Source/MutableRuntime/Private/MuR/CodeRunner.cpp
+++ b/Engine/Plugins/Experimental/Mutable/Source/MutableRuntime/Private/MuR/CodeRunner.cpp
@@ -97,6 +97,12 @@ static FAutoConsoleVariableRef CVarGlobalProjectionLodBias (
 	TEXT("Lod bias applied to the lod resulting form the best mip computation, only used if a min filter method different than None is used."),
 	ECVF_Default);
 
+bool bUseProjectionVectorImpl = false;
+static FAutoConsoleVariableRef CVarUseProjectionVectorImpl (
+	TEXT("mutable.UseProjectionVectorImpl"),
+	bUseProjectionVectorImpl,
+	TEXT("If set to true, enables the vectorized implementation of the projection pixel processing."),
+	ECVF_Default);
 }
 
 
@@ -4279,7 +4285,7 @@ namespace mu
 								FadeStartRad, FadeEndRad, FMath::Frac(Data.RasterMesh.MipValue),
 								layout, args.blockIndex,
 								CropMin, UncroppedSize,
-								&scratch);
+								&scratch, bUseProjectionVectorImpl);
 							break;
 
 						case PROJECTOR_TYPE::WRAPPING:
diff --git a/Engine/Plugins/Experimental/Mutable/Source/MutableRuntime/Private/MuR/OpImageProject.cpp b/Engine/Plugins/Experimental/Mutable/Source/MutableRuntime/Private/MuR/OpImageProject.cpp
--- a/Engine/Plugins/Experimental/Mutable/Source/MutableRuntime/Private/MuR/OpImageProject.cpp
+++ b/Engine/Plugins/Experimental/Mutable/Source/MutableRuntime/Private/MuR/OpImageProject.cpp
@@ -24,6 +24,8 @@
 #include "Misc/AssertionMacros.h"
 #include "Math/UnrealMathUtility.h"
 
+#include "Math/VectorRegister.h"
+
 namespace mu
 {
 
@@ -197,7 +199,7 @@ namespace
 
 
     //---------------------------------------------------------------------------------------------
-    template<int PIXEL_SIZE, bool bUseSmoothInterpolator = false, bool HAS_MASK = true>
+    template<int PIXEL_SIZE, bool bUseVectorImpl = false, bool HAS_MASK = true>
     class RasterProjectedPixelProcessorLinear_UBYTE
     {
     public:
@@ -229,8 +231,8 @@ namespace
 				Source1DataPtr = nullptr;
 			}
 
-			bIsRGBFadingEnabled = bInIsRGBFadingEnabled;
-			bIsAlphaFadingEnabled = bInIsAlphaFadingEnabled;
+			RGBFadingEnabledValue = -static_cast<int8>(bInIsRGBFadingEnabled);
+			AlphaFadingEnabledValue = -static_cast<int8>(bInIsAlphaFadingEnabled);
 
 			const float FadeStartCos = FMath::Cos(FadeStart);
             FadeEndCos = FMath::Cos(FadeEnd);
@@ -248,8 +250,20 @@ namespace
             check(GetImageFormatData(Source->GetFormat() ).m_bytesPerBlock == PIXEL_SIZE);
         }
 
+		inline void ProcessPixel(uint8* BufferPosPtr, float Varying[4]) const
+		{
+			if constexpr (bUseVectorImpl)
+			{
+				ProcessPixelVectorImpl(BufferPosPtr, Varying);
+			}
+			else
+			{
+				ProcessPixelImpl(BufferPosPtr, Varying);
+			}
+		}
+
         //-----------------------------------------------------------------------------------------
-        inline void ProcessPixel( uint8* BufferPosPtr, float Varying[4] ) const
+        inline void ProcessPixelImpl( uint8* BufferPosPtr, float Varying[4] ) const
         {
 			if (!Source0DataPtr)
 			{
@@ -294,14 +308,7 @@ namespace
 					{
 						auto ComputeInterpolator = [](float T) -> uint16
 						{
-							if constexpr (bUseSmoothInterpolator)
-							{
-								return static_cast<uint16>(255.0f * T * T * (3.0f - 2.0f * T));
-							}
-							else
-							{
-								return static_cast<uint16>(255.0f * T);
-							}
+							return static_cast<uint16>(255.0f * T);
 						};
 
 						const FVector2f SizeF(Size.X, Size.Y);
@@ -375,11 +382,13 @@ namespace
 					}
 
 					static_assert(PIXEL_SIZE <= 4);
+
+					// FadingEnabledValue will be -1 if enabled and 0 otherwise.
 					const uint16 FadingMask[4] = { 
-						static_cast<uint16>(bIsRGBFadingEnabled), 
-						static_cast<uint16>(bIsRGBFadingEnabled), 
-						static_cast<uint16>(bIsRGBFadingEnabled), 
-						static_cast<uint16>(bIsAlphaFadingEnabled)};
+						static_cast<uint16>(-RGBFadingEnabledValue), 
+						static_cast<uint16>(-RGBFadingEnabledValue), 
+						static_cast<uint16>(-RGBFadingEnabledValue), 
+						static_cast<uint16>(-AlphaFadingEnabledValue)};
 
 					for (int32 I = 0; I < PIXEL_SIZE; ++I)
 					{
@@ -390,6 +399,166 @@ namespace
             }
         }
 
+        //-----------------------------------------------------------------------------------------
+        inline void ProcessPixelVectorImpl( uint8* BufferPosPtr, float Varying[4] ) const
+        {
+			if (!Source0DataPtr)
+			{
+				return;
+			}
+
+			const bool bDepthClamp = (Varying[2] < 0.0f) | (Varying[2] > 1.0f);
+            const float AngleCos = Varying[3];
+
+			float Factor = FMath::Clamp((AngleCos - FadeEndCos) * OneOverFadeRangeTimes255, 0.0f, 255.0f);
+			Factor = Factor < 255.0f ? Factor : static_cast<float>(!bDepthClamp) * 255.0f;
+
+			float MaskFactor = AngleCos > FadeEndCos ? 255.0f : 0.0f;
+
+			if constexpr (HAS_MASK)
+			{
+				// Only read from memory if needed.
+				if (Factor > UE_SMALL_NUMBER)
+				{
+					MaskFactor = ((MaskFactor > 0.0f) && MaskDataPtr) 
+							? static_cast<float>(MaskDataPtr[(BufferPosPtr - TargetDataPtr) / PIXEL_SIZE]) 
+							: MaskFactor;
+					Factor = (Factor * MaskFactor) * (1.0f/255.0f);
+				}
+			}
+
+            if (Factor > UE_SMALL_NUMBER)
+            {
+				const FVector2f Uv = { Varying[0], Varying[1] };
+
+                if ((Uv.X >= 0.0f) & (Uv.X < 1.0f) & (Uv.Y >= 0.0f) & (Uv.Y < 1.0f))
+                {
+					using namespace GlobalVectorConstants;
+
+					const VectorRegister4Int SizeI = 
+						MakeVectorRegisterInt(Source0SizeX, Source0SizeY, Source1SizeX, Source1SizeY);
+
+					const VectorRegister4Float SizeF = VectorIntToFloat(SizeI);
+
+					const VectorRegister4Float CoordsF = VectorMin(VectorMax(
+						VectorMultiplyAdd(MakeVectorRegister(Uv.X, Uv.Y, Uv.X, Uv.Y), SizeF, FloatMinusOneHalf),
+						FloatZero), 
+						VectorSubtract(SizeF, FloatOne));
+
+					const VectorRegister4Float Frac = VectorSubtract(CoordsF, VectorFloor(CoordsF));
+					
+					const VectorRegister4Int CoordsI = VectorFloatToInt(CoordsF);
+					const VectorRegister4Int CoordsIPlusOne = VectorIntMin(
+							VectorIntAdd(CoordsI, IntOne), VectorIntSubtract(SizeI, IntOne));
+
+					alignas(VectorRegister4Int) int32 CoordsIData[4];
+					VectorIntStoreAligned(CoordsI, CoordsIData);
+
+					alignas(VectorRegister4Int) int32 CoordsIPlusOneData[4];
+					VectorIntStoreAligned(CoordsIPlusOne, CoordsIPlusOneData);
+
+					uint8 const * const Pixel000Ptr = Source0DataPtr + (CoordsIData[1]        * Source0SizeX + CoordsIData[0])        * PIXEL_SIZE;
+					uint8 const * const Pixel010Ptr = Source0DataPtr + (CoordsIData[1]        * Source0SizeX + CoordsIPlusOneData[0]) * PIXEL_SIZE;
+					uint8 const * const Pixel001Ptr = Source0DataPtr + (CoordsIPlusOneData[1] * Source0SizeX + CoordsIData[0])        * PIXEL_SIZE;
+					uint8 const * const Pixel011Ptr = Source0DataPtr + (CoordsIPlusOneData[1] * Source0SizeX + CoordsIPlusOneData[0]) * PIXEL_SIZE;
+
+					uint8 const * const Pixel100Ptr = Source1DataPtr + (CoordsIData[3]        * Source1SizeX + CoordsIData[2])        * PIXEL_SIZE;
+					uint8 const * const Pixel110Ptr = Source1DataPtr + (CoordsIData[3]        * Source1SizeX + CoordsIPlusOneData[2]) * PIXEL_SIZE;
+					uint8 const * const Pixel101Ptr = Source1DataPtr + (CoordsIPlusOneData[3] * Source1SizeX + CoordsIData[2])        * PIXEL_SIZE;
+					uint8 const * const Pixel111Ptr = Source1DataPtr + (CoordsIPlusOneData[3] * Source1SizeX + CoordsIPlusOneData[2]) * PIXEL_SIZE;
+
+					auto LoadPixel = [](const uint8* Ptr) -> VectorRegister4Float
+					{
+						constexpr VectorRegister4Float OneOver255 = 
+							MakeVectorRegisterFloatConstant(1.0f/255.0f, 1.0f/255.0f, 1.0f/255.0f, 1.0f/255.0f);
+
+						if constexpr (PIXEL_SIZE == 4)
+						{
+							const uint32 PackedData = *reinterpret_cast<const uint32*>(Ptr);
+							return VectorMultiply(MakeVectorRegister(
+								static_cast<float>((PackedData >> (8 * 0)) & 0xFF),
+								static_cast<float>((PackedData >> (8 * 1)) & 0xFF),
+								static_cast<float>((PackedData >> (8 * 2)) & 0xFF),
+								static_cast<float>((PackedData >> (8 * 3)) & 0xFF)), OneOver255);
+						}
+						else
+						{
+							alignas(VectorRegister4Float) float PixelData[4];
+							for (int32 C = 0; C < PIXEL_SIZE; ++C)
+							{
+								PixelData[C] = static_cast<float>(Ptr[C]);
+							}
+
+							return VectorMultiply(VectorLoadAligned(&(PixelData[0])), OneOver255);
+						}
+					};
+
+					const VectorRegister4Float Pixel000 = LoadPixel(Pixel000Ptr);
+					const VectorRegister4Float Pixel010 = LoadPixel(Pixel010Ptr);
+					const VectorRegister4Float Pixel001 = LoadPixel(Pixel001Ptr);
+					const VectorRegister4Float Pixel011 = LoadPixel(Pixel011Ptr);
+
+					const VectorRegister4Float Pixel100 = LoadPixel(Pixel100Ptr);
+					const VectorRegister4Float Pixel110 = LoadPixel(Pixel110Ptr);
+					const VectorRegister4Float Pixel101 = LoadPixel(Pixel101Ptr);
+					const VectorRegister4Float Pixel111 = LoadPixel(Pixel111Ptr);
+					
+					const VectorRegister4Float Frac0X = VectorReplicate(Frac, 0);
+					const VectorRegister4Float Frac0Y = VectorReplicate(Frac, 1);
+
+					// Bilerp image 0
+					VectorRegister4Float Sample0 = VectorMultiplyAdd(Frac0X, VectorSubtract(Pixel010, Pixel000), Pixel000);
+					Sample0 = VectorMultiplyAdd(
+						Frac0Y, 
+						VectorSubtract(
+							VectorMultiplyAdd(Frac0X, VectorSubtract(Pixel011, Pixel001), Pixel001), 
+							Sample0), 
+						Sample0);
+
+					const VectorRegister4Float Frac1X = VectorReplicate(Frac, 2);
+					const VectorRegister4Float Frac1Y = VectorReplicate(Frac, 3);
+
+					// Bilerp image 1
+					VectorRegister4Float Sample1 = VectorMultiplyAdd(Frac1X, VectorSubtract(Pixel110, Pixel100), Pixel100);
+					Sample1 = VectorMultiplyAdd(
+						Frac1Y, 
+						VectorSubtract(
+							VectorMultiplyAdd(Frac1X, VectorSubtract(Pixel111, Pixel101), Pixel101), 
+							Sample1), 
+						Sample1);
+
+					VectorRegister4Float Result = VectorMultiply(
+						// Lerp mip
+						VectorMultiplyAdd(VectorSetFloat1(MipInterpolationFactor), VectorSubtract(Sample1, Sample0), Sample0),
+						// Fadding factor
+						VectorSelect(
+							// The selction mask needs to be 0xFFFFFFFF if enabled 0 otherwise. FadingEnabledValue will be -1
+							// if enabled and 0 otherwise.
+							VectorCastIntToFloat(MakeVectorRegisterInt( 
+								static_cast<int32>(RGBFadingEnabledValue), 
+								static_cast<int32>(RGBFadingEnabledValue), 
+								static_cast<int32>(RGBFadingEnabledValue), 
+								static_cast<int32>(AlphaFadingEnabledValue))), 
+							VectorSetFloat1(Factor), 
+							VectorSetFloat1(MaskFactor)));
+
+					if constexpr (PIXEL_SIZE == 4)
+					{
+						VectorStoreByte4(Result, BufferPosPtr);
+						VectorResetFloatRegisters();
+					}
+					else
+					{
+						const AlignedFloat4 ResultData(VectorMin(VectorMax(Result, FloatZero), Float255));
+						for (int32 C = 0; C < PIXEL_SIZE; ++C)
+						{
+							BufferPosPtr[C] = static_cast<uint8>(ResultData[C]);		
+						}				
+					}
+                }
+            }
+        }
+
         const uint8* Source0DataPtr;
         const uint8* Source1DataPtr;
         uint16 Source0SizeX, Source0SizeY;
@@ -402,7 +571,7 @@ namespace
 		// precomputation of 255 / (FadeStartCos - FadeEndCos) 
 		float OneOverFadeRangeTimes255;
 
-		bool bIsRGBFadingEnabled, bIsAlphaFadingEnabled;
+		int8 RGBFadingEnabledValue, AlphaFadingEnabledValue;
 
         const uint8* TargetDataPtr;
         const uint8* MaskDataPtr;
@@ -784,7 +953,7 @@ void ImageRasterProjectedPlanar( const Mesh* pMesh,
 	float fadeStart, float fadeEnd, float MipInterpolationFactor,
 	int layout, int block,
 	UE::Math::TIntVector2<uint16> CropMin, UE::Math::TIntVector2<uint16> UncroppedSize,
-	SCRATCH_IMAGE_PROJECT* scratch )
+	SCRATCH_IMAGE_PROJECT* scratch, bool bUseVectorImplementation)
 {
     check( !pMask || pMask->GetSizeX() == pImage->GetSizeX() );
     check( !pMask || pMask->GetSizeY() == pImage->GetSizeY() );
@@ -792,7 +961,6 @@ void ImageRasterProjectedPlanar( const Mesh* pMesh,
 
 	MUTABLE_CPUPROFILER_SCOPE(ImageProject);
 
-	constexpr bool bUseSmoothInterpolator = false;
     if ( pMesh->m_staticFormatFlags & (1<<SMF_PROJECT) )
     {
         // Mesh-optimised version
@@ -809,12 +977,26 @@ void ImageRasterProjectedPlanar( const Mesh* pMesh,
 			}
 			else if (SamplingMethod == ESamplingMethod::BiLinear)
 			{
-				const RasterProjectedPixelProcessorLinear_UBYTE<3, bUseSmoothInterpolator> pixelProc(pSource,
-					pImage->GetData(),
-					pMask ? pMask->GetData() : nullptr,
-					bIsRGBFadingEnabled, bIsAlphaFadingEnabled,
-					fadeStart, fadeEnd, MipInterpolationFactor);
-				ImageRasterProjected_Optimised(pMesh, pImage, pixelProc, fadeEnd, CropMin, UncroppedSize, scratch);
+				if (bUseVectorImplementation)
+				{
+					constexpr bool bUseVectorImpl = true;
+					const RasterProjectedPixelProcessorLinear_UBYTE<3, bUseVectorImpl> pixelProc(pSource,
+						pImage->GetData(),
+						pMask ? pMask->GetData() : nullptr,
+						bIsRGBFadingEnabled, bIsAlphaFadingEnabled,
+						fadeStart, fadeEnd, MipInterpolationFactor);
+					ImageRasterProjected_Optimised(pMesh, pImage, pixelProc, fadeEnd, CropMin, UncroppedSize, scratch);
+				}
+				else
+				{
+					constexpr bool bUseVectorImpl = false;
+					const RasterProjectedPixelProcessorLinear_UBYTE<3, bUseVectorImpl> pixelProc(pSource,
+						pImage->GetData(),
+						pMask ? pMask->GetData() : nullptr,
+						bIsRGBFadingEnabled, bIsAlphaFadingEnabled,
+						fadeStart, fadeEnd, MipInterpolationFactor);
+					ImageRasterProjected_Optimised(pMesh, pImage, pixelProc, fadeEnd, CropMin, UncroppedSize, scratch);
+				}
 			}
 			else
 			{
@@ -838,12 +1020,26 @@ void ImageRasterProjectedPlanar( const Mesh* pMesh,
 				}
 				else if (SamplingMethod == ESamplingMethod::BiLinear)
 				{
-					const RasterProjectedPixelProcessorLinear_UBYTE<4, bUseSmoothInterpolator, true> pixelProc(pSource,
-						pImage->GetData(),
-						pMask->GetData(),
-						bIsRGBFadingEnabled, bIsAlphaFadingEnabled,
-						fadeStart, fadeEnd, MipInterpolationFactor);
-					ImageRasterProjected_Optimised(pMesh, pImage, pixelProc, fadeEnd, CropMin, UncroppedSize, scratch);
+					if (bUseVectorImplementation)
+					{
+						constexpr bool bUseVectorImpl = true;
+						const RasterProjectedPixelProcessorLinear_UBYTE<4, bUseVectorImpl, true> pixelProc(pSource,
+							pImage->GetData(),
+							pMask->GetData(),
+							bIsRGBFadingEnabled, bIsAlphaFadingEnabled,
+							fadeStart, fadeEnd, MipInterpolationFactor);
+						ImageRasterProjected_Optimised(pMesh, pImage, pixelProc, fadeEnd, CropMin, UncroppedSize, scratch);
+					}
+					else
+					{
+						constexpr bool bUseVectorImpl = false;
+						const RasterProjectedPixelProcessorLinear_UBYTE<4, bUseVectorImpl, true> pixelProc(pSource,
+							pImage->GetData(),
+							pMask->GetData(),
+							bIsRGBFadingEnabled, bIsAlphaFadingEnabled,
+							fadeStart, fadeEnd, MipInterpolationFactor);
+						ImageRasterProjected_Optimised(pMesh, pImage, pixelProc, fadeEnd, CropMin, UncroppedSize, scratch);
+					}
 				}
 				else
 				{
@@ -863,12 +1059,26 @@ void ImageRasterProjectedPlanar( const Mesh* pMesh,
 				}
 				else if (SamplingMethod == ESamplingMethod::BiLinear)
 				{
-					const RasterProjectedPixelProcessorLinear_UBYTE<4, bUseSmoothInterpolator, false> pixelProc(pSource,
-						pImage->GetData(),
-						nullptr,
-						bIsRGBFadingEnabled, bIsAlphaFadingEnabled,
-						fadeStart, fadeEnd, MipInterpolationFactor);
-					ImageRasterProjected_Optimised(pMesh, pImage, pixelProc, fadeEnd, CropMin, UncroppedSize, scratch);
+					if (bUseVectorImplementation)
+					{
+						constexpr bool bUseVectorImpl = true;
+						const RasterProjectedPixelProcessorLinear_UBYTE<4, bUseVectorImpl, false> pixelProc(pSource,
+							pImage->GetData(),
+							nullptr,
+							bIsRGBFadingEnabled, bIsAlphaFadingEnabled,
+							fadeStart, fadeEnd, MipInterpolationFactor);
+						ImageRasterProjected_Optimised(pMesh, pImage, pixelProc, fadeEnd, CropMin, UncroppedSize, scratch);
+					}
+					else
+					{
+						constexpr bool bUseVectorImpl = false;
+						const RasterProjectedPixelProcessorLinear_UBYTE<4, bUseVectorImpl, false> pixelProc(pSource,
+							pImage->GetData(),
+							nullptr,
+							bIsRGBFadingEnabled, bIsAlphaFadingEnabled,
+							fadeStart, fadeEnd, MipInterpolationFactor);
+						ImageRasterProjected_Optimised(pMesh, pImage, pixelProc, fadeEnd, CropMin, UncroppedSize, scratch);
+					}
 				}
 				else
 				{
@@ -889,12 +1099,26 @@ void ImageRasterProjectedPlanar( const Mesh* pMesh,
 			}
 			else if (SamplingMethod == ESamplingMethod::BiLinear)
 			{
-				const RasterProjectedPixelProcessorLinear_UBYTE<1, bUseSmoothInterpolator> pixelProc(pSource,
-					pImage->GetData(),
-					pMask ? pMask->GetData() : nullptr,
-					bIsRGBFadingEnabled, bIsAlphaFadingEnabled,
-					fadeStart, fadeEnd, MipInterpolationFactor);
-				ImageRasterProjected_Optimised(pMesh, pImage, pixelProc, fadeEnd, CropMin, UncroppedSize, scratch);
+				if (bUseVectorImplementation)
+				{
+					constexpr bool bUseVectorImpl = true;
+					const RasterProjectedPixelProcessorLinear_UBYTE<1, bUseVectorImpl> pixelProc(pSource,
+						pImage->GetData(),
+						pMask ? pMask->GetData() : nullptr,
+						bIsRGBFadingEnabled, bIsAlphaFadingEnabled,
+						fadeStart, fadeEnd, MipInterpolationFactor);
+					ImageRasterProjected_Optimised(pMesh, pImage, pixelProc, fadeEnd, CropMin, UncroppedSize, scratch);
+				}
+				else
+				{
+					constexpr bool bUseVectorImpl = false;
+					const RasterProjectedPixelProcessorLinear_UBYTE<1, bUseVectorImpl> pixelProc(pSource,
+						pImage->GetData(),
+						pMask ? pMask->GetData() : nullptr,
+						bIsRGBFadingEnabled, bIsAlphaFadingEnabled,
+						fadeStart, fadeEnd, MipInterpolationFactor);
+					ImageRasterProjected_Optimised(pMesh, pImage, pixelProc, fadeEnd, CropMin, UncroppedSize, scratch);
+				}
 			}
 			else
 			{
diff --git a/Engine/Plugins/Experimental/Mutable/Source/MutableRuntime/Private/MuR/OpImageProject.h b/Engine/Plugins/Experimental/Mutable/Source/MutableRuntime/Private/MuR/OpImageProject.h
--- a/Engine/Plugins/Experimental/Mutable/Source/MutableRuntime/Private/MuR/OpImageProject.h
+++ b/Engine/Plugins/Experimental/Mutable/Source/MutableRuntime/Private/MuR/OpImageProject.h
@@ -24,7 +24,7 @@ namespace mu
 		float fadeStart, float fadeEnd, float MipInterpolationFactor,
 		int layout, int block,
 		UE::Math::TIntVector2<uint16> CropMin, UE::Math::TIntVector2<uint16> UncroppedSize,
-		SCRATCH_IMAGE_PROJECT* scratch );
+		SCRATCH_IMAGE_PROJECT* scratch, bool bUseVectorImplementation = false);
 
     extern void ImageRasterProjectedCylindrical( const Mesh* pMesh, Image* pTargetImage,
 		const Image* pSource, const Image* pMask,
