[mutable] Implemented Tiling LLM data recollection
- Implements LLM testing for Tiling instance generation operation
- Implements generic way of running LLM data recollection over the MutableProgram targeted to be ran .
- Implemented alternative way of handling LLM vs non-LLM executions of mutable program due standard way being buggy.
- Added 2 new target.cs an build.cs to handle the compilation of MutableCompile and MutableTests in LLM mode. Previous target files represent non-LLM programs.
- Enabled all LLM tests since dependency for "heaptrack_print" has been removed (only available before in linux)
- Added debug log message (only visible in IDE log window).
- Added capability for the MTU to upload data to "Development" database if running in local machine.
- Updated Mutable core to take care of invoking the update of the LLM data recollector. This is done because our program does not have "frames" and therefore the LLM system does not recollect data if not directly asked to do so.
- Updated MutableTest and MutableCompile to accept the activation of csv generation argument (-llmcsv)
- Added variable to control if Bulk optimized tests are ran or not. Not setable from argument. Usefull for local testing and debugging.

diff --git a/Engine/Plugins/Experimental/Mutable/Source/MutableRuntime/Private/MuR/SystemPrivate.h b/Engine/Plugins/Experimental/Mutable/Source/MutableRuntime/Private/MuR/SystemPrivate.h
--- a/Engine/Plugins/Experimental/Mutable/Source/MutableRuntime/Private/MuR/SystemPrivate.h
+++ b/Engine/Plugins/Experimental/Mutable/Source/MutableRuntime/Private/MuR/SystemPrivate.h
@@ -18,6 +18,20 @@ namespace mu
 {
 	class ExtensionDataStreamer;
 
+	// Call the tick of the LLM system (we do this to simulate a frame since the LLM system is not entirelly designed to run over a program)
+	inline void UpdateLLMStats()
+	{
+		// This code will only be compiled (and ran) if the global definition to enable LLM tracking is set to 1 for the host program
+		// Ex : 			GlobalDefinitions.Add("LLM_ENABLED_IN_CONFIG=1");
+#if ENABLE_LOW_LEVEL_MEM_TRACKER
+		FLowLevelMemTracker& MemTracker = FLowLevelMemTracker::Get();
+		if (MemTracker.IsEnabled())
+		{
+			MemTracker.UpdateStatsPerFrame();
+		}
+#endif
+	}
+
 	constexpr uint64 AllParametersMask = TNumericLimits<uint64> ::Max();
 
 
@@ -821,6 +835,8 @@ namespace mu
 		{
 			check(at.At < m_resources.size_code());
 			m_resources[at] = TPair<int, Ptr<const RefCounted>>(2, v);
+
+			mu::UpdateLLMStats();
 		}
 
 		void SetMesh(FCacheAddress at, Ptr<const Mesh> v)
@@ -833,12 +849,16 @@ namespace mu
 
 			check(at.At < m_resources.size_code());
 			m_resources[at] = TPair<int, Ptr<const RefCounted>>(2, v);
+
+			mu::UpdateLLMStats();
 		}
 
 		void SetLayout(FCacheAddress at, Ptr<const Layout> v)
 		{
 			check(at.At < m_resources.size_code());
 			m_resources[at] = TPair<int, Ptr<const RefCounted>>(1, v);
+
+			mu::UpdateLLMStats();
 		}
 
 		void SetString(FCacheAddress at, Ptr<const String> v)
