[Mutable] Fix mutable.StreamingCacheMemory 0 case.

#rnx

diff --git a/Engine/Plugins/Experimental/Mutable/Source/CustomizableObject/Private/MuCO/CustomizableObjectSystem.cpp b/Engine/Plugins/Experimental/Mutable/Source/CustomizableObject/Private/MuCO/CustomizableObjectSystem.cpp
--- a/Engine/Plugins/Experimental/Mutable/Source/CustomizableObject/Private/MuCO/CustomizableObjectSystem.cpp
+++ b/Engine/Plugins/Experimental/Mutable/Source/CustomizableObject/Private/MuCO/CustomizableObjectSystem.cpp
@@ -85,7 +85,7 @@ static TAutoConsoleVariable<int32> CVarStreamingMemory(
 #else
 	(12 * 1024),
 #endif
-	TEXT("Limit the amount of memory (in KB) to use to cache streaming data when building characters. More memory reduces the object construction time. Defaults: Desktop = 12288, Others = 3072"),
+	TEXT("Limit the amount of memory (in KB) to use to cache streaming data when building characters. More memory reduces the object construction time. 0 means no restriction. Defaults: Desktop = 12288, Others = 3072"),
 	ECVF_Scalability);
 
 
diff --git a/Engine/Plugins/Experimental/Mutable/Source/MutableRuntime/Private/MuR/System.cpp b/Engine/Plugins/Experimental/Mutable/Source/MutableRuntime/Private/MuR/System.cpp
--- a/Engine/Plugins/Experimental/Mutable/Source/MutableRuntime/Private/MuR/System.cpp
+++ b/Engine/Plugins/Experimental/Mutable/Source/MutableRuntime/Private/MuR/System.cpp
@@ -885,7 +885,6 @@ namespace mu
     uint64 FModelCache::EnsureCacheBelowBudget( uint64 additionalMemory,
 												TFunctionRef<bool(const Model*,int)> isRomLockedFunc )
     {
-
 		uint64 totalMemory = 0;
         for (FModelCacheEntry& m : m_cachePerModel)
         {
@@ -903,6 +902,13 @@ namespace mu
 			}
         }
 
+    	// If budget is zero, we don't unload anything here, and we assume it is managed
+        // somewhere else.
+  		if (!m_romBudget)
+		{
+			return totalMemory;
+		}
+    	
         if (totalMemory>0)
         {
             totalMemory += additionalMemory;
