[Mutable] Small fix to default LOD Management. Replace IsLocallyControlled with IsPlayerControlled, since the first one returns true for both player and bots.
- Fix wrong MeshID invalidating generated meshes.
- Add CVar to ignore min/max LODs to generate (Mutable.IgnoreMinMaxLOD)

#rnx

diff --git a/Engine/Plugins/Experimental/Mutable/Source/CustomizableObject/Private/MuCO/CustomizableInstanceLODManagement.cpp b/Engine/Plugins/Experimental/Mutable/Source/CustomizableObject/Private/MuCO/CustomizableInstanceLODManagement.cpp
--- a/Engine/Plugins/Experimental/Mutable/Source/CustomizableObject/Private/MuCO/CustomizableInstanceLODManagement.cpp
+++ b/Engine/Plugins/Experimental/Mutable/Source/CustomizableObject/Private/MuCO/CustomizableInstanceLODManagement.cpp
@@ -329,7 +329,7 @@ void UCustomizableInstanceLODManagement::UpdateInstanceDistsAndLODs()
 					AActor* ParentParentActor = AttachParentParentComponent ? AttachParentParentComponent->GetOwner() : nullptr;
 
 					const APawn* Pawn = Cast<APawn>(ParentParentActor);
-					if (Pawn && Pawn->IsLocallyControlled())
+					if (Pawn && Pawn->IsPlayerControlled())
 					{
 						COI->SetMinSquareDistToPlayer(-1.f);
 					}
diff --git a/Engine/Plugins/Experimental/Mutable/Source/CustomizableObject/Private/MuCO/CustomizableObjectInstance.cpp b/Engine/Plugins/Experimental/Mutable/Source/CustomizableObject/Private/MuCO/CustomizableObjectInstance.cpp
--- a/Engine/Plugins/Experimental/Mutable/Source/CustomizableObject/Private/MuCO/CustomizableObjectInstance.cpp
+++ b/Engine/Plugins/Experimental/Mutable/Source/CustomizableObject/Private/MuCO/CustomizableObjectInstance.cpp
@@ -4956,7 +4956,7 @@ bool UCustomizableInstancePrivateData::BuildOrCopyRenderData(const TSharedPtr<FM
 
 		TRACE_CPUPROFILER_EVENT_SCOPE_TEXT(*FString::Printf(TEXT("BuildOrCopyRenderData_BuildData: LOD %d"), LODIndex));
 
-		SetLastMeshId(Component.Id, LODIndex, Component.Mesh->GetId());
+		SetLastMeshId(Component.Id, LODIndex, Component.MeshID);
 
 		UnrealConversionUtils::SetupRenderSections(
 			SkeletalMesh,
@@ -6265,6 +6265,13 @@ int32 UCustomizableObjectInstance::GetCurrentMaxLOD() const
 	return CurrentMaxLOD;
 }
 
+#if !UE_BUILD_SHIPPING
+static bool bIgnoreMinMaxLOD = false;
+FAutoConsoleVariableRef CVarMutableIgnoreMinMaxLOD(
+	TEXT("Mutable.IgnoreMinMaxLOD"),
+	bIgnoreMinMaxLOD,
+	TEXT("The limits on the number of LODs to generate will be ignored."));
+#endif
 
 void UCustomizableObjectInstance::SetRequestedLODs(int32 InMinLOD, int32 InMaxLOD, const TArray<uint16>& InRequestedLODsPerComponent)
 {
@@ -6274,6 +6281,15 @@ void UCustomizableObjectInstance::SetRequestedLODs(int32 InMinLOD, int32 InMaxLO
 		return;
 	}
 
+#if !UE_BUILD_SHIPPING
+	// Ignore Min/Max LOD limits. Mainly used for debug
+	if (!bIgnoreMinMaxLOD)
+	{
+		InMinLOD = 0;
+		InMaxLOD = MAX_int32;
+	}
+#endif
+
 	// Clamp Min LOD
 	InMinLOD = FMath::Min(FMath::Max(InMinLOD, static_cast<int32>(PrivateData->FirstLODAvailable)), PrivateData->FirstLODAvailable + PrivateData->NumMaxLODsToStream);
 
diff --git a/Engine/Plugins/Experimental/Mutable/Source/CustomizableObject/Private/MuCO/CustomizableSkeletalComponent.cpp b/Engine/Plugins/Experimental/Mutable/Source/CustomizableObject/Private/MuCO/CustomizableSkeletalComponent.cpp
--- a/Engine/Plugins/Experimental/Mutable/Source/CustomizableObject/Private/MuCO/CustomizableSkeletalComponent.cpp
+++ b/Engine/Plugins/Experimental/Mutable/Source/CustomizableObject/Private/MuCO/CustomizableSkeletalComponent.cpp
@@ -258,11 +258,11 @@ void UCustomizableSkeletalComponent::UpdateDistFromComponentToPlayer(const AActo
 
 				if (ViewCenter && ViewCenter->IsValidLowLevel())
 				{
-					APawn* pawn = Cast<APawn>(ParentActor);
-					bool bIsPawn = pawn ? pawn->IsLocallyControlled() : false;
-					CustomizableObjectInstance->SetIsPlayerOrNearIt(bIsPawn);
+					APawn* Pawn = Cast<APawn>(ParentActor);
+					bool bIsPlayer = Pawn ? Pawn->IsPlayerControlled() : false;
+					CustomizableObjectInstance->SetIsPlayerOrNearIt(bIsPlayer);
 
-					if (bIsPawn)
+					if (bIsPlayer)
 					{
 						SquareDist = -0.01f; // Negative value to give the player character more priority than any other character
 					}
