[mutable] Fixed introduced in CL-25352243

#rnx

diff --git a/Engine/Plugins/Experimental/Mutable/Source/CustomizableObjectEditor/Private/MuCOE/GenerateMutableSource/GenerateMutableSourceSurface.cpp b/Engine/Plugins/Experimental/Mutable/Source/CustomizableObjectEditor/Private/MuCOE/GenerateMutableSource/GenerateMutableSourceSurface.cpp
--- a/Engine/Plugins/Experimental/Mutable/Source/CustomizableObjectEditor/Private/MuCOE/GenerateMutableSource/GenerateMutableSourceSurface.cpp
+++ b/Engine/Plugins/Experimental/Mutable/Source/CustomizableObjectEditor/Private/MuCOE/GenerateMutableSource/GenerateMutableSourceSurface.cpp
@@ -1027,29 +1027,29 @@ mu::NodeSurfacePtr GenerateMutableSourceSurface(const UEdGraphPin * Pin, FMutabl
 			{
 				mu::NodeMeshPtr MeshPtr = AddMeshNode;
 
-				const TArray<UCustomizableObjectLayout*> Layouts = TypedNodeExt->GetLayouts();
-
-				if (Layouts.Num())
-				{
-					//TODO: Implement support for multiple UV channels (e.g. Add warning for vertices which have a block in a layout but not in the other)
-					mu::NodeMeshFragmentPtr MeshFrag = new mu::NodeMeshFragment();
-				
-					MeshFrag->SetMesh(MeshPtr);
-					MeshFrag->SetLayoutOrGroup(0);
-					MeshFrag->SetFragmentType(mu::NodeMeshFragment::FT_LAYOUT_BLOCKS);
-					MeshFrag->SetBlockCount(Layouts[0]->Blocks.Num());
-				
-					for (int i = 0; i < Layouts[0]->Blocks.Num(); ++i)
-					{
-						MeshFrag->SetBlock(i, i);
-					}
-				
-					MeshPtr = MeshFrag;
-				}
-				else
-				{
-					GenerationContext.Compiler->CompilerLog(LOCTEXT("ExtendMaterialLayoutMissing","Skeletal Mesh without Layout Node linked to an Extend Material. A 4x4 layout will be added as default layout."), Node);
-				}
+				//const TArray<UCustomizableObjectLayout*> Layouts = TypedNodeExt->GetLayouts();
+				//
+				//if (Layouts.Num())
+				//{
+				//	//TODO: Implement support for multiple UV channels (e.g. Add warning for vertices which have a block in a layout but not in the other)
+				//	mu::NodeMeshFragmentPtr MeshFrag = new mu::NodeMeshFragment();
+				//
+				//	MeshFrag->SetMesh(MeshPtr);
+				//	MeshFrag->SetLayoutOrGroup(0);
+				//	MeshFrag->SetFragmentType(mu::NodeMeshFragment::FT_LAYOUT_BLOCKS);
+				//	MeshFrag->SetBlockCount(Layouts[0]->Blocks.Num());
+				//
+				//	for (int i = 0; i < Layouts[0]->Blocks.Num(); ++i)
+				//	{
+				//		MeshFrag->SetBlock(i, i);
+				//	}
+				//
+				//	MeshPtr = MeshFrag;
+				//}
+				//else
+				//{
+				//	GenerationContext.Compiler->CompilerLog(LOCTEXT("ExtendMaterialLayoutMissing","Skeletal Mesh without Layout Node linked to an Extend Material. A 4x4 layout will be added as default layout."), Node);
+				//}
 
 				mu::NodeMeshFormatPtr MeshFormat = new mu::NodeMeshFormat();
 				SetSurfaceFormat( GenerationContext,
