[Mutable] Fix static analyzer warning.

#rnx

diff --git a/Engine/Plugins/Experimental/Mutable/Source/MutableRuntime/Private/MuR/OpImageProject.cpp b/Engine/Plugins/Experimental/Mutable/Source/MutableRuntime/Private/MuR/OpImageProject.cpp
--- a/Engine/Plugins/Experimental/Mutable/Source/MutableRuntime/Private/MuR/OpImageProject.cpp
+++ b/Engine/Plugins/Experimental/Mutable/Source/MutableRuntime/Private/MuR/OpImageProject.cpp
@@ -231,8 +231,8 @@ namespace
 				Source1DataPtr = nullptr;
 			}
 
-			RGBFadingEnabledValue = -static_cast<int8>(bInIsRGBFadingEnabled);
-			AlphaFadingEnabledValue = -static_cast<int8>(bInIsAlphaFadingEnabled);
+			RGBFadingEnabledValue = bInIsRGBFadingEnabled ? -1 : 0;
+			AlphaFadingEnabledValue = bInIsAlphaFadingEnabled ? -1 : 0;
 
 			const float FadeStartCos = FMath::Cos(FadeStart);
             FadeEndCos = FMath::Cos(FadeEnd);
