[Mutable] Remove additional LODBias added to server builds since it introduces a difference in the order of parameters between server and client.
+ Discard streamed data on server builds

#rnx

diff --git a/Engine/Plugins/Experimental/Mutable/Source/CustomizableObject/Private/MuCO/CustomizableObject.cpp b/Engine/Plugins/Experimental/Mutable/Source/CustomizableObject/Private/MuCO/CustomizableObject.cpp
--- a/Engine/Plugins/Experimental/Mutable/Source/CustomizableObject/Private/MuCO/CustomizableObject.cpp
+++ b/Engine/Plugins/Experimental/Mutable/Source/CustomizableObject/Private/MuCO/CustomizableObject.cpp
@@ -1985,6 +1985,12 @@ void UCustomizableObjectBulk::CookAdditionalFilesOverride(const TCHAR* PackageFi
 	const ITargetPlatform* TargetPlatform,
 	TFunctionRef<void(const TCHAR* Filename, void* Data, int64 Size)> WriteAdditionalFile)
 {
+	// Don't save streamed data on server builds since it won't be used anyway.
+	if (TargetPlatform->IsServerOnly())
+	{
+		return;
+	}
+	
 	check(CustomizableObject);
 	
 	FMutableCachedPlatformData* PlatformData = CustomizableObject->CachedPlatformsData.Find(TargetPlatform->PlatformName());
diff --git a/Engine/Plugins/Experimental/Mutable/Source/CustomizableObjectEditor/Private/MuCOE/GenerateMutableSource/GenerateMutableSource.cpp b/Engine/Plugins/Experimental/Mutable/Source/CustomizableObjectEditor/Private/MuCOE/GenerateMutableSource/GenerateMutableSource.cpp
--- a/Engine/Plugins/Experimental/Mutable/Source/CustomizableObjectEditor/Private/MuCOE/GenerateMutableSource/GenerateMutableSource.cpp
+++ b/Engine/Plugins/Experimental/Mutable/Source/CustomizableObjectEditor/Private/MuCOE/GenerateMutableSource/GenerateMutableSource.cpp
@@ -1438,16 +1438,8 @@ int32 ComputeLODBias(const FMutableGraphGenerationContext& GenerationContext, co
 	}
 	
 	int32 LODBias = 0;
-	// We used to calculate the lod bias directly from the group like this:
-	//int LODBias = 0;
-	//if (LODSettings.TextureLODGroups.IsValidIndex(ReferenceTexture->LODGroup))
-	//{
-	//	const FTextureLODGroup& Group = LODSettings.GetTextureLODGroup(ReferenceTexture->LODGroup);
-	//	LODBias = Group.LODBias;
-	//}
-
-	// ...but now it seems to be more complicated and we may need this:
-	// This is not 100% correct either because it makes assumptions about the final texture size
+
+	// This is not 100% correct because it makes assumptions about the final texture size
 	// that may not be correct, but if the reference texture is really representative of the average 
 	// case, then it is as good as we can do.
 	if (ReferenceTexture)
@@ -1458,14 +1450,6 @@ int32 ComputeLODBias(const FMutableGraphGenerationContext& GenerationContext, co
 			ReferenceTexture->LODGroup, ReferenceTexture->LODBias, false, ReferenceTexture->MipGenSettings, ReferenceTexture->IsCurrentlyVirtualTextured());
 	}
 
-	// Reduce data size and compilation time for server platforms.
-	// \todo: review this hack
-	bool bIsServer = GenerationContext.Options.TargetPlatform->GetTargetPlatformInfo().PlatformType == EBuildTargetType::Server;
-	if (bIsServer)
-	{
-		LODBias += 4;
-	}
-
 	// Increment the LOD bias per each LOD if we are using automatic LODs
 	if (GenerationContext.CurrentLOD > 0
 		&&
