[mutable] Fixed Edit Material node details view:
- Now layout selector is updated automatically when the UV channel changes.
- Parent Material Value is reset when the nodes is disconnected from an object node.

#rnx

diff --git a/Engine/Plugins/Experimental/Mutable/Source/CustomizableObjectEditor/Private/MuCOE/Nodes/CustomizableObjectNodeEditMaterialBaseDetails.cpp b/Engine/Plugins/Experimental/Mutable/Source/CustomizableObjectEditor/Private/MuCOE/Nodes/CustomizableObjectNodeEditMaterialBaseDetails.cpp
--- a/Engine/Plugins/Experimental/Mutable/Source/CustomizableObjectEditor/Private/MuCOE/Nodes/CustomizableObjectNodeEditMaterialBaseDetails.cpp
+++ b/Engine/Plugins/Experimental/Mutable/Source/CustomizableObjectEditor/Private/MuCOE/Nodes/CustomizableObjectNodeEditMaterialBaseDetails.cpp
@@ -27,6 +27,7 @@ void FCustomizableObjectNodeEditMaterialBaseDetails::CustomizeDetails(IDetailLay
 {
 	FCustomizableObjectNodeParentedMaterialDetails::CustomizeDetails(DetailBuilder);
 
+	DetailBuilderPtr = &DetailBuilder;
 	NodeEditMaterialBase = nullptr;
 
 	const TArray<TWeakObjectPtr<UObject>>& SelectedObjects = DetailBuilder.GetDetailsView()->GetSelectedObjects();
@@ -42,7 +43,7 @@ void FCustomizableObjectNodeEditMaterialBaseDetails::CustomizeDetails(IDetailLay
 		// Add all the materials to the combobox
 		TSharedPtr<FString> LayoutToSelect = GenerateLayoutComboboxOptions();
 
-		TSharedRef<IPropertyHandle> LayoutProperty = DetailBuilder.GetProperty("ParentLayoutIndex");
+		TSharedRef<IPropertyHandle> LayoutProperty = DetailBuilder.GetProperty(GET_MEMBER_NAME_CHECKED(UCustomizableObjectNodeEditMaterialBase, ParentLayoutIndex));
 
 		BlocksCategory.AddCustomRow(LOCTEXT("BlocksRow", "Blocks"))
 		[
@@ -121,6 +122,7 @@ void FCustomizableObjectNodeEditMaterialBaseDetails::OnLayoutComboBoxSelectionCh
 		if (LayoutOptionNames[OptionIndex] == Selection)
 		{
 			LayoutProperty->SetValue(LayoutOptionReferences[OptionIndex]);
+			DetailBuilderPtr->ForceRefreshDetails();
 			break;	
 		}
 	}
@@ -135,5 +137,4 @@ void FCustomizableObjectNodeEditMaterialBaseDetails::OnParentComboBoxSelectionCh
 	LayoutComboBox->SetSelectedItem(LayoutToSelect);
 }
 
-
 #undef LOCTEXT_NAMESPACE
diff --git a/Engine/Plugins/Experimental/Mutable/Source/CustomizableObjectEditor/Private/MuCOE/Nodes/CustomizableObjectNodeEditMaterialBaseDetails.h b/Engine/Plugins/Experimental/Mutable/Source/CustomizableObjectEditor/Private/MuCOE/Nodes/CustomizableObjectNodeEditMaterialBaseDetails.h
--- a/Engine/Plugins/Experimental/Mutable/Source/CustomizableObjectEditor/Private/MuCOE/Nodes/CustomizableObjectNodeEditMaterialBaseDetails.h
+++ b/Engine/Plugins/Experimental/Mutable/Source/CustomizableObjectEditor/Private/MuCOE/Nodes/CustomizableObjectNodeEditMaterialBaseDetails.h
@@ -32,6 +32,9 @@ private:
 	
 	UCustomizableObjectNodeEditMaterialBase* NodeEditMaterialBase;
 
+	// Pointer of the details to refresh the view whenever is necessary
+	IDetailLayoutBuilder* DetailBuilderPtr;
+
 	// Prepares combobox options for the Layout
 	TSharedPtr<FString> GenerateLayoutComboboxOptions();
 
diff --git a/Engine/Plugins/Experimental/Mutable/Source/CustomizableObjectEditor/Private/MuCOE/Nodes/CustomizableObjectNodeUseMaterial.cpp b/Engine/Plugins/Experimental/Mutable/Source/CustomizableObjectEditor/Private/MuCOE/Nodes/CustomizableObjectNodeUseMaterial.cpp
--- a/Engine/Plugins/Experimental/Mutable/Source/CustomizableObjectEditor/Private/MuCOE/Nodes/CustomizableObjectNodeUseMaterial.cpp
+++ b/Engine/Plugins/Experimental/Mutable/Source/CustomizableObjectEditor/Private/MuCOE/Nodes/CustomizableObjectNodeUseMaterial.cpp
@@ -85,6 +85,12 @@ void FCustomizableObjectNodeUseMaterial::PinConnectionListChangedWork(UEdGraphPi
 {
 	if (Pin == OutputPin())
 	{
+		if (Pin->LinkedTo.Num() == 0)
+		{
+			FCustomizableObjectNodeParentedMaterial& NodeParentedMaterial = GetNodeParentedMaterial();
+			NodeParentedMaterial.SetParentNode(nullptr, FGuid());
+		}
+
 		if (const TSharedPtr<ICustomizableObjectEditor> Editor = GetNode().GetGraphEditor())
 		{
 			Editor->UpdateGraphNodeProperties();
