[Mutable] Fix crashes when compiling root objects with AutoLOD strategy set to Inherit.
Set AutoLOD strategy to AutomaticFromMesh by default.

#rnx

diff --git a/Engine/Plugins/Experimental/Mutable/Source/CustomizableObjectEditor/Classes/MuCOE/Nodes/CustomizableObjectNodeObject.h b/Engine/Plugins/Experimental/Mutable/Source/CustomizableObjectEditor/Classes/MuCOE/Nodes/CustomizableObjectNodeObject.h
--- a/Engine/Plugins/Experimental/Mutable/Source/CustomizableObjectEditor/Classes/MuCOE/Nodes/CustomizableObjectNodeObject.h
+++ b/Engine/Plugins/Experimental/Mutable/Source/CustomizableObjectEditor/Classes/MuCOE/Nodes/CustomizableObjectNodeObject.h
@@ -164,7 +164,7 @@ public:
 	int32 NumLODs;
 
 	UPROPERTY(EditAnywhere, Category = CustomizableObject)
-	ECustomizableObjectAutomaticLODStrategy AutoLODStrategy;
+	ECustomizableObjectAutomaticLODStrategy AutoLODStrategy = ECustomizableObjectAutomaticLODStrategy::AutomaticFromMesh;
 
 	UPROPERTY(EditAnywhere, Category = CustomizableObject, meta = (ClampMin = "1"))
 	int32 NumMeshComponents = 1;
diff --git a/Engine/Plugins/Experimental/Mutable/Source/CustomizableObjectEditor/Private/MuCOE/CustomizableObjectCompiler.cpp b/Engine/Plugins/Experimental/Mutable/Source/CustomizableObjectEditor/Private/MuCOE/CustomizableObjectCompiler.cpp
--- a/Engine/Plugins/Experimental/Mutable/Source/CustomizableObjectEditor/Private/MuCOE/CustomizableObjectCompiler.cpp
+++ b/Engine/Plugins/Experimental/Mutable/Source/CustomizableObjectEditor/Private/MuCOE/CustomizableObjectCompiler.cpp
@@ -618,6 +618,14 @@ mu::NodeObjectPtr FCustomizableObjectCompiler::GenerateMutableRoot(
 		ActualRootObject = Root->ParentObject ? GetFullGraphRootObject(Root, VisitedObjects) : Object;
 		GenerationContext.bDisableTextureLayoutManagementFlag = ActualRootObject->bDisableTextureLayoutManagement;
 
+
+	}
+
+	// Ensure that the CO has a valid AutoLODStrategy on the ActualRoot.
+	if (ActualRoot->AutoLODStrategy == ECustomizableObjectAutomaticLODStrategy::Inherited)
+	{
+		CompilerLog(LOCTEXT("RootInheritsFromParent", "Error! Base CustomizableObject's LOD Strategy can't be set to 'Inherit from parent object'"), ActualRoot);
+		return nullptr;
 	}
 
 	// Make sure we have a valid Reference SkeletalMesh and Skeleton for each component
