[Mutable] Add per surface BoneMaps to reduce the required bones per section on the final mesh.
+ Addapt MutableMeshPreviewUtils and related classes to use the new BoneMaps structure.
+ Remove the need to specify a reference skeletal mesh to generate the meshes in the debugger.
+ Fix MutableMeshViewport FOV.
+ Change the naming of some variables to follow the UE standards.

#rnx

diff --git a/Engine/Plugins/Experimental/Mutable/Source/CustomizableObject/Private/MuCO/CustomizableObjectInstance.cpp b/Engine/Plugins/Experimental/Mutable/Source/CustomizableObject/Private/MuCO/CustomizableObjectInstance.cpp
--- a/Engine/Plugins/Experimental/Mutable/Source/CustomizableObject/Private/MuCO/CustomizableObjectInstance.cpp
+++ b/Engine/Plugins/Experimental/Mutable/Source/CustomizableObject/Private/MuCO/CustomizableObjectInstance.cpp
@@ -3506,9 +3506,11 @@ bool UCustomizableInstancePrivateData::BuildSkeletonData(const TSharedPtr<FMutab
 				continue;
 			}
 
-			for (uint16& BoneIndex : Component.BoneMap)
+			for (uint32 BoneMapIndex = Component.FirstBoneMap; BoneMapIndex < Component.FirstBoneMap + Component.BoneMapCount; ++BoneMapIndex)
 			{
-				BoneIndex = BoneToFinalBoneIndexMap[MutSkeletonData.BoneNames[BoneIndex]];
+				const int32 BoneIndex = OperationData->InstanceUpdateData.BoneMaps[BoneMapIndex];
+				const FName& BoneName = MutSkeletonData.BoneNames[BoneIndex];
+				OperationData->InstanceUpdateData.BoneMaps[BoneMapIndex] = BoneToFinalBoneIndexMap[BoneName];
 			}
 
 			for (uint16& BoneIndex : Component.ActiveBones)
@@ -3642,9 +3644,11 @@ bool UCustomizableInstancePrivateData::CopySkeletonData(const TSharedPtr<FMutabl
 				continue;
 			}
 
-			for (uint16& BoneIndex : Component.BoneMap)
+			for (uint32 BoneMapIndex = Component.FirstBoneMap; BoneMapIndex < Component.FirstBoneMap + Component.BoneMapCount; ++BoneMapIndex)
 			{
-				BoneIndex = BoneToFinalBoneIndexMap[MutSkeletonData.BoneNames[BoneIndex]];
+				const int32 BoneIndex = OperationData->InstanceUpdateData.BoneMaps[BoneMapIndex];
+				const FName& BoneName = MutSkeletonData.BoneNames[BoneIndex];
+				OperationData->InstanceUpdateData.BoneMaps[BoneMapIndex] = BoneToFinalBoneIndexMap[BoneName];
 			}
 
 			TArray<uint16> UniqueActiveBones;
@@ -3902,7 +3906,7 @@ void UCustomizableInstancePrivateData::BuildOrCopyMorphTargetsData(const TShared
 				}
 
 				int FirstVertex, VerticesCount, FirstIndex, IndiciesCount;
-				Component.Mesh->GetSurface(Section, &FirstVertex, &VerticesCount, &FirstIndex, &IndiciesCount);
+				Component.Mesh->GetSurface(Section, &FirstVertex, &VerticesCount, &FirstIndex, &IndiciesCount, nullptr, nullptr);
 
 				for (int32 VertexIdx = FirstVertex; VertexIdx < FirstVertex + VerticesCount; ++VertexIdx)
 				{
@@ -4084,7 +4088,7 @@ void UCustomizableInstancePrivateData::BuildOrCopyClothingData(const TSharedPtr<
 				for (int32 Section = 0; Section < SurfaceCount; ++Section)
 				{
 					int FirstVertex, VerticesCount, FirstIndex, IndicesCount;
-					MutableMesh->GetSurface(Section, &FirstVertex, &VerticesCount, &FirstIndex, &IndicesCount);
+					MutableMesh->GetSurface(Section, &FirstVertex, &VerticesCount, &FirstIndex, &IndicesCount, nullptr, nullptr);
 
 					if (VerticesCount == 0 || IndicesCount == 0)
 					{
@@ -4728,7 +4732,7 @@ void UCustomizableInstancePrivateData::BuildOrCopyClothingData(const TSharedPtr<
 				{
 					// Check that is a valid surface.
 					int32 FirstVertex, VerticesCount, FirstIndex, IndicesCount;
-					MutableMesh->GetSurface(Section, &FirstVertex, &VerticesCount, &FirstIndex, &IndicesCount);
+					MutableMesh->GetSurface(Section, &FirstVertex, &VerticesCount, &FirstIndex, &IndicesCount, nullptr, nullptr);
 
 					if (VerticesCount == 0 || IndicesCount == 0 )
 					{
@@ -4962,7 +4966,8 @@ bool UCustomizableInstancePrivateData::BuildOrCopyRenderData(const TSharedPtr<FM
 			SkeletalMesh,
 			Component.Mesh,
 			LODIndex,
-			Component.BoneMap);
+			OperationData->InstanceUpdateData.BoneMaps,
+			Component.FirstBoneMap);
 
 		UnrealConversionUtils::CopyMutableVertexBuffers(
 			SkeletalMesh,
diff --git a/Engine/Plugins/Experimental/Mutable/Source/CustomizableObject/Private/MuCO/CustomizableObjectSystem.cpp b/Engine/Plugins/Experimental/Mutable/Source/CustomizableObject/Private/MuCO/CustomizableObjectSystem.cpp
--- a/Engine/Plugins/Experimental/Mutable/Source/CustomizableObject/Private/MuCO/CustomizableObjectSystem.cpp
+++ b/Engine/Plugins/Experimental/Mutable/Source/CustomizableObject/Private/MuCO/CustomizableObjectSystem.cpp
@@ -1624,34 +1624,40 @@ namespace impl
 					SkeletonIds.AddUnique(Mesh->GetSkeletonID(SkeletonIndex));
 				}
 
+				// Add bonemaps from this mesh
+				const TArray<uint16>& BoneMap = Mesh->GetBoneMap();
+				CurrentLODComponent.FirstBoneMap = OperationData->InstanceUpdateData.BoneMaps.Num();
+				CurrentLODComponent.BoneMapCount = BoneMap.Num();
 
-				const int32 MaxBoneIndex = Mesh->GetBonePoseCount();
+				OperationData->InstanceUpdateData.BoneMaps.Reserve(CurrentLODComponent.FirstBoneMap + CurrentLODComponent.BoneMapCount);
 
-				// Add active bones and new names to the BoneNames array 
-				CurrentLODComponent.BoneMap.Reserve(MaxBoneIndex);
-				CurrentLODComponent.ActiveBones.Reserve(MaxBoneIndex);
+				const mu::SkeletonPtrConst Skeleton = Mesh->GetSkeleton();
+				for (const uint16& BoneIndex : BoneMap)
+				{
+					const FName BoneName = Skeleton->GetBoneName(BoneIndex);
+					check(BoneName != NAME_None);
 
+					const int32 FinalBoneIndex = BoneNames.AddUnique(BoneName);
+					OperationData->InstanceUpdateData.BoneMaps.Add(FinalBoneIndex);
+				}
+				
+				// Add active bone indices and poses
+				const int32 MaxBoneIndex = Mesh->GetBonePoseCount();
+				CurrentLODComponent.ActiveBones.Reserve(MaxBoneIndex);
 				for (int32 BonePoseIndex = 0; BonePoseIndex < MaxBoneIndex; ++BonePoseIndex)
 				{
 					const FName BoneName = Mesh->GetBonePoseName(BonePoseIndex);
 					check(BoneName != NAME_None);
 
-					int32 BoneIndex = BoneNames.Find(BoneName);
-					if (BoneIndex == INDEX_NONE)
-					{
-						BoneIndex = BoneNames.Add(BoneName);
+					const int32 BoneIndex = BoneNames.AddUnique(BoneName); 
+					CurrentLODComponent.ActiveBones.Add(BoneIndex);
 
+					if (!BoneMatricesWithScale.Contains(BoneName))
+					{
 						FTransform3f Transform;
 						Mesh->GetBoneTransform(BonePoseIndex, Transform);
 						BoneMatricesWithScale.Emplace(BoneName, Transform.Inverse().ToMatrixWithScale());
 					}
-
-					if (EnumHasAnyFlags(Mesh->GetBoneUsageFlags(BonePoseIndex), mu::EBoneUsageFlags::Skinning))
-					{
-						CurrentLODComponent.BoneMap.Add(BoneIndex);
-					}
-
-					CurrentLODComponent.ActiveBones.Add(BoneIndex);
 				}
 			}
 		}
diff --git a/Engine/Plugins/Experimental/Mutable/Source/CustomizableObject/Private/MuCO/CustomizableObjectSystemPrivate.h b/Engine/Plugins/Experimental/Mutable/Source/CustomizableObject/Private/MuCO/CustomizableObjectSystemPrivate.h
--- a/Engine/Plugins/Experimental/Mutable/Source/CustomizableObject/Private/MuCO/CustomizableObjectSystemPrivate.h
+++ b/Engine/Plugins/Experimental/Mutable/Source/CustomizableObject/Private/MuCO/CustomizableObjectSystemPrivate.h
@@ -383,11 +383,9 @@ struct FInstanceUpdateData
 		//uint32 FirstActiveBone;
 		//uint32 ActiveBoneCount;
 
-		// \TODO: Flatten
-		TArray<uint16> BoneMap;
 		/** Range in the external Bones array */
-		//uint32 FirstBoneMap;
-		//uint32 BoneMapCount;
+		uint32 FirstBoneMap = 0;
+		uint32 BoneMapCount = 0;
 	};
 
 	struct FLOD
@@ -404,6 +402,8 @@ struct FInstanceUpdateData
 	TArray<FVector> Vectors;
 	TArray<FScalar> Scalars;
 
+	TArray<uint16> BoneMaps;
+
 	struct FSkeletonData
 	{
 		int16 ComponentIndex = INDEX_NONE;
diff --git a/Engine/Plugins/Experimental/Mutable/Source/CustomizableObject/Private/MuCO/UnrealConversionUtils.cpp b/Engine/Plugins/Experimental/Mutable/Source/CustomizableObject/Private/MuCO/UnrealConversionUtils.cpp
--- a/Engine/Plugins/Experimental/Mutable/Source/CustomizableObject/Private/MuCO/UnrealConversionUtils.cpp
+++ b/Engine/Plugins/Experimental/Mutable/Source/CustomizableObject/Private/MuCO/UnrealConversionUtils.cpp
@@ -98,52 +98,12 @@ namespace UnrealConversionUtils
 	}
 
 
-
-	void BuildRefSkeleton(FInstanceUpdateData::FSkeletonData* OutMutSkeletonData
-		, const FReferenceSkeleton& InSourceReferenceSkeleton, const TArray<bool>& InUsedBones,
-		FReferenceSkeleton& InRefSkeleton, const USkeleton* InSkeleton)
-	{
-		const int32 SourceBoneCount = InSourceReferenceSkeleton.GetRawBoneNum();
-		const TArray<FTransform>& SourceRawMeshBonePose = InSourceReferenceSkeleton.GetRawRefBonePose();
-
-		MUTABLE_CPUPROFILER_SCOPE(BuildRefSkeleton);
-
-		// Build new RefSkeleton	
-		FReferenceSkeletonModifier RefSkeletonModifier(InRefSkeleton, InSkeleton);
-
-		const TArray<FMeshBoneInfo>& BoneInfo = InSourceReferenceSkeleton.GetRawRefBoneInfo();
-
-		TMap<FName, uint16> BoneToFinalBoneIndexMap;
-		BoneToFinalBoneIndexMap.Reserve(SourceBoneCount);
-
-		uint32 FinalBoneCount = 0;
-		for (int32 BoneIndex = 0; BoneIndex < SourceBoneCount; ++BoneIndex)
-		{
-			if (!InUsedBones[BoneIndex])
-			{
-				continue;
-			}
-
-			FName BoneName = BoneInfo[BoneIndex].Name;
-
-			// Build a bone to index map so we can remap BoneMaps and ActiveBoneIndices later on
-			BoneToFinalBoneIndexMap.Add(BoneName, FinalBoneCount);
-			FinalBoneCount++;
-
-			// Find parent index
-			const int32 SourceParentIndex = BoneInfo[BoneIndex].ParentIndex;
-			const int32 ParentIndex = SourceParentIndex != INDEX_NONE ? BoneToFinalBoneIndexMap[BoneInfo[SourceParentIndex].Name] : INDEX_NONE;
-
-			RefSkeletonModifier.Add(FMeshBoneInfo(BoneName, BoneName.ToString(), ParentIndex), SourceRawMeshBonePose[BoneIndex]);
-		}
-	}
-
-
 	void SetupRenderSections(
 		const USkeletalMesh* OutSkeletalMesh,
 		const mu::MeshPtrConst InMutableMesh,
 		const int32 MeshLODIndex,
-		const TArray<uint16>& InBoneMap)
+		const TArray<uint16>& InBoneMap,
+		const int32 InFirstBoneMapIndex)
 	{
 		check(InMutableMesh);
 
@@ -169,7 +129,9 @@ namespace UnrealConversionUtils
 			int32 IndexCount;
 			int32 FirstVertex;
 			int32 VertexCount;
-			InMutableMesh->GetSurface(SurfaceIndex, &FirstVertex, &VertexCount, &FirstIndex, &IndexCount);
+			int32 FirstBone;
+			int32 BoneCount;
+			InMutableMesh->GetSurface(SurfaceIndex, &FirstVertex, &VertexCount, &FirstIndex, &IndexCount, &FirstBone, &BoneCount);
 			FSkelMeshRenderSection& Section = OutSkeletalMesh->GetResourceForRendering()->LODRenderData[MeshLODIndex].RenderSections[SurfaceIndex];
 
 			Section.DuplicatedVerticesBuffer.Init(1, TMap<int, TArray<int32>>());
@@ -185,7 +147,15 @@ namespace UnrealConversionUtils
 			Section.BaseVertexIndex = FirstVertex;
 			Section.MaxBoneInfluences = NumBoneInfluences;
 			Section.NumVertices = VertexCount;
-			Section.BoneMap.Append(InBoneMap);
+
+			// InBoneMaps may contain bonemaps from other sections. Copy the bones belonging to this mesh.
+			FirstBone += InFirstBoneMapIndex;
+			
+			Section.BoneMap.Reserve(BoneCount);
+			for (int32 BoneMapIndex = 0; BoneMapIndex < BoneCount; ++BoneMapIndex, ++FirstBone)
+			{
+				Section.BoneMap.Add(InBoneMap[FirstBone]);
+			}
 		}
 	}
 
diff --git a/Engine/Plugins/Experimental/Mutable/Source/CustomizableObject/Public/MuCO/UnrealConversionUtils.h b/Engine/Plugins/Experimental/Mutable/Source/CustomizableObject/Public/MuCO/UnrealConversionUtils.h
--- a/Engine/Plugins/Experimental/Mutable/Source/CustomizableObject/Public/MuCO/UnrealConversionUtils.h
+++ b/Engine/Plugins/Experimental/Mutable/Source/CustomizableObject/Public/MuCO/UnrealConversionUtils.h
@@ -20,24 +20,6 @@ namespace UnrealConversionUtils
 	 * current pipelines (instance update and USkeletal mesh generation for a mesh viewport)
 	 */
 	
-	
-	/**
-	 * Builds the reference skeleton by adding bones to the reference skeleton modifier provided. When exiting the scope
-	 * of this method those bones get generated on the reference skeleton. It also updates an array with the used bones.
-	 * @param OutMutSkeletonData - Object map used later when remaking bones.
-	 * @param InSourceReferenceSkeleton - The reference skeleton from the source object. Is used to get the skeleton bone info
-	 * @param InUsedBones - Array with the same count as the bones present on the skeleton and determines what bones are used and what not.
-	 * @param InRefSkeleton - The reference skeleton to be updated with new bone data
-	 * @param InSkeleton - Skeleton to be updated as consequence of the destruction of the FReferenceSkeletonModifier
-	 */
-	CUSTOMIZABLEOBJECT_API void BuildRefSkeleton(
-		FInstanceUpdateData::FSkeletonData* OutMutSkeletonData,
-		const FReferenceSkeleton& InSourceReferenceSkeleton,
-		const TArray<bool>& InUsedBones,
-		FReferenceSkeleton& InRefSkeleton,
-		const USkeleton* InSkeleton );
-
-	
 	/**
 	 * Prepares the render sections found on the InSkeletalMesh and sets them up accordingly what the InMutableMesh requires
 	 * @param OutSkeletalMesh - The Skeletal mesh whose sections are ought to be updated
@@ -49,7 +31,8 @@ namespace UnrealConversionUtils
 		const USkeletalMesh* OutSkeletalMesh,
 		const mu::MeshPtrConst InMutableMesh,
 		const int32 InMeshLODIndex,
-		const TArray<uint16>& InBoneMap);
+		const TArray<uint16>& InBoneMap,
+		const int32 InFirstBoneMapIndex);
 
 
 	/** Performs a copy of the data found on the vertex buffers on the mutable mesh to the buffers of the skeletal mesh
diff --git a/Engine/Plugins/Experimental/Mutable/Source/CustomizableObjectEditor/Private/MuCOE/FMutableViewportClient.cpp b/Engine/Plugins/Experimental/Mutable/Source/CustomizableObjectEditor/Private/MuCOE/FMutableViewportClient.cpp
--- a/Engine/Plugins/Experimental/Mutable/Source/CustomizableObjectEditor/Private/MuCOE/FMutableViewportClient.cpp
+++ b/Engine/Plugins/Experimental/Mutable/Source/CustomizableObjectEditor/Private/MuCOE/FMutableViewportClient.cpp
@@ -14,4 +14,10 @@ FMutableMeshViewportClient::FMutableMeshViewportClient(const TSharedRef<FAdvance
 	FAdvancedPreviewScene* PreviewSceneCasted = static_cast<FAdvancedPreviewScene*>(PreviewScene);
 	//PreviewSceneCasted->SetEnvironmentVisibility(false);
 	PreviewSceneCasted->SetFloorVisibility(false);
+
+	// Set standard FOV
+	const float FOVMin = 5.f;
+	const float FOVMax = 170.f;
+	ViewFOV = FMath::Clamp<float>(53.43f, FOVMin, FOVMax);
+
 }
diff --git a/Engine/Plugins/Experimental/Mutable/Source/CustomizableObjectEditor/Private/MuCOE/GenerateMutableSource/GenerateMutableSourceMesh.cpp b/Engine/Plugins/Experimental/Mutable/Source/CustomizableObjectEditor/Private/MuCOE/GenerateMutableSource/GenerateMutableSourceMesh.cpp
--- a/Engine/Plugins/Experimental/Mutable/Source/CustomizableObjectEditor/Private/MuCOE/GenerateMutableSource/GenerateMutableSourceMesh.cpp
+++ b/Engine/Plugins/Experimental/Mutable/Source/CustomizableObjectEditor/Private/MuCOE/GenerateMutableSource/GenerateMutableSourceMesh.cpp
@@ -633,7 +633,7 @@ mu::MeshPtr ConvertSkeletalMeshToMutable(USkeletalMesh* InSkeletalMesh, const TS
 			{
 				if (ErrorMessage.IsEmpty())
 				{
-					GenerationContext.Compiler->CompilerLog(FText::FromString(ErrorMessage), CurrentNode, EMessageSeverity::Info);
+					GenerationContext.Compiler->CompilerLog(FText::FromString(ErrorMessage), CurrentNode, EMessageSeverity::Warning);
 				}
 				return nullptr;
 			}
@@ -712,20 +712,30 @@ mu::MeshPtr ConvertSkeletalMeshToMutable(USkeletalMesh* InSkeletalMesh, const TS
 			bBoneMapModified = bBoneMapModified || SourceBoneMap[BoneIndex] != FinalBoneIndex;
 		}
 
+		const int32 NumBonesBoneMap = BoneMap.Num();
+		const int32 NumRequiredBones = RequiredBones.Num();
+
+		// BoneMap mapping the BoneMap indices to those of the mu::skeleton
+		TArray<uint16> MutableBoneMap;
+		MutableBoneMap.Reserve(NumBonesBoneMap);
+
+		// Mapping of RequiredBones from RefSkeletonIndex to mu::Skeleton BoneIndex
 		TMap<int32, int32> InverseBoneMap;
+		InverseBoneMap.Reserve(NumRequiredBones);
 
-		const int32 BoneMapCount = BoneMap.Num();
-		for (int32 BoneIndex = 0; BoneIndex < BoneMapCount; ++BoneIndex)
+		for (int32 BoneIndex = 0; BoneIndex < NumBonesBoneMap; ++BoneIndex)
 		{
 			int32 RefSkeletonIndex = BoneMap[BoneIndex];
 			InverseBoneMap.Add(RefSkeletonIndex, BoneIndex);
+			MutableBoneMap.Add(BoneIndex);
 		}
 
+		MutableMesh->SetBoneMap(MutableBoneMap);
+
 		// Create the skeleton and poses for this mesh
 		mu::SkeletonPtr MutableSkeleton = new mu::Skeleton;
 		MutableMesh->SetSkeleton(MutableSkeleton);
 
-		const int32 NumRequiredBones = RequiredBones.Num();
 		MutableMesh->SetBonePoseCount(NumRequiredBones);
 		MutableSkeleton->SetBoneCount(NumRequiredBones);
 
@@ -733,16 +743,15 @@ mu::MeshPtr ConvertSkeletalMeshToMutable(USkeletalMesh* InSkeletalMesh, const TS
 		for (int32 RequiredBoneIndex = 0; RequiredBoneIndex < NumRequiredBones; ++RequiredBoneIndex)
 		{
 			const int32 RefSkelIndex = RequiredBones[RequiredBoneIndex];
-
 			const FMeshBoneInfo& BoneInfo = RefBoneInfo[RefSkelIndex];
-
-			const FString BoneName = BoneInfo.Name.ToString();
-			const int32 BoneMapIndex = InverseBoneMap.FindOrAdd(RefSkelIndex, InverseBoneMap.Num());
-			const int32 ParentBoneMapIndex = BoneInfo.ParentIndex != INDEX_NONE ? InverseBoneMap.FindOrAdd(BoneInfo.ParentIndex, InverseBoneMap.Num()) : INDEX_NONE;
+			
+			const int32 BoneIndex = InverseBoneMap.FindOrAdd(RefSkelIndex, InverseBoneMap.Num());
+			const int32 ParentBoneIndex = BoneInfo.ParentIndex != INDEX_NONE ? InverseBoneMap.FindOrAdd(BoneInfo.ParentIndex, InverseBoneMap.Num()) : INDEX_NONE;
 
 			// Set bone hierarchy
-			MutableSkeleton->SetBoneName(BoneMapIndex, StringCast<ANSICHAR>(*BoneName).Get());
-			MutableSkeleton->SetBoneParent(BoneMapIndex, ParentBoneMapIndex);
+			const FString BoneName = BoneInfo.Name.ToString();
+			MutableSkeleton->SetBoneName(BoneIndex, StringCast<ANSICHAR>(*BoneName).Get());
+			MutableSkeleton->SetBoneParent(BoneIndex, ParentBoneIndex);
 
 			// Set bone pose
 			FMatrix44f BaseInvMatrix = InSkeletalMesh->GetRefBasesInvMatrix()[RefSkelIndex];
@@ -750,9 +759,10 @@ mu::MeshPtr ConvertSkeletalMeshToMutable(USkeletalMesh* InSkeletalMesh, const TS
 			BaseInvTransform.SetFromMatrix(BaseInvMatrix);
 
 			mu::EBoneUsageFlags BoneUsageFlags = mu::EBoneUsageFlags::None;
-			EnumAddFlags(BoneUsageFlags, BoneMapIndex < BoneMapCount ? mu::EBoneUsageFlags::Skinning : mu::EBoneUsageFlags::None);
+			EnumAddFlags(BoneUsageFlags, BoneIndex < NumBonesBoneMap ? mu::EBoneUsageFlags::Skinning : mu::EBoneUsageFlags::None);
 			EnumAddFlags(BoneUsageFlags, BoneInfo.ParentIndex == INDEX_NONE ? mu::EBoneUsageFlags::Root : mu::EBoneUsageFlags::None);
-			MutableMesh->SetBonePose(BoneMapIndex, StringCast<ANSICHAR>(*BoneName).Get(), BaseInvTransform.Inverse(), BoneUsageFlags);
+
+			MutableMesh->SetBonePose(BoneIndex, StringCast<ANSICHAR>(*BoneName).Get(), BaseInvTransform.Inverse(), BoneUsageFlags);
 		}
 	}
 
@@ -2895,7 +2905,7 @@ mu::NodeMeshPtr GenerateMutableSourceMesh(const UEdGraphPin * Pin,
 				MeshData.bHasVertexColors = TypedNodeSkel->SkeletalMesh->GetHasVertexColors();
 				FSkeletalMeshModel* importModel = TypedNodeSkel->SkeletalMesh->GetImportedModel();
 				MeshData.NumTexCoordChannels = importModel->LODModels[LOD].NumTexCoords;
-				MeshData.MaxBoneIndexTypeSizeBytes = importModel->LODModels[LOD].RequiredBones.Num() > 256 ? 2 : 1;
+				MeshData.MaxBoneIndexTypeSizeBytes = MutableMesh->GetBoneMap().Num() > 256 ? 2 : 1;
 				MeshData.MaxNumBonesPerVertex = importModel->LODModels[LOD].GetMaxBoneInfluences();
 				
 				// When mesh data is combined we will get an upper and lower bound of the number of triangles.
diff --git a/Engine/Plugins/Experimental/Mutable/Source/CustomizableObjectEditor/Private/MuCOE/MutableMeshPreviewUtils.cpp b/Engine/Plugins/Experimental/Mutable/Source/CustomizableObjectEditor/Private/MuCOE/MutableMeshPreviewUtils.cpp
--- a/Engine/Plugins/Experimental/Mutable/Source/CustomizableObjectEditor/Private/MuCOE/MutableMeshPreviewUtils.cpp
+++ b/Engine/Plugins/Experimental/Mutable/Source/CustomizableObjectEditor/Private/MuCOE/MutableMeshPreviewUtils.cpp
@@ -276,264 +276,60 @@ namespace MutableMeshPreviewUtils
 			return bCanBeConverted;
 		}
 
-		/** Loads up the OutSkeletalMesh with all the bone data found on the InMutableMesh if found to be a skeleton there. If not it will fail
-		 * @param InMutableMesh - The mutable mesh to read bone data from
-		 * @param OutFinalBoneNames - An array of names containing the names of the bones to be used on the skeletal mesh
-		 * @param OutBoneMap - Array with all the bones used on the final mesh
-		 */ 
-		void PrepareSkeletonData(mu::MeshPtrConst InMutableMesh, TArray<FName>& OutFinalBoneNames, TArray<uint16>& OutBoneMap)
+
+		/** Provides the OutSkeletalMesh with a new skeleton
+		 * @param OutSkeletalMesh - The skeletal mesh whose skeleton we want to generate.
+		 * @return True if the operation could be performed successfully, false if not.
+		 */
+		bool BuildSkeletalMeshSkeletonData(mu::MeshPtrConst InMutableMesh, USkeletalMesh* OutSkeletalMesh)
 		{
-			if (!InMutableMesh)
+			if (!InMutableMesh || !OutSkeletalMesh)
 			{
-				return;
+				return false;
 			}
 			
-			// Helpers
-			TMap<FName, int16> BoneNameToIndex;
-			TArray<int16> ParentBoneIndices;
-			TArray<int16> ParentChain;
-			TArray<bool> UsedBones;
-			
+			// The skeleton will only have one bone, and all bon indices will be mapped to it.
+			TObjectPtr<USkeleton> Skeleton = NewObject<USkeleton>();
 
-			// Use first valid LOD bone count as a potential total number of bones, used for pre-allocating data arrays
-			if ( InMutableMesh && InMutableMesh->GetSkeleton())
+			// Build new RefSkeleton	
 			{
-				const int32 TotalPossibleBones = InMutableMesh->GetSkeleton()->GetBoneCount();
-				
-				// Helpers
-				BoneNameToIndex.Empty(TotalPossibleBones);
-				ParentBoneIndices.Empty(TotalPossibleBones);
-
-				ParentChain.SetNumZeroed(TotalPossibleBones);
-
-				// Out Data
-				OutFinalBoneNames.Reserve(TotalPossibleBones);
+				// Scope is important, FReferenceSkeletonModifier will rebuild the reference skeleon on destroy
+				FReferenceSkeletonModifier RefSkeletonModifier(Skeleton);
+				RefSkeletonModifier.Add(FMeshBoneInfo(FName("Root"), FString("Root"), INDEX_NONE), FTransform::Identity);
 			}
 
-			// Locate what bones are being used by asking each vertex with what bone index it is related with.
-			// We end up having an array of Bool variables where each index corresponds to a bone on the mutable skeleton
-			// and the value determines if the bone is found to be used by any of the vertices located on the
-			// mutable mesh surfaces
-			{
-				if (!InMutableMesh || !InMutableMesh->GetSkeleton())
-				{
-					return;
-				}
-
-				const mu::Skeleton* Skeleton = InMutableMesh->GetSkeleton().get();
-				const int32 BoneCount = Skeleton->GetBoneCount();
-
-				const mu::FMeshBufferSet& MutableMeshVertexBuffers = InMutableMesh->GetVertexBuffers();
-
-				const int32 NumVerticesLODModel = InMutableMesh->GetVertexCount();
-				const int32 SurfaceCount = InMutableMesh->GetSurfaceCount();
-
-				mu::MESH_BUFFER_FORMAT BoneIndexFormat = mu::MBF_NONE;
-				int32 BoneIndexComponents = 0;
-				int32 BoneIndexOffset = 0;
-				int32 BoneIndexBuffer = -1;
-				int32 BoneIndexChannel = -1;
-				MutableMeshVertexBuffers.FindChannel(mu::MBS_BONEINDICES, 0, &BoneIndexBuffer, &BoneIndexChannel);
-				if (BoneIndexBuffer >= 0 || BoneIndexChannel >= 0)
-				{
-					MutableMeshVertexBuffers.GetChannel(BoneIndexBuffer, BoneIndexChannel,
-						nullptr, nullptr, &BoneIndexFormat, &BoneIndexComponents, &BoneIndexOffset);
-				}
-
-				const int32 ElementSize = MutableMeshVertexBuffers.GetElementSize(BoneIndexBuffer);
-
-				const uint8_t* BufferStart = MutableMeshVertexBuffers.GetBufferData(BoneIndexBuffer) + BoneIndexOffset;
-
-				const int NumBoneInfluences = BoneIndexComponents;
-				{
-					UsedBones.Empty(BoneCount);
-					UsedBones.AddDefaulted(BoneCount);
-
-					for (int32 Surface = 0; Surface < SurfaceCount; Surface++)
-					{
-						int32 FirstIndex;
-						int32 IndexCount;
-						int32 FirstVertex;
-						int32 VertexCount;
-						InMutableMesh->GetSurface(Surface, &FirstVertex, &VertexCount, &FirstIndex, &IndexCount);
-
-						if (VertexCount == 0 || IndexCount == 0)
-						{
-							continue;
-						}
-
-						const uint8_t* VertexBoneIndexPtr = BufferStart + FirstVertex * ElementSize;
-
-						if (BoneIndexFormat == mu::MBF_UINT8)
-						{
-							for (int32 v = 0; v < VertexCount; ++v)
-							{
-								for (int32 i = 0; i < BoneIndexComponents; ++i)
-								{
-									const size_t SectionBoneIndex = VertexBoneIndexPtr[i];
-									UsedBones[SectionBoneIndex] = true;
-								}
-								VertexBoneIndexPtr += ElementSize;
-							}
-						}
-						else if (BoneIndexFormat == mu::MBF_UINT16)
-						{
-							for (int32 v = 0; v < VertexCount; ++v)
-							{
-								for (int32 i = 0; i < BoneIndexComponents; ++i)
-								{
-									const size_t SectionBoneIndex = ((uint16*)VertexBoneIndexPtr)[i];
-									UsedBones[SectionBoneIndex] = true;
-								}
-								VertexBoneIndexPtr += ElementSize;
-							}
-						}
-						else if (BoneIndexFormat == mu::MBF_UINT32)
-						{
-							for (int32 v = 0; v < VertexCount; ++v)
-							{
-								for (int32 i = 0; i < BoneIndexComponents; ++i)
-								{
-									const size_t SectionBoneIndex = ((uint32_t*)VertexBoneIndexPtr)[i];
-									UsedBones[SectionBoneIndex] = true;
-								}
-								VertexBoneIndexPtr += ElementSize;
-							}
-						}
-						else
-						{
-							// Unsupported bone index format in generated mutable mesh
-							check(false);
-						}
+			OutSkeletalMesh->SetSkeleton(Skeleton);
+			OutSkeletalMesh->SetRefSkeleton(Skeleton->GetReferenceSkeleton());
+			OutSkeletalMesh->GetRefBasesInvMatrix().Empty(1);
+			OutSkeletalMesh->CalculateInvRefMatrices();
 
-					}
-				}
+			// Compute imported bounds
+			FBoxSphereBounds Bounds;
 
-				// Fill the BoneMap with the data found on the mutable skeleton and with the aid of the
-				// bool array with all the bones we are using in this mesh from all the ones found on the mu::Skeleton
+			const int32 BonePoseCount = InMutableMesh->GetBonePoseCount();
+			if (BonePoseCount > 0)
+			{
+				// Extract bounds from the pose of the mesh
+				TArray<FVector> Points;
+				Points.Reserve(BonePoseCount);
+				for (int32 BoneIndex = 0; BoneIndex < BonePoseCount; ++BoneIndex)
 				{
-					OutBoneMap.Reserve(BoneCount);
-
-					for (int32 BoneIndex = 0; BoneIndex < BoneCount; ++BoneIndex)
-					{
-						if (!UsedBones[BoneIndex])
-						{
-							// Add root as a placeholder
-							OutBoneMap.Add(0);
-							continue;
-						}
-
-						FName BoneName = Skeleton->GetBoneName(BoneIndex);
-
-						int16* FinalBoneIndexPtr = BoneNameToIndex.Find(BoneName);
-						if (!FinalBoneIndexPtr) // New bone!
-						{
-							// Ensure parent chain is valid
-							int16 FinalParentIndex = INDEX_NONE;
-
-							const int16 MutParentIndex = Skeleton->GetBoneParent(BoneIndex);
-							if (MutParentIndex != INDEX_NONE) // Bone has a parent, ensure the parent exists
-							{
-								// Ensure parent chain until root exists
-								uint16 ParentChainCount = 0;
-
-								// Find parents to add
-								int16 NextMutParentIndex = MutParentIndex;
-								while (NextMutParentIndex != INDEX_NONE)
-								{
-									if (int16* Found = BoneNameToIndex.Find(Skeleton->GetBoneName(NextMutParentIndex)))
-									{
-										FinalParentIndex = *Found;
-										break;
-									}
-
-									ParentChain[ParentChainCount] = NextMutParentIndex;
-									++ParentChainCount;
-
-									NextMutParentIndex = Skeleton->GetBoneParent(NextMutParentIndex);
-								}
-
-								// Add parent bones to the list and to the active bones array
-								for (int16 ParentChainIndex = ParentChainCount - 1; ParentChainIndex >= 0; --ParentChainIndex)
-								{
-									FName ChainParentName = Skeleton->GetBoneName(ParentChain[ParentChainIndex]);
-
-									// Set the parent for the bone we're about to add (previous ParentIndex)
-									ParentBoneIndices.Add(FinalParentIndex);
-
-									// Add the parent
-									FinalParentIndex = OutFinalBoneNames.Num();
-									BoneNameToIndex.Add(ChainParentName, FinalParentIndex);
-									OutFinalBoneNames.Add(ChainParentName);
-								}
-							}
-
-							const int16 NewBoneIndex = OutFinalBoneNames.Num();
-							OutBoneMap.Add(NewBoneIndex);
+					FTransform3f Transform;
+					InMutableMesh->GetBoneTransform(BoneIndex, Transform);
 
-							// Add new bone to the ReferenceSkeleton
-							BoneNameToIndex.Add(BoneName, NewBoneIndex);
-							OutFinalBoneNames.Add(BoneName);
-							ParentBoneIndices.Add(FinalParentIndex);
-						}
-						else
-						{
-							const int16 FinalBoneIndex = *FinalBoneIndexPtr;
-							OutBoneMap.Add(FinalBoneIndex);
-						}
-					}
+					Points.Add(FVector(Transform.GetTranslation()));
 				}
 
-				OutBoneMap.Shrink();
+				Bounds = FBoxSphereBounds(Points.GetData(), Points.Num());
+				Bounds.ExpandBy(1.2f);
 			}
-
-			OutFinalBoneNames.Shrink();
-		}
-
-		/** Provides the OutSkeletalMesh with a new reference skeleton based on the data found on the other method inputs
-		 * @param InRefSkeletalMesh - The skeletal mesh to be used as reference for the copying of data not held by the mutable
-		 * @param OutSkeletalMesh - The skeletal mesh whose reference skeleton we want to generate.
-		 * @return True if the operation could be performed successfully, false if not.
-		 */
-		bool BuildSkeletalMeshSkeletonData(const USkeletalMesh* InRefSkeletalMesh, USkeletalMesh* OutSkeletalMesh)
-		{
-			if (!OutSkeletalMesh || !InRefSkeletalMesh)
+			else
 			{
-				return false;
+				// Set bounds even if they're not correct
+				Bounds = FBoxSphereBounds(FSphere(FVector(0,0,0), 1000));
 			}
 
-			// get a copy of the skeleton used by the reference skeletal mesh
-			const USkeleton* OriginalSkeleton = InRefSkeletalMesh->GetSkeleton();
-			if (!OriginalSkeleton)
-			{
-				return false;
-			}
-			
-			TObjectPtr<USkeleton> Skeleton = DuplicateObject(OriginalSkeleton,OriginalSkeleton->GetExternalPackage()) ;
-			OutSkeletalMesh->SetSkeleton(Skeleton);
-			
-			const FReferenceSkeleton& SourceReferenceSkeleton = Skeleton->GetReferenceSkeleton();
-			const int32 SourceBoneCount = SourceReferenceSkeleton.GetRawBoneNum();
-			
-			// Used to add them in strictly increasing order
-			TArray<bool> UsedBones;
-			UsedBones.Init(true, SourceBoneCount);	
-
-			// New reference skeleton
-			FReferenceSkeleton RefSkeleton;
-			
-			// Build RefSkeleton
-			UnrealConversionUtils::BuildRefSkeleton(
-				nullptr,
-				SourceReferenceSkeleton,
-				UsedBones,
-				RefSkeleton,
-				Skeleton);
-
-			OutSkeletalMesh->SetRefSkeleton(RefSkeleton);
-
-			OutSkeletalMesh->GetRefBasesInvMatrix().Empty(RefSkeleton.GetRawBoneNum());
-			OutSkeletalMesh->CalculateInvRefMatrices();
+			OutSkeletalMesh->SetImportedBounds(Bounds);
 
 			return true;
 		}
@@ -570,20 +366,26 @@ namespace MutableMeshPreviewUtils
 		 * @param OutSkeletalMesh - The skeletal mesh to set up.
 		 * @return True if the operation could be performed successfully, false if not.
 		 */
-		bool BuildSkeletalMeshRenderData(
-			const mu::MeshPtrConst InMutableMesh, const USkeletalMesh* InRefSkeletalMesh,
-			const TArray<uint16>& InBoneMap, USkeletalMesh* OutSkeletalMesh)
+		bool BuildSkeletalMeshRenderData(const mu::MeshPtrConst InMutableMesh, USkeletalMesh* OutSkeletalMesh)
 		{
 			OutSkeletalMesh->SetHasBeenSimplified(false);
 			OutSkeletalMesh->SetHasVertexColors(false);
 
+			// Find how many bones the bonemap could have
+			int32 NumBoneMapIndices = !InMutableMesh->GetBoneMap().IsEmpty() ? InMutableMesh->GetBoneMap().Num() : InMutableMesh->GetBonePoseCount();
+			
+			// Fill the bonemap with zeros
+			TArray<uint16> BoneMap;
+			BoneMap.SetNumZeroed(FMath::Max(NumBoneMapIndices, 1));
+
 			// Load buffer data found on the mutable model onto the out skeletal mesh
 			// It includes vertex and index buffers
 			UnrealConversionUtils::SetupRenderSections(
 				OutSkeletalMesh,
 				InMutableMesh,
 				0,
-				InBoneMap);
+				BoneMap,
+				0);
 
 			UnrealConversionUtils::CopyMutableVertexBuffers(
 				OutSkeletalMesh,
@@ -592,19 +394,18 @@ namespace MutableMeshPreviewUtils
 			
 			FSkeletalMeshLODRenderData& LODModel = OutSkeletalMesh->GetResourceForRendering()->LODRenderData[0];
 
-			// Generate the active bones array based on the bones found on BoneMap
-			TArray<uint16> ActiveBones;
+			// Ensure there's at least one bone in the bonemap of each render section, the root.
+			for (FSkelMeshRenderSection& RenderSection : LODModel.RenderSections)
 			{
-				ActiveBones.Reserve(InBoneMap.Num());
-				for (const uint16& BoneMapIndex : InBoneMap)
+				if (RenderSection.BoneMap.IsEmpty())
 				{
-					ActiveBones.AddUnique(BoneMapIndex);
+					RenderSection.BoneMap.Add(0);
 				}
-				ActiveBones.Sort();
 			}
 
-			LODModel.ActiveBoneIndices.Append(ActiveBones);
-			LODModel.RequiredBones.Append(ActiveBones);
+			// Add root as the only required bone
+			LODModel.ActiveBoneIndices.Add(0);
+			LODModel.RequiredBones.Add(0);
 			
 			// Copy index buffers or fail to generate the mesh
 			if (!UnrealConversionUtils::CopyMutableIndexBuffers(InMutableMesh, LODModel))
@@ -615,12 +416,6 @@ namespace MutableMeshPreviewUtils
 			// Update LOD and streaming data
 			LODModel.bIsLODOptional = false;
 			LODModel.bStreamedDataInlined = false;
-
-			if (InRefSkeletalMesh)
-			{
-				const FBoxSphereBounds Bounds = ((USkeletalMesh*)InRefSkeletalMesh)->GetBounds();
-				OutSkeletalMesh->SetImportedBounds(Bounds);
-			}
 			
 			return true;
 		}
@@ -630,11 +425,9 @@ namespace MutableMeshPreviewUtils
 	 * Publicly accessible functions of MutableMeshPreviewUtils
 	 */
 
-	USkeletalMesh* GenerateSkeletalMeshFromMutableMesh(mu::MeshPtrConst InMutableMesh,
-		const USkeletalMesh* InReferenceSkeletalMesh)
+	USkeletalMesh* GenerateSkeletalMeshFromMutableMesh(mu::MeshPtrConst InMutableMesh)
 	{
-		if (!InMutableMesh ||
-			!InReferenceSkeletalMesh)
+		if (!InMutableMesh)
 		{
 			return nullptr;
 		}
@@ -668,83 +461,31 @@ namespace MutableMeshPreviewUtils
 			// later found and used to generate a proper Skeletal Mesh
 		}
 		
-		
-		// Bone processed data containers
-		TArray<uint16> BoneMap;
-		TArray<FName> BoneNames;
-
-		// Provided a mutable skeleton do load the BoneNames and BoneMap from there
-		if (InMutableMesh->GetSkeleton())
-		{
-			// Prepare the skeleton data and get the bone names to be used on the generated mesh alongside with the indices of
-			// such bones on the final mesh.
-			PrepareSkeletonData(InMutableMesh,BoneNames,BoneMap);
-		}
-		else
-		{
-			// No mutable skeleton has been loaded.
-			// Use the data already set on the reference skeletal mesh to do the work of PrepareSkeletonData
-			// but without having access to the mutable skeleton
-
-			// Add all reference bones set on the reference skeleton onto our new mesh.
-			const FReferenceSkeleton& ReferenceMeshReferenceSkeleton = InReferenceSkeletalMesh->GetRefSkeleton();
-			for	(int32 RefBone = 0; RefBone < ReferenceMeshReferenceSkeleton.GetRawBoneNum(); RefBone++)
-			{
-				BoneMap.Add(RefBone);
-			}
-		}
-		
 		// Generate a new skeletal mesh to be filled up with the data found on the mutable mesh and the reference skeletal mesh
-		USkeletalMesh* GeneratedSkeletalMesh =
-			NewObject<USkeletalMesh>(GetTransientPackage(), NAME_None, RF_Transient);
+		USkeletalMesh* GeneratedSkeletalMesh = NewObject<USkeletalMesh>();
 
-		// build the reference skeleton for the Generated Skeletal mesh
-		bool bSuccess = BuildSkeletalMeshSkeletonData(InReferenceSkeletalMesh,GeneratedSkeletalMesh);
+		// Build the reference skeleton for the Generated Skeletal mesh
+		bool bSuccess = BuildSkeletalMeshSkeletonData(InMutableMesh, GeneratedSkeletalMesh);
 		if (!bSuccess)
 		{
 			return nullptr;
 		}
 		
-		// Load the data present on the reference skeletal mesh that could not be grabbed from the mu::Mesh
-		GeneratedSkeletalMesh->SetImportedBounds(InReferenceSkeletalMesh->GetBounds()); // TODO Generate Bounds
-		GeneratedSkeletalMesh->SetPhysicsAsset(InReferenceSkeletalMesh->GetPhysicsAsset());
-		GeneratedSkeletalMesh->SetEnablePerPolyCollision(InReferenceSkeletalMesh->GetEnablePerPolyCollision());
-
-		GeneratedSkeletalMesh->AllocateResourceForRendering();
-
 		// Prepare the mesh to be filled with rendering data (buffers)
 		{
-			GeneratedSkeletalMesh->GetResourceForRendering()->LODRenderData.Add(( GeneratedSkeletalMesh->GetResourceForRendering()->LODRenderData.Num()) ? &GeneratedSkeletalMesh->GetResourceForRendering()->LODRenderData[0] : new FSkeletalMeshLODRenderData());
-			GeneratedSkeletalMesh->GetLODInfoArray().Add(FSkeletalMeshLODInfo());
+			GeneratedSkeletalMesh->AllocateResourceForRendering();
+			GeneratedSkeletalMesh->GetResourceForRendering()->LODRenderData.Add(new FSkeletalMeshLODRenderData());
 
-			FSkeletalMeshLODInfo& LastLODInfo = GeneratedSkeletalMesh->GetLODInfoArray().Last();
-
-	#if WITH_EDITORONLY_DATA
+			FSkeletalMeshLODInfo& LastLODInfo = GeneratedSkeletalMesh->GetLODInfoArray().AddDefaulted_GetRef();
 			LastLODInfo.BuildSettings.bUseFullPrecisionUVs = true;
-	#endif
-
-			constexpr int32 LODIndex = 0;
-			if (InReferenceSkeletalMesh->GetLODInfoArray().IsValidIndex(LODIndex))
-			{
-				const FSkeletalMeshLODInfo& RefLODIndex = InReferenceSkeletalMesh->GetLODInfoArray()[LODIndex];
-				LastLODInfo.ScreenSize = RefLODIndex.ScreenSize;
-				LastLODInfo.LODHysteresis = RefLODIndex.LODHysteresis;
-				LastLODInfo.bSupportUniformlyDistributedSampling = RefLODIndex.bSupportUniformlyDistributedSampling;
-				LastLODInfo.bAllowCPUAccess = RefLODIndex.bAllowCPUAccess;
-			}
-			else
-			{
-				LastLODInfo.ScreenSize = 0.3f / (LODIndex + 1);
-				LastLODInfo.LODHysteresis = 0.02f;
-				UE_LOG(LogMutable, Warning, TEXT("Setting default values for LOD ScreenSize and LODHysteresis because the values cannot be found in the reference mesh"));
-			}
+			LastLODInfo.bAllowCPUAccess = false;
 		}
 
 		// Set the material data and initialize the sections that will be used later
 		BuildSkeletalMeshElementData(InMutableMesh,GeneratedSkeletalMesh);
 		
 		// Load all data from the mutable mesh onto the skeletal mesh object
-		bSuccess = BuildSkeletalMeshRenderData(InMutableMesh,InReferenceSkeletalMesh,BoneMap,GeneratedSkeletalMesh);
+		bSuccess = BuildSkeletalMeshRenderData(InMutableMesh, GeneratedSkeletalMesh);
 		if (!bSuccess)
 		{
 			return nullptr;
@@ -755,14 +496,4 @@ namespace MutableMeshPreviewUtils
 		
 		return GeneratedSkeletalMesh;
 	}
-
 }
-
-
-
-
-
-
-
-
-
diff --git a/Engine/Plugins/Experimental/Mutable/Source/CustomizableObjectEditor/Private/MuCOE/SMutableMeshViewer.cpp b/Engine/Plugins/Experimental/Mutable/Source/CustomizableObjectEditor/Private/MuCOE/SMutableMeshViewer.cpp
--- a/Engine/Plugins/Experimental/Mutable/Source/CustomizableObjectEditor/Private/MuCOE/SMutableMeshViewer.cpp
+++ b/Engine/Plugins/Experimental/Mutable/Source/CustomizableObjectEditor/Private/MuCOE/SMutableMeshViewer.cpp
@@ -198,12 +198,6 @@ void SMutableMeshViewer::Construct(const FArguments& InArgs)
 }
 
 
-void SMutableMeshViewer::AddReferencedObjects(FReferenceCollector& Collector)
-{
-	Collector.AddReferencedObject(SelectedReferenceMesh);
-}
-
-
 void SMutableMeshViewer::SetMesh(const mu::MeshPtrConst& InMesh)
 {
 	if (InMesh != MutableMesh)
@@ -232,80 +226,10 @@ void SMutableMeshViewer::SetMesh(const mu::MeshPtrConst& InMesh)
 	}
 }
 
-
-void SMutableMeshViewer::OnSelectedReferenceSkeletalMeshChanged(const FAssetData& AssetData)
-{
-	SelectedReferenceMesh = nullptr;
-
-	const UObject* SelectedAsset = AssetData.GetAsset();
-	if (SelectedAsset)
-	{
-		SelectedReferenceMesh = Cast<USkeletalMesh>(SelectedAsset);
-	}
-
-	// Notify the viewport from the selected reference mesh change
-	MeshViewport->SetReferenceMesh(SelectedReferenceMesh);
-}
-
-bool SMutableMeshViewer::OnShouldFilterAsset(const FAssetData& AssetData) const
-{
-	if (!AssetData.IsValid())
-	{
-		// Filter it out if it is not valid
-		return true;
-	}
-
-	// This is too slow to run for all 
-	const USkeleton* ParsedResourceSkeleton = Cast<USkeletalMesh>(AssetData.GetAsset())->GetSkeleton();
-	if (!ParsedResourceSkeleton)
-	{
-		// Only meshes with skeletons
-		return true;
-	}
-
-	return false;
-}
-
-FString SMutableMeshViewer::OnObjectPath() const
-{
-	if (SelectedReferenceMesh)
-	{
-		return SelectedReferenceMesh->GetPathName();
-	}
-
-	return FString();
-}
-
 TSharedRef<SWidget> SMutableMeshViewer::GenerateViewportSlates()
 {
 	TSharedRef<SWidget> Container = SNew(SVerticalBox)
 
-		// Reference selection space
-		+ SVerticalBox::Slot()
-		.AutoHeight()
-		[
-			SAssignNew(ReferenceMeshSelectionSpace, SBorder)
-			[
-				SNew(SHorizontalBox)
-
-				+ SHorizontalBox::Slot()
-				.AutoWidth()
-				[
-					SNew(STextBlock)
-					.Text(FText(LOCTEXT("ReferenceMesh", "Reference Mesh : ")))
-				]
-
-				+ SHorizontalBox::Slot()
-				[
-					SNew(SObjectPropertyEntryBox)
-					.AllowedClass(USkeletalMesh::StaticClass())
-					.ObjectPath(this, &SMutableMeshViewer::OnObjectPath)
-					.OnObjectChanged(this, &SMutableMeshViewer::OnSelectedReferenceSkeletalMeshChanged)
-					.OnShouldFilterAsset(this, &SMutableMeshViewer::OnShouldFilterAsset)
-				]
-			]
-		]
-
 		// User warning messages 
 		// + SVerticalBox::Slot()
 		// .AutoHeight()
diff --git a/Engine/Plugins/Experimental/Mutable/Source/CustomizableObjectEditor/Private/MuCOE/SMutableMeshViewer.h b/Engine/Plugins/Experimental/Mutable/Source/CustomizableObjectEditor/Private/MuCOE/SMutableMeshViewer.h
--- a/Engine/Plugins/Experimental/Mutable/Source/CustomizableObjectEditor/Private/MuCOE/SMutableMeshViewer.h
+++ b/Engine/Plugins/Experimental/Mutable/Source/CustomizableObjectEditor/Private/MuCOE/SMutableMeshViewer.h
@@ -45,7 +45,7 @@ struct FBufferElement
 };
 
 /** Widget designed to show the statistical data from a Mutable Mesh*/
-class SMutableMeshViewer : public SCompoundWidget, public FGCObject
+class SMutableMeshViewer : public SCompoundWidget
 {
 public:
 	SLATE_BEGIN_ARGS(SMutableMeshViewer) {}
@@ -55,17 +55,6 @@ public:
 	/** Builds the widget */
 	void Construct(const FArguments& InArgs);
 
-	// SWidget interface
-	// void Tick(const FGeometry& AllottedGeometry, const double InCurrentTime, const float InDeltaTime) override;
-
-	// FGCObject interface
-	virtual void AddReferencedObjects(FReferenceCollector& Collector) override;
-
-	virtual FString GetReferencerName() const override
-	{
-		return TEXT("SMutableMeshViewer");
-	}
-
 	/** Set the Mutable Mesh to be used for this widget */
 	void SetMesh(const mu::MeshPtrConst& InMesh);
 
@@ -111,9 +100,6 @@ private:
 	TSharedPtr<SListView<TSharedPtr<FBufferElement>>> FaceBuffersSlateView;
 
 private:
-	
-	/** The border where the object responsible for providing the user of a way to select objects is */
-	TSharedPtr<SBorder> ReferenceMeshSelectionSpace;
 
 	/** Generates all slate objects related with the Mesh Viewport Slate */
 	TSharedRef<SWidget> GenerateViewportSlates();
@@ -185,27 +171,4 @@ private:
 	TSharedRef<ITableRow> OnGenerateBufferChannelRow(TSharedPtr<FBufferChannelElement> InBufferChannel,
 	                                                 const TSharedRef<STableViewBase>& OwnerTable);
 	
-	
-	/*
-	 * Reference skeleton mesh selection
-	 */
-
-	/** Currently selected reference mesh */
-	const USkeletalMesh* SelectedReferenceMesh = nullptr;
-
-	/** Callback invoked when the selected asset on the reference skeletal mesh selection space has changed
-	 * @param AssetData - The data of the selected asset on the selection space
-	 */
-	void OnSelectedReferenceSkeletalMeshChanged(const FAssetData& AssetData);
-
-	/** Filters the assets found in order to know if the asset should be exposed on the UI for usage
-	 * @param AssetData - The data of the asset being checked for validation
-	 * @return - True means that the asset WILL NOT appear on the list of selectable resources
-	 */
-	bool OnShouldFilterAsset(const FAssetData& AssetData) const;
-
-	/**
-	 * Provides the SObjectPropertyEntryBox a way to know what to display in the selected object space
-	 */
-	FString OnObjectPath() const;
 };
diff --git a/Engine/Plugins/Experimental/Mutable/Source/CustomizableObjectEditor/Private/MuCOE/SMutableMeshViewport.cpp b/Engine/Plugins/Experimental/Mutable/Source/CustomizableObjectEditor/Private/MuCOE/SMutableMeshViewport.cpp
--- a/Engine/Plugins/Experimental/Mutable/Source/CustomizableObjectEditor/Private/MuCOE/SMutableMeshViewport.cpp
+++ b/Engine/Plugins/Experimental/Mutable/Source/CustomizableObjectEditor/Private/MuCOE/SMutableMeshViewport.cpp
@@ -47,22 +47,8 @@ void SMutableMeshViewport::SetMesh(const mu::MeshPtrConst& InMesh)
 	RefreshViewportContents();
 }
 
-void SMutableMeshViewport::SetReferenceMesh(const USkeletalMesh* InReferenceMesh)
-{
-	if (ReferenceSkeletalMesh == InReferenceMesh)
-	{
-		return;
-	}
-
-	ReferenceSkeletalMesh = InReferenceMesh;
-	
-	// Update viewport 
-	RefreshViewportContents();
-}
-
 void SMutableMeshViewport::AddReferencedObjects(FReferenceCollector& Collector)
 {
-	Collector.AddReferencedObject(ReferenceSkeletalMesh);
 	Collector.AddReferencedObject(SkeletalMeshComponent);
 }
 
@@ -125,8 +111,7 @@ void SMutableMeshViewport::SendMeshToViewport()
 bool SMutableMeshViewport::GenerateUnrealMesh()
 {
 	// Convert the provided mutable mesh to unreal with heavy aid of the Reference Skeletal Mesh
-	USkeletalMesh* GeneratedSkeletalMesh =
-		MutableMeshPreviewUtils::GenerateSkeletalMeshFromMutableMesh(MutableMesh,ReferenceSkeletalMesh);
+	USkeletalMesh* GeneratedSkeletalMesh = MutableMeshPreviewUtils::GenerateSkeletalMeshFromMutableMesh(MutableMesh);
 	
 	if (!GeneratedSkeletalMesh)
 	{
diff --git a/Engine/Plugins/Experimental/Mutable/Source/CustomizableObjectEditor/Private/MuCOE/SMutableMeshViewport.h b/Engine/Plugins/Experimental/Mutable/Source/CustomizableObjectEditor/Private/MuCOE/SMutableMeshViewport.h
--- a/Engine/Plugins/Experimental/Mutable/Source/CustomizableObjectEditor/Private/MuCOE/SMutableMeshViewport.h
+++ b/Engine/Plugins/Experimental/Mutable/Source/CustomizableObjectEditor/Private/MuCOE/SMutableMeshViewport.h
@@ -34,11 +34,6 @@ public:
      */
     void SetMesh(const mu::MeshPtrConst& InMesh);
 
-	/** Set the reference mesh used by the mutable to skeletal mesh conversion process
-	 * @param InReferenceMesh - The Skeletal mesh to be used to aid on the conversion of the MutableMesh to a SkeletalMesh
-	 */
-	void SetReferenceMesh(const USkeletalMesh* InReferenceMesh);
-
 	// FGCObject interface
 	virtual void AddReferencedObjects(FReferenceCollector& Collector) override;
 	virtual FString GetReferencerName() const override
@@ -63,9 +58,6 @@ private:
 	 *	UI  
 	 */
 	
-	/** Skeleton provided by the user */
-	const USkeletalMesh* ReferenceSkeletalMesh = nullptr;
-	
 	/** The preview scene that we are viewing */
 	TSharedPtr<FAdvancedPreviewScene> PreviewScene = nullptr;
 
diff --git a/Engine/Plugins/Experimental/Mutable/Source/CustomizableObjectEditor/Public/MuCOE/MutableMeshPreviewUtils.h b/Engine/Plugins/Experimental/Mutable/Source/CustomizableObjectEditor/Public/MuCOE/MutableMeshPreviewUtils.h
--- a/Engine/Plugins/Experimental/Mutable/Source/CustomizableObjectEditor/Public/MuCOE/MutableMeshPreviewUtils.h
+++ b/Engine/Plugins/Experimental/Mutable/Source/CustomizableObjectEditor/Public/MuCOE/MutableMeshPreviewUtils.h
@@ -14,17 +14,12 @@ class USkeletalMesh;
 namespace MutableMeshPreviewUtils
 {
 	/**
-	 * Method designed to provide a way of getting a USkeletalMesh based on the data found on a mutable mesh. The
-	 * usage of a Reference Skeletal mesh is required to aid on the generation of the skeletal mesh.
+	 * Method designed to provide a way of getting a USkeletalMesh based on the data found on a mutable mesh. 
 	 * @param InMutableMesh - The mutable mesh used as base for the generation of the USkeletalMesh
-	 * @param InReferenceSkeletalMesh - Skeletal mesh used as source for the data not stored by mutable mesh. It must be
-	 * "similar" in terms of bone structure to the mutable mesh. Example: If generating a mesh from a mutable mesh found
-	 * on an instance (using the mutable debugger) you may want to provide as reference skeletal mesh the one used by
-	 * the Customizable Object.
 	 * @return - A new Skeletal mesh based on the provided data. It may be a nullptr if the conversion fails.
 	 */
 	CUSTOMIZABLEOBJECTEDITOR_API USkeletalMesh* GenerateSkeletalMeshFromMutableMesh(
-		mu::MeshPtrConst InMutableMesh, const USkeletalMesh* InReferenceSkeletalMesh);
+		mu::MeshPtrConst InMutableMesh);
 
 }
 
diff --git a/Engine/Plugins/Experimental/Mutable/Source/MutableRuntime/Private/MuR/CodeRunner.cpp b/Engine/Plugins/Experimental/Mutable/Source/MutableRuntime/Private/MuR/CodeRunner.cpp
--- a/Engine/Plugins/Experimental/Mutable/Source/MutableRuntime/Private/MuR/CodeRunner.cpp
+++ b/Engine/Plugins/Experimental/Mutable/Source/MutableRuntime/Private/MuR/CodeRunner.cpp
@@ -1307,18 +1307,12 @@ namespace mu
 
                 if (pA && pB && pA->GetVertexCount() && pB->GetVertexCount())
                 {
-                    pResult = MeshMerge( pA.get(), pB.get() );
+                    pResult = MeshMerge( pA.get(), pB.get(), !args.newSurfaceID);
 
-                    if (!args.newSurfaceID)
+                    if (args.newSurfaceID)
                     {
-                        // Make it a single surface.
-                        pResult->m_surfaces.Empty();
-                        pResult->EnsureSurfaceData();
-                    }
-                    else
-                    {
-                        check(pB->GetSurfaceCount()==1);
-                        pResult->m_surfaces.Last().m_id=args.newSurfaceID;
+						check(pB->GetSurfaceCount() == 1);
+						pResult->m_surfaces.Last().m_id = args.newSurfaceID;
                     }
                 }
                 else if (pA && pA->GetVertexCount())
diff --git a/Engine/Plugins/Experimental/Mutable/Source/MutableRuntime/Private/MuR/Mesh.cpp b/Engine/Plugins/Experimental/Mutable/Source/MutableRuntime/Private/MuR/Mesh.cpp
--- a/Engine/Plugins/Experimental/Mutable/Source/MutableRuntime/Private/MuR/Mesh.cpp
+++ b/Engine/Plugins/Experimental/Mutable/Source/MutableRuntime/Private/MuR/Mesh.cpp
@@ -72,6 +72,7 @@ MeshPtr Mesh::Clone() const
 	
 	// Clone bone poses
 	pResult->BonePoses = BonePoses;
+	pResult->BoneMapIndices = BoneMapIndices;
 
 	// Clone SkeletonIDs
 	pResult->SkeletonIDs = SkeletonIDs;
@@ -153,6 +154,12 @@ MeshPtr Mesh::Clone(EMeshCloneFlags Flags) const
 		pResult->BonePoses = BonePoses;
 	}
 
+	// Clone BoneMap
+	if (EnumHasAnyFlags(Flags, EMeshCloneFlags::WithBoneMap))
+	{
+		pResult->BoneMapIndices = BoneMapIndices;
+	}
+
 	// Clone SkeletonIDs
 	if (EnumHasAnyFlags(Flags, EMeshCloneFlags::WithSkeletonIDs))
 	{
@@ -297,9 +304,10 @@ int Mesh::GetSurfaceCount() const
 
 
 //---------------------------------------------------------------------------------------------
-void Mesh::GetSurface( int surfaceIndex,
-                       int* firstVertex, int* vertexCount,
-                       int* firstIndex, int* indexCount ) const
+void Mesh::GetSurface( int32 surfaceIndex,
+                       int32* firstVertex, int32* vertexCount,
+                       int32* firstIndex, int32* indexCount,
+					   int32* BoneIndex, int32* BoneCount) const
 {
     int count = GetSurfaceCount();
 
@@ -312,6 +320,8 @@ void Mesh::GetSurface( int surfaceIndex,
             if (vertexCount) *vertexCount = surf.m_vertexCount;
             if (firstIndex) *firstIndex = surf.m_firstIndex;
             if (indexCount) *indexCount = surf.m_indexCount;
+            if (BoneIndex) *BoneIndex = surf.BoneMapIndex;
+            if (BoneCount) *BoneCount = surf.BoneMapCount;
         }
         else
         {
@@ -320,6 +330,8 @@ void Mesh::GetSurface( int surfaceIndex,
             if (vertexCount) *vertexCount = GetVertexCount();
             if (firstIndex) *firstIndex = 0;
             if (indexCount) *indexCount = GetIndexCount();
+			if (BoneIndex) *BoneIndex = 0;
+			if (BoneCount) *BoneCount = BoneMapIndices.Num();
         }
     }
     else
@@ -329,6 +341,8 @@ void Mesh::GetSurface( int surfaceIndex,
         if (vertexCount) *vertexCount = 0;
         if (firstIndex) *firstIndex = 0;
         if (indexCount) *indexCount = 0;
+		if (BoneIndex) *BoneIndex = 0;
+		if (BoneCount) *BoneCount = 0;
     }
 }
 
@@ -542,6 +556,8 @@ void mu::Mesh::GetBoneTransform(int32 BoneIndex, FTransform3f& Transform) const
 	Transform = BoneIndex > INDEX_NONE ? BonePoses[BoneIndex].BoneTransform : FTransform3f::Identity;
 }
 
+
+//---------------------------------------------------------------------------------------------
 EBoneUsageFlags Mesh::GetBoneUsageFlags(int32 BoneIndex) const
 {
 	check(BoneIndex >= 0 && BoneIndex < BonePoses.Num());
@@ -549,6 +565,20 @@ EBoneUsageFlags Mesh::GetBoneUsageFlags(int32 BoneIndex) const
 }
 
 
+//---------------------------------------------------------------------------------------------
+void Mesh::SetBoneMap(const TArray<uint16>& InBoneMap)
+{
+	BoneMapIndices = InBoneMap;
+}
+
+
+//---------------------------------------------------------------------------------------------
+const TArray<uint16>& Mesh::GetBoneMap() const
+{
+	return BoneMapIndices;
+}
+
+
 //---------------------------------------------------------------------------------------------
 int32 Mesh::GetSkeletonIDsCount() const
 {
@@ -973,6 +1003,7 @@ void Mesh::EnsureSurfaceData()
 		MESH_SURFACE s;
 		s.m_vertexCount = m_VertexBuffers.GetElementCount();
 		s.m_indexCount = m_IndexBuffers.GetElementCount();
+		s.BoneMapCount = BoneMapIndices.Num();
 		m_surfaces.Add(s);
 	}
 }
@@ -987,6 +1018,87 @@ void Mesh::ResetBufferIndices()
 }
 
 
+//---------------------------------------------------------------------------------------------
+void UnserialiseLegacySurfaces(InputArchive& arch, TArray<MESH_SURFACE>& OutMeshSurfaces)
+{
+	struct FMeshSurfaceLegacy
+	{
+		FMeshSurfaceLegacy()
+		{}
+
+		int32 m_firstVertex = 0;
+		int32 m_vertexCount = 0;
+		int32 m_firstIndex = 0;
+		int32 m_indexCount = 0;
+		uint32 m_id = 0;
+
+		void Unserialise(InputArchive& arch)
+		{
+			arch >> m_firstVertex;
+			arch >> m_vertexCount;
+			arch >> m_firstIndex;
+			arch >> m_indexCount;
+			arch >> m_id;
+		}
+	}; 
+
+	TArray<FMeshSurfaceLegacy> LegacyMeshSurfaces;
+	arch >> LegacyMeshSurfaces;
+	
+	const int32 NumSurfaces = LegacyMeshSurfaces.Num();
+	OutMeshSurfaces.SetNumZeroed(NumSurfaces);
+
+	MESH_SURFACE* Surfaces = OutMeshSurfaces.GetData();
+	FMeshSurfaceLegacy* LegacySurfaces = LegacyMeshSurfaces.GetData();
+
+	for (int32 SurfaceIndex = 0; SurfaceIndex < NumSurfaces; ++SurfaceIndex)
+	{
+		Surfaces->m_firstVertex = LegacySurfaces->m_firstVertex;
+		Surfaces->m_vertexCount = LegacySurfaces->m_vertexCount;
+		Surfaces->m_firstIndex = LegacySurfaces->m_firstIndex;
+		Surfaces->m_indexCount = LegacySurfaces->m_indexCount;
+		Surfaces->m_id = LegacySurfaces->m_id;
+	}
+}
+
+
+//-------------------------------------------------------------------------------------------------
+void MESH_SURFACE::Serialise(OutputArchive& arch) const
+{
+
+	const int32 ver = 0;
+	arch << ver;
+
+	arch << m_firstVertex;
+	arch << m_vertexCount;
+	arch << m_firstIndex;
+	arch << m_indexCount;
+	arch << BoneMapIndex;
+	arch << BoneMapCount;
+	
+	arch << m_id;
+
+}
+
+
+//-------------------------------------------------------------------------------------------------
+void MESH_SURFACE::Unserialise(InputArchive& arch)
+{
+	int32 ver = 0;
+	arch >> ver;
+	check(ver == 0);
+
+	arch >> m_firstVertex;
+	arch >> m_vertexCount;
+	arch >> m_firstIndex;
+	arch >> m_indexCount;
+	arch >> BoneMapIndex;
+	arch >> BoneMapCount;
+		 
+	arch >> m_id;
+}
+
+
 //-------------------------------------------------------------------------------------------------
 void Mesh::FACE_GROUP::Serialise(OutputArchive& arch) const
 {
@@ -1049,7 +1161,7 @@ void Mesh::FBonePose::Unserialise(InputArchive& arch)
 //-------------------------------------------------------------------------------------------------
 void Mesh::Serialise(OutputArchive& arch) const
 {
-	uint32 ver = 15;
+	uint32 ver = 16;
 	arch << ver;
 
 	arch << m_IndexBuffers;
@@ -1070,6 +1182,7 @@ void Mesh::Serialise(OutputArchive& arch) const
 	arch << m_tags;
 
 	arch << BonePoses;
+	arch << BoneMapIndices;
 
 	arch << AdditionalPhysicsBodies;
 }
@@ -1080,7 +1193,7 @@ void Mesh::Unserialise(InputArchive& arch)
 {
 	uint32 ver;
 	arch >> ver;
-	check(ver <= 15);
+	check(ver <= 16);
 
 	arch >> m_IndexBuffers;
 	arch >> m_VertexBuffers;
@@ -1104,7 +1217,16 @@ void Mesh::Unserialise(InputArchive& arch)
 	}
 
 	arch >> m_staticFormatFlags;
-	arch >> m_surfaces;
+
+	if (ver >= 16)
+	{
+		arch >> m_surfaces;
+	}
+	else
+	{
+		// Deserialize LegacySurfaces
+		UnserialiseLegacySurfaces(arch, m_surfaces);
+	}
 	arch >> m_faceGroups;
 
 	arch >> m_tags;
@@ -1127,6 +1249,25 @@ void Mesh::Unserialise(InputArchive& arch)
 		}
 	}
 
+	if (ver >= 16)
+	{
+		arch >> BoneMapIndices;
+	}
+	else
+	{
+		const int32 NumBonePoses = BonePoses.Num();
+		BoneMapIndices.SetNum(NumBonePoses);
+		for (int32 BoneIndex = 0; BoneIndex < NumBonePoses; ++BoneIndex)
+		{
+			BoneMapIndices[BoneIndex] = BoneIndex;
+		}
+
+		for (MESH_SURFACE& Surface : m_surfaces)
+		{
+			Surface.BoneMapCount = NumBonePoses;
+		}
+	}
+
 	if (ver >= 15)
 	{
 		arch >> AdditionalPhysicsBodies;
@@ -1606,7 +1747,6 @@ void Mesh::Log( FString& out, int32 BufferElementLimit)
 }
 
 
-MUTABLE_IMPLEMENT_POD_VECTOR_SERIALISABLE(MESH_SURFACE)
 MUTABLE_IMPLEMENT_ENUM_SERIALISABLE(EBoneUsageFlags)
 MUTABLE_IMPLEMENT_ENUM_SERIALISABLE(EMeshBufferType)
 MUTABLE_IMPLEMENT_ENUM_SERIALISABLE(EShapeBindingMethod)
diff --git a/Engine/Plugins/Experimental/Mutable/Source/MutableRuntime/Private/MuR/OpMeshFormat.cpp b/Engine/Plugins/Experimental/Mutable/Source/MutableRuntime/Private/MuR/OpMeshFormat.cpp
--- a/Engine/Plugins/Experimental/Mutable/Source/MutableRuntime/Private/MuR/OpMeshFormat.cpp
+++ b/Engine/Plugins/Experimental/Mutable/Source/MutableRuntime/Private/MuR/OpMeshFormat.cpp
@@ -663,20 +663,21 @@ namespace mu
 			pResult->m_layouts.Add(Layout->Clone());
 		}
 
-		pResult->ResetStaticFormatFlags();
-		pResult->EnsureSurfaceData();
-
 		pResult->m_tags = pSource->m_tags;
 
 		pResult->m_AdditionalBuffers = pSource->m_AdditionalBuffers;
 
 		pResult->BonePoses = pSource->BonePoses;
+		pResult->BoneMapIndices = pSource->BoneMapIndices;
 
 		pResult->SkeletonIDs = pSource->SkeletonIDs;
 
 		// A shallow copy is done here, it should not be a problem.
 		pResult->AdditionalPhysicsBodies = pSource->AdditionalPhysicsBodies;
 
+		pResult->ResetStaticFormatFlags();
+		pResult->EnsureSurfaceData();
+
 		return pResult;
 
 	}
diff --git a/Engine/Plugins/Experimental/Mutable/Source/MutableRuntime/Private/MuR/OpMeshMerge.h b/Engine/Plugins/Experimental/Mutable/Source/MutableRuntime/Private/MuR/OpMeshMerge.h
--- a/Engine/Plugins/Experimental/Mutable/Source/MutableRuntime/Private/MuR/OpMeshMerge.h
+++ b/Engine/Plugins/Experimental/Mutable/Source/MutableRuntime/Private/MuR/OpMeshMerge.h
@@ -16,7 +16,7 @@ namespace mu
 	//---------------------------------------------------------------------------------------------
 	//! Merge two meshes into one new mesh
 	//---------------------------------------------------------------------------------------------
-	inline MeshPtr MeshMerge(const Mesh* pFirst, const Mesh* pSecond)
+	inline MeshPtr MeshMerge(const Mesh* pFirst, const Mesh* pSecond, bool bMergeSurfaces)
 	{
 		MUTABLE_CPUPROFILER_SCOPE(MeshMerge);
 
@@ -34,122 +34,113 @@ namespace mu
 		{
 			MUTABLE_CPUPROFILER_SCOPE(Indices);
 
-			int firstCount = pFirst->GetIndexBuffers().GetElementCount();
-			int secondCount = pSecond->GetIndexBuffers().GetElementCount();
-			pResult->GetIndexBuffers().SetElementCount(firstCount + secondCount);
+			const int32 FirstCount = pFirst->GetIndexBuffers().GetElementCount();
+			const int32 SecondCount = pSecond->GetIndexBuffers().GetElementCount();
+			pResult->GetIndexBuffers().SetElementCount(FirstCount + SecondCount);
 
 			check(pFirst->GetIndexBuffers().GetBufferCount() <= 1);
 			check(pSecond->GetIndexBuffers().GetBufferCount() <= 1);
 			pResult->GetIndexBuffers().SetBufferCount(1);
 
+			MESH_BUFFER& ResultIndexBuffer = pResult->GetIndexBuffers().m_buffers[0];
+
+			const MESH_BUFFER& FirstIndexBuffer = pFirst->GetIndexBuffers().m_buffers[0];
+			const MESH_BUFFER& SecondIndexBuffer = pSecond->GetIndexBuffers().m_buffers[0];
+
+			// Avoid unused variable warnings
+			(void)FirstIndexBuffer;
+			(void)SecondIndexBuffer;
+
 			// This will be changed below if need to change the format of the index buffers.
-			MESH_BUFFER_FORMAT requiredBufferFormatChange = MBF_NONE;
+			MESH_BUFFER_FORMAT IndexBufferFormat = MBF_NONE;
 
-			if (firstCount && secondCount)
+			if (FirstCount && SecondCount)
 			{
-				const MESH_BUFFER& first = pFirst->GetIndexBuffers().m_buffers[0];
-				const MESH_BUFFER& second = pSecond->GetIndexBuffers().m_buffers[0];
-				check(first.m_channels == second.m_channels);
-				check(first.m_elementSize == second.m_elementSize);
-				check(first.m_channels.IsEmpty() || first.m_channels[0].m_format == second.m_channels[0].m_format);
-
-				// Avoid unused variable warnings
-				(void)first;
-				(void)second;
+				check(!FirstIndexBuffer.m_channels.IsEmpty());
+				check(FirstIndexBuffer.m_channels == SecondIndexBuffer.m_channels);
+				check(FirstIndexBuffer.m_elementSize == SecondIndexBuffer.m_elementSize);
 
 				// We need to know the total number of vertices in case we need to adjust the index buffer format.
-				uint64_t totalVertexCount = pFirst->GetVertexBuffers().GetElementCount()
-					+ pSecond->GetVertexBuffers().GetElementCount();
-				uint64_t maxValueBits = GetMeshFormatData(pFirst->GetIndexBuffers().m_buffers[0].m_channels[0].m_format).m_maxValueBits;
-				uint64_t maxSupportedVertices = uint64_t(1) << maxValueBits;
+				const uint64 totalVertexCount = pFirst->GetVertexBuffers().GetElementCount() + pSecond->GetVertexBuffers().GetElementCount();
+				const uint64 maxValueBits = GetMeshFormatData(pFirst->GetIndexBuffers().m_buffers[0].m_channels[0].m_format).m_maxValueBits;
+				const uint64 maxSupportedVertices = uint64(1) << maxValueBits;
+				
 				if (totalVertexCount > maxSupportedVertices)
 				{
-					if (totalVertexCount > 0xFFFF)
-					{
-						requiredBufferFormatChange = MBF_UINT32;
-					}
-					else
-					{
-						requiredBufferFormatChange = MBF_UINT16;
-					}
+					IndexBufferFormat = totalVertexCount > MAX_uint16 ? MBF_UINT32 : MBF_UINT16;
 				}
-			}
 
-			MESH_BUFFER& result = pResult->GetIndexBuffers().m_buffers[0];
-
-			if (requiredBufferFormatChange != MBF_NONE)
+			}
+			
+			if (IndexBufferFormat != MBF_NONE)
 			{
 				// We only support vertex indices in case of having to change the format.
-				check(pFirst->GetIndexBuffers().m_buffers[0].m_channels.Num() == 1);
-
-				result.m_channels.SetNum(1);
-				result.m_channels[0].m_semantic = MBS_VERTEXINDEX;
-				result.m_channels[0].m_format = requiredBufferFormatChange;
-				result.m_channels[0].m_componentCount = 1;
-				result.m_channels[0].m_semanticIndex = 0;
-				result.m_channels[0].m_offset = 0;
-				result.m_elementSize = GetMeshFormatData(requiredBufferFormatChange).m_size;
+				check(FirstIndexBuffer.m_channels.Num() == 1);
+
+				ResultIndexBuffer.m_channels.SetNum(1);
+				ResultIndexBuffer.m_channels[0].m_semantic = MBS_VERTEXINDEX;
+				ResultIndexBuffer.m_channels[0].m_format = IndexBufferFormat;
+				ResultIndexBuffer.m_channels[0].m_componentCount = 1;
+				ResultIndexBuffer.m_channels[0].m_semanticIndex = 0;
+				ResultIndexBuffer.m_channels[0].m_offset = 0;
+				ResultIndexBuffer.m_elementSize = GetMeshFormatData(IndexBufferFormat).m_size;
 			}
-			else if (firstCount)
+			else if (FirstCount)
 			{
-				const MESH_BUFFER& first = pFirst->GetIndexBuffers().m_buffers[0];
-				result.m_channels = first.m_channels;
-				result.m_elementSize = first.m_elementSize;
+				ResultIndexBuffer.m_channels = FirstIndexBuffer.m_channels;
+				ResultIndexBuffer.m_elementSize = FirstIndexBuffer.m_elementSize;
 			}
-			else if (secondCount)
+			else if (SecondCount)
 			{
-				const MESH_BUFFER& second = pSecond->GetIndexBuffers().m_buffers[0];
-				result.m_channels = second.m_channels;
-				result.m_elementSize = second.m_elementSize;
+				ResultIndexBuffer.m_channels = SecondIndexBuffer.m_channels;
+				ResultIndexBuffer.m_elementSize = SecondIndexBuffer.m_elementSize;
 			}
 
-			result.m_data.SetNum(result.m_elementSize * (firstCount + secondCount));
+			ResultIndexBuffer.m_data.SetNum(ResultIndexBuffer.m_elementSize * (FirstCount + SecondCount));
 
-			check(result.m_channels.Num() == 1);
-			check(result.m_channels[0].m_semantic == MBS_VERTEXINDEX);
+			check(ResultIndexBuffer.m_channels.Num() == 1);
+			check(ResultIndexBuffer.m_channels[0].m_semantic == MBS_VERTEXINDEX);
 
-			if (result.m_data.Num())
+			if (!ResultIndexBuffer.m_data.IsEmpty())
 			{
-				if (firstCount)
+				if (FirstCount)
 				{
-					const MESH_BUFFER& first = pFirst->GetIndexBuffers().m_buffers[0];
-
-					if (requiredBufferFormatChange == MBF_NONE
-						|| requiredBufferFormatChange == first.m_channels[0].m_format)
+					if (IndexBufferFormat == MBF_NONE
+						|| IndexBufferFormat == FirstIndexBuffer.m_channels[0].m_format)
 					{
-						FMemory::Memcpy(&result.m_data[0],
-							&first.m_data[0],
-							first.m_elementSize * firstCount);
+						FMemory::Memcpy(&ResultIndexBuffer.m_data[0],
+							&FirstIndexBuffer.m_data[0],
+							FirstIndexBuffer.m_elementSize * FirstCount);
 					}
 					else
 					{
 						// Conversion required
-						const uint8_t* pSource = &first.m_data[0];
-						uint8_t* pDest = &result.m_data[0];
-						switch (requiredBufferFormatChange)
+						const uint8_t* pSource = &FirstIndexBuffer.m_data[0];
+						uint8_t* pDest = &ResultIndexBuffer.m_data[0];
+						switch (IndexBufferFormat)
 						{
 						case MBF_UINT32:
 						{
-							switch (first.m_channels[0].m_format)
+							switch (FirstIndexBuffer.m_channels[0].m_format)
 							{
 							case MBF_UINT16:
 							{
-								for (int v = 0; v < firstCount; ++v)
+								for (int v = 0; v < FirstCount; ++v)
 								{
 									*(uint32_t*)pDest = *(const uint16*)pSource;
-									pSource += first.m_elementSize;
-									pDest += result.m_elementSize;
+									pSource += FirstIndexBuffer.m_elementSize;
+									pDest += ResultIndexBuffer.m_elementSize;
 								}
 								break;
 							}
 
 							case MBF_UINT8:
 							{
-								for (int v = 0; v < firstCount; ++v)
+								for (int v = 0; v < FirstCount; ++v)
 								{
 									*(uint32_t*)pDest = *(const uint8_t*)pSource;
-									pSource += first.m_elementSize;
-									pDest += result.m_elementSize;
+									pSource += FirstIndexBuffer.m_elementSize;
+									pDest += ResultIndexBuffer.m_elementSize;
 								}
 								break;
 							}
@@ -163,16 +154,16 @@ namespace mu
 
 						case MBF_UINT16:
 						{
-							switch (first.m_channels[0].m_format)
+							switch (FirstIndexBuffer.m_channels[0].m_format)
 							{
 
 							case MBF_UINT8:
 							{
-								for (int v = 0; v < firstCount; ++v)
+								for (int v = 0; v < FirstCount; ++v)
 								{
 									*(uint16*)pDest = *(const uint8_t*)pSource;
-									pSource += first.m_elementSize;
-									pDest += result.m_elementSize;
+									pSource += FirstIndexBuffer.m_elementSize;
+									pDest += ResultIndexBuffer.m_elementSize;
 								}
 								break;
 							}
@@ -191,29 +182,28 @@ namespace mu
 					}
 				}
 
-				if (secondCount)
+				if (SecondCount)
 				{
-					const MESH_BUFFER& second = pSecond->GetIndexBuffers().m_buffers[0];
-					const uint8_t* pSource = &second.m_data[0];
-					uint8_t* pDest = &result.m_data[result.m_elementSize * firstCount];
+					const uint8_t* pSource = &SecondIndexBuffer.m_data[0];
+					uint8_t* pDest = &ResultIndexBuffer.m_data[ResultIndexBuffer.m_elementSize * FirstCount];
 
 					uint32_t firstVertexCount = pFirst->GetVertexBuffers().GetElementCount();
 
-					if (requiredBufferFormatChange == MBF_NONE
-						|| requiredBufferFormatChange == second.m_channels[0].m_format)
+					if (IndexBufferFormat == MBF_NONE
+						|| IndexBufferFormat == SecondIndexBuffer.m_channels[0].m_format)
 					{
-						switch (second.m_channels[0].m_format)
+						switch (SecondIndexBuffer.m_channels[0].m_format)
 						{
 						case MBF_INT32:
 						case MBF_UINT32:
 						case MBF_NINT32:
 						case MBF_NUINT32:
 						{
-							for (int v = 0; v < secondCount; ++v)
+							for (int v = 0; v < SecondCount; ++v)
 							{
 								*(uint32_t*)pDest = firstVertexCount + *(const uint32_t*)pSource;
-								pSource += second.m_elementSize;
-								pDest += result.m_elementSize;
+								pSource += SecondIndexBuffer.m_elementSize;
+								pDest += ResultIndexBuffer.m_elementSize;
 							}
 							break;
 						}
@@ -223,11 +213,11 @@ namespace mu
 						case MBF_NINT16:
 						case MBF_NUINT16:
 						{
-							for (int v = 0; v < secondCount; ++v)
+							for (int v = 0; v < SecondCount; ++v)
 							{
 								*(uint16*)pDest = uint16(firstVertexCount) + *(const uint16*)pSource;
-								pSource += second.m_elementSize;
-								pDest += result.m_elementSize;
+								pSource += SecondIndexBuffer.m_elementSize;
+								pDest += ResultIndexBuffer.m_elementSize;
 							}
 							break;
 						}
@@ -237,11 +227,11 @@ namespace mu
 						case MBF_NINT8:
 						case MBF_NUINT8:
 						{
-							for (int v = 0; v < secondCount; ++v)
+							for (int v = 0; v < SecondCount; ++v)
 							{
 								*(uint8_t*)pDest = uint8_t(firstVertexCount) + *(const uint8_t*)pSource;
-								pSource += second.m_elementSize;
-								pDest += result.m_elementSize;
+								pSource += SecondIndexBuffer.m_elementSize;
+								pDest += ResultIndexBuffer.m_elementSize;
 							}
 							break;
 						}
@@ -254,23 +244,23 @@ namespace mu
 					else
 					{
 						// Format conversion required
-						switch (requiredBufferFormatChange)
+						switch (IndexBufferFormat)
 						{
 
 						case MBF_UINT32:
 						{
-							switch (second.m_channels[0].m_format)
+							switch (SecondIndexBuffer.m_channels[0].m_format)
 							{
 							case MBF_INT16:
 							case MBF_UINT16:
 							case MBF_NINT16:
 							case MBF_NUINT16:
 							{
-								for (int v = 0; v < secondCount; ++v)
+								for (int v = 0; v < SecondCount; ++v)
 								{
 									*(uint32_t*)pDest = uint32_t(firstVertexCount) + *(const uint16*)pSource;
-									pSource += second.m_elementSize;
-									pDest += result.m_elementSize;
+									pSource += SecondIndexBuffer.m_elementSize;
+									pDest += ResultIndexBuffer.m_elementSize;
 								}
 								break;
 							}
@@ -280,11 +270,11 @@ namespace mu
 							case MBF_NINT8:
 							case MBF_NUINT8:
 							{
-								for (int v = 0; v < secondCount; ++v)
+								for (int v = 0; v < SecondCount; ++v)
 								{
 									*(uint32_t*)pDest = uint32_t(firstVertexCount) + *(const uint8_t*)pSource;
-									pSource += second.m_elementSize;
-									pDest += result.m_elementSize;
+									pSource += SecondIndexBuffer.m_elementSize;
+									pDest += ResultIndexBuffer.m_elementSize;
 								}
 								break;
 							}
@@ -299,18 +289,18 @@ namespace mu
 
 						case MBF_UINT16:
 						{
-							switch (second.m_channels[0].m_format)
+							switch (SecondIndexBuffer.m_channels[0].m_format)
 							{
 							case MBF_INT8:
 							case MBF_UINT8:
 							case MBF_NINT8:
 							case MBF_NUINT8:
 							{
-								for (int v = 0; v < secondCount; ++v)
+								for (int v = 0; v < SecondCount; ++v)
 								{
 									*(uint16*)pDest = uint16(firstVertexCount) + *(const uint8_t*)pSource;
-									pSource += second.m_elementSize;
-									pDest += result.m_elementSize;
+									pSource += SecondIndexBuffer.m_elementSize;
+									pDest += ResultIndexBuffer.m_elementSize;
 								}
 								break;
 							}
@@ -338,20 +328,20 @@ namespace mu
 		{
 			MUTABLE_CPUPROFILER_SCOPE(Faces);
 
-			int firstCount = pFirst->GetFaceBuffers().GetElementCount();
-			int secondCount = pSecond->GetFaceBuffers().GetElementCount();
-			pResult->GetFaceBuffers().SetElementCount(firstCount + secondCount);
+			const int32 FirstCount = pFirst->GetFaceBuffers().GetElementCount();
+			const int32 SecondCount = pSecond->GetFaceBuffers().GetElementCount();
+			pResult->GetFaceBuffers().SetElementCount(FirstCount + SecondCount);
 			pResult->GetFaceBuffers().SetBufferCount(pFirst->GetFaceBuffers().GetBufferCount());
 
 			// Merge only the buffers present in the first mesh
-			for (int b = 0; b < pResult->GetFaceBuffers().GetBufferCount(); ++b)
+			for (int32 b = 0; b < pResult->GetFaceBuffers().GetBufferCount(); ++b)
 			{
 				const MESH_BUFFER& first = pFirst->GetFaceBuffers().m_buffers[b];
 
 				MESH_BUFFER& result = pResult->GetFaceBuffers().m_buffers[b];
 				result.m_channels = first.m_channels;
 				result.m_elementSize = first.m_elementSize;
-				result.m_data.SetNum(result.m_elementSize * (firstCount + secondCount));
+				result.m_data.SetNum(result.m_elementSize * (FirstCount + SecondCount));
 
 				MESH_BUFFER_SEMANTIC semantic = MBS_NONE;
 				int semanticIndex = 0;
@@ -361,15 +351,15 @@ namespace mu
 				pFirst->GetFaceBuffers().GetChannel
 				(b, 0, &semantic, &semanticIndex, &type, &components, &offset);
 
-				if (firstCount)
+				if (FirstCount)
 				{
 					FMemory::Memcpy(&result.m_data[0],
 						&first.m_data[0],
-						result.m_elementSize * firstCount
+						result.m_elementSize * FirstCount
 					);
 				}
 
-				if (secondCount)
+				if (SecondCount)
 				{
 					// Find in the second mesh
 					int otherBuffer = -1;
@@ -384,15 +374,15 @@ namespace mu
 						check(first.m_channels == second.m_channels);
 
 						// Raw copy
-						FMemory::Memcpy(&result.m_data[result.m_elementSize * firstCount],
+						FMemory::Memcpy(&result.m_data[result.m_elementSize * FirstCount],
 							&second.m_data[0],
-							result.m_elementSize * secondCount);
+							result.m_elementSize * SecondCount);
 					}
 					else
 					{
 						// Fill with zeroes
-						FMemory::Memzero(&result.m_data[result.m_elementSize * firstCount],
-							result.m_elementSize * secondCount);
+						FMemory::Memzero(&result.m_data[result.m_elementSize * FirstCount],
+							result.m_elementSize * SecondCount);
 					}
 				}
 			}
@@ -422,18 +412,6 @@ namespace mu
 		}
 
 
-		// Surfaces
-		//-----------------
-		pResult->m_surfaces = pFirst->m_surfaces;
-
-		for (const auto& s : pSecond->m_surfaces)
-		{
-			pResult->m_surfaces.Add(s);
-			pResult->m_surfaces.Last().m_firstVertex += pFirst->GetVertexCount();
-			pResult->m_surfaces.Last().m_firstIndex += pFirst->GetIndexCount();
-		}
-
-
 		// Skeleton
 		//---------------------------
 
@@ -446,20 +424,17 @@ namespace mu
 		}
 
 		// Do they have the same skeleton?
-		bool remapNeeded = pFirst->GetSkeleton() != pSecond->GetSkeleton();
+		bool bMergeSkeletons = pFirst->GetSkeleton() != pSecond->GetSkeleton();
 
 		// Are they different skeletons but with the same data?
-		if (remapNeeded && pFirst->GetSkeleton() && pSecond->GetSkeleton())
+		if (bMergeSkeletons && pFirst->GetSkeleton() && pSecond->GetSkeleton())
 		{
-			remapNeeded = !(*pFirst->GetSkeleton()
-				==
-				*pSecond->GetSkeleton());
+			bMergeSkeletons = !(*pFirst->GetSkeleton() == *pSecond->GetSkeleton());
 		}
 
+		TArray<uint16> SecondToResultBoneIndices;
 
-		TArray<int32> secondToFirstBones;
-
-		if (remapNeeded)
+		if (bMergeSkeletons)
 		{
 			MUTABLE_CPUPROFILER_SCOPE(Skeleton);
 			Ptr<Skeleton> pResultSkeleton;
@@ -467,47 +442,46 @@ namespace mu
 			mu::SkeletonPtrConst pFirstSkeleton = pFirst->GetSkeleton();
 			mu::SkeletonPtrConst pSecondSkeleton = pSecond->GetSkeleton();
 
-			const int32 pFirstBoneCount = pFirstSkeleton ? pFirstSkeleton->GetBoneCount() : 0;
-			const int32 pSecondBoneCount = pSecondSkeleton ? pSecondSkeleton->GetBoneCount() : 0;
+			const int32 NumBonesFirst = pFirstSkeleton ? pFirstSkeleton->GetBoneCount() : 0;
+			const int32 NumBonesSecond = pSecondSkeleton ? pSecondSkeleton->GetBoneCount() : 0;
 
 			pResultSkeleton = pFirstSkeleton ? pFirstSkeleton->Clone() : new Skeleton;
 			pResult->SetSkeleton(pResultSkeleton);
 
-			secondToFirstBones.SetNumUninitialized(pSecondBoneCount);
+			SecondToResultBoneIndices.SetNumUninitialized(NumBonesSecond);
 
 			// Merge pSecond and build the remap array 
-			for (int32 ob = 0; ob < pSecondBoneCount; ++ob)
+			for (int32 SecondBoneIndex = 0; SecondBoneIndex < NumBonesSecond; ++SecondBoneIndex)
 			{
-				int32 tb = INDEX_NONE;
-				for (int32 c = 0; tb < 0 && c < pResultSkeleton->m_bones.Num(); ++c)
+				int16 BoneIndex = INDEX_NONE;
+				for (int32 FirstBoneIndex = 0; FirstBoneIndex < NumBonesFirst; ++FirstBoneIndex)
 				{
-					if (pResultSkeleton->m_bones[c] == pSecondSkeleton->GetBoneName(ob))
+					if (pResultSkeleton->m_bones[FirstBoneIndex] == pSecondSkeleton->GetBoneName(SecondBoneIndex))
 					{
-						tb = c;
+						BoneIndex = FirstBoneIndex;
 						break;
 					}
 				}
 
 				// Add a new bone
-				if (tb == INDEX_NONE)
+				if (BoneIndex == INDEX_NONE)
 				{
-					tb = pResultSkeleton->m_bones.Num();
-					pResultSkeleton->m_bones.Add(pSecondSkeleton->m_bones[ob]);
+					BoneIndex = pResultSkeleton->m_bones.Add(pSecondSkeleton->m_bones[SecondBoneIndex]);
 
 					// Add an incorrect index, to be fixed below in case the parent index is later in the bone array.
-					pResultSkeleton->m_boneParents.Add(pSecondSkeleton->m_boneParents[ob]);
+					pResultSkeleton->m_boneParents.Add(pSecondSkeleton->m_boneParents[SecondBoneIndex]);
 				}
 
-				secondToFirstBones[ob] = tb;
+				SecondToResultBoneIndices[SecondBoneIndex] = (uint16)BoneIndex;
 			}
 
 			// Fix second mesh bone parents
-			for (int32 ob = pFirstBoneCount; ob < pResultSkeleton->m_boneParents.Num(); ++ob)
+			for (int32 ob = NumBonesFirst; ob < pResultSkeleton->m_boneParents.Num(); ++ob)
 			{
 				int16_t secondMeshIndex = pResultSkeleton->m_boneParents[ob];
 				if (secondMeshIndex != INDEX_NONE)
 				{
-					pResultSkeleton->m_boneParents[ob] = (uint16)secondToFirstBones[secondMeshIndex];
+					pResultSkeleton->m_boneParents[ob] = SecondToResultBoneIndices[secondMeshIndex];
 				}
 			}
 		}
@@ -517,6 +491,76 @@ namespace mu
 		}
 
 
+		// Surfaces
+		//---------------------------
+		
+		// Remap bone indices if we merge surfaces since bonemaps will be merged too.
+		bool bRemapBoneIndices = false;
+		TArray<uint16> RemappedBoneMapIndices;
+
+		// Used to know the format of the bone index buffer
+		uint32 MaxBoneMapIndices = 0;
+		const int32 NumSecondBoneMapIndices = pSecond->BoneMapIndices.Num();
+
+		{
+			MUTABLE_CPUPROFILER_SCOPE(Surfaces);
+			
+			const int32 NumFirstBoneMapIndices = pFirst->BoneMapIndices.Num();
+			pResult->BoneMapIndices = pFirst->BoneMapIndices;
+
+			if (bMergeSurfaces)
+			{
+				// Merge BoneMaps
+				RemappedBoneMapIndices.SetNumUninitialized(NumSecondBoneMapIndices);
+
+				for (uint16 SecondBoneMapIndex = 0; SecondBoneMapIndex < NumSecondBoneMapIndices; ++SecondBoneMapIndex)
+				{
+					const int32 SecondBoneIndex = pSecond->BoneMapIndices[SecondBoneMapIndex];
+					const int32 BoneIndex = bMergeSkeletons ? SecondToResultBoneIndices[SecondBoneIndex] : SecondBoneIndex;
+					const int32 BoneMapIndex = pResult->BoneMapIndices.AddUnique(BoneIndex);
+					RemappedBoneMapIndices[SecondBoneMapIndex] = BoneMapIndex;
+
+					bRemapBoneIndices = bRemapBoneIndices || BoneMapIndex != SecondBoneMapIndex;
+				}
+
+				MESH_SURFACE& NewSurface = pResult->m_surfaces.AddDefaulted_GetRef();
+				NewSurface.m_vertexCount = pFirst->GetVertexCount() + pSecond->GetVertexCount();
+				NewSurface.m_indexCount = pFirst->GetIndexCount() + pSecond->GetIndexCount();
+				NewSurface.BoneMapCount = pResult->BoneMapIndices.Num();
+			}
+			else
+			{
+				// Add the bonemaps of the second mesh with the fixed BoneIndices
+				const int32 NumBoneMapIndices = NumFirstBoneMapIndices + NumSecondBoneMapIndices;
+				pResult->BoneMapIndices.Reserve(NumBoneMapIndices);
+
+				for (const uint16& BoneIndex : pSecond->BoneMapIndices)
+				{
+					pResult->BoneMapIndices.Add(bMergeSkeletons ? SecondToResultBoneIndices[BoneIndex] : BoneIndex);
+				}
+
+				// Add pFirst surfaces
+				pResult->m_surfaces = pFirst->m_surfaces;
+
+				const int32 FirstVertexIndex = pFirst->GetVertexCount();
+				const int32 FirstIndexIndex = pFirst->GetIndexCount();
+
+				check(pSecond->m_surfaces.Num() == 1);
+				MESH_SURFACE& NewSurface = pResult->m_surfaces.Add_GetRef(pSecond->m_surfaces[0]);
+				NewSurface.m_firstVertex += FirstVertexIndex;
+				NewSurface.m_firstIndex += FirstIndexIndex;
+				NewSurface.BoneMapIndex += NumFirstBoneMapIndices;
+			}
+
+			for (const MESH_SURFACE& Surface : pResult->m_surfaces)
+			{
+				MaxBoneMapIndices = FMath::Max(MaxBoneMapIndices, Surface.BoneMapCount);
+			}
+
+			pResult->BoneMapIndices.Shrink();
+		}
+
+
 		// Pose
 		//---------------------------
 		{
@@ -708,18 +752,44 @@ namespace mu
 		{
             MUTABLE_CPUPROFILER_SCOPE(Vertices);
 
-            int firstCount = pFirst->GetVertexBuffers().GetElementCount();
-			int secondCount = pSecond->GetVertexBuffers().GetElementCount();
+            const int32 FirstCount = pFirst->GetVertexBuffers().GetElementCount();
+			const int32 SecondCount = pSecond->GetVertexBuffers().GetElementCount();
 
 			// TODO: when formats match, which at runtime should be always.
-			bool fastPath = pFirst->GetVertexBuffers().HasSameFormat( pSecond->GetVertexBuffers() );
+			bool FastPath = pFirst->GetVertexBuffers().HasSameFormat( pSecond->GetVertexBuffers() );
+
+			// Check if the format of the BoneIndex buffer has to change
+			bool bChangeBoneIndicesFormat = false;
+			MESH_BUFFER_FORMAT BoneIndexFormat = MaxBoneMapIndices > MAX_uint8 ? MBF_UINT16 : MBF_UINT8;
 
-			fastPath = fastPath && !remapNeeded;
+			{
+				// Iterate all vertex buffers 
+				FMeshBufferSet& VertexBuffers = pResult->GetVertexBuffers();
+				for (int32 VertexBufferIndex = 0; !bChangeBoneIndicesFormat &&  VertexBufferIndex < VertexBuffers.m_buffers.Num(); ++VertexBufferIndex)
+				{
+					MESH_BUFFER& Buffer = VertexBuffers.m_buffers[VertexBufferIndex];
+
+					const int32 elemSize = VertexBuffers.GetElementSize(VertexBufferIndex);
+					const int32 firstSize = FirstCount * elemSize;
+
+					const int32 ChannelsCount = VertexBuffers.GetBufferChannelCount(VertexBufferIndex);
+					for (int32 ChannelIndex = 0; ChannelIndex < ChannelsCount; ++ChannelIndex)
+					{
+						if (Buffer.m_channels[ChannelIndex].m_semantic == MBS_BONEINDICES)
+						{
+							bChangeBoneIndicesFormat = Buffer.m_channels[ChannelIndex].m_format != BoneIndexFormat;
+							break;
+						}
+					}
+				}
+
+				FastPath = FastPath && !bChangeBoneIndicesFormat;
+			}
 
 			MeshPtrConst pVFirst;
 			MeshPtrConst pVSecond;
 
-			if ( !fastPath )
+			if ( !FastPath )
 			{
                 MUTABLE_CPUPROFILER_SCOPE(SlowPath);
 
@@ -815,29 +885,21 @@ namespace mu
                     }
                 }
 
-				// Change the format of the bone indices buffer if there are more than 256 bones.
-                int pResultBoneCount =
-                    pResult->GetSkeleton()  ? pResult->GetSkeleton()->GetBoneCount() : 0;
-                if (pResultBoneCount > 256)
+				// Change the format of the bone indices buffer
+				if(bChangeBoneIndicesFormat)
 				{
-					// Desired format
-					MESH_BUFFER_FORMAT format = pResultBoneCount > 65535 ? MBF_UINT32 : MBF_UINT16;
-
 					// Iterate all vertex buffers and update the format
 					FMeshBufferSet& VertexBuffers = pResult->GetVertexBuffers();
 					for (int32 VertexBufferIndex = 0; VertexBufferIndex < VertexBuffers.m_buffers.Num(); ++VertexBufferIndex)
 					{
 						MESH_BUFFER& result = VertexBuffers.m_buffers[VertexBufferIndex];
 
-						const int32 elemSize = VertexBuffers.GetElementSize(VertexBufferIndex);
-						const int32 firstSize = firstCount * elemSize;
-
 						const int32 ChannelsCount = VertexBuffers.GetBufferChannelCount(VertexBufferIndex);
 						for (int32 ChannelIndex = 0; ChannelIndex < ChannelsCount; ++ChannelIndex)
 						{
 							if (result.m_channels[ChannelIndex].m_semantic == MBS_BONEINDICES)
 							{
-								result.m_channels[ChannelIndex].m_format = format;
+								result.m_channels[ChannelIndex].m_format = BoneIndexFormat;
 
 								// Reset offsets
 								int32 offset = 0;
@@ -855,7 +917,7 @@ namespace mu
 				}
 
                 // Allocate vertices
-                pResult->GetVertexBuffers().SetElementCount( firstCount + secondCount );
+                pResult->GetVertexBuffers().SetElementCount( FirstCount + SecondCount );
 
 				// Convert the source meshes to the new format
                 if (pFirst->GetVertexBuffers().HasSameFormat(pResult->GetVertexBuffers()))
@@ -890,7 +952,7 @@ namespace mu
 
 
 			pResult->m_VertexBuffers = pVFirst->m_VertexBuffers;
-			pResult->GetVertexBuffers().SetElementCount( firstCount + secondCount );
+			pResult->GetVertexBuffers().SetElementCount( FirstCount + SecondCount );
 
             // first copy all the vertex data
 			{
@@ -898,132 +960,112 @@ namespace mu
 				{
 					MESH_BUFFER& result = pResult->GetVertexBuffers().m_buffers[vb];
 					const MESH_BUFFER& second = pVSecond->GetVertexBuffers().m_buffers[vb];
-					if ( secondCount )
+					if ( SecondCount )
 					{
 						int elemSize = pResult->GetVertexBuffers().GetElementSize((int)vb);
-						int firstSize = firstCount * elemSize;
-						int secondSize = secondCount * elemSize;
+						int firstSize = FirstCount * elemSize;
+						int secondSize = SecondCount * elemSize;
 
 						FMemory::Memcpy( &result.m_data[firstSize], &second.m_data[0], secondSize );
 					}
 				}
 			}
 
-            // TODO This could eventually be part of the mesh format: force the same skeleton.
-            if (remapNeeded)
+
+            if (bRemapBoneIndices)
             {
                 MUTABLE_CPUPROFILER_SCOPE(Remap);
 
-                // We need to remap the bone indices of the second mesh vertices that we already copied
-                // to result
-				for ( int32 vb=0; vb<pResult->GetVertexBuffers().m_buffers.Num(); ++vb )
+				// We need to remap the bone indices of the second mesh vertices that we already copied
+				// to result
+				check(!RemappedBoneMapIndices.IsEmpty())
+
+				// Iterate all vertex buffers and update the format
+				FMeshBufferSet& VertexBuffers = pResult->GetVertexBuffers();
+				for (int32 VertexBufferIndex = 0; VertexBufferIndex < VertexBuffers.m_buffers.Num(); ++VertexBufferIndex)
 				{
-					MESH_BUFFER& result = pResult->GetVertexBuffers().m_buffers[vb];
+					MESH_BUFFER& ResultBuffer = VertexBuffers.m_buffers[VertexBufferIndex];
 
-					int elemSize = pResult->GetVertexBuffers().GetElementSize((int)vb);
-                    int firstSize = firstCount * elemSize;
+					const int32 ElemSize = VertexBuffers.GetElementSize(VertexBufferIndex);
+					const int32 FirstSize = FirstCount * ElemSize;
 
-					for ( int c=0; c<pResult->GetVertexBuffers().GetBufferChannelCount( (int)vb ); ++c )
+					const int32 ChannelsCount = VertexBuffers.GetBufferChannelCount(VertexBufferIndex);
+					for (int32 ChannelIndex = 0; ChannelIndex < ChannelsCount; ++ChannelIndex)
 					{
-						// Get info about the destination channel
-						MESH_BUFFER_SEMANTIC semantic = MBS_NONE;
-						int semanticIndex = 0;
-						MESH_BUFFER_FORMAT format = MBF_NONE;
-						int components = 0;
-						int offset = 0;
-						pResult->GetVertexBuffers().GetChannel
-								( (int)vb, c, &semantic, &semanticIndex, &format, &components, &offset );
+						if (ResultBuffer.m_channels[ChannelIndex].m_semantic != MBS_BONEINDICES)
+						{
+							continue;
+						}
+						
+						int32 ResultOffset = FirstSize + ResultBuffer.m_channels[ChannelIndex].m_offset;
 
-                        int resultOffset = firstSize + offset;
+						const int32 NumComponents = ResultBuffer.m_channels[ChannelIndex].m_componentCount;
 
-						// Check special channels that need extra work
-						if ( semantic==MBS_BONEINDICES )
+						// Bone indices may need remapping
+						for (int32 VertexIndex = 0; VertexIndex < pVSecond->GetVertexCount(); ++VertexIndex)
 						{
-							// Bone indices may need remapping
-                            for ( int si=0; si<pVSecond->GetVertexCount(); ++si )
+							switch (BoneIndexFormat)
 							{
-								switch ( format )
-								{
-								case MBF_INT8:
-								case MBF_UINT8:
+							case MBF_INT8:
+							case MBF_UINT8:
+							{
+								uint8* pD = reinterpret_cast<uint8*>
+									(&ResultBuffer.m_data[ResultOffset]);
+
+								for (int32 ComponentIndex = 0; ComponentIndex < NumComponents; ++ComponentIndex)
 								{
-                                    uint8_t* pD = reinterpret_cast<uint8_t*>
-											( &result.m_data[resultOffset] );
+									uint8 BoneMapIndex = pD[ComponentIndex];
 
-                                    int comp=0;
-                                    for ( ; comp<components; ++comp )
+									// be defensive
+									if (BoneMapIndex < NumSecondBoneMapIndices)
 									{
-                                        uint8_t bone = pD[comp];
-
-                                        // be defensive
-                                        if (bone<secondToFirstBones.Num())
-                                        {
-                                            pD[comp] = (uint8_t)secondToFirstBones[ bone ];
-                                        }
-                                        else
-                                        {
-                                            pD[comp] = 0;
-                                        }
+										pD[ComponentIndex] = (uint8)RemappedBoneMapIndices[BoneMapIndex];
+									}
+									else
+									{
+										pD[ComponentIndex] = 0;
 									}
-
-									resultOffset += elemSize;
-									break;
 								}
 
-								case MBF_INT16:
-								case MBF_UINT16:
-								{
-                                    uint16* pD = reinterpret_cast<uint16*>
-											( &result.m_data[resultOffset] );
+								ResultOffset += ElemSize;
+								break;
+							}
 
-                                    int comp=0;
-                                    for ( ; comp<components; ++comp )
-									{
-                                        uint16 bone = pD[comp];
-
-                                        // be defensive
-                                        if (bone<secondToFirstBones.Num())
-                                        {
-                                            pD[comp] = (uint16)secondToFirstBones[ bone ];
-                                        }
-                                        else
-                                        {
-                                            pD[comp] = 0;
-                                        }
-                                    }
-
-									resultOffset += elemSize;
-									break;
-								}
+							case MBF_INT16:
+							case MBF_UINT16:
+							{
+								uint16* pD = reinterpret_cast<uint16*>
+									(&ResultBuffer.m_data[ResultOffset]);
 
-								case MBF_INT32:
-								case MBF_UINT32:
+								for (int32 ComponentIndex = 0; ComponentIndex < NumComponents; ++ComponentIndex)
 								{
-                                    uint32_t* pD = reinterpret_cast<uint32_t*>
-											( &result.m_data[resultOffset] );
+									uint16 BoneMapIndex = pD[ComponentIndex];
 
-                                    for ( int comp=0; comp<components; ++comp )
+									// be defensive
+									if (BoneMapIndex < NumSecondBoneMapIndices)
+									{
+										pD[ComponentIndex] = (uint16)RemappedBoneMapIndices[BoneMapIndex];
+									}
+									else
 									{
-                                        uint32_t bone = pD[comp];
-
-                                        // be defensive
-                                        if (bone< (uint32_t)secondToFirstBones.Num())
-                                        {
-                                            pD[comp] = (uint32_t)secondToFirstBones[ bone ];
-                                        }
-                                        else
-                                        {
-                                            pD[comp] = 0;
-                                        }
-                                    }
-
-									resultOffset += elemSize;
-									break;
+										pD[ComponentIndex] = 0;
+									}
 								}
 
-								default:
-									checkf(false, TEXT("Format not supported."));
-								}
+								ResultOffset += ElemSize;
+								break;
+							}
+
+							case MBF_INT32:
+							case MBF_UINT32:
+							{
+								// Unreal does not support 32 bit bone indices
+								checkf(false, TEXT("Format not supported."));
+								break;
+							}
+
+							default:
+								checkf(false, TEXT("Format not supported."));
 							}
 						}
 					}
@@ -1055,7 +1097,7 @@ namespace mu
 
 		return pResult;
 	}
-
+	
 
     //---------------------------------------------------------------------------------------------
     //!
@@ -1098,150 +1140,47 @@ namespace mu
     //---------------------------------------------------------------------------------------------
     //! Return null if there is no need to remap the mesh
     //---------------------------------------------------------------------------------------------
-    inline MeshPtr MeshRemapSkeleton( const Mesh* pSource, const Skeleton* pSkeleton )
+    inline MeshPtr MeshRemapSkeleton( const Mesh* SourceMesh, const Skeleton* Skeleton )
     {
-        if ( pSource->GetVertexCount()==0
+        if (SourceMesh->GetVertexCount()==0
              ||
-             !pSource->GetSkeleton()
+             !SourceMesh->GetSkeleton()
              ||
-             pSource->GetSkeleton()->GetBoneCount()==0 )
+			SourceMesh->GetSkeleton()->GetBoneCount()==0 )
         {
             return nullptr;
         }
 
-        auto pSourceSkeleton = pSource->GetSkeleton();
-
-        MeshPtr pResult = pSource->Clone();
-        pResult->SetSkeleton( pSkeleton );
-
-        // Build a skeleton map
-        TArray<int> sourceToSkeleton;
-        for ( const auto& boneName: pSourceSkeleton->m_bones )
-        {
-            sourceToSkeleton.Add( pResult->GetSkeleton()->FindBone(boneName.c_str()) );
-        }
-
-        bool somethingRemapped = false;
-
-        // Convert all bone index channels
-        FMeshBufferSet& BufSet = pResult->m_VertexBuffers;
-        for (auto& buf: BufSet.m_buffers)
-        {
-            for (auto& chan: buf.m_channels)
-            {
-                if (chan.m_semantic==MBS_BONEINDICES)
-                {
-                    uint8_t* pData = &buf.m_data[chan.m_offset];
+        mu::SkeletonPtrConst SourceSkeleton = SourceMesh->GetSkeleton();
+		const TArray<uint16>& SourceBoneMaps = SourceMesh->GetBoneMap();
 
-                    switch ( chan.m_format )
-                    {
-                    case MBF_INT8:
-                    case MBF_UINT8:
-                    {
-                        for (uint32_t v=0;v<BufSet.m_elementCount; ++v)
-                        {
-                            uint8_t* pTypedData = (uint8_t*)pData;
+        // Remap the indices of the bonemap to those of the new skeleton
+		TArray<uint16> RemappedBoneMap;
 
-                            for (int c=0; c<chan.m_componentCount; ++c)
-                            {
-                                size_t boneIndex = size_t( pTypedData[c] );
-                                //check( boneIndex < sourceToSkeleton.Num() );
-                                if ( boneIndex < sourceToSkeleton.Num() )
-                                {
-                                    int bone = sourceToSkeleton[pTypedData[c]];
-                                    if ( pTypedData[c] != uint8_t( bone ) )
-                                    {
-                                        pTypedData[c] = uint8_t( bone );
-                                        somethingRemapped = true;
-                                    }
-                                }
-                                else
-                                {
-                                    UE_LOG(LogMutableCore, Error, TEXT("Bone index out of range.") );
-                                }
-                            }
+		const int32 NumBonesBoneMap = SourceBoneMaps.Num();
+		RemappedBoneMap.Reserve(NumBonesBoneMap);
 
-                            pData += buf.m_elementSize;
-                        }
-                        break;
-                    }
-
-                    case MBF_INT16:
-                    case MBF_UINT16:
-                    {
-                        for (uint32_t v=0;v<BufSet.m_elementCount; ++v)
-                        {
-                            uint16* pTypedData = (uint16*)pData;
-
-                            for (int c=0; c<chan.m_componentCount; ++c)
-                            {
-                                size_t boneIndex = size_t(pTypedData[c]);
-                                //check( boneIndex < sourceToSkeleton.size() );
-                                if ( boneIndex < sourceToSkeleton.Num() )
-                                {
-                                    int bone = sourceToSkeleton[boneIndex];
-                                    if ( pTypedData[c] != uint16(bone) )
-                                    {
-                                        pTypedData[c] = uint16(bone);
-                                        somethingRemapped = true;
-                                    }
-                                }
-                                else
-                                {
-                                    UE_LOG(LogMutableCore, Error, TEXT("Bone index out of range.") );
-                                }
-                            }
-
-                            pData += buf.m_elementSize;
-                        }
-                        break;
-                    }
+		bool bBonesRemapped = false;
 
-                    case MBF_INT32:
-                    case MBF_UINT32: 
-                    {
-                        for (uint32_t v=0;v<BufSet.m_elementCount; ++v)
-                        {
-                            uint32_t* pTypedData = (uint32_t*)pData;
-
-                            for (int c=0; c<chan.m_componentCount; ++c)
-                            {
-                                size_t boneIndex = size_t( pTypedData[c] );
-                                //check( boneIndex < sourceToSkeleton.size() );
-                                if ( boneIndex < sourceToSkeleton.Num() )
-                                {
-                                    int bone = sourceToSkeleton[pTypedData[c]];
-                                    if ( pTypedData[c] != uint32_t( bone ) )
-                                    {
-                                        pTypedData[c] = uint32_t( bone );
-                                        somethingRemapped = true;
-                                    }
-                                }
-                                else
-                                {
-                                    UE_LOG(LogMutableCore, Error, TEXT("Bone index out of range.") );
-                                }
-                            }
-
-                            pData += buf.m_elementSize;
-                        }
-                        break;
-                    }
+		for (const uint16& SourceBoneIndex : SourceBoneMaps)
+		{
+			const char* BoneName = SourceSkeleton->GetBoneName(SourceBoneIndex);
+			const uint16 BoneIndex = Skeleton->FindBone(BoneName);
 
-                    default:
-						checkf(false, TEXT("Not implemented."));
-						break;
-                    }
-                }
-            }
-        }
+			bBonesRemapped = bBonesRemapped || BoneIndex != SourceBoneIndex;
+			RemappedBoneMap.Add(BoneIndex);
+		}
 
-        if (!somethingRemapped)
+        if (!bBonesRemapped)
         {
-            pResult = nullptr;
+            return nullptr;
         }
 
+		MeshPtr pResult = SourceMesh->Clone();
+		pResult->SetSkeleton(Skeleton);
+		pResult->SetBoneMap(RemappedBoneMap);
+
         return pResult;
     }
-
+	
 }
diff --git a/Engine/Plugins/Experimental/Mutable/Source/MutableRuntime/Private/MuR/OpMeshOptimizeSkinning.h b/Engine/Plugins/Experimental/Mutable/Source/MutableRuntime/Private/MuR/OpMeshOptimizeSkinning.h
--- a/Engine/Plugins/Experimental/Mutable/Source/MutableRuntime/Private/MuR/OpMeshOptimizeSkinning.h
+++ b/Engine/Plugins/Experimental/Mutable/Source/MutableRuntime/Private/MuR/OpMeshOptimizeSkinning.h
@@ -23,38 +23,53 @@ namespace mu
 			return nullptr;
 		}
 
-		// Prepare an array of BoneIndex to SkinnedBoneIndex to remap indices.
-		TArray<uint16> RemapedBoneIndices;
-
-		mu::SkeletonPtrConst Skeleton = InMesh->GetSkeleton();
+		uint32 MaxBoneMapIndex = 0;
+		for (const MESH_SURFACE& Surface : InMesh->m_surfaces)
+		{
+			MaxBoneMapIndex = FMath::Max(MaxBoneMapIndex, Surface.BoneMapCount);
+		}
 
-		if (!Skeleton)
+		// We can't optimize the skinning if the mesh requires 16 bit bone indices 
+		if (MaxBoneMapIndex > MAX_uint8)
 		{
 			return nullptr;
 		}
 
-		RemapedBoneIndices.AddZeroed(Skeleton->GetBoneCount());
-
-		int32 SkinnedBonesCount = 0;
+		bool bRequiresFormatChange = false;
 
-		const int32 BonePoseCount = InMesh->GetBonePoseCount();
-		for (int32 BonePoseIndex = 0; BonePoseIndex < BonePoseCount; ++BonePoseIndex)
+		// Desired format if MaxBoneMapIndex <= MAX_uint8
+		const MESH_BUFFER_FORMAT DesiredBoneIndexFormat = mu::MBF_UINT8;
+		
+		// Iterate all vertex buffers and check if the BoneIndices buffers have the desired format
+		const FMeshBufferSet& InMeshVertexBuffers = InMesh->GetVertexBuffers();
+		for (int32 VertexBufferIndex = 0; !bRequiresFormatChange && VertexBufferIndex < InMeshVertexBuffers.m_buffers.Num(); ++VertexBufferIndex)
 		{
-			if (EnumHasAnyFlags(InMesh->BonePoses[BonePoseIndex].BoneUsageFlags, EBoneUsageFlags::Skinning))
+			const MESH_BUFFER& Buffer = InMeshVertexBuffers.m_buffers[VertexBufferIndex];
+
+			const int32 ChannelsCount = InMeshVertexBuffers.GetBufferChannelCount(VertexBufferIndex);
+			for (int32 ChannelIndex = 0; ChannelIndex < ChannelsCount; ++ChannelIndex)
 			{
-				const int32 SkeletonBoneIndex = Skeleton->FindBone(InMesh->BonePoses[BonePoseIndex].BoneName.c_str());
-				RemapedBoneIndices[SkeletonBoneIndex] = SkinnedBonesCount;
-				++SkinnedBonesCount;
+				if (Buffer.m_channels[ChannelIndex].m_semantic == MBS_BONEINDICES)
+				{
+					bRequiresFormatChange = Buffer.m_channels[ChannelIndex].m_format != DesiredBoneIndexFormat;
+					break;
+				}
 			}
 		}
 
-		// Find the desired bone index format
-		const MESH_BUFFER_FORMAT DesiredBoneIndexFormat = SkinnedBonesCount < MAX_uint8 ? mu::MBF_UINT8 : mu::MBF_UINT16;
+		// Not sure if bRequiresFormatChange will ever be true.
+		if (!bRequiresFormatChange)
+		{
+			return nullptr;
+		}
+
+		MUTABLE_CPUPROFILER_SCOPE(MeshOptimizeSkinning_Format);
 
-		const mu::FMeshBufferSet& SourceVertexBuffers = InMesh->GetVertexBuffers();
+		// Format Bone Indices. For some reason, the bone index format is MBF_UINT16 and it should be MBF_UINT8.
 
-		const int32 VertexBuffersCount = SourceVertexBuffers.GetBufferCount();
-		const int32 ElementCount = SourceVertexBuffers.GetElementCount();
+		// TODO: Replace by MeshFormatInPlace once implemented
+		const int32 VertexBuffersCount = InMeshVertexBuffers.GetBufferCount();
+		const int32 ElementCount = InMeshVertexBuffers.GetElementCount();
 
 		// Clone mesh without VertexBuffers, they will be copied manually.
 		MeshPtr Result = InMesh->Clone(~(EMeshCloneFlags::WithVertexBuffers));
@@ -66,8 +81,8 @@ namespace mu
 
 		for (int32 BufferIndex = 0; BufferIndex < VertexBuffersCount; ++BufferIndex)
 		{
-			const MESH_BUFFER& SourceBuffer = SourceVertexBuffers.m_buffers[BufferIndex];
-			const int32 ChannelsCount = SourceVertexBuffers.GetBufferChannelCount(BufferIndex);
+			const MESH_BUFFER& SourceBuffer = InMeshVertexBuffers.m_buffers[BufferIndex];
+			const int32 ChannelsCount = InMeshVertexBuffers.GetBufferChannelCount(BufferIndex);
 
 			int32 BoneIndexChannelIndex = INDEX_NONE;
 
@@ -92,7 +107,7 @@ namespace mu
 				if (Channel.m_semantic == MBS_BONEINDICES)
 				{
 					SourceBoneIndexFormat = Channel.m_format;
-					Format = DesiredBoneIndexFormat;
+					Format = mu::MBF_UINT8;
 					BoneIndexChannelIndex = ChannelIndex;
 				}
 
@@ -109,94 +124,59 @@ namespace mu
 			}
 
 			// Copy buffers
-			if (BoneIndexChannelIndex != INDEX_NONE)
+			if (BoneIndexChannelIndex != INDEX_NONE && SourceBoneIndexFormat == mu::MBF_UINT16)
 			{
 				// This buffer has a BoneIndex channel that needs to be manually fixed. 
 				VertexBuffers.SetBufferCount(BufferIndex + 1);
 				VertexBuffers.SetBuffer(BufferIndex, ElementSize, ChannelsCount, Semantics.GetData(), SemanticIndices.GetData(), Formats.GetData(), Components.GetData(), Offsets.GetData());
 				uint8* Data = VertexBuffers.GetBufferData(BufferIndex);
 
-				const uint8* SourceData = SourceBuffer.m_data.GetData();
+				const uint16* SourceData = (const uint16*)SourceBuffer.m_data.GetData();
 				const int32 SourceBoneIndexSize = GetMeshFormatData(SourceBoneIndexFormat).m_size;
 
-				const int32 BoneIndexSize = GetMeshFormatData(DesiredBoneIndexFormat).m_size;
+				const int32 BoneIndexSize = GetMeshFormatData(mu::MBF_UINT8).m_size;
 				const int32 BoneIndexComponentCount = Components[BoneIndexChannelIndex];
 				const int32 BoneIndexChannelOffset = Offsets[BoneIndexChannelIndex];
 
 				const int32 TailSize = AuxOffset - (BoneIndexChannelOffset + BoneIndexComponentCount * BoneIndexSize);
 
-				for (int32 VertexIndex = 0; VertexIndex < ElementCount; ++VertexIndex)
+
+				for (const MESH_SURFACE& Surface : InMesh->m_surfaces)
 				{
-					FMemory::Memcpy(Data, SourceData, BoneIndexChannelOffset);
-					Data += BoneIndexChannelOffset;
-					SourceData += BoneIndexChannelOffset;
+					const uint32 NumBonesInBoneMap = Surface.BoneMapCount;
 
-					for (int32 ComponentIndex = 0; ComponentIndex < BoneIndexComponentCount; ++ComponentIndex)
+					for (int32 VertexIndex = 0; VertexIndex < Surface.m_vertexCount; ++VertexIndex)
 					{
-						switch (DesiredBoneIndexFormat)
-						{
-						case MESH_BUFFER_FORMAT::MBF_UINT8:
+						FMemory::Memcpy(Data, SourceData, BoneIndexChannelOffset);
+						Data += BoneIndexChannelOffset;
+						SourceData += BoneIndexChannelOffset;
+
+						for (int32 ComponentIndex = 0; ComponentIndex < BoneIndexComponentCount; ++ComponentIndex)
 						{
-							switch (SourceBoneIndexFormat)
+							const uint16 SourceIndex = *SourceData;
+							if (SourceIndex < NumBonesInBoneMap)
 							{
-							case MESH_BUFFER_FORMAT::MBF_UINT8:
-							{
-								int32 SourceIndex = *SourceData;
-								if (SourceIndex < RemapedBoneIndices.Num())
-								{
-									*Data = (uint8)RemapedBoneIndices[SourceIndex];
-								}
-								break;
-							}
-							case MESH_BUFFER_FORMAT::MBF_UINT16:
-							{
-								int32 SourceIndex = *((uint16*)SourceData);
-								if (SourceIndex< RemapedBoneIndices.Num())
-								{
-									*Data = (uint8)RemapedBoneIndices[SourceIndex];
-								}
-								break;
+								*Data = (uint8)SourceIndex;
 							}
-							case MESH_BUFFER_FORMAT::MBF_INT32:
-							{
-								int32 SourceIndex = *((int32*)SourceData);
-								if (SourceIndex < RemapedBoneIndices.Num())
-								{
-									*Data = (uint8)RemapedBoneIndices[SourceIndex];
-								}
-								break;
-							}
-							default:
-								unimplemented();
-								break;
-							}
-							break;
-						}
-						case MESH_BUFFER_FORMAT::MBF_UINT16:
-						{
-							check(SourceBoneIndexFormat == MESH_BUFFER_FORMAT::MBF_UINT16);
-							*((uint16*)Data) = RemapedBoneIndices[*((uint16*)SourceData)];
-							break;
-						}
-						default:
-							unimplemented();
-							break;
-						}
 
-						Data += BoneIndexSize;
-						SourceData += SourceBoneIndexSize;
-					}
+							Data += BoneIndexSize;
+							SourceData += SourceBoneIndexSize;
+						}
 
 
-					FMemory::Memcpy(Data, SourceData, TailSize);
-					Data += TailSize;
-					SourceData += TailSize;
+						FMemory::Memcpy(Data, SourceData, TailSize);
+						Data += TailSize;
+						SourceData += TailSize;
+					}
 				}
 			}
 			else
 			{
+				// SourceBoneIndexFormat must be mu::MBF_UINT8 or none if the buffer doesn't have BoneIndices
+				check(SourceBoneIndexFormat == mu::MBF_NONE || SourceBoneIndexFormat == mu::MBF_UINT8);
+
 				// Add buffers that don't require a fix up
-				VertexBuffers.AddBuffer(SourceVertexBuffers, BufferIndex);
+				VertexBuffers.AddBuffer(InMeshVertexBuffers, BufferIndex);
 			}
 
 		}
diff --git a/Engine/Plugins/Experimental/Mutable/Source/MutableRuntime/Public/MuR/Mesh.h b/Engine/Plugins/Experimental/Mutable/Source/MutableRuntime/Public/MuR/Mesh.h
--- a/Engine/Plugins/Experimental/Mutable/Source/MutableRuntime/Public/MuR/Mesh.h
+++ b/Engine/Plugins/Experimental/Mutable/Source/MutableRuntime/Public/MuR/Mesh.h
@@ -50,6 +50,9 @@ namespace mu
 			m_firstIndex = 0;
 			m_indexCount = 0;
 			m_id = 0;
+
+			BoneMapIndex = 0;
+			BoneMapCount = 0;
 		}
 
 		int32 m_firstVertex;
@@ -57,7 +60,9 @@ namespace mu
 		int32 m_firstIndex;
 		int32 m_indexCount;
 		uint32 m_id;
-
+		
+		uint32 BoneMapIndex;
+		uint32 BoneMapCount;
 
 		//!
 		inline bool operator==(const MESH_SURFACE& o) const
@@ -66,9 +71,15 @@ namespace mu
 				&& m_vertexCount == o.m_vertexCount
 				&& m_firstIndex == o.m_firstIndex
 				&& m_indexCount == o.m_indexCount
-				&& m_id == o.m_id;
+				&& m_id == o.m_id
+				&& BoneMapIndex == o.BoneMapIndex
+				&& BoneMapCount == o.BoneMapCount;
 		}
 
+		inline void Serialise(OutputArchive& arch) const;
+
+		inline void Unserialise(InputArchive& arch);
+
 	};
 
 
@@ -121,8 +132,9 @@ namespace mu
 		WithAdditionalBuffers = 1 << 10,
 		WithLayouts           = 1 << 11,
 		WithPoses			  = 1 << 12,
-		WithSkeletonIDs		  = 1 << 13,
-		WithAdditionalPhysics = 1 << 14
+		WithBoneMap			  = 1 << 13,
+		WithSkeletonIDs		  = 1 << 14,
+		WithAdditionalPhysics = 1 << 15
 	};
 	
 	ENUM_CLASS_FLAGS(EMeshCloneFlags);
@@ -183,9 +195,10 @@ namespace mu
         //! Get the number of surfaces defined in this mesh. Surfaces are buffer-contiguous mesh
         //! fragments that share common properties (usually material)
         int GetSurfaceCount() const;
-        void GetSurface( int surfaceIndex,
-                         int* firstVertex, int* vertexCount,
-                         int* firstIndex, int* indexCount ) const;
+        void GetSurface( int32 surfaceIndex,
+                         int32* FirstVertex, int32* VertexCount,
+                         int32* FirstIndex, int32* IndexCount,
+						 int32* FirstBone, int32* BoneCount) const;
 
         //! Return an internal id that can be used to match mesh surfaces and instance surfaces.
         //! Only valid for meshes that are part of instances.
@@ -293,6 +306,12 @@ namespace mu
 		//! 
 		EBoneUsageFlags GetBoneUsageFlags(int32 BoneIndex) const;
 
+		//! Set the bonemap of this mesh
+		void SetBoneMap(const TArray<uint16>& InBoneMap);
+
+		//! Return an array containing the bonemap indices of all surfaces in the mesh.
+		const TArray<uint16>& GetBoneMap() const;
+
 		//!
 		int32 GetSkeletonIDsCount() const;
 
@@ -402,6 +421,9 @@ namespace mu
 		// taking into consideration the meshes being used.
 		TArray<FBonePose> BonePoses;
 
+		// Array containing the bonemap indices of all surfaces in the mesh.
+		TArray<uint16> BoneMapIndices;
+
 		//!
 		inline void Serialise(OutputArchive& arch) const;
 
@@ -417,6 +439,7 @@ namespace mu
 			if (equal) equal = (m_FaceBuffers == o.m_FaceBuffers);
 			if (equal) equal = (m_layouts.Num() == o.m_layouts.Num());
 			if (equal) equal = (BonePoses.Num() == o.BonePoses.Num());
+			if (equal) equal = (BoneMapIndices.Num() == o.BoneMapIndices.Num());
 			if (equal && m_pSkeleton != o.m_pSkeleton)
 			{
 				if (m_pSkeleton && o.m_pSkeleton)
@@ -450,6 +473,8 @@ namespace mu
 				equal &= BonePoses[i] == o.BonePoses[i];
 			}
 
+			if (equal) equal = BoneMapIndices == o.BoneMapIndices;
+
 			equal &= AdditionalPhysicsBodies.Num() == o.AdditionalPhysicsBodies.Num();
 			for (int32 i = 0; equal && i < AdditionalPhysicsBodies.Num(); ++i)
 			{
@@ -520,7 +545,6 @@ namespace mu
     };
 
 
-	MUTABLE_DEFINE_POD_VECTOR_SERIALISABLE(MESH_SURFACE)
 	MUTABLE_DEFINE_ENUM_SERIALISABLE(EBoneUsageFlags)
 	MUTABLE_DEFINE_ENUM_SERIALISABLE(EMeshBufferType)
 	MUTABLE_DEFINE_ENUM_SERIALISABLE(EShapeBindingMethod)
diff --git a/Engine/Plugins/Experimental/Mutable/Source/MutableRuntime/Public/MuR/Model.h b/Engine/Plugins/Experimental/Mutable/Source/MutableRuntime/Public/MuR/Model.h
--- a/Engine/Plugins/Experimental/Mutable/Source/MutableRuntime/Public/MuR/Model.h
+++ b/Engine/Plugins/Experimental/Mutable/Source/MutableRuntime/Public/MuR/Model.h
@@ -13,7 +13,7 @@
 //! data structures. Compiled models are not necessarily compatible when the runtime is updated,
 //! so this version number can be used externally to verify this. It is not used internally, and
 //! serializing models from different versions than this runtime will probably result in a crash.
-#define MUTABLE_COMPILED_MODEL_CODE_VERSION		uint32( 65 )
+#define MUTABLE_COMPILED_MODEL_CODE_VERSION		uint32( 66 )
 #define MUTABLE_PARAMETERS_VERSION              uint32( 2 )
 
 
