[Mutable] Fix Mip array out of range on special CIS builds.

Have not been able to reproduce, but given the task information it looks like a an unsinged value wrapping arround.

#rnx

diff --git a/Engine/Plugins/Experimental/Mutable/Source/CustomizableObject/Private/MuCO/UnrealMutableImageProvider.cpp b/Engine/Plugins/Experimental/Mutable/Source/CustomizableObject/Private/MuCO/UnrealMutableImageProvider.cpp
--- a/Engine/Plugins/Experimental/Mutable/Source/CustomizableObject/Private/MuCO/UnrealMutableImageProvider.cpp
+++ b/Engine/Plugins/Experimental/Mutable/Source/CustomizableObject/Private/MuCO/UnrealMutableImageProvider.cpp
@@ -195,7 +195,7 @@ mu::EImageFormat GetMutablePixelFormat(EPixelFormat InTextureFormat)
 			}
 #endif
 
-			uint8 MipIndex = MipmapsToSkip < TextureToLoad->GetPlatformData()->Mips.Num() ? MipmapsToSkip : TextureToLoad->GetPlatformData()->Mips.Num() - 1;
+			int32 MipIndex = MipmapsToSkip < TextureToLoad->GetPlatformData()->Mips.Num() ? MipmapsToSkip : TextureToLoad->GetPlatformData()->Mips.Num() - 1;
 
 			// Mips in the mip tail are inlined and can't be streamed, find the smallest mip available.
 			for (; MipIndex > 0; --MipIndex)
@@ -420,7 +420,7 @@ mu::FImageDesc FUnrealMutableImageProvider::GetImageDesc(mu::EXTERNAL_IMAGE_ID i
 			// not being modified by the game thread at the moment and the texture cannot be GCed because of the AddReferencedObjects
 			// in the FUnrealMutableImageProvider
 
-			uint8 MipIndex = MipmapsToSkip < TextureToLoad->GetPlatformData()->Mips.Num() ? MipmapsToSkip : TextureToLoad->GetPlatformData()->Mips.Num() - 1;
+			int32 MipIndex = MipmapsToSkip < TextureToLoad->GetPlatformData()->Mips.Num() ? MipmapsToSkip : TextureToLoad->GetPlatformData()->Mips.Num() - 1;
 
 			// Mips in the mip tail are inlined and can't be streamed, find the smallest mip available.
 			for (; MipIndex > 0; --MipIndex)
