[Mutable] Reuse textures between LODs when using AutoLODs from mesh

#rnx

diff --git a/Engine/Plugins/Experimental/Mutable/Source/CustomizableObject/Classes/MuCO/CustomizableObject.h b/Engine/Plugins/Experimental/Mutable/Source/CustomizableObject/Classes/MuCO/CustomizableObject.h
--- a/Engine/Plugins/Experimental/Mutable/Source/CustomizableObject/Classes/MuCO/CustomizableObject.h
+++ b/Engine/Plugins/Experimental/Mutable/Source/CustomizableObject/Classes/MuCO/CustomizableObject.h
@@ -182,6 +182,9 @@ struct FCompilationOptions
 	// Used to enable AnimBp override physics mainipualtion.  
 	bool bAnimBpPhysicsManipulationEnabled = false;
 
+	// Used to enable Reuse textures between lods. Requires automatic LODs to be enabled.  
+	bool bReuseTexturesEnabled = false;
+
 	// Used to reduce the number of notifications when compiling objects
 	bool bSilentCompilation = true;
 
@@ -1410,7 +1413,7 @@ private:
 	// This is a manual version number for the binary blobs in this asset.
 	// Increasing it invalidates all the previously compiled models.
 	// Warning: If while merging code both versions have changed, take the highest+1.
-	static const int32 CurrentSupportedVersion = 384;
+	static const int32 CurrentSupportedVersion = 385;
 
 public:
 
diff --git a/Engine/Plugins/Experimental/Mutable/Source/CustomizableObject/Private/MuCO/CustomizableObjectInstance.cpp b/Engine/Plugins/Experimental/Mutable/Source/CustomizableObject/Private/MuCO/CustomizableObjectInstance.cpp
--- a/Engine/Plugins/Experimental/Mutable/Source/CustomizableObject/Private/MuCO/CustomizableObjectInstance.cpp
+++ b/Engine/Plugins/Experimental/Mutable/Source/CustomizableObject/Private/MuCO/CustomizableObjectInstance.cpp
@@ -5541,10 +5541,14 @@ void UCustomizableInstancePrivateData::BuildMaterials(const TSharedPtr<FMutableO
 {
 	MUTABLE_CPUPROFILER_SCOPE(UCustomizableInstancePrivateData::BuildMaterials)
 
-		UCustomizableObject* CustomizableObject = Public->GetCustomizableObject();
+	UCustomizableObject* CustomizableObject = Public->GetCustomizableObject();
 
 	const int32 LODCount = OperationData->InstanceUpdateData.LODs.Num();
 
+	// Find skipped LODs. The following valid LOD will be copied into them. 
+	TArray<bool> LODsSkipped;
+	LODsSkipped.SetNum(OperationData->CurrentMaxLOD);
+
 	TArray<FGeneratedTexture> NewGeneratedTextures;
 
 	GeneratedMaterials.Reset();
@@ -5573,27 +5577,29 @@ void UCustomizableInstancePrivateData::BuildMaterials(const TSharedPtr<FMutableO
 		SkeletalMesh->GetMaterials().Reset();
 
 		// Maps serializations of FMutableMaterialPlaceholder to Created Dynamic Material instances, used to reuse materials across LODs
-		TMap<FString, FMutableMaterialPlaceholder*> ReuseMaterialCache;
-		// Stores the FMutableMaterialPlaceholders and keeps them in memory while the component's LODs are being processed
-		TArray<TSharedPtr<FMutableMaterialPlaceholder>> MutableMaterialPlaceholderArray;
+		TMap<FString, TSharedPtr<FMutableMaterialPlaceholder>> ReuseMaterialCache;
+
+		// Map of SharedSurfaceId to FMutableMaterialPlaceholder serialization key 
+		TMap<int32, FString> SharedSurfacesCache;
 
 		MUTABLE_CPUPROFILER_SCOPE(BuildMaterials_LODLoop);
 
-		int32 LastValidLODIndex = OperationData->CurrentMaxLOD;
-		for (int32 LODIndex = OperationData->CurrentMaxLOD; LODIndex >= FirstLODAvailable; --LODIndex)
+		for (int32 LODIndex = FirstLODAvailable; LODIndex <= OperationData->CurrentMaxLOD; LODIndex++)
 		{
-			if (SkeletalMesh->GetLODInfoArray().Num() != 0)
+			const FInstanceUpdateData::FLOD& LOD = OperationData->InstanceUpdateData.LODs[LODIndex];
+			const FInstanceUpdateData::FComponent& Component = OperationData->InstanceUpdateData.Components[LOD.FirstComponent + ComponentIndex];
+
+			if (!Component.bGenerated)
 			{
-				SkeletalMesh->GetLODInfoArray()[LODIndex].LODMaterialMap.Reset();
+				LODsSkipped[LODIndex] = true;
+				continue;
 			}
 
-			if (LODIndex >= OperationData->CurrentMinLOD && OperationData->InstanceUpdateData.Components[OperationData->InstanceUpdateData.LODs[LODIndex].FirstComponent + ComponentIndex].bGenerated)
+			if (SkeletalMesh->GetLODInfoArray().Num() != 0)
 			{
-				LastValidLODIndex = LODIndex;
+				SkeletalMesh->GetLODInfoArray()[LODIndex].LODMaterialMap.Reset();
 			}
 
-			const FInstanceUpdateData::FLOD& LOD = OperationData->InstanceUpdateData.LODs[LastValidLODIndex];
-			const FInstanceUpdateData::FComponent& Component = OperationData->InstanceUpdateData.Components[LOD.FirstComponent + ComponentIndex];
 			check(Component.bGenerated);
 
 			const FMutableRefSkeletalMeshData* RefSkeletalMeshData = CustomizableObject->GetRefSkeletalMeshData(Component.Id);
@@ -5603,17 +5609,31 @@ void UCustomizableInstancePrivateData::BuildMaterials(const TSharedPtr<FMutableO
 			{
 				const FInstanceUpdateData::FSurface& Surface = OperationData->InstanceUpdateData.Surfaces[Component.FirstSurface + SurfaceIndex];
 
-				MutableMaterialPlaceholderArray.Add(TSharedPtr<FMutableMaterialPlaceholder>(new FMutableMaterialPlaceholder));
-				FMutableMaterialPlaceholder& MutableMaterialPlaceholder = *MutableMaterialPlaceholderArray.Last();
+				const uint32 ReferencedMaterialIndex = ObjectToInstanceIndexMap[Surface.MaterialIndex];
+				UMaterialInterface* MaterialTemplate = ReferencedMaterials[ReferencedMaterialIndex];
+				if (!MaterialTemplate)
+				{
+					UE_LOG(LogMutable, Error, TEXT("Build Materials: Missing referenced template to use as parent material on CustomizableObject [%s]."), *CustomizableObject->GetName());
+					continue;
+				}
 
-				if (LastValidLODIndex != LODIndex)
+				// Reuse surface between LODs when using AutomaticLODs from mesh.
+				if (const FString* SharedSurfaceSerialization = SharedSurfacesCache.Find(Surface.SurfaceId))
 				{
-					// Copy information from the LastValidLODIndex
-					SkeletalMesh->GetLODInfoArray()[LODIndex].LODMaterialMap = SkeletalMesh->GetLODInfoArray()[LastValidLODIndex].LODMaterialMap;
-					SkeletalMesh->GetResourceForRendering()->LODRenderData[LODIndex].RenderSections[SurfaceIndex].MaterialIndex = SkeletalMesh->GetResourceForRendering()->LODRenderData[LastValidLODIndex].RenderSections[SurfaceIndex].MaterialIndex;
+					TSharedPtr<FMutableMaterialPlaceholder>* FoundMaterialPlaceholder = ReuseMaterialCache.Find(*SharedSurfaceSerialization);
+					check(FoundMaterialPlaceholder);
+
+					const int32 MaterialIndex = (*FoundMaterialPlaceholder)->MatIndex;
+					check(MaterialIndex >= 0);
+
+					int32 LODMaterialIndex = SkeletalMesh->GetLODInfoArray()[LODIndex].LODMaterialMap.Add(MaterialIndex);
+					SkeletalMesh->GetResourceForRendering()->LODRenderData[LODIndex].RenderSections[SurfaceIndex].MaterialIndex = LODMaterialIndex;
 					continue;
 				}
 
+				TSharedPtr<FMutableMaterialPlaceholder> MutableMaterialPlaceholderPtr(new FMutableMaterialPlaceholder);
+				FMutableMaterialPlaceholder& MutableMaterialPlaceholder = *MutableMaterialPlaceholderPtr;
+
 				{
 					MUTABLE_CPUPROFILER_SCOPE(ParamLoop);
 
@@ -5703,7 +5723,7 @@ void UCustomizableInstancePrivateData::BuildMaterials(const TSharedPtr<FMutableO
 						UTexture2D* Texture = nullptr;
 
 						// \TODO: Change this key to a struct.
-						FString TextureReuseCacheRef = bReuseTextures ? FString::Printf(TEXT("%d-%d-%d-%d"), LastValidLODIndex, ComponentIndex, Surface.SurfaceId, ImageIndex) : FString();
+						FString TextureReuseCacheRef = bReuseTextures ? FString::Printf(TEXT("%d-%d-%d-%d"), Image.BaseLOD, ComponentIndex, Surface.SurfaceId, ImageIndex) : FString();
 
 						// If the mutable image is null, it must be in the cache
 						FMutableImageCacheKey ImageCacheKey = { Image.ImageID, OperationData->MipsToSkip };
@@ -5851,10 +5871,10 @@ void UCustomizableInstancePrivateData::BuildMaterials(const TSharedPtr<FMutableO
 													if (MutableMipDataProviderFactory)
 													{
 														MutableMipDataProviderFactory->CustomizableObjectInstance = Public;
-														check(LastValidLODIndex < 256 && ComponentIndex < 256 && ImageIndex < 256);
+														check(LODIndex < 256 && ComponentIndex < 256 && ImageIndex < 256);
 														MutableMipDataProviderFactory->ImageRef.ImageID = Image.ImageID;
 														MutableMipDataProviderFactory->ImageRef.SurfaceId = Surface.SurfaceId;
-														MutableMipDataProviderFactory->ImageRef.LOD = uint8(LastValidLODIndex);
+														MutableMipDataProviderFactory->ImageRef.LOD = uint8(Image.BaseLOD);
 														MutableMipDataProviderFactory->ImageRef.Component = uint8(ComponentIndex);
 														MutableMipDataProviderFactory->ImageRef.Image = uint8(ImageIndex);
 														MutableMipDataProviderFactory->UpdateContext = UpdateContext;
@@ -5928,114 +5948,138 @@ void UCustomizableInstancePrivateData::BuildMaterials(const TSharedPtr<FMutableO
 					}
 				}
 
-				const uint32 ReferencedMaterialIndex = ObjectToInstanceIndexMap[Surface.MaterialIndex];
-				UMaterialInterface* MaterialTemplate = ReferencedMaterials[ReferencedMaterialIndex];
 				MutableMaterialPlaceholder.ParentMaterial = MaterialTemplate;
 
-				FString MaterialPlaceholderSerialization = MutableMaterialPlaceholder.GetSerialization();
-				FMutableMaterialPlaceholder** FoundMaterialPlaceholder = ReuseMaterialCache.Find(MaterialPlaceholderSerialization);
+				const FString MaterialPlaceholderSerialization = MutableMaterialPlaceholder.GetSerialization();
+				check(MaterialPlaceholderSerialization != FString("null"));
 
-				if (FoundMaterialPlaceholder)
+				int32 MatIndex = INDEX_NONE;
+				
+				if (TSharedPtr<FMutableMaterialPlaceholder>* FoundMaterialPlaceholder = ReuseMaterialCache.Find(MaterialPlaceholderSerialization))
 				{
-					int32 MatIndex = (*FoundMaterialPlaceholder)->MatIndex;
-					check(MatIndex >= 0);
-					int32 LODMaterialIndex = SkeletalMesh->GetLODInfoArray()[LODIndex].LODMaterialMap.Add(MatIndex);
-					SkeletalMesh->GetResourceForRendering()->LODRenderData[LODIndex].RenderSections[SurfaceIndex].MaterialIndex = LODMaterialIndex;
+					MatIndex = (*FoundMaterialPlaceholder)->MatIndex;
 				}
-				else
+				else // Material not cached, create a new one
 				{
-					if (MutableMaterialPlaceholder.ParentMaterial)
-					{
-						check(MaterialPlaceholderSerialization != FString("null"));
-						ReuseMaterialCache.Add(MaterialPlaceholderSerialization, &MutableMaterialPlaceholder);
-					}
+					MUTABLE_CPUPROFILER_SCOPE(BuildMaterials_CreateMaterial);
+
+					ReuseMaterialCache.Add(MaterialPlaceholderSerialization, MutableMaterialPlaceholderPtr);
+
+					SharedSurfacesCache.Add(Surface.SurfaceId, MaterialPlaceholderSerialization);
 
 					FGeneratedMaterial Material;
 					UMaterialInstanceDynamic* MaterialInstance = nullptr;
+					UMaterialInterface* ActualMaterialInterface = MaterialTemplate;
 
-					if (MaterialTemplate)
+					if (MutableMaterialPlaceholder.Params.Num() != 0)
 					{
-						MUTABLE_CPUPROFILER_SCOPE(BuildMaterials_CreateMaterial);
-
-						UMaterialInterface* ActualMaterialInterface = MaterialTemplate;
-
-						if (MutableMaterialPlaceholder.Params.Num() != 0)
-						{
-							MaterialInstance = UMaterialInstanceDynamic::Create(MaterialTemplate, GetTransientPackage());
-							ActualMaterialInterface = MaterialInstance;
-						}
-
-						if (SkeletalMesh)
-						{
-							const int32 MatIndex = SkeletalMesh->GetMaterials().Num();
-							MutableMaterialPlaceholder.MatIndex = MatIndex;
+						MaterialInstance = UMaterialInstanceDynamic::Create(MaterialTemplate, GetTransientPackage());
+						ActualMaterialInterface = MaterialInstance;
+					}
 
-							// Set up SkeletalMaterial data
-							FSkeletalMaterial& SkeletalMaterial = SkeletalMesh->GetMaterials().Add_GetRef(ActualMaterialInterface);
-							SkeletalMaterial.MaterialSlotName = CustomizableObject->ReferencedMaterialSlotNames[Surface.MaterialIndex];
-							SetMeshUVChannelDensity(SkeletalMaterial.UVChannelData, RefSkeletalMeshData->Settings.DefaultUVChannelDensity);
+					if (SkeletalMesh)
+					{
+						MatIndex = SkeletalMesh->GetMaterials().Num();
+						MutableMaterialPlaceholder.MatIndex = MatIndex;
 
-							const int32 LODMaterialIndex = SkeletalMesh->GetLODInfoArray()[LODIndex].LODMaterialMap.Add(MatIndex);
-							SkeletalMesh->GetResourceForRendering()->LODRenderData[LODIndex].RenderSections[SurfaceIndex].MaterialIndex = LODMaterialIndex;
-						}
+						// Set up SkeletalMaterial data
+						FSkeletalMaterial& SkeletalMaterial = SkeletalMesh->GetMaterials().Add_GetRef(ActualMaterialInterface);
+						SkeletalMaterial.MaterialSlotName = CustomizableObject->ReferencedMaterialSlotNames[Surface.MaterialIndex];
+						SetMeshUVChannelDensity(SkeletalMaterial.UVChannelData, RefSkeletalMeshData->Settings.DefaultUVChannelDensity);
+					}
 
-						if (MaterialInstance)
+					if (MaterialInstance)
+					{
+						for (const FMutableMaterialPlaceholder::FMutableMaterialPlaceHolderParam& Param : MutableMaterialPlaceholder.Params)
 						{
-							for (const FMutableMaterialPlaceholder::FMutableMaterialPlaceHolderParam& Param : MutableMaterialPlaceholder.Params)
+							switch (Param.Type)
 							{
-								switch (Param.Type)
+							case FMutableMaterialPlaceholder::EPlaceHolderParamType::Vector:
+								if (Param.LayerIndex < 0)
 								{
-								case FMutableMaterialPlaceholder::EPlaceHolderParamType::Vector:
-									if (Param.LayerIndex < 0)
-									{
-										MaterialInstance->SetVectorParameterValue(Param.ParamName, Param.Vector);
-									}
-									else
-									{
-										FMaterialParameterInfo ParameterInfo = FMaterialParameterInfo(Param.ParamName, EMaterialParameterAssociation::LayerParameter, Param.LayerIndex);
-										MaterialInstance->SetVectorParameterValueByInfo(ParameterInfo, Param.Vector);
-									}
+									MaterialInstance->SetVectorParameterValue(Param.ParamName, Param.Vector);
+								}
+								else
+								{
+									FMaterialParameterInfo ParameterInfo = FMaterialParameterInfo(Param.ParamName, EMaterialParameterAssociation::LayerParameter, Param.LayerIndex);
+									MaterialInstance->SetVectorParameterValueByInfo(ParameterInfo, Param.Vector);
+								}
 
-									break;
+								break;
 
-								case FMutableMaterialPlaceholder::EPlaceHolderParamType::Scalar:
-									if (Param.LayerIndex < 0)
-									{
-										MaterialInstance->SetScalarParameterValue(FName(Param.ParamName), Param.Scalar);
-									}
-									else
-									{
-										FMaterialParameterInfo ParameterInfo = FMaterialParameterInfo(Param.ParamName, EMaterialParameterAssociation::LayerParameter, Param.LayerIndex);
-										MaterialInstance->SetScalarParameterValueByInfo(ParameterInfo, Param.Scalar);
-									}
+							case FMutableMaterialPlaceholder::EPlaceHolderParamType::Scalar:
+								if (Param.LayerIndex < 0)
+								{
+									MaterialInstance->SetScalarParameterValue(FName(Param.ParamName), Param.Scalar);
+								}
+								else
+								{
+									FMaterialParameterInfo ParameterInfo = FMaterialParameterInfo(Param.ParamName, EMaterialParameterAssociation::LayerParameter, Param.LayerIndex);
+									MaterialInstance->SetScalarParameterValueByInfo(ParameterInfo, Param.Scalar);
+								}
 
-									break;
+								break;
 
-								case FMutableMaterialPlaceholder::EPlaceHolderParamType::Texture:
-									if (Param.LayerIndex < 0)
-									{
-										MaterialInstance->SetTextureParameterValue(Param.ParamName, Param.Texture.Texture);
-									}
-									else
-									{
-										FMaterialParameterInfo ParameterInfo = FMaterialParameterInfo(Param.ParamName, EMaterialParameterAssociation::LayerParameter, Param.LayerIndex);
-										MaterialInstance->SetTextureParameterValueByInfo(ParameterInfo, Param.Texture.Texture);
-									}
+							case FMutableMaterialPlaceholder::EPlaceHolderParamType::Texture:
+								if (Param.LayerIndex < 0)
+								{
+									MaterialInstance->SetTextureParameterValue(Param.ParamName, Param.Texture.Texture);
+								}
+								else
+								{
+									FMaterialParameterInfo ParameterInfo = FMaterialParameterInfo(Param.ParamName, EMaterialParameterAssociation::LayerParameter, Param.LayerIndex);
+									MaterialInstance->SetTextureParameterValueByInfo(ParameterInfo, Param.Texture.Texture);
+								}
 
-									Material.Textures.Add(Param.Texture);
+								Material.Textures.Add(Param.Texture);
 
-									break;
-								}
+								break;
 							}
 						}
+					}
+
+					GeneratedMaterials.Add(Material);
+				}
+
+				int32 LODMaterialIndex = SkeletalMesh->GetLODInfoArray()[LODIndex].LODMaterialMap.Add(MatIndex);
+				SkeletalMesh->GetResourceForRendering()->LODRenderData[LODIndex].RenderSections[SurfaceIndex].MaterialIndex = LODMaterialIndex;
+
+			}
+		}
+
+		{
+			// Copy data from valid LODs into the skipped ones.
+			int32 LastValidLODIndex = OperationData->CurrentMaxLOD;
+			for (int32 LODIndex = OperationData->CurrentMaxLOD; LODIndex < FirstLODAvailable; ++LODIndex)
+			{
+				// Copy information from the LastValidLODIndex
+				if (LODsSkipped[LODIndex])
+				{
+					SkeletalMesh->GetLODInfoArray()[LODIndex].LODMaterialMap = SkeletalMesh->GetLODInfoArray()[LastValidLODIndex].LODMaterialMap;
 
-						GeneratedMaterials.Add(Material);
+					TIndirectArray<FSkeletalMeshLODRenderData>& LODRenderData = SkeletalMesh->GetResourceForRendering()->LODRenderData;
+
+					const int32 NumRenderSections = LODRenderData[LODIndex].RenderSections.Num();
+					check(NumRenderSections == LODRenderData[LastValidLODIndex].RenderSections.Num());
+
+					if (NumRenderSections == LODRenderData[LastValidLODIndex].RenderSections.Num())
+					{
+						for (int32 RenderSectionIndex = 0; RenderSectionIndex < NumRenderSections; ++RenderSectionIndex)
+						{
+							const int32 MaterialIndex = LODRenderData[LastValidLODIndex].RenderSections[RenderSectionIndex].MaterialIndex;
+							LODRenderData[LODIndex].RenderSections[RenderSectionIndex].MaterialIndex = MaterialIndex;
+						}
 					}
 				}
+				else
+				{
+					LastValidLODIndex = LODIndex;
+				}
 			}
 		}
 	}
 
+	
 	{
 		MUTABLE_CPUPROFILER_SCOPE(BuildMaterials_Exchange);
 
diff --git a/Engine/Plugins/Experimental/Mutable/Source/CustomizableObject/Private/MuCO/CustomizableObjectSystem.cpp b/Engine/Plugins/Experimental/Mutable/Source/CustomizableObject/Private/MuCO/CustomizableObjectSystem.cpp
--- a/Engine/Plugins/Experimental/Mutable/Source/CustomizableObject/Private/MuCO/CustomizableObjectSystem.cpp
+++ b/Engine/Plugins/Experimental/Mutable/Source/CustomizableObject/Private/MuCO/CustomizableObjectSystem.cpp
@@ -1285,6 +1285,9 @@ namespace impl
 			}
 		}
 
+		// Map SharedSurfaceId to surface index
+		TArray<int32> SurfacesSharedId;
+
 		// Generate the mesh and gather all the required resource Ids
 		OperationData->InstanceUpdateData.LODs.SetNum(Instance->GetLODCount());
 		for (int32 MutableLODIndex = 0; MutableLODIndex < Instance->GetLODCount(); ++MutableLODIndex)
@@ -1356,6 +1359,9 @@ namespace impl
 					int32 InstanceSurfaceIndex = Instance->FindSurfaceById(MutableLODIndex, ComponentIndex, SurfaceId);
 					check(Component.bReuseMesh || Component.Mesh->GetVertexCount() == 0 || InstanceSurfaceIndex >= 0);
 
+					int32 BaseSurfaceIndex = InstanceSurfaceIndex;
+					int32 BaseLODIndex = MutableLODIndex;
+
 					if (InstanceSurfaceIndex >= 0)
 					{
 						OperationData->InstanceUpdateData.Surfaces.Push({});
@@ -1365,6 +1371,25 @@ namespace impl
 						Surface.MaterialIndex = Instance->GetSurfaceCustomId(MutableLODIndex, ComponentIndex, InstanceSurfaceIndex);
 						Surface.SurfaceId = SurfaceId;
 
+						const int32 SharedSurfaceId = Instance->GetSharedSurfaceId(MutableLODIndex, ComponentIndex, InstanceSurfaceIndex);
+						const int32 SharedSurfaceIndex = SurfacesSharedId.Find(SharedSurfaceId);
+
+						SurfacesSharedId.Add(SharedSurfaceId);
+
+						if (SharedSurfaceId != INDEX_NONE)
+						{
+							if (SharedSurfaceIndex >= 0)
+							{
+								Surface = OperationData->InstanceUpdateData.Surfaces[SharedSurfaceIndex];
+								continue;
+							}
+
+							// Find the first LOD where this surface can be found
+							Instance->FindBaseSurfaceBySharedId(ComponentIndex, SharedSurfaceId, BaseSurfaceIndex, BaseLODIndex);
+
+							Surface.SurfaceId = Instance->GetSurfaceId(BaseLODIndex, ComponentIndex, BaseSurfaceIndex);
+						}
+
 						// Images
 						Surface.FirstImage = OperationData->InstanceUpdateData.Images.Num();
 						Surface.ImageCount = Instance->GetImageCount(MutableLODIndex, ComponentIndex, InstanceSurfaceIndex);
@@ -1375,11 +1400,10 @@ namespace impl
 							OperationData->InstanceUpdateData.Images.Push({});
 							FInstanceUpdateData::FImage& Image = OperationData->InstanceUpdateData.Images.Last();
 							Image.Name = Instance->GetImageName(MutableLODIndex, ComponentIndex, InstanceSurfaceIndex, ImageIndex);
-							Image.ImageID = Instance->GetImageId(MutableLODIndex, ComponentIndex, InstanceSurfaceIndex, ImageIndex);
+							Image.ImageID = Instance->GetImageId(BaseLODIndex, ComponentIndex, BaseSurfaceIndex, ImageIndex);
 							Image.FullImageSizeX = 0;
 							Image.FullImageSizeY = 0;
-							Image.LOD = MutableLODIndex;
-
+							Image.BaseLOD = BaseLODIndex;
 						}
 
 						// Vectors
@@ -1505,7 +1529,7 @@ namespace impl
 				}
 				else
 				{
-					Image.Image = System->GetImage(OperationData->InstanceID, Image.ImageID, MipsToSkip, Image.LOD);
+					Image.Image = System->GetImage(OperationData->InstanceID, Image.ImageID, MipsToSkip, Image.BaseLOD);
 				}
 
 				check(Image.Image);
diff --git a/Engine/Plugins/Experimental/Mutable/Source/CustomizableObject/Private/MuCO/CustomizableObjectSystemPrivate.h b/Engine/Plugins/Experimental/Mutable/Source/CustomizableObject/Private/MuCO/CustomizableObjectSystemPrivate.h
--- a/Engine/Plugins/Experimental/Mutable/Source/CustomizableObject/Private/MuCO/CustomizableObjectSystemPrivate.h
+++ b/Engine/Plugins/Experimental/Mutable/Source/CustomizableObject/Private/MuCO/CustomizableObjectSystemPrivate.h
@@ -322,7 +322,10 @@ struct FInstanceUpdateData
 	{
 		FString Name;
 		mu::RESOURCE_ID ImageID;
-		int32 LOD;
+		
+		// LOD of the ImageId. If the texture is shared between LOD, first LOD where this image can be found. 
+		int32 BaseLOD;
+		
 		uint16 FullImageSizeX, FullImageSizeY;
 		mu::Ptr<const mu::Image> Image;
 		TWeakObjectPtr<UTexture2D> Cached;
diff --git a/Engine/Plugins/Experimental/Mutable/Source/CustomizableObjectEditor/Private/MuCOE/CustomizableObjectCompiler.cpp b/Engine/Plugins/Experimental/Mutable/Source/CustomizableObjectEditor/Private/MuCOE/CustomizableObjectCompiler.cpp
--- a/Engine/Plugins/Experimental/Mutable/Source/CustomizableObjectEditor/Private/MuCOE/CustomizableObjectCompiler.cpp
+++ b/Engine/Plugins/Experimental/Mutable/Source/CustomizableObjectEditor/Private/MuCOE/CustomizableObjectCompiler.cpp
@@ -621,6 +621,10 @@ mu::NodeObjectPtr FCustomizableObjectCompiler::GenerateMutableRoot(
 
 	}
 
+	// Reuse textures when using automatic LODs from mesh if the LayoutManagement is disabled.
+	GenerationContext.Options.bReuseTexturesEnabled = ActualRoot->AutoLODStrategy == ECustomizableObjectAutomaticLODStrategy::AutomaticFromMesh &&
+		GenerationContext.bDisableTextureLayoutManagementFlag;
+
 	// Ensure that the CO has a valid AutoLODStrategy on the ActualRoot.
 	if (ActualRoot->AutoLODStrategy == ECustomizableObjectAutomaticLODStrategy::Inherited)
 	{
diff --git a/Engine/Plugins/Experimental/Mutable/Source/CustomizableObjectEditor/Private/MuCOE/GenerateMutableSource/GenerateMutableSource.h b/Engine/Plugins/Experimental/Mutable/Source/CustomizableObjectEditor/Private/MuCOE/GenerateMutableSource/GenerateMutableSource.h
--- a/Engine/Plugins/Experimental/Mutable/Source/CustomizableObjectEditor/Private/MuCOE/GenerateMutableSource/GenerateMutableSource.h
+++ b/Engine/Plugins/Experimental/Mutable/Source/CustomizableObjectEditor/Private/MuCOE/GenerateMutableSource/GenerateMutableSource.h
@@ -716,6 +716,24 @@ struct FMutableGraphGenerationContext
 	// Stores the parameters generated in the node tables
 	TMap<const class UCustomizableObjectNodeTable*, TArray<FGuid>> GeneratedParametersInTables;
 
+	struct FSharedSurfaces
+	{
+		FSharedSurfaces(int32 InSurfaceId, mu::NodeSurfaceNewPtr InNodeSurface)
+		{
+			check(InSurfaceId != INDEX_NONE);
+			check(InNodeSurface);
+			SharedSurfaceId = InSurfaceId;
+			NodeSurface = InNodeSurface;
+		}
+
+		int32 SharedSurfaceId = INDEX_NONE;
+
+		// NodeSurface of the current LOD 
+		mu::NodeSurfaceNewPtr NodeSurface;
+	};
+
+	// UCustomizableObjectNodeMaterial material to SharedSurfaceId
+	TMap<UCustomizableObjectNodeMaterial*, FSharedSurfaces> SharedSurfaceIds;
 
 	/** Extension Data constants are collected here */
 	FExtensionDataCompilerInterface ExtensionDataCompilerInterface;
diff --git a/Engine/Plugins/Experimental/Mutable/Source/CustomizableObjectEditor/Private/MuCOE/GenerateMutableSource/GenerateMutableSourceSurface.cpp b/Engine/Plugins/Experimental/Mutable/Source/CustomizableObjectEditor/Private/MuCOE/GenerateMutableSource/GenerateMutableSourceSurface.cpp
--- a/Engine/Plugins/Experimental/Mutable/Source/CustomizableObjectEditor/Private/MuCOE/GenerateMutableSource/GenerateMutableSourceSurface.cpp
+++ b/Engine/Plugins/Experimental/Mutable/Source/CustomizableObjectEditor/Private/MuCOE/GenerateMutableSource/GenerateMutableSourceSurface.cpp
@@ -210,6 +210,31 @@ void SetSurfaceFormat( FMutableGraphGenerationContext& GenerationContext,
 }
 
 
+void UpdateSharedSurfaceId(FMutableGraphGenerationContext& GenerationContext, UCustomizableObjectNodeMaterial* NodeMaterial, mu::NodeSurfaceNewPtr InNodeSurface)
+{
+	const bool bCanShareSurface = GenerationContext.Options.bReuseTexturesEnabled && GenerationContext.CurrentAutoLODStrategy == ECustomizableObjectAutomaticLODStrategy::AutomaticFromMesh;
+
+	// Set shared surface Id to reuse materials and textures between LODs if automatic LODs from mesh is being used
+	if (bCanShareSurface && InNodeSurface)
+	{
+		FMutableGraphGenerationContext::FSharedSurfaces& SharedSurface = GenerationContext.SharedSurfaceIds.FindOrAdd(NodeMaterial, { GenerationContext.SharedSurfaceIds.Num(), InNodeSurface });
+		SharedSurface.NodeSurface = InNodeSurface;
+
+		InNodeSurface->SetSharedSurfaceId(SharedSurface.SharedSurfaceId);
+	}
+
+	// Invalidate the shared surface Id. 
+	else if (!bCanShareSurface && !InNodeSurface)
+	{
+		FMutableGraphGenerationContext::FSharedSurfaces* SharedSurface = GenerationContext.SharedSurfaceIds.Find(NodeMaterial);
+		if (SharedSurface && SharedSurface->NodeSurface)
+		{
+			SharedSurface->NodeSurface->SetSharedSurfaceId(INDEX_NONE);
+		}
+	}
+}
+
+
 mu::NodeSurfacePtr GenerateMutableSourceSurface(const UEdGraphPin * Pin, FMutableGraphGenerationContext & GenerationContext, FMutableGraphSurfaceGenerationData& SurfaceData)
 {
 	check(Pin)
@@ -299,6 +324,9 @@ mu::NodeSurfacePtr GenerateMutableSourceSurface(const UEdGraphPin * Pin, FMutabl
 			// Used ReferencedMaterialsIndex instead of TypedNodeMat->Material->GetName() to prevent material name collisions
 			SurfNode->SetCustomID(ReferencedMaterialsIndex);
 
+			// Set shared surface Id to reuse materials and textures between LODs if automatic LODs from mesh is being used
+			UpdateSharedSurfaceId(GenerationContext, TypedNodeMat, SurfNode);
+
 			// Take slot name from skeletal mesh if one can be found, else leave empty.
 			// Keep Referenced Materials and Materail Slot Names synchronized even if no material name can be found.
 			const bool IsNewSlotName = GenerationContext.ReferencedMaterialSlotNames.Num() == ReferencedMaterialsIndex;
@@ -1148,6 +1176,9 @@ mu::NodeSurfacePtr GenerateMutableSourceSurface(const UEdGraphPin * Pin, FMutabl
 				const int32 UVIndex = ParentMaterialNode->GetImageUVLayout(ImageIndex);
 				LayoutOperationUVNormalizedWarning(GenerationContext, ParentMaterialNode->GetMeshPin(), Node, UVIndex);
 				LayoutOperationUVNormalizedWarning(GenerationContext, TypedNodeExt->AddMeshPin(), Node, UVIndex);
+
+				// Validate if the ParentMaterialNode can be shared between LODs.
+				UpdateSharedSurfaceId(GenerationContext, ParentMaterialNode, nullptr);
 			}
 		
 			for (const FString& Tag : TypedNodeExt->Tags)
@@ -1162,7 +1193,7 @@ mu::NodeSurfacePtr GenerateMutableSourceSurface(const UEdGraphPin * Pin, FMutabl
 		mu::NodeSurfaceEditPtr SurfNode = new mu::NodeSurfaceEdit();
 		Result = SurfNode;
 
-		UCustomizableObjectNodeMaterialBase* ParentMaterialNode = TypedNodeRem->GetParentMaterialNode();
+		UCustomizableObjectNodeMaterial* ParentMaterialNode = TypedNodeRem->GetParentMaterialNode();
 		if (!ParentMaterialNode)
 		{
 			GenerationContext.Compiler->CompilerLog(LOCTEXT("ParentMissing", "Parent node not set (or not found)."), Node);
@@ -1189,6 +1220,9 @@ mu::NodeSurfacePtr GenerateMutableSourceSurface(const UEdGraphPin * Pin, FMutabl
 				SurfNode->SetMesh(MeshPatch.get());
 				MeshPatch->SetMessageContext(Node);
 			}
+
+			// Validate if the ParentMaterialNode can be shared between LODs.
+			UpdateSharedSurfaceId(GenerationContext, ParentMaterialNode, nullptr);
 		}
 	}
 
@@ -1272,6 +1306,9 @@ mu::NodeSurfacePtr GenerateMutableSourceSurface(const UEdGraphPin * Pin, FMutabl
 			}
 
 			LayoutOperationUVNormalizedWarning(GenerationContext, ParentMaterialNode->GetMeshPin(), Node, TypedNodeRemBlocks->ParentLayoutIndex);
+
+			// Validate if the ParentMaterialNode can be shared between LODs.
+			UpdateSharedSurfaceId(GenerationContext, ParentMaterialNode, nullptr);
 		}
 	}
 
@@ -1382,6 +1419,9 @@ mu::NodeSurfacePtr GenerateMutableSourceSurface(const UEdGraphPin * Pin, FMutabl
 				}
 				
 				LayoutOperationUVNormalizedWarning(GenerationContext, ParentMaterialNode->GetMeshPin(), Node, ParentMaterialNode->GetImageUVLayout(ImageIndex));
+
+				// Validate if the ParentMaterialNode can be shared between LODs.
+				UpdateSharedSurfaceId(GenerationContext, ParentMaterialNode, nullptr);
 			}
 		}
 	}
@@ -1391,7 +1431,7 @@ mu::NodeSurfacePtr GenerateMutableSourceSurface(const UEdGraphPin * Pin, FMutabl
 		mu::NodeSurfaceEditPtr SurfNode = new mu::NodeSurfaceEdit();
 		Result = SurfNode;
 
-		UCustomizableObjectNodeMaterialBase* ParentMaterialNode = TypedNodeMorph->GetParentMaterialNode();
+		UCustomizableObjectNodeMaterial* ParentMaterialNode = TypedNodeMorph->GetParentMaterialNode();
 		if (!ParentMaterialNode)
 		{
 			GenerationContext.Compiler->CompilerLog(LOCTEXT("ParentMissing", "Parent node not set (or not found)."), Node);
@@ -1462,6 +1502,9 @@ mu::NodeSurfacePtr GenerateMutableSourceSurface(const UEdGraphPin * Pin, FMutabl
 					}
 				}
 			}
+
+			// Validate if the ParentMaterialNode can be shared between LODs.
+			UpdateSharedSurfaceId(GenerationContext, ParentMaterialNode, nullptr);
 		}
 	}
 
diff --git a/Engine/Plugins/Experimental/Mutable/Source/MutableRuntime/Private/MuR/CodeRunner.cpp b/Engine/Plugins/Experimental/Mutable/Source/MutableRuntime/Private/MuR/CodeRunner.cpp
--- a/Engine/Plugins/Experimental/Mutable/Source/MutableRuntime/Private/MuR/CodeRunner.cpp
+++ b/Engine/Plugins/Experimental/Mutable/Source/MutableRuntime/Private/MuR/CodeRunner.cpp
@@ -658,8 +658,9 @@ namespace mu
                     pResult->GetPrivate()->SetSurfaceName( 0, 0, sindex, strName );
 
                     // IDs
-                    pResult->GetPrivate()->m_lods[0].m_components[0].m_surfaces[sindex].m_internalID = args.id;
-                    pResult->GetPrivate()->m_lods[0].m_components[0].m_surfaces[sindex].m_customID = args.externalId;
+                    pResult->GetPrivate()->m_lods[0].m_components[0].m_surfaces[sindex].InternalId = args.id;
+                    pResult->GetPrivate()->m_lods[0].m_components[0].m_surfaces[sindex].ExternalId = args.ExternalId;
+                    pResult->GetPrivate()->m_lods[0].m_components[0].m_surfaces[sindex].SharedId = args.SharedSurfaceId;
                 }
                 GetMemory().SetInstance( item, pResult );
                 break;
diff --git a/Engine/Plugins/Experimental/Mutable/Source/MutableRuntime/Private/MuR/Instance.cpp b/Engine/Plugins/Experimental/Mutable/Source/MutableRuntime/Private/MuR/Instance.cpp
--- a/Engine/Plugins/Experimental/Mutable/Source/MutableRuntime/Private/MuR/Instance.cpp
+++ b/Engine/Plugins/Experimental/Mutable/Source/MutableRuntime/Private/MuR/Instance.cpp
@@ -174,7 +174,7 @@ namespace mu
              comp>=0 && comp<m_pD->m_lods[lod].m_components.Num() &&
              surf>=0 && surf<m_pD->m_lods[lod].m_components[comp].m_surfaces.Num() )
         {
-            return m_pD->m_lods[lod].m_components[comp].m_surfaces[surf].m_internalID;
+            return m_pD->m_lods[lod].m_components[comp].m_surfaces[surf].InternalId;
         }
 		else
 		{
@@ -193,7 +193,7 @@ namespace mu
 		{
 			for (int i = 0; i < m_pD->m_lods[lod].m_components[comp].m_surfaces.Num(); ++i)
 			{
-				if (m_pD->m_lods[lod].m_components[comp].m_surfaces[i].m_internalID == id)
+				if (m_pD->m_lods[lod].m_components[comp].m_surfaces[i].InternalId == id)
 				{
 					return i;
 				}
@@ -207,6 +207,49 @@ namespace mu
         return -1;
     }
 
+	
+	//---------------------------------------------------------------------------------------------
+	void Instance::FindBaseSurfaceBySharedId(int32 CompIndex, int32 SharedId, int32& OutSurfaceIndex, int32& OutLODIndex) const
+	{
+		for (int32 LodIndex = 0; LodIndex < m_pD->m_lods.Num(); LodIndex++)
+		{
+			if (m_pD->m_lods[LodIndex].m_components.IsValidIndex(CompIndex))
+			{
+				for (int32 SurfaceIndex = 0; SurfaceIndex < m_pD->m_lods[LodIndex].m_components[CompIndex].m_surfaces.Num(); ++SurfaceIndex)
+				{
+					if (m_pD->m_lods[LodIndex].m_components[CompIndex].m_surfaces[SurfaceIndex].SharedId == SharedId)
+					{
+						OutSurfaceIndex = SurfaceIndex;
+						OutLODIndex = LodIndex;
+						return;
+					}
+				}
+			}
+
+		}
+
+		OutSurfaceIndex = INDEX_NONE;
+		OutLODIndex = INDEX_NONE;
+	}
+
+
+	//---------------------------------------------------------------------------------------------
+	int32 Instance::GetSharedSurfaceId(int32 LodIndex, int32 CompIndex, int32 SurfaceIndex) const
+	{
+		if (LodIndex >= 0 && LodIndex < m_pD->m_lods.Num() &&
+			CompIndex >= 0 && CompIndex < m_pD->m_lods[LodIndex].m_components.Num() &&
+			SurfaceIndex >= 0 && SurfaceIndex < m_pD->m_lods[LodIndex].m_components[CompIndex].m_surfaces.Num())
+		{
+			return m_pD->m_lods[LodIndex].m_components[CompIndex].m_surfaces[SurfaceIndex].SharedId;
+		}
+		else
+		{
+			check(false);
+		}
+
+		return 0;
+	}
+	
 
     //---------------------------------------------------------------------------------------------
     uint32_t Instance::GetSurfaceCustomId( int lod, int comp, int surf ) const
@@ -215,7 +258,7 @@ namespace mu
              comp>=0 && comp<m_pD->m_lods[lod].m_components.Num() &&
              surf>=0 && surf<m_pD->m_lods[lod].m_components[comp].m_surfaces.Num() )
         {
-            return m_pD->m_lods[lod].m_components[comp].m_surfaces[surf].m_customID;
+            return m_pD->m_lods[lod].m_components[comp].m_surfaces[surf].ExternalId;
         }
 		else
 		{
diff --git a/Engine/Plugins/Experimental/Mutable/Source/MutableRuntime/Private/MuR/InstancePrivate.h b/Engine/Plugins/Experimental/Mutable/Source/MutableRuntime/Private/MuR/InstancePrivate.h
--- a/Engine/Plugins/Experimental/Mutable/Source/MutableRuntime/Private/MuR/InstancePrivate.h
+++ b/Engine/Plugins/Experimental/Mutable/Source/MutableRuntime/Private/MuR/InstancePrivate.h
@@ -14,8 +14,9 @@ namespace mu
     struct INSTANCE_SURFACE
 	{
 		string m_name;
-        uint32_t m_internalID=0;
-        uint32_t m_customID=0;
+        uint32_t InternalId=0;
+        uint32_t ExternalId =0;
+        uint32_t SharedId =0;
 
 		struct IMAGE
 		{
diff --git a/Engine/Plugins/Experimental/Mutable/Source/MutableRuntime/Private/MuR/Operations.h b/Engine/Plugins/Experimental/Mutable/Source/MutableRuntime/Private/MuR/Operations.h
--- a/Engine/Plugins/Experimental/Mutable/Source/MutableRuntime/Private/MuR/Operations.h
+++ b/Engine/Plugins/Experimental/Mutable/Source/MutableRuntime/Private/MuR/Operations.h
@@ -1009,7 +1009,11 @@ namespace mu
             ADDRESS instance;
             ADDRESS value;
             uint32 id;
-            uint32 externalId;
+            uint32 ExternalId;
+
+			// Id used to identify shared surfaces between lods.
+			int32 SharedSurfaceId;
+
             ADDRESS name;
 
             // Index in the FProgram::m_parameterLists with the parameters that are relevant
diff --git a/Engine/Plugins/Experimental/Mutable/Source/MutableRuntime/Public/MuR/Instance.h b/Engine/Plugins/Experimental/Mutable/Source/MutableRuntime/Public/MuR/Instance.h
--- a/Engine/Plugins/Experimental/Mutable/Source/MutableRuntime/Public/MuR/Instance.h
+++ b/Engine/Plugins/Experimental/Mutable/Source/MutableRuntime/Public/MuR/Instance.h
@@ -88,6 +88,23 @@ namespace mu
         //! \param id ID of the surface to look for.
         int FindSurfaceById( int lod, int comp, uint32_t id ) const;
 
+
+		//! Find the base surface index and Lod index when reusing surfaces between LODs. Return the surface index
+		//! and the LOD it belongs to.
+		//! \param Comp - Index of the component, from 0 to GetComponentCount(lod)-1
+		//! \param SharedSurfaceId - Id of the surface to look for (as returned by GetSharedSurfaceId).
+		//! \param OutSurfaceIndex - Index of the surface in the OutLODIndex lod. 
+		//! \param OutLODIndex - Index of the first LOD where the surface can be found. 
+		void FindBaseSurfaceBySharedId(int32 CompIndex, int32 SharedId, int32& OutSurfaceIndex, int32& OutLODIndex) const;
+
+		
+		//! Get an id that can be used to find the same surface on other LODs
+		//! \param lod Index of the level of detail, from 0 to GetLODCount()-1
+		//! \param comp Index of the component, from 0 to GetComponentCount(lod)-1
+		//! \param surf Index of the surface, from 0 to GetSurfaceCount(lod,comp)-1
+		int32 GetSharedSurfaceId(int32 lod, int32 comp, int32 surf) const;
+
+
         //! Get an optional, opaque application-defined identifier for this surface. The meaning of
         //! this ID depends on each application, and it is specified when creating the source data
         //! that generates this surface.
diff --git a/Engine/Plugins/Experimental/Mutable/Source/MutableRuntime/Public/MuR/Model.h b/Engine/Plugins/Experimental/Mutable/Source/MutableRuntime/Public/MuR/Model.h
--- a/Engine/Plugins/Experimental/Mutable/Source/MutableRuntime/Public/MuR/Model.h
+++ b/Engine/Plugins/Experimental/Mutable/Source/MutableRuntime/Public/MuR/Model.h
@@ -13,7 +13,7 @@
 //! data structures. Compiled models are not necessarily compatible when the runtime is updated,
 //! so this version number can be used externally to verify this. It is not used internally, and
 //! serializing models from different versions than this runtime will probably result in a crash.
-#define MUTABLE_COMPILED_MODEL_CODE_VERSION		uint32( 66 )
+#define MUTABLE_COMPILED_MODEL_CODE_VERSION		uint32( 67 )
 #define MUTABLE_PARAMETERS_VERSION              uint32( 2 )
 
 
diff --git a/Engine/Plugins/Experimental/Mutable/Source/MutableTools/Private/MuT/ASTOpInstanceAdd.cpp b/Engine/Plugins/Experimental/Mutable/Source/MutableTools/Private/MuT/ASTOpInstanceAdd.cpp
--- a/Engine/Plugins/Experimental/Mutable/Source/MutableTools/Private/MuT/ASTOpInstanceAdd.cpp
+++ b/Engine/Plugins/Experimental/Mutable/Source/MutableTools/Private/MuT/ASTOpInstanceAdd.cpp
@@ -44,7 +44,8 @@ namespace mu
 				instance == other->instance &&
 				value == other->value &&
 				id == other->id &&
-				externalId == other->externalId &&
+				ExternalId == other->ExternalId &&
+				SharedSurfaceId == other->SharedSurfaceId &&
 				name == other->name;
 		}
 		return false;
@@ -59,7 +60,8 @@ namespace mu
 		n->instance = mapChild(instance.child());
 		n->value = mapChild(value.child());
 		n->id = id;
-		n->externalId = externalId;
+		n->ExternalId = ExternalId;
+		n->SharedSurfaceId = SharedSurfaceId;
 		n->name = name;
 		return n;
 	}
@@ -116,8 +118,8 @@ namespace mu
 			OP::InstanceAddArgs args;
 			memset(&args, 0, sizeof(args));
 			args.id = id;
-			args.externalId = externalId;
-
+			args.ExternalId = ExternalId;
+			args.SharedSurfaceId = SharedSurfaceId;
 			args.name = program.AddConstant(name);
 
 			if (instance) args.instance = instance->linkedAddress;
diff --git a/Engine/Plugins/Experimental/Mutable/Source/MutableTools/Private/MuT/ASTOpInstanceAdd.h b/Engine/Plugins/Experimental/Mutable/Source/MutableTools/Private/MuT/ASTOpInstanceAdd.h
--- a/Engine/Plugins/Experimental/Mutable/Source/MutableTools/Private/MuT/ASTOpInstanceAdd.h
+++ b/Engine/Plugins/Experimental/Mutable/Source/MutableTools/Private/MuT/ASTOpInstanceAdd.h
@@ -26,7 +26,8 @@ struct FProgram;
 		ASTChild instance;
 		ASTChild value;
 		uint32_t id = 0;
-		uint32_t externalId = 0;
+		uint32_t ExternalId = 0;
+		int32 SharedSurfaceId = INDEX_NONE;
 		string name;
 
 	public:
diff --git a/Engine/Plugins/Experimental/Mutable/Source/MutableTools/Private/MuT/CodeGenerator.cpp b/Engine/Plugins/Experimental/Mutable/Source/MutableTools/Private/MuT/CodeGenerator.cpp
--- a/Engine/Plugins/Experimental/Mutable/Source/MutableTools/Private/MuT/CodeGenerator.cpp
+++ b/Engine/Plugins/Experimental/Mutable/Source/MutableTools/Private/MuT/CodeGenerator.cpp
@@ -1793,7 +1793,8 @@ namespace mu
                 sop->value = surfaceGenerationResult.surfaceOp;
 
                 sop->id = surfaceID;
-                sop->externalId = its.node->GetPrivate()->m_customID;
+                sop->ExternalId = its.node->GetPrivate()->ExternalId;
+                sop->SharedSurfaceId = its.node->GetPrivate()->SharedSurfaceId;
                 Ptr<ASTOp> surfaceAt = sop;
 
 				// TODO: This could be done earlier?
diff --git a/Engine/Plugins/Experimental/Mutable/Source/MutableTools/Private/MuT/CodeGenerator_Image.cpp b/Engine/Plugins/Experimental/Mutable/Source/MutableTools/Private/MuT/CodeGenerator_Image.cpp
--- a/Engine/Plugins/Experimental/Mutable/Source/MutableTools/Private/MuT/CodeGenerator_Image.cpp
+++ b/Engine/Plugins/Experimental/Mutable/Source/MutableTools/Private/MuT/CodeGenerator_Image.cpp
@@ -1478,11 +1478,16 @@ namespace mu
 
 			// Match the block id of the block we are generating with the id that resulted in the generated mesh
 			GeneratedLayoutBlockId = -1;
-			if (node.m_layout<MeshResult.GeneratedLayouts.Num() 
-				&& 
-				LayoutBlockIndex >= 0 && LayoutBlockIndex< MeshResult.GeneratedLayouts[node.m_layout]->m_blocks.Num() )
+			
+			mu::LayoutPtrConst Layout = MeshResult.GeneratedLayouts.IsValidIndex(node.m_layout) ? MeshResult.GeneratedLayouts[node.m_layout] : nullptr;
+			if (Layout && Layout->m_blocks.IsValidIndex(LayoutBlockIndex))
 			{
-				GeneratedLayoutBlockId = MeshResult.GeneratedLayouts[node.m_layout]->m_blocks[LayoutBlockIndex].m_id;
+				GeneratedLayoutBlockId = Layout->m_blocks[LayoutBlockIndex].m_id;
+			}
+			else if (Layout && Layout->m_blocks.Num() == 1)
+			{
+				// Layout management disabled, use the only block available
+				GeneratedLayoutBlockId = Layout->m_blocks[0].m_id;
 			}
 			else
 			{
diff --git a/Engine/Plugins/Experimental/Mutable/Source/MutableTools/Private/MuT/NodeSurfaceNew.cpp b/Engine/Plugins/Experimental/Mutable/Source/MutableTools/Private/MuT/NodeSurfaceNew.cpp
--- a/Engine/Plugins/Experimental/Mutable/Source/MutableTools/Private/MuT/NodeSurfaceNew.cpp
+++ b/Engine/Plugins/Experimental/Mutable/Source/MutableTools/Private/MuT/NodeSurfaceNew.cpp
@@ -184,10 +184,24 @@ namespace mu
     }
 
 
+	//---------------------------------------------------------------------------------------------
+	void NodeSurfaceNew::SetSharedSurfaceId(int32 SharedSurfaceId)
+    {
+		m_pD->SharedSurfaceId = SharedSurfaceId;
+    }
+
+
+	//---------------------------------------------------------------------------------------------
+	int32 NodeSurfaceNew::GetSharedSurfaceId() const
+    {
+		return m_pD->SharedSurfaceId;
+    }
+
+
     //---------------------------------------------------------------------------------------------
     void NodeSurfaceNew::SetCustomID( uint32_t id )
     {
-        m_pD->m_customID = id;
+        m_pD->ExternalId = id;
     }
 
 
diff --git a/Engine/Plugins/Experimental/Mutable/Source/MutableTools/Private/MuT/NodeSurfaceNewPrivate.h b/Engine/Plugins/Experimental/Mutable/Source/MutableTools/Private/MuT/NodeSurfaceNewPrivate.h
--- a/Engine/Plugins/Experimental/Mutable/Source/MutableTools/Private/MuT/NodeSurfaceNewPrivate.h
+++ b/Engine/Plugins/Experimental/Mutable/Source/MutableTools/Private/MuT/NodeSurfaceNewPrivate.h
@@ -31,7 +31,8 @@ namespace mu
 		static NODE_TYPE s_type;
 
 		string m_name;
-        uint32_t m_customID=0;
+        uint32_t ExternalId =0;
+        int32 SharedSurfaceId =INDEX_NONE;
 
 		struct MESH
 		{
@@ -170,11 +171,12 @@ namespace mu
         //!
 		void Serialise( OutputArchive& arch ) const
 		{
-            uint32_t ver = 7;
+            uint32_t ver = 8;
 			arch << ver;
 
 			arch << m_name;
-            arch << m_customID;
+            arch << ExternalId;
+			arch << SharedSurfaceId;
 
 			arch << m_meshes;
 			arch << m_images;
@@ -189,10 +191,15 @@ namespace mu
 		{
             uint32_t ver;
 			arch >> ver;
-            check(ver==7);
+            check(ver<=8);
 
 			arch >> m_name;
-            arch >> m_customID;
+            arch >> ExternalId;
+
+			if (ver >= 8)
+			{
+				arch >> SharedSurfaceId;
+			}
 
 			arch >> m_meshes;
 			arch >> m_images;
diff --git a/Engine/Plugins/Experimental/Mutable/Source/MutableTools/Public/MuT/NodeSurfaceNew.h b/Engine/Plugins/Experimental/Mutable/Source/MutableTools/Public/MuT/NodeSurfaceNew.h
--- a/Engine/Plugins/Experimental/Mutable/Source/MutableTools/Public/MuT/NodeSurfaceNew.h
+++ b/Engine/Plugins/Experimental/Mutable/Source/MutableTools/Public/MuT/NodeSurfaceNew.h
@@ -85,6 +85,12 @@ namespace mu
         //! See Instance::GetSurfaceCustomId
         void SetCustomID( uint32 id );
 
+		//! Add an id that will be used to identify the same surface in other LODs 
+		void SetSharedSurfaceId( int32 SharedSurfaceId );
+
+		//! Return the surface ID shared between LODs 
+		int32 GetSharedSurfaceId() const;
+
         //! \name Tags
         //! \{
 
