[mutable] Fixed possible table node crash.
+ Now, if a table cell is empty, the default value will be used to avoid a null pointer (except for SKM which support null meshes)
+ Fixed case where table materials where added to the generation context cache and the textures sent to the core were the wrong ones.

#rnx

diff --git a/Engine/Plugins/Experimental/Mutable/Source/CustomizableObjectEditor/Private/MuCOE/GenerateMutableSource/GenerateMutableSourceImage.cpp b/Engine/Plugins/Experimental/Mutable/Source/CustomizableObjectEditor/Private/MuCOE/GenerateMutableSource/GenerateMutableSourceImage.cpp
--- a/Engine/Plugins/Experimental/Mutable/Source/CustomizableObjectEditor/Private/MuCOE/GenerateMutableSource/GenerateMutableSourceImage.cpp
+++ b/Engine/Plugins/Experimental/Mutable/Source/CustomizableObjectEditor/Private/MuCOE/GenerateMutableSource/GenerateMutableSourceImage.cpp
@@ -141,6 +141,8 @@ mu::NodeImagePtr GenerateMutableSourceImage(const UEdGraphPin* Pin, FMutableGrap
 		return static_cast<mu::NodeImage*>(Generated->Node.get());
 	}
 
+	bool bDoNotAddToGeneratedCache = false;
+
 	mu::NodeImagePtr Result;
 	
 	if (const UCustomizableObjectNodeTexture* TypedNodeTex = Cast<UCustomizableObjectNodeTexture>(Node))
@@ -943,6 +945,9 @@ mu::NodeImagePtr GenerateMutableSourceImage(const UEdGraphPin* Pin, FMutableGrap
 
 				if (Pin->PinType.PinCategory == Schema->PC_MaterialAsset)
 				{
+					// Material pins have to skip the cache of nodes or they will return always the same column node
+					bDoNotAddToGeneratedCache = true;
+
 					ColumnName = GenerationContext.CurrentMaterialTableParameterId;
 				}
 
@@ -982,8 +987,11 @@ mu::NodeImagePtr GenerateMutableSourceImage(const UEdGraphPin* Pin, FMutableGrap
 		GenerationContext.Compiler->CompilerLog(LOCTEXT("UnimplementedNode", "Node type not implemented yet."), Node);
 	}
 
-	GenerationContext.Generated.Add(Key, FGeneratedData(Node, Result));
-	GenerationContext.GeneratedNodes.Add(Node);
+	if (!bDoNotAddToGeneratedCache)
+	{
+		GenerationContext.Generated.Add(Key, FGeneratedData(Node, Result));
+		GenerationContext.GeneratedNodes.Add(Node);
+	}
 
 	if (Result)
 	{
diff --git a/Engine/Plugins/Experimental/Mutable/Source/CustomizableObjectEditor/Private/MuCOE/GenerateMutableSource/GenerateMutableSourceTable.cpp b/Engine/Plugins/Experimental/Mutable/Source/CustomizableObjectEditor/Private/MuCOE/GenerateMutableSource/GenerateMutableSourceTable.cpp
--- a/Engine/Plugins/Experimental/Mutable/Source/CustomizableObjectEditor/Private/MuCOE/GenerateMutableSource/GenerateMutableSourceTable.cpp
+++ b/Engine/Plugins/Experimental/Mutable/Source/CustomizableObjectEditor/Private/MuCOE/GenerateMutableSource/GenerateMutableSourceTable.cpp
@@ -29,13 +29,15 @@ void FillTableColumn(const UCustomizableObjectNodeTable* TableNode,	mu::TablePtr
 	{
 		UObject* Object = SoftObjectProperty->GetPropertyValue(CellData).LoadSynchronous();
 
-		if (!Object)
+		if (SoftObjectProperty->PropertyClass->IsChildOf(USkeletalMesh::StaticClass()))
 		{
-			return;
-		}
+			USkeletalMesh* SkeletalMesh = Cast<USkeletalMesh>(Object);
+
+			if(!SkeletalMesh)
+			{
+				return;
+			}
 
-		if (USkeletalMesh* SkeletalMesh = Cast<USkeletalMesh>(Object))
-		{
 			// Getting Animation Blueprint and Animation Slot
 			FString AnimBP, AnimSlot, GameplayTag, AnimBPAssetTag;
 			TArray<FGameplayTag> GameplayTags;
@@ -256,8 +258,15 @@ void FillTableColumn(const UCustomizableObjectNodeTable* TableNode,	mu::TablePtr
 			}
 		}
 
-		else if (UStaticMesh* StaticMesh = Cast<UStaticMesh>(Object))
+		else if (SoftObjectProperty->PropertyClass->IsChildOf(UStaticMesh::StaticClass()))
 		{
+			UStaticMesh* StaticMesh = Cast<UStaticMesh>(Object);
+
+			if (!StaticMesh)
+			{
+				return;
+			}
+
 			// Getting reference Mesh
 			UStaticMesh* ReferenceStaticMesh = TableNode->GetColumnDefaultAssetByType<UStaticMesh>(ColumnName);
 
@@ -324,8 +333,18 @@ void FillTableColumn(const UCustomizableObjectNodeTable* TableNode,	mu::TablePtr
 			}
 		}
 
-		else if (UTexture2D* Texture = Cast<UTexture2D>(Object))
+		else if (SoftObjectProperty->PropertyClass->IsChildOf(UTexture2D::StaticClass()))
 		{
+			UTexture2D* Texture = Cast<UTexture2D>(Object);
+
+			if (!Texture)
+			{
+				Texture = TableNode->GetColumnDefaultAssetByType<UTexture2D>(ColumnName);
+
+				FString msg = FString::Printf(TEXT("Texture from column [%s] row [%s] is null. The default texture will be used instead."), *ColumnName, *RowName);
+				GenerationContext.Compiler->CompilerLog(FText::FromString(msg), TableNode);
+			}
+
 			// Getting column index from column name
 			CurrentColumn = MutableTable->FindColumn(StringCast<ANSICHAR>(*ColumnName).Get());
 
@@ -357,8 +376,18 @@ void FillTableColumn(const UCustomizableObjectNodeTable* TableNode,	mu::TablePtr
 			}
 		}
 
-		else if (UMaterialInstance* Material = Cast<UMaterialInstance>(Object))
+		else if (SoftObjectProperty->PropertyClass->IsChildOf(UMaterialInstance::StaticClass()))
 		{
+			UMaterialInstance* Material = Cast<UMaterialInstance>(Object);
+
+			if (!Material)
+			{
+				Material = TableNode->GetColumnDefaultAssetByType<UMaterialInstance>(ColumnName);
+
+				FString msg = FString::Printf(TEXT("Material Instance from column [%s] row [%s] is null. The default Material Instance will be used instead."), *ColumnName, *RowName);
+				GenerationContext.Compiler->CompilerLog(FText::FromString(msg), TableNode);
+			}
+
 			//Adding an empty column for searching purposes
 			if (MutableTable.get()->FindColumn(StringCast<ANSICHAR>(*ColumnName).Get()) == -1)
 			{
@@ -436,9 +465,18 @@ void FillTableColumn(const UCustomizableObjectNodeTable* TableNode,	mu::TablePtr
 						{
 							if (InstanceTexture.ParameterInfo.Name == ReferenceTexture.ParameterInfo.Name)
 							{
-								TextureToConvert = Cast<UTexture2D>(InstanceTexture.ParameterValue);
+								if (InstanceTexture.ParameterValue)
+								{
+									TextureToConvert = Cast<UTexture2D>(InstanceTexture.ParameterValue);
+								}
+								else
+								{
+									FString ParamName = ReferenceTexture.ParameterInfo.Name.ToString();
+									FString msg = FString::Printf(TEXT("Parameter [%s] from material instance of column [%s] row [%s] is null. The parameter texture of the default material will be used instead."), *ParamName,  *ColumnName, *RowName);
+									GenerationContext.Compiler->CompilerLog(FText::FromString(msg), TableNode);
+								}
+
 								ModifiedTextureParameters++;
-								
 								break;
 							}
 						}
