[Mutable] Fix MergePhysicsAssets setting being ignored when converting meshes to Mutable.

#rnx

diff --git a/Engine/Plugins/Experimental/Mutable/Source/CustomizableObjectEditor/Private/MuCOE/GenerateMutableSource/GenerateMutableSourceMesh.cpp b/Engine/Plugins/Experimental/Mutable/Source/CustomizableObjectEditor/Private/MuCOE/GenerateMutableSource/GenerateMutableSourceMesh.cpp
--- a/Engine/Plugins/Experimental/Mutable/Source/CustomizableObjectEditor/Private/MuCOE/GenerateMutableSource/GenerateMutableSourceMesh.cpp
+++ b/Engine/Plugins/Experimental/Mutable/Source/CustomizableObjectEditor/Private/MuCOE/GenerateMutableSource/GenerateMutableSourceMesh.cpp
@@ -1444,7 +1444,7 @@ mu::MeshPtr ConvertSkeletalMeshToMutable(USkeletalMesh* InSkeletalMesh, const TS
 		}
 	}
 
-	if (!bIgnorePhysics && InSkeletalMesh->GetPhysicsAsset() && MutableMesh->GetSkeleton())
+	if (!bIgnorePhysics && InSkeletalMesh->GetPhysicsAsset() && MutableMesh->GetSkeleton() && GenerationContext.Options.bPhysicsAssetMergeEnabled)
 	{
 		// Find BodySetups with relevant bones.
 		TArray<TObjectPtr<USkeletalBodySetup>>& SkeletalBodySetups = InSkeletalMesh->GetPhysicsAsset()->SkeletalBodySetups;
