[Mutable] Enable by default vector optimized image project pixel processor.

- Sets mutable.UseProjectectionVectorImpl cvar default to true.
- The optimization has been testesd and confirmed to work on all relevant platforms.
- Removed deprecated cvar mutable.ForceBestQualityProjectionSampling

#rnx

diff --git a/Engine/Plugins/Experimental/Mutable/Source/MutableRuntime/Private/MuR/CodeRunner.cpp b/Engine/Plugins/Experimental/Mutable/Source/MutableRuntime/Private/MuR/CodeRunner.cpp
--- a/Engine/Plugins/Experimental/Mutable/Source/MutableRuntime/Private/MuR/CodeRunner.cpp
+++ b/Engine/Plugins/Experimental/Mutable/Source/MutableRuntime/Private/MuR/CodeRunner.cpp
@@ -75,14 +75,6 @@
 namespace
 {
 
-// This cvar will be deprecated, use CVarForceProjectionSamplingMode with value 1 to achieve the same result.
-bool bForceBestQualityProjectionSampling = false;
-static FAutoConsoleVariableRef CVarForceBestQualityProjectionSampling (
-	TEXT("mutable.ForceBestQualityProjectionSampling"),
-	bForceBestQualityProjectionSampling,
-	TEXT("DEPRECATED, please use mutable.ForceProjectionMode with value 1 to achieve the same result. Force mutable to use projection mode BiLinear + TotalAreaHeuristic."),
-	ECVF_Default);
-
 int32 ForcedProjectionMode = -1;
 static FAutoConsoleVariableRef CVarForceProjectionSamplingMode (
 	TEXT("mutable.ForceProjectionMode"),
@@ -97,7 +89,7 @@ static FAutoConsoleVariableRef CVarGlobalProjectionLodBias (
 	TEXT("Lod bias applied to the lod resulting form the best mip computation, only used if a min filter method different than None is used."),
 	ECVF_Default);
 
-bool bUseProjectionVectorImpl = false;
+bool bUseProjectionVectorImpl = true;
 static FAutoConsoleVariableRef CVarUseProjectionVectorImpl (
 	TEXT("mutable.UseProjectionVectorImpl"),
 	bUseProjectionVectorImpl,
@@ -4084,11 +4076,7 @@ namespace mu
 				{
 					EMinFilterMethod MinFilterMethod = Invoke([&]() -> EMinFilterMethod
 					{
-						if (bForceBestQualityProjectionSampling)
-						{
-							return EMinFilterMethod::TotalAreaHeuristic;
-						}
-						else if (ForcedProjectionMode == 0)
+						if (ForcedProjectionMode == 0)
 						{
 							return EMinFilterMethod::None;
 						}
@@ -4229,11 +4217,7 @@ namespace mu
 
 				ESamplingMethod SamplingMethod = Invoke([&]() -> ESamplingMethod
 				{
-					if (bForceBestQualityProjectionSampling)
-					{
-						return ESamplingMethod::BiLinear;
-					}
-					else if (ForcedProjectionMode == 0)
+					if (ForcedProjectionMode == 0)
 					{
 						return ESamplingMethod::Point;
 					}
