[Mutable] Automatic Node Skeletal Mesh.

Node Skeletal Mesh will no longer get out of date unless the data of a conected pin changes in its Skeletal Mesh Asset:
We consider pin data the following:
- Section Index
- LOD Index
- Texture Parameter Id
- UV Index

Changelist:
- NodeSkeletalMesh removed LODs data structure. Moved to the PinData system.
- Changed the names of NodeSkeletalMesh pins to be more concise.
- Changed the name of some MaterialIndex variables to SectionIndex (they are not the same!).
- PinData added equal operator that takes into account inheritance.
- NodeStaticMesh deprecated old ImagePins.
- Minor pointer to reference function changes.

#rnx

diff --git a/Engine/Plugins/Experimental/Mutable/Source/CustomizableObject/Public/MuCO/CustomizableObjectCustomVersion.h b/Engine/Plugins/Experimental/Mutable/Source/CustomizableObject/Public/MuCO/CustomizableObjectCustomVersion.h
--- a/Engine/Plugins/Experimental/Mutable/Source/CustomizableObject/Public/MuCO/CustomizableObjectCustomVersion.h
+++ b/Engine/Plugins/Experimental/Mutable/Source/CustomizableObject/Public/MuCO/CustomizableObjectCustomVersion.h
@@ -74,6 +74,8 @@ struct CUSTOMIZABLEOBJECT_API FCustomizableObjectCustomVersion
 
 		RemovedParameterDecorations,
 
+		AutomaticNodeSkeletalMesh,
+
 		// -----<new versions can be added above this line>-------------------------------------------------
 		VersionPlusOne,
 		LatestVersion = VersionPlusOne - 1
diff --git a/Engine/Plugins/Experimental/Mutable/Source/CustomizableObjectEditor/Classes/MuCOE/Nodes/CustomizableObjectNode.h b/Engine/Plugins/Experimental/Mutable/Source/CustomizableObjectEditor/Classes/MuCOE/Nodes/CustomizableObjectNode.h
--- a/Engine/Plugins/Experimental/Mutable/Source/CustomizableObjectEditor/Classes/MuCOE/Nodes/CustomizableObjectNode.h
+++ b/Engine/Plugins/Experimental/Mutable/Source/CustomizableObjectEditor/Classes/MuCOE/Nodes/CustomizableObjectNode.h
@@ -24,6 +24,15 @@ class UCustomizableObjectNodePinData : public UObject
 
 public:
 	UCustomizableObjectNodePinData();
+
+	/** False if inherited types are different. */
+	bool operator==(const UCustomizableObjectNodePinData& Other) const;
+	
+	virtual bool operator!=(const UCustomizableObjectNodePinData& Other) const;
+
+protected:
+	/** Virtual function used to perform comparision between different UCustomizableObjectNodePinData inherited types. */
+	virtual bool Equals(const UCustomizableObjectNodePinData& Other) const;
 };
 
 /** Encapsulation of parameters for the FPostEditChangePropertyDelegate delegate function.
diff --git a/Engine/Plugins/Experimental/Mutable/Source/CustomizableObjectEditor/Classes/MuCOE/Nodes/CustomizableObjectNodeMesh.h b/Engine/Plugins/Experimental/Mutable/Source/CustomizableObjectEditor/Classes/MuCOE/Nodes/CustomizableObjectNodeMesh.h
--- a/Engine/Plugins/Experimental/Mutable/Source/CustomizableObjectEditor/Classes/MuCOE/Nodes/CustomizableObjectNodeMesh.h
+++ b/Engine/Plugins/Experimental/Mutable/Source/CustomizableObjectEditor/Classes/MuCOE/Nodes/CustomizableObjectNodeMesh.h
@@ -20,7 +20,7 @@ public:
 
 	// Own interface
 	/** Returns the multiple Layouts of a given Mesh pin. A pin can have multiple Layouts since it can have multiple UVs. Override. */
-	virtual TArray<UCustomizableObjectLayout*> GetLayouts(const UEdGraphPin* OutPin) const PURE_VIRTUAL(UCustomizableObjectNodeMesh::GetLayouts, return {};);
+	virtual TArray<UCustomizableObjectLayout*> GetLayouts(const UEdGraphPin& OutPin) const PURE_VIRTUAL(UCustomizableObjectNodeMesh::GetLayouts, return {};);
 
 	virtual UTexture2D* FindTextureForPin(const UEdGraphPin* Pin) const PURE_VIRTUAL(UCustomizableObjectNodeMesh::FindTextureForPin, return {};);
 
@@ -29,9 +29,12 @@ public:
 	/** Returns the Unreal mesh (e.g., USkeletalMesh, UStaticMesh...). */
 	virtual UObject* GetMesh() const PURE_VIRTUAL(UCustomizableObjectNodeMesh::GetMesh, return {};);;
 
-	/** Returns the out mesh pin associated to the given LOD and Material. Override. */
-	virtual UEdGraphPin* GetMeshPin(int32 LOD, int MaterialIndex) const PURE_VIRTUAL(UCustomizableObjectNodeMesh::GetMeshPin, return {};);
+	/** Returns the output Mesh pin associated to the given LODIndex and SectionIndex. Override. */
+	virtual UEdGraphPin* GetMeshPin(int32 LOD, int32 SectionIndex) const PURE_VIRTUAL(UCustomizableObjectNodeMesh::GetMeshPin, return {};);
 
+	/** Returns the input Layout pin associated to the given LODIndex, SectionIndex and LayoutIndex. Override. */
+	virtual UEdGraphPin* GetLayoutPin(int32 LODIndex, int32 SectionIndex, int32 LayoutIndex) const PURE_VIRTUAL(UCustomizableObjectNodeMesh::GetLayoutPin, return {};);
+	
 	/** Given a pin, return the Section and Layout index.
 	 *
 	 * Will always return a valid result. The result can be valid but out of sync with respect the Unreal mesh asset.
diff --git a/Engine/Plugins/Experimental/Mutable/Source/CustomizableObjectEditor/Classes/MuCOE/Nodes/CustomizableObjectNodeSkeletalMesh.h b/Engine/Plugins/Experimental/Mutable/Source/CustomizableObjectEditor/Classes/MuCOE/Nodes/CustomizableObjectNodeSkeletalMesh.h
--- a/Engine/Plugins/Experimental/Mutable/Source/CustomizableObjectEditor/Classes/MuCOE/Nodes/CustomizableObjectNodeSkeletalMesh.h
+++ b/Engine/Plugins/Experimental/Mutable/Source/CustomizableObjectEditor/Classes/MuCOE/Nodes/CustomizableObjectNodeSkeletalMesh.h
@@ -4,7 +4,9 @@
 
 #include "GameplayTagContainer.h"
 #include "MuCOE/Nodes/CustomizableObjectNodeMesh.h"
+#include "MuCOE/CustomizableObjectEditor_Deprecated.h"
 #include "SGraphNode.h"
+#include "MuCOE/RemapPins/CustomizableObjectNodeRemapPinsByNameDefaultPin.h"
 
 #include "CustomizableObjectNodeSkeletalMesh.generated.h"
 
@@ -57,8 +59,6 @@ public:
 	const FSlateBrush* GetExpressionPreviewArrow() const;
 	EVisibility ExpressionPreviewVisibility() const;
 
-public:
-
 	// Single property that only draws the combo box widget of the skeletal mesh
 	TSharedPtr<ISinglePropertyView> SkeletalMeshSelector;
 
@@ -80,45 +80,85 @@ private:
 };
 
 
+/** Remap pins by pin PinData. */
+UCLASS()
+class UCustomizableObjectNodeSkeletalMeshRemapPinsBySection : public UCustomizableObjectNodeRemapPinsByNameDefaultPin
+{
+	GENERATED_BODY()
+public:
+	virtual bool Equal(const UCustomizableObjectNode& Node, const UEdGraphPin& OldPin, const UEdGraphPin& NewPin) const override;
+};
+
 
-USTRUCT()
-struct CUSTOMIZABLEOBJECTEDITOR_API FCustomizableObjectNodeSkeletalMeshMaterial
+/** PinData of a pin that belongs to a Skeletal Mesh Section. */
+UCLASS()
+class UCustomizableObjectNodeSkeletalMeshPinDataSection : public UCustomizableObjectNodePinData
 {
-	GENERATED_USTRUCT_BODY()
+	GENERATED_BODY()
 
-	UPROPERTY()
-	FString Name;
+public:
+	void Init(int32 InLODIndex, int32 InSectionIndex);
 
-	UPROPERTY()
-	TObjectPtr<UEdGraphPin_Deprecated> MeshPin = nullptr;
+	int32 GetLODIndex() const;
 
-	UPROPERTY()
-	TArray< TObjectPtr<UEdGraphPin_Deprecated> > LayoutPins;
+	int32 GetSectionIndex() const;
 
-	UPROPERTY()
-	TArray< TObjectPtr<UEdGraphPin_Deprecated> > ImagePins;
+protected:
+	virtual bool Equals(const UCustomizableObjectNodePinData& Other) const override;
 
-	UPROPERTY()
-	FEdGraphPinReference MeshPinRef;
+private:
+	int32 LODIndex = -1;
+	int32 SectionIndex = -1;
+};
 
-	UPROPERTY()
-	TArray<FEdGraphPinReference> LayoutPinsRef;
 
-	UPROPERTY()
-	TArray<FEdGraphPinReference> ImagePinsRef;
+/** PinData of a Mesh pin. */
+UCLASS()
+class UCustomizableObjectNodeSkeletalMeshPinDataMesh : public UCustomizableObjectNodeSkeletalMeshPinDataSection
+{
+	GENERATED_BODY()
 };
 
 
-USTRUCT()
-struct CUSTOMIZABLEOBJECTEDITOR_API FCustomizableObjectNodeSkeletalMeshLOD
+/** PinData of a Image pin. */
+UCLASS()
+class UCustomizableObjectNodeSkeletalMeshPinDataImage : public UCustomizableObjectNodeSkeletalMeshPinDataSection
 {
-	GENERATED_USTRUCT_BODY()
+	GENERATED_BODY()
 
-	UPROPERTY()
-	TArray<FCustomizableObjectNodeSkeletalMeshMaterial> Materials;
+public:
+	void Init(int32 InLODIndex, int32 InSectionIndex, FGuid InTextureParameterId);
+
+	FGuid GetTextureParameterId() const;
+
+protected:
+	virtual bool Equals(const UCustomizableObjectNodePinData& Other) const override;
+
+private:
+	FGuid TextureParameterId;
+};
+
+
+/** PinData of a Layout pin. */
+UCLASS()
+class UCustomizableObjectNodeSkeletalMeshPinDataLayout : public UCustomizableObjectNodeSkeletalMeshPinDataSection
+{
+	GENERATED_BODY()
+
+public:
+	void Init(int32 InLODIndex, int32 InSectionIndex, int32 InUVIndex);
+
+	int32 GetUVIndex() const;
+
+protected:
+	virtual bool Equals(const UCustomizableObjectNodePinData& Other) const override;
+
+private:
+	int32 UVIndex = -1;
 };
 
 
+
 UCLASS()
 class CUSTOMIZABLEOBJECTEDITOR_API UCustomizableObjectNodeSkeletalMesh : public UCustomizableObjectNodeMesh
 {
@@ -128,10 +168,6 @@ public:
 	UPROPERTY(EditAnywhere, Category = CustomizableObject)
 	TObjectPtr<USkeletalMesh> SkeletalMesh;
 
-	/** Images */
-	UPROPERTY()
-	TArray<FCustomizableObjectNodeSkeletalMeshLOD> LODs;
-
 	/** Default pin when there is no mesh. */
 	UPROPERTY()
 	FEdGraphPinReference DefaultPin;
@@ -162,8 +198,7 @@ public:
 
 	// UObject interface.
 	virtual void PostEditChangeProperty(FPropertyChangedEvent& PropertyChangedEvent) override;
-	virtual void Serialize(FArchive& Ar) override;
-
+	
 	// UEdGraphNode interface
 	virtual FText GetNodeTitle(ENodeTitleType::Type TitleType) const override;
 	virtual FLinearColor GetNodeTitleColor() const override;
@@ -172,19 +207,22 @@ public:
 	// UCustomizableObjectNode interface
 	virtual void BackwardsCompatibleFixup() override;
 	virtual void AllocateDefaultPins(UCustomizableObjectNodeRemapPins* RemapPins) override;
-	virtual UCustomizableObjectNodeRemapPinsByName* CreateRemapPinsByName() const override;
+	virtual UCustomizableObjectNodeRemapPins* CreateRemapPinsDefault() const override;
 
 	// UCustomizableObjectNodeMesh interface
 	virtual UTexture2D* FindTextureForPin(const UEdGraphPin* Pin) const override;
 	virtual void GetUVChannelForPin(const UEdGraphPin* Pin, TArray<FVector2f>& OutSegments, int32 UVIndex) const override;
-	virtual TArray<UCustomizableObjectLayout*> GetLayouts(const UEdGraphPin* OutPin) const override;
+	virtual TArray<UCustomizableObjectLayout*> GetLayouts(const UEdGraphPin& MeshPin) const override;
 	virtual UObject* GetMesh() const override;
-	virtual UEdGraphPin* GetMeshPin(int32 LOD, int MaterialIndex) const override;
+	virtual UEdGraphPin* GetMeshPin(int32 LOD, int32 SectionIndex) const override;
+	virtual UEdGraphPin* GetLayoutPin(int32 LODIndex, int32 SectionIndex, int32 LayoutIndex) const override;
 	virtual void GetPinSection(const UEdGraphPin& Pin, int32& OutLODIndex, int32& OutSectionIndex, int32& OutLayoutIndex) const override;
+
+	// Own interface
 	
 	/** Returns the material associated to the given output pin. */
 	UMaterialInterface* GetMaterialFor(const UEdGraphPin* Pin) const;
-	FSkeletalMaterial* GetSkeletalMaterialFor(const UEdGraphPin* Pin) const;
+	FSkeletalMaterial* GetSkeletalMaterialFor(const UEdGraphPin& Pin) const;
 
 	virtual bool ProvidesCustomPinRelevancyTest() const override { return true; }
 	virtual bool IsPinRelevant(const UEdGraphPin* Pin) const override;
@@ -195,7 +233,7 @@ public:
 	// Creates the SGraph Node widget for the thumbnail
 	virtual TSharedPtr<SGraphNode> CreateVisualWidget() override;
 
-	// Check if previous UV channels pins of the same LOD and material are also linked to a layout node
+	/** Check if lower UV channels pins of the same LOD and Section are also linked to a Layout Node. */
 	bool CheckIsValidLayout(const UEdGraphPin* Pin, int32& LayoutIndex, FString& MaterialName);
 
 	// Determines if the Node is collapsed or not
@@ -203,7 +241,13 @@ public:
 
 	// Pointer to the SGraphNode Skeletal Mesh
 	TWeakPtr< SGraphNodeSkeletalMesh > GraphNodeSkeletalMesh;
+
 private:
-	UMaterialInterface* GetMaterialInterfaceFor(const int LODIndex, const int MaterialIndex, const FSkeletalMeshModel* ImportedModel = nullptr) const;
-	FSkeletalMaterial* GetSkeletalMaterialFor(const int LODIndex, const int MaterialIndex, const FSkeletalMeshModel* ImportedModel = nullptr) const;
+	UMaterialInterface* GetMaterialInterfaceFor(const int32 LODIndex, const int32 MaterialIndex, const FSkeletalMeshModel* ImportedModel = nullptr) const;
+	FSkeletalMaterial* GetSkeletalMaterialFor(const int32 LODIndex, const int32 MaterialIndex, const FSkeletalMeshModel* ImportedModel = nullptr) const;
+
+	// Deprecated
+	
+	UPROPERTY()
+	TArray<FCustomizableObjectNodeSkeletalMeshLOD> LODs_DEPRECATED;
 };
diff --git a/Engine/Plugins/Experimental/Mutable/Source/CustomizableObjectEditor/Classes/MuCOE/Nodes/CustomizableObjectNodeStaticMesh.h b/Engine/Plugins/Experimental/Mutable/Source/CustomizableObjectEditor/Classes/MuCOE/Nodes/CustomizableObjectNodeStaticMesh.h
--- a/Engine/Plugins/Experimental/Mutable/Source/CustomizableObjectEditor/Classes/MuCOE/Nodes/CustomizableObjectNodeStaticMesh.h
+++ b/Engine/Plugins/Experimental/Mutable/Source/CustomizableObjectEditor/Classes/MuCOE/Nodes/CustomizableObjectNodeStaticMesh.h
@@ -84,13 +84,13 @@ struct CUSTOMIZABLEOBJECTEDITOR_API FCustomizableObjectNodeStaticMeshMaterial
 	FString Name;
 
 	UPROPERTY()
-	TObjectPtr<UEdGraphPin_Deprecated> MeshPin = nullptr;
+	TObjectPtr<UEdGraphPin_Deprecated> MeshPin_DEPRECATED;
 
 	UPROPERTY()
-	TObjectPtr<UEdGraphPin_Deprecated> LayoutPin = nullptr;
+	TObjectPtr<UEdGraphPin_Deprecated> LayoutPin_DEPRECATED;
 
 	UPROPERTY()
-	TArray< TObjectPtr<UEdGraphPin_Deprecated>> ImagePins;
+	TArray< TObjectPtr<UEdGraphPin_Deprecated>> ImagePins_DEPRECATED;
 
 	UPROPERTY()
 	FEdGraphPinReference MeshPinRef;
@@ -151,9 +151,10 @@ public:
 	// UCustomizableObjectNodeMesh interface
 	virtual UTexture2D* FindTextureForPin(const UEdGraphPin* Pin) const override;
 	virtual void GetUVChannelForPin(const UEdGraphPin* Pin, TArray<FVector2f>& OutSegments, int32 UVIndex) const override;
-	virtual TArray<UCustomizableObjectLayout*> GetLayouts(const UEdGraphPin* OutPin) const override;
+	virtual TArray<UCustomizableObjectLayout*> GetLayouts(const UEdGraphPin& OutPin) const override;
 	virtual UObject* GetMesh() const override;
-	virtual UEdGraphPin* GetMeshPin(int32 LOD, int MaterialIndex) const override;
+	virtual UEdGraphPin* GetMeshPin(int32 LOD, int32 SectionIndex) const override;
+	virtual UEdGraphPin* GetLayoutPin(int32 LODIndex, int32 SectionIndex, int32 LayoutIndex) const override;
 	virtual void GetPinSection(const UEdGraphPin& Pin, int32& OutLODIndex, int32& OutSectionIndex, int32& OutLayoutIndex) const override;
 
 	/** Returns the material assossiated to the given output pin. */
diff --git a/Engine/Plugins/Experimental/Mutable/Source/CustomizableObjectEditor/Private/MuCOE/CustomizableObjectEditor.cpp b/Engine/Plugins/Experimental/Mutable/Source/CustomizableObjectEditor/Private/MuCOE/CustomizableObjectEditor.cpp
--- a/Engine/Plugins/Experimental/Mutable/Source/CustomizableObjectEditor/Private/MuCOE/CustomizableObjectEditor.cpp
+++ b/Engine/Plugins/Experimental/Mutable/Source/CustomizableObjectEditor/Private/MuCOE/CustomizableObjectEditor.cpp
@@ -2551,39 +2551,6 @@ void FCustomizableObjectEditor::OnEnterText(const FText & NewText, ETextCommit::
 							}
 						}
 					}
-					else if (ArrayType == "FCustomizableObjectNodeSkeletalMeshLOD")
-					{
-						TArray<FCustomizableObjectNodeSkeletalMeshLOD>  Array =
-							*ArrayProperty->ContainerPtrToValuePtr<TArray<FCustomizableObjectNodeSkeletalMeshLOD>>(Node);
-						
-						FString LODs = "LODs";
-
-						for (int a = 0; a < Array.Num(); ++a)
-						{
-							if (LODs.Contains(VariableName))
-							{
-								LogSearchResult(CustomizableObjectNode, "Variable", found, LODs);
-								found = true;
-							}
-
-							for (int s = 0; s < Array[a].Materials.Num(); ++s)
-							{
-								FString Name = "Name";
-
-								if (Name.Contains(VariableName))
-								{
-									LogSearchResult(CustomizableObjectNode, "Variable", found, Name);
-									found = true;
-								}
-
-								if (Array[a].Materials[s].Name.Contains(NewText.ToString()))
-								{
-									LogSearchResult(CustomizableObjectNode, "Value", found, Array[a].Materials[s].Name);
-									found = true;
-								}
-							}
-						}
-					}
 					else if (ArrayType == "FCustomizableObjectNodeStaticMeshLOD")
 					{
 						TArray<FCustomizableObjectNodeStaticMeshLOD>  Array =
diff --git a/Engine/Plugins/Experimental/Mutable/Source/CustomizableObjectEditor/Private/MuCOE/GenerateMutableSource/GenerateMutableSourceMesh.cpp b/Engine/Plugins/Experimental/Mutable/Source/CustomizableObjectEditor/Private/MuCOE/GenerateMutableSource/GenerateMutableSourceMesh.cpp
--- a/Engine/Plugins/Experimental/Mutable/Source/CustomizableObjectEditor/Private/MuCOE/GenerateMutableSource/GenerateMutableSourceMesh.cpp
+++ b/Engine/Plugins/Experimental/Mutable/Source/CustomizableObjectEditor/Private/MuCOE/GenerateMutableSource/GenerateMutableSourceMesh.cpp
@@ -2767,11 +2767,12 @@ mu::NodeMeshPtr GenerateMutableSourceMesh(const UEdGraphPin * Pin,
 		{
 			int32 LOD;
 			int32 SectionIndex;
-			int32 LayoutIndex;
-			TypedNodeSkel->GetPinSection(*Pin, LOD, SectionIndex, LayoutIndex);
-
-			check(SectionIndex < TypedNodeSkel->LODs[LOD].Materials.Num());
 
+			{
+				int32 LayoutIndex;
+				TypedNodeSkel->GetPinSection(*Pin, LOD, SectionIndex, LayoutIndex);
+			}
+				
 			// See if we need to select another LOD because of automatic LOD generation
 			if (GenerationContext.CurrentAutoLODStrategy == ECustomizableObjectAutomaticLODStrategy::AutomaticFromMesh)
 			{
@@ -2916,28 +2917,22 @@ mu::NodeMeshPtr GenerateMutableSourceMesh(const UEdGraphPin * Pin,
 			// Layouts
 			// If we didn't find a layout, but we are generating LODs and this LOD is automatic, reuse the first valid LOD layout
 			bool LayoutFound = false;
-
-			int LayoutLOD = LOD;
-			while (!LayoutFound
-				&&
-				LayoutLOD >= 0
-				)
+			
+			int32 LayoutLOD = LOD;
+			while (!LayoutFound && LayoutLOD >= 0)
 			{
-				if (TypedNodeSkel->LODs.IsValidIndex(LayoutLOD)
-					&&
-					TypedNodeSkel->LODs[LayoutLOD].Materials.IsValidIndex(SectionIndex))
+				const int32 NumLayouts = TypedNodeSkel->SkeletalMesh->GetImportedModel()->LODModels[LayoutLOD].Sections.Num();
+				
+				MeshNode->SetLayoutCount(NumLayouts);
+				for (int32 LayoutIndex = 0; LayoutIndex < NumLayouts; ++LayoutIndex)
 				{
-					MeshNode->SetLayoutCount(TypedNodeSkel->LODs[LayoutLOD].Materials[SectionIndex].LayoutPinsRef.Num());
-					for (LayoutIndex = 0; LayoutIndex < TypedNodeSkel->LODs[LayoutLOD].Materials[SectionIndex].LayoutPinsRef.Num(); ++LayoutIndex)
+					if (UEdGraphPin* LayoutPin = TypedNodeSkel->GetLayoutPin(LayoutLOD, SectionIndex, LayoutIndex))
 					{
-						if (TypedNodeSkel->LODs[LayoutLOD].Materials[SectionIndex].LayoutPinsRef[LayoutIndex].Get())
+						if (const UEdGraphPin* ConnectedPin = FollowInputPin(*LayoutPin))
 						{
-							if (const UEdGraphPin* ConnectedPin = FollowInputPin(*TypedNodeSkel->LODs[LayoutLOD].Materials[SectionIndex].LayoutPinsRef[LayoutIndex].Get()))
-							{
-								mu::NodeLayoutPtr LayoutNode = GenerateMutableSourceLayout(ConnectedPin, GenerationContext, bLinkedToExtendMaterial);
-								MeshNode->SetLayout(LayoutIndex, LayoutNode);
-								LayoutFound = true;
-							}
+							mu::NodeLayoutPtr LayoutNode = GenerateMutableSourceLayout(ConnectedPin, GenerationContext, bLinkedToExtendMaterial);
+							MeshNode->SetLayout(LayoutIndex, LayoutNode);
+							LayoutFound = true;
 						}
 					}
 				}
diff --git a/Engine/Plugins/Experimental/Mutable/Source/CustomizableObjectEditor/Private/MuCOE/GenerateMutableSource/GenerateMutableSourceSurface.cpp b/Engine/Plugins/Experimental/Mutable/Source/CustomizableObjectEditor/Private/MuCOE/GenerateMutableSource/GenerateMutableSourceSurface.cpp
--- a/Engine/Plugins/Experimental/Mutable/Source/CustomizableObjectEditor/Private/MuCOE/GenerateMutableSource/GenerateMutableSourceSurface.cpp
+++ b/Engine/Plugins/Experimental/Mutable/Source/CustomizableObjectEditor/Private/MuCOE/GenerateMutableSource/GenerateMutableSourceSurface.cpp
@@ -341,7 +341,7 @@ mu::NodeSurfacePtr GenerateMutableSourceSurface(const UEdGraphPin * Pin, FMutabl
 					{
 						if (const UCustomizableObjectNodeSkeletalMesh* SkeletalMeshNode = Cast<UCustomizableObjectNodeSkeletalMesh>(SkeletalMeshPin->GetOwningNode()))
 						{
-							if (const FSkeletalMaterial* SkeletalMaterial = SkeletalMeshNode->GetSkeletalMaterialFor(SkeletalMeshPin))
+							if (const FSkeletalMaterial* SkeletalMaterial = SkeletalMeshNode->GetSkeletalMaterialFor(*SkeletalMeshPin))
 							{
 								if (IsNewSlotName)
 								{
diff --git a/Engine/Plugins/Experimental/Mutable/Source/CustomizableObjectEditor/Private/MuCOE/Nodes/CustomizableObjectNode.cpp b/Engine/Plugins/Experimental/Mutable/Source/CustomizableObjectEditor/Private/MuCOE/Nodes/CustomizableObjectNode.cpp
--- a/Engine/Plugins/Experimental/Mutable/Source/CustomizableObjectEditor/Private/MuCOE/Nodes/CustomizableObjectNode.cpp
+++ b/Engine/Plugins/Experimental/Mutable/Source/CustomizableObjectEditor/Private/MuCOE/Nodes/CustomizableObjectNode.cpp
@@ -498,6 +498,24 @@ UCustomizableObjectNodePinData::UCustomizableObjectNodePinData()
 }
 
 
+bool UCustomizableObjectNodePinData::operator==(const UCustomizableObjectNodePinData& Other) const
+{
+	return Equals(Other);
+}
+
+
+bool UCustomizableObjectNodePinData::operator!=(const UCustomizableObjectNodePinData& Other) const
+{
+	return !operator==(Other);
+}
+
+
+bool UCustomizableObjectNodePinData::Equals(const UCustomizableObjectNodePinData& Other) const
+{
+	return GetClass() == Other.GetClass();	
+}
+
+
 void UCustomizableObjectNode::Serialize(FArchive& Ar)
 {
 	Super::Serialize(Ar);
diff --git a/Engine/Plugins/Experimental/Mutable/Source/CustomizableObjectEditor/Private/MuCOE/Nodes/CustomizableObjectNodeExtendMaterial.cpp b/Engine/Plugins/Experimental/Mutable/Source/CustomizableObjectEditor/Private/MuCOE/Nodes/CustomizableObjectNodeExtendMaterial.cpp
--- a/Engine/Plugins/Experimental/Mutable/Source/CustomizableObjectEditor/Private/MuCOE/Nodes/CustomizableObjectNodeExtendMaterial.cpp
+++ b/Engine/Plugins/Experimental/Mutable/Source/CustomizableObjectEditor/Private/MuCOE/Nodes/CustomizableObjectNodeExtendMaterial.cpp
@@ -185,7 +185,7 @@ TArray<UCustomizableObjectLayout*> UCustomizableObjectNodeExtendMaterial::GetLay
 			{
 				if (const UCustomizableObjectNodeSkeletalMesh* MeshNode = Cast<UCustomizableObjectNodeSkeletalMesh>(SourceMeshPin->GetOwningNode()))
 				{
-					Result = MeshNode->GetLayouts(SourceMeshPin);
+					Result = MeshNode->GetLayouts(*SourceMeshPin);
 				}
 				else if (const UCustomizableObjectNodeTable* TableNode = Cast<UCustomizableObjectNodeTable>(SourceMeshPin->GetOwningNode()))
 				{
diff --git a/Engine/Plugins/Experimental/Mutable/Source/CustomizableObjectEditor/Private/MuCOE/Nodes/CustomizableObjectNodeMaterial.cpp b/Engine/Plugins/Experimental/Mutable/Source/CustomizableObjectEditor/Private/MuCOE/Nodes/CustomizableObjectNodeMaterial.cpp
--- a/Engine/Plugins/Experimental/Mutable/Source/CustomizableObjectEditor/Private/MuCOE/Nodes/CustomizableObjectNodeMaterial.cpp
+++ b/Engine/Plugins/Experimental/Mutable/Source/CustomizableObjectEditor/Private/MuCOE/Nodes/CustomizableObjectNodeMaterial.cpp
@@ -694,7 +694,7 @@ TArray<UCustomizableObjectLayout*> UCustomizableObjectNodeMaterial::GetLayouts()
 			{
 				if (const UCustomizableObjectNodeSkeletalMesh* MeshNode = Cast<UCustomizableObjectNodeSkeletalMesh>(SourceMeshPin->GetOwningNode()))
 				{
-					Result = MeshNode->GetLayouts(SourceMeshPin);
+					Result = MeshNode->GetLayouts(*SourceMeshPin);
 				}
 				else if (const UCustomizableObjectNodeTable* TableNode = Cast<UCustomizableObjectNodeTable>(SourceMeshPin->GetOwningNode()))
 				{
diff --git a/Engine/Plugins/Experimental/Mutable/Source/CustomizableObjectEditor/Private/MuCOE/Nodes/CustomizableObjectNodeSkeletalMesh.cpp b/Engine/Plugins/Experimental/Mutable/Source/CustomizableObjectEditor/Private/MuCOE/Nodes/CustomizableObjectNodeSkeletalMesh.cpp
--- a/Engine/Plugins/Experimental/Mutable/Source/CustomizableObjectEditor/Private/MuCOE/Nodes/CustomizableObjectNodeSkeletalMesh.cpp
+++ b/Engine/Plugins/Experimental/Mutable/Source/CustomizableObjectEditor/Private/MuCOE/Nodes/CustomizableObjectNodeSkeletalMesh.cpp
@@ -40,68 +40,6 @@ void UCustomizableObjectNodeSkeletalMesh::PostEditChangeProperty(FPropertyChange
 }
 
 
-void UCustomizableObjectNodeSkeletalMesh::Serialize(FArchive& Ar)
-{
-	Super::Serialize(Ar);
-
-	Ar.UsingCustomVersion(FCustomizableObjectCustomVersion::GUID);
-
-	if (Ar.CustomVer(FCustomizableObjectCustomVersion::GUID) < FCustomizableObjectCustomVersion::PostLoadToCustomVersion)
-	{
-		for (FCustomizableObjectNodeSkeletalMeshLOD& LOD : LODs)
-		{
-			for(FCustomizableObjectNodeSkeletalMeshMaterial& Material : LOD.Materials)
-			{
-				if (Material.MeshPin && !Material.MeshPinRef.Get())
-				{
-					UEdGraphPin* AuxPin = UEdGraphPin::FindPinCreatedFromDeprecatedPin(Material.MeshPin);
-					Material.MeshPinRef.SetPin(AuxPin);
-				}
-
-				if (!Material.LayoutPinsRef.Num())
-				{
-					if (Material.LayoutPins.Num())
-					{
-						for (UEdGraphPin_Deprecated* LayoutPin : Material.LayoutPins)
-						{
-							UEdGraphPin* AuxPin = UEdGraphPin::FindPinCreatedFromDeprecatedPin(LayoutPin);
-							FEdGraphPinReference AuxEdGraphPinReference(AuxPin);
-							Material.LayoutPinsRef.Add(AuxEdGraphPinReference);
-						}
-					}
-					else
-					{
-						FString MaterialLayoutName = Material.Name + " Layout";
-						for (UEdGraphPin* Pin : GetAllNonOrphanPins())
-						{
-							if (Pin
-								&& Pin->Direction == EEdGraphPinDirection::EGPD_Input
-								&& (MaterialLayoutName == Helper_GetPinName(Pin)
-									|| MaterialLayoutName == Pin->PinFriendlyName.ToString()))
-							{
-								FEdGraphPinReference AuxEdGraphPinReference(Pin);
-								Material.LayoutPinsRef.Add(AuxEdGraphPinReference);
-								break;
-							}
-						}
-					}
-				}
-
-				if (!Material.ImagePinsRef.Num())
-				{
-					for (UEdGraphPin_Deprecated* ImagePin : Material.ImagePins)
-					{
-						UEdGraphPin* AuxPin = UEdGraphPin::FindPinCreatedFromDeprecatedPin(ImagePin);
-						FEdGraphPinReference AuxEdGraphPinReference(AuxPin);
-						Material.ImagePinsRef.Add(AuxEdGraphPinReference);
-					}
-				}
-			}
-		}
-	}
-}
-
-
 void UCustomizableObjectNodeSkeletalMesh::AllocateDefaultPins(UCustomizableObjectNodeRemapPins* RemapPins)
 {
 	// Pass information to the remap pins action context
@@ -112,189 +50,95 @@ void UCustomizableObjectNodeSkeletalMesh::AllocateDefaultPins(UCustomizableObjec
 	
 	const UEdGraphSchema_CustomizableObject* Schema = GetDefault<UEdGraphSchema_CustomizableObject>();
 
-	TArray<FCustomizableObjectNodeSkeletalMeshLOD> OldLODs(LODs);
-	LODs.Empty();
-
-	if (SkeletalMesh)
+	if (!SkeletalMesh)
 	{
-		int NumLODs = SkeletalMesh->GetLODInfoArray().Num();
-		TMap<FString, TArray<int32>> MaterialNameUseCount; // Stores how many times the material name has been used per LOD
-		TMap<FString, TArray<int32>> ImageNameUseCount; // Stores how many times the image name has been used per LOD
-
-		if (const FSkeletalMeshModel* ImportedModel = SkeletalMesh->GetImportedModel())
+		DefaultPin = CustomCreatePin(EGPD_Output, Schema->PC_Mesh, FName(SKELETAL_MESH_PIN_NAME));
+		return;
+	}
+	else
+	{
+		DefaultPin = {};
+	}
+	
+	if (const FSkeletalMeshModel* ImportedModel = SkeletalMesh->GetImportedModel())
+	{
+		const int32 NumLODs = SkeletalMesh->GetLODInfoArray().Num();
+		for (int32 LODIndex = 0; LODIndex < NumLODs; ++LODIndex)
 		{
-			for (int32 LODIndex = 0; LODIndex < NumLODs; ++LODIndex)
+			const int32 NumSections = ImportedModel->LODModels[LODIndex].Sections.Num();
+			for (int32 SectionIndex = 0; SectionIndex < NumSections; ++SectionIndex)
 			{
-				LODs.Add(FCustomizableObjectNodeSkeletalMeshLOD());
-
-				// Reset the material array and handle the actual materials from the mesh
-				int32 NumMaterials = ImportedModel->LODModels[LODIndex].Sections.Num();
-
-				FString OldNameForMap;
+				UMaterialInterface* MaterialInterface = GetMaterialInterfaceFor(LODIndex, SectionIndex, ImportedModel);
+				
+				FString SectionFriendlyName = MaterialInterface ? MaterialInterface->GetName() : FString::Printf(TEXT("Section %i"), SectionIndex);
 
-				for (int32 MaterialIndex = 0; MaterialIndex < NumMaterials; ++MaterialIndex)
+				// Mesh
 				{
-					LODs[LODIndex].Materials.Add(FCustomizableObjectNodeSkeletalMeshMaterial());
-
-					UMaterialInterface* MaterialInterface = GetMaterialInterfaceFor(LODIndex, MaterialIndex, ImportedModel);
-
-					FString MaterialName = LOCTEXT("Unnamed Material", "Unnamed Material").ToString();
-					if (MaterialInterface)
-					{
-						MaterialName = MaterialInterface->GetName();
-						LODs[LODIndex].Materials[MaterialIndex].Name = MaterialName;
-					}
-
-					FString PrefixId; // Prefix is empty if only 1 LOD and 1 material
-					FString FriendlyPrefixId; // For retrocompatibility, only friendly names have all indexes
-
-					if (NumLODs > 1) // Add LOD ids to prefix if more than one LOD
-					{
-						FString Aux = FString::Printf(TEXT("LOD%d-"), LODIndex);
-						if (LODIndex > 0)
-						{
-							PrefixId.Append(Aux);
-						}
-						FriendlyPrefixId.Append(Aux);
-					}
-
-					int32 MaterialCurrentCount = 0;
-					if (TArray<int32>* MUsageCountP = MaterialNameUseCount.Find(MaterialName))
+					FString MeshName = FString::Printf(TEXT("LOD %i - Section %i - Mesh"), LODIndex, SectionIndex);
+
+					UCustomizableObjectNodeSkeletalMeshPinDataMesh* PinData = NewObject<UCustomizableObjectNodeSkeletalMeshPinDataMesh>();
+					PinData->Init(LODIndex, SectionIndex);
+					
+					UEdGraphPin* Pin = CustomCreatePin(EGPD_Output, Schema->PC_Mesh, FName(*MeshName), PinData);
+					Pin->PinFriendlyName = FText::FromString(FString::Printf(TEXT("LOD %i - %s - Mesh"), LODIndex, *SectionFriendlyName));
+					Pin->PinToolTip = MeshName;
+				}
+				
+				// Layout
+				{
+					if (ImportedModel->LODModels.Num() > LODIndex)
 					{
-						if (MUsageCountP->Num() > LODIndex)
-						{
-							(*MUsageCountP)[LODIndex]++;
-							MaterialCurrentCount = (*MUsageCountP)[LODIndex];
-						}
-						else
+						const uint32 NumUVs = ImportedModel->LODModels[LODIndex].NumTexCoords;
+						for (uint32 UVIndex = 0; UVIndex < NumUVs; ++UVIndex)
 						{
-							MUsageCountP->AddZeroed(LODIndex - (MUsageCountP->Num() - 1));
-							(*MUsageCountP)[LODIndex] = MaterialCurrentCount;
+							FString PinName = FString::Printf(TEXT("LOD %i - Section %i - UV %i"), LODIndex, SectionIndex, UVIndex);
+							
+							UCustomizableObjectNodeSkeletalMeshPinDataLayout* PinData = NewObject<UCustomizableObjectNodeSkeletalMeshPinDataLayout>();
+							PinData->Init(LODIndex, SectionIndex, UVIndex);
+							
+							UEdGraphPin* Pin = CustomCreatePin(EGPD_Input, Schema->PC_Layout, FName(*PinName), PinData);
+							Pin->PinFriendlyName = FText::FromString(FString::Printf(TEXT("LOD %i - %s - UV %i"), LODIndex, *SectionFriendlyName, UVIndex));
+							Pin->PinToolTip = PinName;
 						}
 					}
-					else
-					{
-						TArray<int32> NewMaterialNameUseCount;
-						NewMaterialNameUseCount.Add(MaterialCurrentCount);
-						MaterialNameUseCount.Add(MaterialName, NewMaterialNameUseCount);
-					}
-
-					if (MaterialCurrentCount > 0) // Add section ids to prefix if a section/material name is used more than once
-					{
-						FString Aux = FString::Printf(TEXT("%d-"), MaterialCurrentCount);
-						PrefixId.Append(Aux);
-					}
+				}
 
-					UEdGraphPin* Pin = nullptr;
+				// Images
+				if (MaterialInterface)
+				{
+					const UMaterial* Material = MaterialInterface->GetMaterial();
 
-					{
-						// Mesh
-						FString MeshName = PrefixId + MaterialName + FString(" Mesh");
-						Pin = CustomCreatePin(EGPD_Output, Schema->PC_Mesh, FName(*MeshName));
-						Pin->PinFriendlyName = FText::FromString(MeshName);
-						Pin->PinToolTip = MeshName;
-						LODs[LODIndex].Materials[MaterialIndex].MeshPinRef.SetPin(Pin);
-
-						// Layout
-						if (ImportedModel->LODModels.Num() > LODIndex)
-						{
-							const uint32 numberOfUVLayouts = ImportedModel->LODModels[LODIndex].NumTexCoords;
-							for (uint32 LayoutIndex = 0; LayoutIndex < numberOfUVLayouts; ++LayoutIndex)
-							{
-								FString LayoutPrefix = "";
-								FString FriendlyLayoutPrefix = "";
-								if (LayoutIndex > 0)
-								{
-									LayoutPrefix = FString::Printf(TEXT("UV%d-"), LayoutIndex);
-								}
-								if (numberOfUVLayouts > 1)
-								{
-									FriendlyLayoutPrefix.Append(FString::Printf(TEXT("UV%d-"), LayoutIndex));
-								}
-								FString LayoutName = PrefixId + LayoutPrefix + MaterialName + FString(" Layout");
-								FString PinFriendlyName = FriendlyPrefixId + FriendlyLayoutPrefix + MaterialName + FString(" Layout");
-								FString PinTooltip = FString::Printf(TEXT("Material %d "), MaterialIndex) + PinFriendlyName;
-								Pin = CustomCreatePin(EGPD_Input, Schema->PC_Layout, FName(*LayoutName));
-								Pin->PinFriendlyName = FText::FromString(PinFriendlyName);
-								Pin->PinToolTip = PinTooltip;
-								Pin->bDefaultValueIsIgnored = true;
-								LODs[LODIndex].Materials[MaterialIndex].LayoutPinsRef.Add(FEdGraphPinReference(Pin));
-							}
-						}
-					}
+					TArray<FMaterialParameterInfo> ImageInfos;
+					TArray<FGuid> ImageIds;
+					Material->GetAllTextureParameterInfo(ImageInfos, ImageIds);
 
-					// Images
-					if (MaterialInterface)
+					check(ImageInfos.Num() == ImageIds.Num())
+					for (int32 Index = 0; Index < ImageInfos.Num(); ++Index)
 					{
-						UMaterial* Material = MaterialInterface->GetMaterial();
-
-						TArray<FMaterialParameterInfo> ImageInfos;
-						TArray<FGuid> ImageIds;
-						Material->GetAllTextureParameterInfo(ImageInfos, ImageIds);
-
-						for (int32 ImageIndex = 0; ImageIndex < ImageInfos.Num(); ++ImageIndex)
-						{
-							FString ImageName = ImageInfos[ImageIndex].Name.ToString();
-							FString Sufix = FString::Printf(TEXT("-mat-%s"), *MaterialName);
-
-
-							int32 ImageCurrentCount = 0;
-							if (TArray<int32>* IUsageCountP = ImageNameUseCount.Find(ImageName))
-							{
-								if (IUsageCountP->Num() > LODIndex)
-								{
-									(*IUsageCountP)[LODIndex]++;
-									ImageCurrentCount = (*IUsageCountP)[LODIndex];
-								}
-								else
-								{
-									IUsageCountP->AddZeroed(LODIndex - (IUsageCountP->Num() - 1));
-									(*IUsageCountP)[LODIndex] = ImageCurrentCount;
-								}
-							}
-							else
-							{
-								TArray<int32> NewImageNameUseCount;
-								NewImageNameUseCount.Add(ImageCurrentCount);
-								ImageNameUseCount.Add(ImageName, NewImageNameUseCount);
-							}
-
-							FString PinName = PrefixId + ImageName + FString::Printf(TEXT("-mat-%s"), *MaterialName);
-							FString PinFriendlyImageName = ImageName;
-
-							if (ImageCurrentCount > 0)
-							{
-								PinName.Append(FString::Printf(TEXT("-%d"), ImageCurrentCount));
-							}
-
-							Pin = CustomCreatePin(EGPD_Output, Schema->PC_Image, FName(*PinName)); // Store the sufix in the category to be able to remove the sufix from the pin name later to retrieve the textures
-							Pin->bDefaultValueIsIgnored = true;
-							Pin->bHidden = true;
-							Pin->PinFriendlyName = FText::FromString(PrefixId + PinFriendlyImageName);
-							Pin->PinToolTip = PinName;
-							LODs[LODIndex].Materials[MaterialIndex].ImagePinsRef.Add(FEdGraphPinReference(Pin));
-						}
+						const FMaterialParameterInfo& ImageInfo = ImageInfos[Index];
+						FGuid& ImageId = ImageIds[Index];
+						
+						FString ImageNameStr = *ImageInfo.Name.ToString();
+						FString PinName = FString::Printf(TEXT("LOD %i - Section %i - Texture Parameter %s"), LODIndex, SectionIndex, *ImageNameStr);
+						
+						UCustomizableObjectNodeSkeletalMeshPinDataImage* PinData = NewObject<UCustomizableObjectNodeSkeletalMeshPinDataImage>();
+						PinData->Init(LODIndex, SectionIndex, ImageId);
+						
+						UEdGraphPin* Pin = CustomCreatePin(EGPD_Output, Schema->PC_Image, FName(*PinName), PinData);
+						Pin->PinFriendlyName = FText::FromString(FString::Printf(TEXT("LOD %i - %s - %s"), LODIndex, *SectionFriendlyName, *ImageNameStr));
+						Pin->PinToolTip = PinName;
+						Pin->bHidden = true;
 					}
 				}
 			}
 		}
-
-		DefaultPin = FEdGraphPinReference();
-	}
-	else
-	{
-		UEdGraphPin* Pin = CustomCreatePin(EGPD_Output, Schema->PC_Mesh, FName(SKELETAL_MESH_PIN_NAME));
-		Pin->bDefaultValueIsIgnored = true;
-		Pin->PinFriendlyName = FText::FromString(SKELETAL_MESH_PIN_NAME);
-		
-		DefaultPin = FEdGraphPinReference(Pin);
 	}
 }
 
 
-UCustomizableObjectNodeRemapPinsByName* UCustomizableObjectNodeSkeletalMesh::CreateRemapPinsByName() const
+UCustomizableObjectNodeRemapPins* UCustomizableObjectNodeSkeletalMesh::CreateRemapPinsDefault() const
 {
-	return NewObject<UCustomizableObjectNodeRemapPinsByNameDefaultPin>();
+	return NewObject<UCustomizableObjectNodeSkeletalMeshRemapPinsBySection>();
 }
 
 
@@ -323,36 +167,46 @@ FLinearColor UCustomizableObjectNodeSkeletalMesh::GetNodeTitleColor() const
 
 UTexture2D* UCustomizableObjectNodeSkeletalMesh::FindTextureForPin(const UEdGraphPin* Pin) const
 {
+	if (!Pin)
+	{
+		return nullptr;
+	}
+	
 	if (!SkeletalMesh)
 	{
 		return nullptr;
 	}
 
-	if (const FSkeletalMeshModel* ImportedModel = SkeletalMesh->GetImportedModel())
+	const FSkeletalMeshModel* ImportedModel = SkeletalMesh->GetImportedModel();
+	if (ImportedModel)
+	{
+		return nullptr;
+	}	
+
+	if (const UCustomizableObjectNodeSkeletalMeshPinDataImage* PinData = Cast<UCustomizableObjectNodeSkeletalMeshPinDataImage>(GetPinData(*Pin)))
 	{
-		for (int LODIndex = 0; LODIndex < LODs.Num(); ++LODIndex)
+		if (const UMaterialInterface* MaterialInterface = GetMaterialInterfaceFor(PinData->GetLODIndex(), PinData->GetSectionIndex(), ImportedModel))
 		{
-			for (int MaterialIndex = 0; MaterialIndex < LODs[LODIndex].Materials.Num(); ++MaterialIndex)
+			const UMaterialInterface* Material = GetMaterialFor(Pin);
+			
+			TArray<FGuid> ParameterIds;
+			TArray<FMaterialParameterInfo> ParameterInfo;
+			Material->GetAllParameterInfoOfType(EMaterialParameterType::Texture, ParameterInfo, ParameterIds);
+
+			check(ParameterIds.Num() == ParameterInfo.Num())
+			for (int32 Index = 0; Index < ParameterIds.Num(); ++Index)
 			{
-				if (UMaterialInterface* MaterialInterface = GetMaterialInterfaceFor(LODIndex, MaterialIndex, ImportedModel))
+				if (ParameterIds[Index] == PinData->GetTextureParameterId())
 				{
-					for (int ImageIndex = 0; ImageIndex < LODs[LODIndex].Materials[MaterialIndex].ImagePinsRef.Num(); ++ImageIndex)
-					{
-						const UEdGraphPin* OtherPin = LODs[LODIndex].Materials[MaterialIndex].ImagePinsRef[ImageIndex].Get();
-						if (OtherPin == Pin)
-						{
-							UTexture* Texture = nullptr;
-							FString TextureString = OtherPin->PinFriendlyName.IsEmptyOrWhitespace() ? Helper_GetPinName(OtherPin) : OtherPin->PinFriendlyName.ToString();
-							FName TextureName = FName(*TextureString);
-							MaterialInterface->GetTextureParameterValue(TextureName, Texture);
-							return Cast<UTexture2D>(Texture);
-						}
-					}
+					UTexture* Texture = nullptr;
+					MaterialInterface->GetTextureParameterValue(ParameterInfo[Index].Name, Texture);
+
+					return Cast<UTexture2D>(Texture);
 				}
 			}
 		}
 	}
-
+	
 	return nullptr;
 }
 
@@ -375,28 +229,26 @@ void UCustomizableObjectNodeSkeletalMesh::GetUVChannelForPin(const UEdGraphPin*
 }
 
 
-TArray<UCustomizableObjectLayout*> UCustomizableObjectNodeSkeletalMesh::GetLayouts(const UEdGraphPin* OutPin) const
+TArray<UCustomizableObjectLayout*> UCustomizableObjectNodeSkeletalMesh::GetLayouts(const UEdGraphPin& MeshPin) const
 {
 	TArray<UCustomizableObjectLayout*> Result;
 
-	for (int LODIndex = 0; LODIndex < LODs.Num(); ++LODIndex)
+	const UCustomizableObjectNodeSkeletalMeshPinDataMesh* MeshPinData = Cast<UCustomizableObjectNodeSkeletalMeshPinDataMesh>(GetPinData(MeshPin));
+	check(MeshPinData); // Not a mesh pin
+
+	for (const UEdGraphPin* Pin : GetAllNonOrphanPins())
 	{
-		for (int MaterialIndex = 0; MaterialIndex < LODs[LODIndex].Materials.Num(); ++MaterialIndex)
+		if (const UCustomizableObjectNodeSkeletalMeshPinDataLayout* PinData = Cast<UCustomizableObjectNodeSkeletalMeshPinDataLayout>(GetPinData(*Pin)))
 		{
-			if (LODs[LODIndex].Materials[MaterialIndex].MeshPinRef.Get() == OutPin)
+			if (PinData->GetLODIndex() == MeshPinData->GetLODIndex() &&
+				PinData->GetSectionIndex() == MeshPinData->GetSectionIndex())
 			{
-				for (int LayoutIndex = 0; LayoutIndex < LODs[LODIndex].Materials[MaterialIndex].LayoutPinsRef.Num(); ++LayoutIndex)
+				if (const UEdGraphPin* ConnectedPin = FollowInputPin(*Pin))
 				{
-					if (UEdGraphPin* LayoutInPin = LODs[LODIndex].Materials[MaterialIndex].LayoutPinsRef[LayoutIndex].Get())
+					const UCustomizableObjectNodeLayoutBlocks* LayoutNode = Cast<UCustomizableObjectNodeLayoutBlocks>(ConnectedPin->GetOwningNode());
+					if (LayoutNode && LayoutNode->Layout)
 					{
-						if (const UEdGraphPin* ConnectedPin = FollowInputPin(*LayoutInPin))
-						{
-							const UCustomizableObjectNodeLayoutBlocks* LayoutNode = Cast<UCustomizableObjectNodeLayoutBlocks>(ConnectedPin->GetOwningNode());
-							if (LayoutNode && LayoutNode->Layout)
-							{
-								Result.Add(LayoutNode->Layout);
-							}
-						}
+						Result.Add(LayoutNode->Layout);
 					}
 				}
 			}
@@ -413,73 +265,73 @@ UObject* UCustomizableObjectNodeSkeletalMesh::GetMesh() const
 }
 
 
-UEdGraphPin* UCustomizableObjectNodeSkeletalMesh::GetMeshPin(const int32 LODIndex, const int MaterialIndex) const
+UEdGraphPin* UCustomizableObjectNodeSkeletalMesh::GetMeshPin(const int32 LODIndex, const int32 SectionIndex) const
 {
-	if (LODIndex < LODs.Num())
+	for (UEdGraphPin* Pin : GetAllNonOrphanPins())
 	{
-		if (const FCustomizableObjectNodeSkeletalMeshLOD& LOD = LODs[LODIndex];
-			MaterialIndex < LOD.Materials.Num())
+		if (const UCustomizableObjectNodeSkeletalMeshPinDataMesh* PinData = Cast<UCustomizableObjectNodeSkeletalMeshPinDataMesh>(GetPinData(*Pin)))
 		{
-			return LOD.Materials[MaterialIndex].MeshPinRef.Get();
+			if (PinData->GetLODIndex() == LODIndex &&
+				PinData->GetSectionIndex() == SectionIndex)
+			{
+				return Pin;
+			}
 		}
 	}
 
 	return nullptr;
 }
 
-void UCustomizableObjectNodeSkeletalMesh::GetPinSection(const UEdGraphPin& Pin, int32& OutLODIndex, int32& OutSectionIndex, int32& OutLayoutIndex) const
-{
-	if (LODs.IsEmpty())
-	{
-		OutLODIndex = -1;
-		OutSectionIndex = -1;
-		OutLayoutIndex = -1;
-
-		return;
-	}
 
-	for (OutLODIndex = 0; OutLODIndex < LODs.Num(); ++OutLODIndex)
+UEdGraphPin* UCustomizableObjectNodeSkeletalMesh::GetLayoutPin(int32 LODIndex, int32 SectionIndex, int32 LayoutIndex) const
+{
+	for (UEdGraphPin* Pin : GetAllNonOrphanPins())
 	{
-		const FCustomizableObjectNodeSkeletalMeshLOD& LOD = LODs[OutLODIndex];
-		
-		for (OutSectionIndex = 0; OutSectionIndex < LOD.Materials.Num(); ++OutSectionIndex)
+		if (const UCustomizableObjectNodeSkeletalMeshPinDataLayout* PinData = Cast<UCustomizableObjectNodeSkeletalMeshPinDataLayout>(GetPinData(*Pin)))
 		{
-			const FCustomizableObjectNodeSkeletalMeshMaterial& Material = LOD.Materials[OutSectionIndex];
-			
-			for (OutLayoutIndex = 0; OutLayoutIndex < Material.LayoutPinsRef.Num(); ++OutLayoutIndex)
+			if (PinData->GetLODIndex() == LODIndex &&
+				PinData->GetSectionIndex() == SectionIndex &&
+				PinData->GetUVIndex() == LayoutIndex)
 			{
-				// Is a Layout pin?
-				if (Material.LayoutPinsRef[OutLayoutIndex].Get() == &Pin)
-				{
-					return;
-				}
-			
-				// Is a Mesh pin?
-				if (Material.MeshPinRef.Get() == &Pin)
-				{
-					return;
-				}
-
-				// Is an Image pin?
-				for (int ImageIndex = 0; ImageIndex < Material.ImagePinsRef.Num(); ++ImageIndex)
-				{
-					if (Material.ImagePinsRef[ImageIndex].Get() == &Pin)
-					{
-						return;
-					}
-				}
+				return Pin;
 			}
 		}
 	}
 
-	check(false); // Pin not found. Probably a new pin has been added which is not contemplated in this function.
+	return nullptr;
+}
+
+
+void UCustomizableObjectNodeSkeletalMesh::GetPinSection(const UEdGraphPin& Pin, int32& OutLODIndex, int32& OutSectionIndex, int32& OutLayoutIndex) const
+{
+	if (const UCustomizableObjectNodeSkeletalMeshPinDataSection* PinData = Cast<UCustomizableObjectNodeSkeletalMeshPinDataSection>(GetPinData(Pin)))
+	{
+		OutLODIndex = PinData->GetLODIndex();
+		OutSectionIndex = PinData->GetSectionIndex();
+
+		if (const UCustomizableObjectNodeSkeletalMeshPinDataLayout* LayoutPinData = Cast<UCustomizableObjectNodeSkeletalMeshPinDataLayout>(GetPinData(Pin)))
+		{
+			OutLayoutIndex = LayoutPinData->GetUVIndex();
+		}
+		else
+		{
+			OutLayoutIndex = -1;
+		}
+		
+		return;
+	}
+	
+	OutLODIndex = -1;
+	OutSectionIndex = -1;
+	OutLayoutIndex = -1;
 }
 
+
 UMaterialInterface* UCustomizableObjectNodeSkeletalMesh::GetMaterialFor(const UEdGraphPin* Pin) const
 {
 	if (SkeletalMesh)
 	{
-		if (FSkeletalMaterial* SkeletalMaterial = GetSkeletalMaterialFor(Pin))
+		if (FSkeletalMaterial* SkeletalMaterial = GetSkeletalMaterialFor(*Pin))
 		{
 			return SkeletalMaterial->MaterialInterface;
 		}
@@ -488,18 +340,18 @@ UMaterialInterface* UCustomizableObjectNodeSkeletalMesh::GetMaterialFor(const UE
 	return nullptr;
 }
 
-FSkeletalMaterial* UCustomizableObjectNodeSkeletalMesh::GetSkeletalMaterialFor(const UEdGraphPin* Pin) const
+
+FSkeletalMaterial* UCustomizableObjectNodeSkeletalMesh::GetSkeletalMaterialFor(const UEdGraphPin& Pin) const
 {
-	check(Pin);
-	
 	int32 LODIndex;
 	int32 SectionIndex;
 	int32 LayoutIndex;
-	GetPinSection(*Pin, LODIndex, SectionIndex, LayoutIndex);
+	GetPinSection(Pin, LODIndex, SectionIndex, LayoutIndex);
 
 	return GetSkeletalMaterialFor(LODIndex, SectionIndex);
 }
 
+
 bool UCustomizableObjectNodeSkeletalMesh::IsPinRelevant(const UEdGraphPin* Pin) const
 {
 	const UEdGraphSchema_CustomizableObject* Schema = GetDefault<UEdGraphSchema_CustomizableObject>();
@@ -517,74 +369,88 @@ bool UCustomizableObjectNodeSkeletalMesh::IsPinRelevant(const UEdGraphPin* Pin)
 	return false;
 }
 
+
 bool UCustomizableObjectNodeSkeletalMesh::IsNodeOutDatedAndNeedsRefresh()
 {
-	bool Result = false;
-
-	if (SkeletalMesh)
+	const bool bOutdated = [&]()
 	{
-		int32 numLods = SkeletalMesh->GetLODNum();
-		if (LODs.Num() != numLods)
+		if (!SkeletalMesh)
 		{
-			Result = true;
+			return false;
 		}
 
-		if (const FSkeletalMeshModel* ImportedModel = SkeletalMesh->GetImportedModel())
+		const FSkeletalMeshModel* ImportedModel = SkeletalMesh->GetImportedModel();
+		if (!ImportedModel)
+		{
+			return false;
+		}
+	
+		for (const UEdGraphPin* Pin : GetAllNonOrphanPins())
 		{
-			for (int LODIndex = 0; !Result && LODIndex < LODs.Num(); ++LODIndex)
+			auto Connected = [](const UEdGraphPin& Pin)
+			{
+				return Pin.Direction == EGPD_Input ? FollowInputPin(Pin) != nullptr : !FollowOutputPinArray(Pin).IsEmpty();	
+			};
+			
+			auto OutdatedSectionPinData = [&](const UCustomizableObjectNodeSkeletalMeshPinDataSection& PinData) -> bool
+			{
+				return !ImportedModel->LODModels.IsValidIndex(PinData.GetLODIndex()) ||
+					!ImportedModel->LODModels[PinData.GetLODIndex()].Sections.IsValidIndex(PinData.GetSectionIndex());
+			};
+		
+			if (const UCustomizableObjectNodeSkeletalMeshPinDataLayout* LayoutPinData = Cast<UCustomizableObjectNodeSkeletalMeshPinDataLayout>(GetPinData(*Pin)))
 			{
-				int32 NumMaterials = ImportedModel->LODModels[LODIndex].Sections.Num();
-				if (LODs[LODIndex].Materials.Num() != NumMaterials)
+				if (Connected(*Pin) &&
+					(OutdatedSectionPinData(*LayoutPinData) ||
+					LayoutPinData->GetUVIndex() < 0 || LayoutPinData->GetUVIndex() >= static_cast<int32>(ImportedModel->LODModels[LayoutPinData->GetLODIndex()].NumTexCoords)))
 				{
-					Result = true;
+					return true;
 				}
-
-				for (int MaterialIndex = 0; !Result && MaterialIndex < NumMaterials; ++MaterialIndex)
+			}
+			else if (const UCustomizableObjectNodeSkeletalMeshPinDataMesh* MeshPinData = Cast<UCustomizableObjectNodeSkeletalMeshPinDataMesh>(GetPinData(*Pin)))
+			{
+				if (Connected(*Pin) &&
+					OutdatedSectionPinData(*MeshPinData))
 				{
-					if (UMaterialInterface* MaterialInterface = GetMaterialInterfaceFor(LODIndex, MaterialIndex, ImportedModel))
-					{
-						if (LODs[LODIndex].Materials[MaterialIndex].Name != MaterialInterface->GetName())
-						{
-							Result = true;
-							break;
-						}
-						UMaterial* Material = MaterialInterface->GetMaterial();
-
-						TArray<FMaterialParameterInfo> ImageInfos;
-						TArray<FGuid> ImageIds;
-						Material->GetAllTextureParameterInfo(ImageInfos, ImageIds);
-
-						if (LODs[LODIndex].Materials[MaterialIndex].ImagePinsRef.Num() != ImageInfos.Num())
-						{
-							Result = true;
-							break;
-						}
+					return true;
+				}
+			}
+			else if (const UCustomizableObjectNodeSkeletalMeshPinDataImage* ImagePinData = Cast<UCustomizableObjectNodeSkeletalMeshPinDataImage>(GetPinData(*Pin)))
+			{				
+				const UMaterialInterface* MaterialInterface = GetMaterialInterfaceFor(ImagePinData->GetLODIndex(), ImagePinData->GetSectionIndex(), ImportedModel);
+				if (!MaterialInterface) // If we had an Image pin for sure we had a MaterialInstance.
+				{
+					return true;
+				}
 
-						for (int ImageIndex = 0; ImageIndex < ImageInfos.Num(); ++ImageIndex)
-						{
-							const FString ImagePinName = LODs[LODIndex].Materials[MaterialIndex].ImagePinsRef[ImageIndex].Get()->GetName();
-							if (!ImagePinName.Contains(ImageInfos[ImageIndex].Name.ToString()))
-							{
-								Result = true;
-								break;
-							}
-						}
-					}
+				TArray<FGuid> ParameterIds;
+				TArray<FMaterialParameterInfo> ParameterInfo;
+				MaterialInterface->GetAllParameterInfoOfType(EMaterialParameterType::Texture, ParameterInfo, ParameterIds);
+				
+				UTexture* Texture = nullptr;
+				if (Connected(*Pin) &&
+					(OutdatedSectionPinData(*ImagePinData) ||
+					!ParameterIds.Contains(ImagePinData->GetTextureParameterId()))) // Check that the Texture Parameter still exists.
+				{
+					return true; 
 				}
 			}
 		}
-	}
 
+		return false;
+	}();
+	
 	// Remove previous compilation warnings
-	if (!Result && bHasCompilerMessage)
+	if (!bOutdated && bHasCompilerMessage)
 	{
 		RemoveWarnings();
 		GetGraph()->NotifyGraphChanged();
 	}
 
-    return Result;
+    return bOutdated;
 }
 
+
 FString UCustomizableObjectNodeSkeletalMesh::GetRefreshMessage() const
 {
     return "Node data outdated. Please refresh node.";
@@ -602,10 +468,59 @@ void UCustomizableObjectNodeSkeletalMesh::BackwardsCompatibleFixup()
 	Super::BackwardsCompatibleFixup();
 
 	const int32 CustomizableObjectCustomVersion = GetLinkerCustomVersion(FCustomizableObjectCustomVersion::GUID);
-
-	if (CustomizableObjectCustomVersion < FCustomizableObjectCustomVersion::NodeSkeletalMeshCorruptedPinRef)
+	
+	if (CustomizableObjectCustomVersion < FCustomizableObjectCustomVersion::PostLoadToCustomVersion)
 	{
-		ReconstructNode();
+		for (FCustomizableObjectNodeSkeletalMeshLOD& LOD : LODs_DEPRECATED)
+		{
+			for(FCustomizableObjectNodeSkeletalMeshMaterial& Material : LOD.Materials)
+			{
+				if (Material.MeshPin_DEPRECATED && !Material.MeshPinRef.Get())
+				{
+					UEdGraphPin* AuxPin = UEdGraphPin::FindPinCreatedFromDeprecatedPin(Material.MeshPin_DEPRECATED);
+					Material.MeshPinRef.SetPin(AuxPin);
+				}
+
+				if (!Material.LayoutPinsRef.Num())
+				{
+					if (Material.LayoutPins_DEPRECATED.Num())
+					{
+						for (UEdGraphPin_Deprecated* LayoutPin : Material.LayoutPins_DEPRECATED)
+						{
+							UEdGraphPin* AuxPin = UEdGraphPin::FindPinCreatedFromDeprecatedPin(LayoutPin);
+							FEdGraphPinReference AuxEdGraphPinReference(AuxPin);
+							Material.LayoutPinsRef.Add(AuxEdGraphPinReference);
+						}
+					}
+					else
+					{
+						FString MaterialLayoutName = Material.Name + " Layout";
+						for (UEdGraphPin* Pin : GetAllNonOrphanPins())
+						{
+							if (Pin
+								&& Pin->Direction == EEdGraphPinDirection::EGPD_Input
+								&& (MaterialLayoutName == Helper_GetPinName(Pin)
+									|| MaterialLayoutName == Pin->PinFriendlyName.ToString()))
+							{
+								FEdGraphPinReference AuxEdGraphPinReference(Pin);
+								Material.LayoutPinsRef.Add(AuxEdGraphPinReference);
+								break;
+							}
+						}
+					}
+				}
+
+				if (!Material.ImagePinsRef.Num())
+				{
+					for (UEdGraphPin_Deprecated* ImagePin : Material.ImagePins_DEPRECATED)
+					{
+						UEdGraphPin* AuxPin = UEdGraphPin::FindPinCreatedFromDeprecatedPin(ImagePin);
+						FEdGraphPinReference AuxEdGraphPinReference(AuxPin);
+						Material.ImagePinsRef.Add(AuxEdGraphPinReference);
+					}
+				}
+			}
+		}
 	}
 
 	if (CustomizableObjectCustomVersion < FCustomizableObjectCustomVersion::ConvertAnimationSlotToFName)
@@ -616,6 +531,67 @@ void UCustomizableObjectNodeSkeletalMesh::BackwardsCompatibleFixup()
 			AnimBlueprintSlot_DEPRECATED = -1; // Unnecessary, just in case anyone tried to use it later in this method.
 		}
 	}
+
+	if (CustomizableObjectCustomVersion < FCustomizableObjectCustomVersion::AutomaticNodeSkeletalMesh)
+	{
+		for (int32 LODIndex = 0; LODIndex < LODs_DEPRECATED.Num(); ++LODIndex)
+		{
+			const FCustomizableObjectNodeSkeletalMeshLOD& LOD = LODs_DEPRECATED[LODIndex];
+			
+			for (int32 SectionIndex = 0; SectionIndex < LOD.Materials.Num(); ++SectionIndex)
+			{
+				const FCustomizableObjectNodeSkeletalMeshMaterial& Section = LOD.Materials[SectionIndex];
+
+				{
+					UCustomizableObjectNodeSkeletalMeshPinDataMesh* PinData = NewObject<UCustomizableObjectNodeSkeletalMeshPinDataMesh>();
+					PinData->Init(LODIndex, SectionIndex);
+				
+					AddPinData(*Section.MeshPinRef.Get(), *PinData);
+				}
+
+				if (SkeletalMesh)
+				{
+					if (const FSkeletalMaterial* SkeletalMaterial = GetSkeletalMaterialFor(LODIndex, SectionIndex))
+					{
+						for (int32 ImageIndex = 0; ImageIndex < Section.ImagePinsRef.Num(); ++ImageIndex)
+						{
+							const UEdGraphPin* ImagePin = Section.ImagePinsRef[ImageIndex].Get();
+
+							FGuid TextureParameterId;
+
+							TArray<FGuid> ParameterIds;
+							TArray<FMaterialParameterInfo> ParameterInfo;
+							SkeletalMaterial->MaterialInterface->GetAllParameterInfoOfType(EMaterialParameterType::Texture, ParameterInfo, ParameterIds);
+							check(ParameterIds.Num() == ParameterInfo.Num())
+							for (int32 Index = 0; Index < ParameterIds.Num(); ++Index)
+							{
+								if (ParameterInfo[Index].Name == ImagePin->PinFriendlyName.ToString())
+								{
+									TextureParameterId = ParameterIds[Index];
+									break;	
+								}
+							}
+							
+							UCustomizableObjectNodeSkeletalMeshPinDataImage* PinData = NewObject<UCustomizableObjectNodeSkeletalMeshPinDataImage>();
+							PinData->Init(LODIndex, SectionIndex, TextureParameterId);
+						
+							AddPinData(*ImagePin, *PinData);
+						}
+					}
+				}
+
+				for (int32 LayoutIndex = 0; LayoutIndex < Section.LayoutPinsRef.Num(); ++LayoutIndex)
+				{
+					UCustomizableObjectNodeSkeletalMeshPinDataLayout* PinData = NewObject<UCustomizableObjectNodeSkeletalMeshPinDataLayout>();
+					PinData->Init(LODIndex, SectionIndex, LayoutIndex);
+				
+					AddPinData(*Section.LayoutPinsRef[LayoutIndex].Get(), *PinData);
+				}
+			}
+		}
+			
+		ReconstructNode();
+	}
 }
 
 
@@ -630,45 +606,46 @@ TSharedPtr<SGraphNode> UCustomizableObjectNodeSkeletalMesh::CreateVisualWidget()
 }
 
 
-bool UCustomizableObjectNodeSkeletalMesh::CheckIsValidLayout(const UEdGraphPin* Pin, int32& LayoutIndex, FString& MaterialName)
+bool UCustomizableObjectNodeSkeletalMesh::CheckIsValidLayout(const UEdGraphPin* InPin, int32& LayoutIndex, FString& MaterialName)
 {
-	if (const UEdGraphPin* ConnectedPin = FollowOutputPin(*Pin))
+	const UEdGraphPin* ConnectedPin = FollowOutputPin(*InPin);
+	if (!ConnectedPin)
 	{
-		int32 LODIndex;
-		int32 SectionIndex;
-		GetPinSection(*ConnectedPin, LODIndex, SectionIndex, LayoutIndex);
+		return true;
+	}
 
-		if (LayoutIndex == 0)
-		{
-			return true;
-		}
-		else
-		{
-			TArray<FEdGraphPinReference> LayoutPins = LODs[LODIndex].Materials[SectionIndex].LayoutPinsRef;
-			MaterialName = LODs[LODIndex].Materials[SectionIndex].Name;
+	int32 LODIndex;
+	int32 SectionIndex;
+	GetPinSection(*ConnectedPin, LODIndex, SectionIndex, LayoutIndex);
 
-			for (int32 i = 0; i < LayoutPins.Num(); ++i)
+	MaterialName = GetMaterialInterfaceFor(LODIndex, SectionIndex)->GetName();
+	
+	if (LayoutIndex == 0)
+	{
+		return true;
+	}
+	
+	TBitArray VisitedLayouts(false, LayoutIndex);
+	
+	for (const UEdGraphPin* Pin : GetAllNonOrphanPins())
+	{
+		if (const UCustomizableObjectNodeSkeletalMeshPinDataLayout* PinData = Cast<UCustomizableObjectNodeSkeletalMeshPinDataLayout>(GetPinData(*Pin)))
+		{
+			if (PinData->GetLODIndex() == LODIndex &&
+				PinData->GetSectionIndex() == SectionIndex &&
+				PinData->GetUVIndex() < LayoutIndex)
 			{
-				if (i == LayoutIndex)
-				{
-					break;
-				}
-
-				if (!FollowInputPin(*LayoutPins[i].Get()))
-				{
-					return false;
-				}
+				VisitedLayouts[PinData->GetUVIndex()] = true;
 			}
-
 		}
 	}
 
-	return true;
+	return !VisitedLayouts.Contains(false);
 }
 
-UMaterialInterface* UCustomizableObjectNodeSkeletalMesh::GetMaterialInterfaceFor(const int LODIndex, const int MaterialIndex, const FSkeletalMeshModel* ImportedModel) const
+UMaterialInterface* UCustomizableObjectNodeSkeletalMesh::GetMaterialInterfaceFor(const int32 LODIndex, const int32 SectionIndex, const FSkeletalMeshModel* ImportedModel) const
 {
-	if (FSkeletalMaterial* SkeletalMaterial = GetSkeletalMaterialFor(LODIndex, MaterialIndex, ImportedModel))
+	if (FSkeletalMaterial* SkeletalMaterial = GetSkeletalMaterialFor(LODIndex, SectionIndex, ImportedModel))
 	{
 		return SkeletalMaterial->MaterialInterface;
 	}
@@ -677,7 +654,7 @@ UMaterialInterface* UCustomizableObjectNodeSkeletalMesh::GetMaterialInterfaceFor
 }
 
 
-FSkeletalMaterial* UCustomizableObjectNodeSkeletalMesh::GetSkeletalMaterialFor(const int LODIndex, const int MaterialIndex, const FSkeletalMeshModel* ImportedModel) const
+FSkeletalMaterial* UCustomizableObjectNodeSkeletalMesh::GetSkeletalMaterialFor(const int32 LODIndex, const int32 SectionIndex, const FSkeletalMeshModel* ImportedModel) const
 {
 	if (!SkeletalMesh)
 	{
@@ -687,9 +664,9 @@ FSkeletalMaterial* UCustomizableObjectNodeSkeletalMesh::GetSkeletalMaterialFor(c
 	// We assume that LODIndex and MaterialIndex are valid for the imported model
 	int32 SkeletalMeshMaterialIndex = -1;
 
-	if (ImportedModel && ImportedModel->LODModels.IsValidIndex(LODIndex) && ImportedModel->LODModels[LODIndex].Sections.IsValidIndex(MaterialIndex))
+	if (ImportedModel && ImportedModel->LODModels.IsValidIndex(LODIndex) && ImportedModel->LODModels[LODIndex].Sections.IsValidIndex(SectionIndex))
 	{
-		SkeletalMeshMaterialIndex = ImportedModel->LODModels[LODIndex].Sections[MaterialIndex].MaterialIndex;
+		SkeletalMeshMaterialIndex = ImportedModel->LODModels[LODIndex].Sections[SectionIndex].MaterialIndex;
 	}
 	else
 	{
@@ -697,9 +674,9 @@ FSkeletalMaterial* UCustomizableObjectNodeSkeletalMesh::GetSkeletalMaterialFor(c
 
 		if (AuxImportedModel)
 		{
-			if (AuxImportedModel->LODModels.IsValidIndex(LODIndex) && AuxImportedModel->LODModels[LODIndex].Sections.IsValidIndex(MaterialIndex))
+			if (AuxImportedModel->LODModels.IsValidIndex(LODIndex) && AuxImportedModel->LODModels[LODIndex].Sections.IsValidIndex(SectionIndex))
 			{
-				SkeletalMeshMaterialIndex = AuxImportedModel->LODModels[LODIndex].Sections[MaterialIndex].MaterialIndex;
+				SkeletalMeshMaterialIndex = AuxImportedModel->LODModels[LODIndex].Sections[SectionIndex].MaterialIndex;
 			}
 		}
 	}
@@ -707,9 +684,9 @@ FSkeletalMaterial* UCustomizableObjectNodeSkeletalMesh::GetSkeletalMaterialFor(c
 	// Check if we have lod info map to get the correct material index
 	if (const FSkeletalMeshLODInfo* LodInfo = SkeletalMesh->GetLODInfo(LODIndex))
 	{
-		if (LodInfo->LODMaterialMap.IsValidIndex(MaterialIndex))
+		if (LodInfo->LODMaterialMap.IsValidIndex(SectionIndex))
 		{
-			SkeletalMeshMaterialIndex = LodInfo->LODMaterialMap[MaterialIndex];
+			SkeletalMeshMaterialIndex = LodInfo->LODMaterialMap[SectionIndex];
 		}
 	}
 	
@@ -848,4 +825,118 @@ EVisibility SGraphNodeSkeletalMesh::ExpressionPreviewVisibility() const
 }
 
 
+bool UCustomizableObjectNodeSkeletalMeshRemapPinsBySection::Equal(const UCustomizableObjectNode& Node, const UEdGraphPin& OldPin, const UEdGraphPin& NewPin) const
+{
+	const UCustomizableObjectNodeSkeletalMeshPinDataSection* PinDataOldPin = Cast<UCustomizableObjectNodeSkeletalMeshPinDataSection>(Node.GetPinData(OldPin));
+	const UCustomizableObjectNodeSkeletalMeshPinDataSection* PinDataNewPin = Cast<UCustomizableObjectNodeSkeletalMeshPinDataSection>(Node.GetPinData(NewPin));
+	if (PinDataOldPin && PinDataNewPin)
+	{
+		return *PinDataOldPin == *PinDataNewPin;
+	}
+	else
+	{
+		return Super::Equal(Node, OldPin, NewPin);	
+	}
+}
+
+
+void UCustomizableObjectNodeSkeletalMeshPinDataSection::Init(int32 InLODIndex, int32 InSectionIndex)
+{
+	LODIndex = InLODIndex;
+	SectionIndex = InSectionIndex;
+}
+
+
+int32 UCustomizableObjectNodeSkeletalMeshPinDataSection::GetLODIndex() const
+{
+	return LODIndex;
+}
+
+
+int32 UCustomizableObjectNodeSkeletalMeshPinDataSection::GetSectionIndex() const
+{
+	return SectionIndex;
+}
+
+
+bool UCustomizableObjectNodeSkeletalMeshPinDataSection::Equals(const UCustomizableObjectNodePinData& Other) const
+{
+	if (GetClass() != Other.GetClass())
+	{
+		return false;	
+	}
+
+	const UCustomizableObjectNodeSkeletalMeshPinDataSection& OtherTyped = static_cast<const UCustomizableObjectNodeSkeletalMeshPinDataSection&>(Other);
+    if (LODIndex != OtherTyped.LODIndex ||
+    	SectionIndex != OtherTyped.SectionIndex)
+    {
+        return false;	            
+    }
+	
+    return Super::Equals(Other);	
+}
+
+
+void UCustomizableObjectNodeSkeletalMeshPinDataImage::Init(int32 InLODIndex, int32 InSectionIndex, FGuid InTextureParameterId)
+{
+	Super::Init(InLODIndex, InSectionIndex);
+
+	TextureParameterId = InTextureParameterId;
+}
+
+
+FGuid UCustomizableObjectNodeSkeletalMeshPinDataImage::GetTextureParameterId() const
+{
+	return TextureParameterId;
+}
+
+
+bool UCustomizableObjectNodeSkeletalMeshPinDataImage::Equals(const UCustomizableObjectNodePinData& Other) const
+{
+	if (GetClass() != Other.GetClass())
+	{
+		return false;	
+	}
+
+	const UCustomizableObjectNodeSkeletalMeshPinDataImage& OtherTyped = static_cast<const UCustomizableObjectNodeSkeletalMeshPinDataImage&>(Other);
+    if (TextureParameterId != OtherTyped.TextureParameterId)
+    {
+        return false;
+    }
+	
+    return Super::Equals(Other);	
+}
+
+
+void UCustomizableObjectNodeSkeletalMeshPinDataLayout::Init(int32 InLODIndex, int32 InSectionIndex, int32 InUVIndex)
+{
+	Super::Init(InLODIndex, InSectionIndex);
+
+	UVIndex = InUVIndex;
+}
+
+
+int32 UCustomizableObjectNodeSkeletalMeshPinDataLayout::GetUVIndex() const
+{
+	return UVIndex;
+}
+
+
+bool UCustomizableObjectNodeSkeletalMeshPinDataLayout::Equals(const UCustomizableObjectNodePinData& Other) const
+{
+	if (GetClass() != Other.GetClass())
+	{
+		return false;	
+	}
+
+	const UCustomizableObjectNodeSkeletalMeshPinDataLayout& OtherTyped = static_cast<const UCustomizableObjectNodeSkeletalMeshPinDataLayout&>(Other);
+    if (UVIndex != OtherTyped.UVIndex)
+    {
+        return false;	            
+    }
+	
+    return Super::Equals(Other);
+}
+
+
 #undef LOCTEXT_NAMESPACE
diff --git a/Engine/Plugins/Experimental/Mutable/Source/CustomizableObjectEditor/Private/MuCOE/Nodes/CustomizableObjectNodeStaticMesh.cpp b/Engine/Plugins/Experimental/Mutable/Source/CustomizableObjectEditor/Private/MuCOE/Nodes/CustomizableObjectNodeStaticMesh.cpp
--- a/Engine/Plugins/Experimental/Mutable/Source/CustomizableObjectEditor/Private/MuCOE/Nodes/CustomizableObjectNodeStaticMesh.cpp
+++ b/Engine/Plugins/Experimental/Mutable/Source/CustomizableObjectEditor/Private/MuCOE/Nodes/CustomizableObjectNodeStaticMesh.cpp
@@ -53,21 +53,21 @@ void UCustomizableObjectNodeStaticMesh::Serialize(FArchive& Ar)
 		{
 			for (FCustomizableObjectNodeStaticMeshMaterial& Material : LOD.Materials)
 			{
-				if (Material.MeshPin && !Material.MeshPinRef.Get())
+				if (Material.MeshPin_DEPRECATED && !Material.MeshPinRef.Get())
 				{
-					UEdGraphPin* AuxPin = UEdGraphPin::FindPinCreatedFromDeprecatedPin(Material.MeshPin);
+					UEdGraphPin* AuxPin = UEdGraphPin::FindPinCreatedFromDeprecatedPin(Material.MeshPin_DEPRECATED);
 					Material.MeshPinRef.SetPin(AuxPin);
 				}
 
-				if (Material.LayoutPin && !Material.LayoutPinRef.Get())
+				if (Material.LayoutPin_DEPRECATED && !Material.LayoutPinRef.Get())
 				{
-					UEdGraphPin* AuxPin = UEdGraphPin::FindPinCreatedFromDeprecatedPin(Material.LayoutPin);
+					UEdGraphPin* AuxPin = UEdGraphPin::FindPinCreatedFromDeprecatedPin(Material.LayoutPin_DEPRECATED);
 					Material.LayoutPinRef.SetPin(AuxPin);
 				}
 
-				if (Material.ImagePins.Num() && !Material.ImagePinsRef.Num())
+				if (Material.ImagePins_DEPRECATED.Num() && !Material.ImagePinsRef.Num())
 				{
-					for (UEdGraphPin_Deprecated* DeprecatedPin : Material.ImagePins)
+					for (UEdGraphPin_Deprecated* DeprecatedPin : Material.ImagePins_DEPRECATED)
 					{
 						UEdGraphPin* AuxPin = UEdGraphPin::FindPinCreatedFromDeprecatedPin(DeprecatedPin);
 						FEdGraphPinReference AuxEdGraphPinReference(AuxPin);
@@ -251,7 +251,7 @@ void UCustomizableObjectNodeStaticMesh::GetUVChannelForPin(const UEdGraphPin* Pi
 }
 
 
-TArray<class UCustomizableObjectLayout*> UCustomizableObjectNodeStaticMesh::GetLayouts(const UEdGraphPin* OutPin) const
+TArray<class UCustomizableObjectLayout*> UCustomizableObjectNodeStaticMesh::GetLayouts(const UEdGraphPin& OutPin) const
 {
 	TArray<class UCustomizableObjectLayout*> Result;
 
@@ -259,7 +259,7 @@ TArray<class UCustomizableObjectLayout*> UCustomizableObjectNodeStaticMesh::GetL
 	{
 		for (int MaterialIndex = 0; MaterialIndex<LODs[LODIndex].Materials.Num(); ++MaterialIndex)
 		{
-			if (LODs[LODIndex].Materials[MaterialIndex].MeshPinRef.Get() == OutPin)
+			if (LODs[LODIndex].Materials[MaterialIndex].MeshPinRef.Get() == &OutPin)
 			{
 				if (UEdGraphPin* LayoutInPin = LODs[LODIndex].Materials[MaterialIndex].LayoutPinRef.Get())
 				{
@@ -286,14 +286,31 @@ UObject* UCustomizableObjectNodeStaticMesh::GetMesh() const
 }
 
 
-UEdGraphPin* UCustomizableObjectNodeStaticMesh::GetMeshPin(const int32 LODIndex, const int MaterialIndex) const
+UEdGraphPin* UCustomizableObjectNodeStaticMesh::GetMeshPin(const int32 LODIndex, const int32 SectionIndex) const
 {
 	if (LODIndex < LODs.Num())
 	{
 		if (const FCustomizableObjectNodeStaticMeshLOD& LOD = LODs[LODIndex];
-			MaterialIndex < LOD.Materials.Num())
+			SectionIndex < LOD.Materials.Num())
 		{
-			return LOD.Materials[MaterialIndex].MeshPinRef.Get();
+			return LOD.Materials[SectionIndex].MeshPinRef.Get();
+		}
+	}
+
+	return nullptr;
+}
+
+
+UEdGraphPin* UCustomizableObjectNodeStaticMesh::GetLayoutPin(int32 LODIndex, int32 SectionIndex, int32 LayoutIndex) const
+{
+	check(LayoutIndex == 1); // Multiple UVs not supported on Static Mesh Node.
+
+	if (LODIndex < LODs.Num())
+	{
+		if (const FCustomizableObjectNodeStaticMeshLOD& LOD = LODs[LODIndex];
+			SectionIndex < LOD.Materials.Num())
+		{
+			return LOD.Materials[SectionIndex].LayoutPinRef.Get();
 		}
 	}
 
diff --git a/Engine/Plugins/Experimental/Mutable/Source/CustomizableObjectEditor/Public/MuCOE/CustomizableObjectEditor_Deprecated.h b/Engine/Plugins/Experimental/Mutable/Source/CustomizableObjectEditor/Public/MuCOE/CustomizableObjectEditor_Deprecated.h
--- a/Engine/Plugins/Experimental/Mutable/Source/CustomizableObjectEditor/Public/MuCOE/CustomizableObjectEditor_Deprecated.h
+++ b/Engine/Plugins/Experimental/Mutable/Source/CustomizableObjectEditor/Public/MuCOE/CustomizableObjectEditor_Deprecated.h
@@ -87,6 +87,47 @@ struct CUSTOMIZABLEOBJECTEDITOR_API FCustomizableObjectNodeExtendMaterialImage
 	FString Name;
 };
 
+// UCustomizableObjectNodeSkeletalMesh
+// Deprecated, do not use!
+USTRUCT()
+struct CUSTOMIZABLEOBJECTEDITOR_API FCustomizableObjectNodeSkeletalMeshMaterial
+{
+	GENERATED_USTRUCT_BODY()
+
+	UPROPERTY()
+	FString Name;
+
+	UPROPERTY()
+	TObjectPtr<UEdGraphPin_Deprecated> MeshPin_DEPRECATED;
+
+	UPROPERTY()
+	TArray<TObjectPtr<UEdGraphPin_Deprecated>> LayoutPins_DEPRECATED;
+
+	UPROPERTY()
+	TArray<TObjectPtr<UEdGraphPin_Deprecated>> ImagePins_DEPRECATED;
+
+	UPROPERTY()
+	FEdGraphPinReference MeshPinRef;
+
+	UPROPERTY()
+	TArray<FEdGraphPinReference> LayoutPinsRef;
+
+	UPROPERTY()
+	TArray<FEdGraphPinReference> ImagePinsRef;
+};
+
+// Deprecated, do not use!
+USTRUCT()
+struct CUSTOMIZABLEOBJECTEDITOR_API FCustomizableObjectNodeSkeletalMeshLOD
+{
+	GENERATED_USTRUCT_BODY()
+
+	UPROPERTY()
+	TArray<FCustomizableObjectNodeSkeletalMeshMaterial> Materials;
+};
+
+
+
 #if UE_ENABLE_INCLUDE_ORDER_DEPRECATED_IN_5_2
 #include "Engine/Texture2D.h"
 #endif
diff --git a/Engine/Plugins/Experimental/Mutable/Source/MutableRuntime/Private/MuR/SystemPrivate.h b/Engine/Plugins/Experimental/Mutable/Source/MutableRuntime/Private/MuR/SystemPrivate.h
--- a/Engine/Plugins/Experimental/Mutable/Source/MutableRuntime/Private/MuR/SystemPrivate.h
+++ b/Engine/Plugins/Experimental/Mutable/Source/MutableRuntime/Private/MuR/SystemPrivate.h
@@ -23,7 +23,7 @@ namespace mu
 	{
 		// This code will only be compiled (and ran) if the global definition to enable LLM tracking is set to 1 for the host program
 		// Ex : 			GlobalDefinitions.Add("LLM_ENABLED_IN_CONFIG=1");
-#if ENABLE_LOW_LEVEL_MEM_TRACKER
+#if ENABLE_LOW_LEVEL_MEM_TRACKER && 0
 		FLowLevelMemTracker& MemTracker = FLowLevelMemTracker::Get();
 		if (MemTracker.IsEnabled())
 		{
