[Mutable] Added Refresh All Child Nodes option.

- Removed Refresh All Material Nodes in Child Objects buttion.
- Added a new generic Refresh All Nodes button which supports all node types.
    - Will reconstruct child nodes from the selected one.
    - The user can chose which node type to reconstruct.
- Added a list of all nodes in the schema.
- Fixed Group Object Node and Object Node names.
- Minor code cleanup (added missing virtuals, unecessary semicolons...).

New refrehs only refreshes nodes connected in the co hierarchy.

#rnx

diff --git a/Engine/Plugins/Experimental/Mutable/Source/CustomizableObjectEditor/Classes/MuCOE/EdGraphSchema_CustomizableObject.h b/Engine/Plugins/Experimental/Mutable/Source/CustomizableObjectEditor/Classes/MuCOE/EdGraphSchema_CustomizableObject.h
--- a/Engine/Plugins/Experimental/Mutable/Source/CustomizableObjectEditor/Classes/MuCOE/EdGraphSchema_CustomizableObject.h
+++ b/Engine/Plugins/Experimental/Mutable/Source/CustomizableObjectEditor/Classes/MuCOE/EdGraphSchema_CustomizableObject.h
@@ -14,6 +14,7 @@ class UEdGraphPin;
 class UObject;
 struct FAssetData;
 struct FEdGraphPinType;
+class ICustomizableObjectEditor;
 
 /** Action to add a node to the graph */
 USTRUCT()
@@ -103,29 +104,37 @@ public:
 	static const FName PC_Enum;
 	static const FName PC_Stack;
 	static const FName PC_MaterialAsset;
+
+private:
+	/** List of all node types. */
+	TArray<UClass*> NodeTypes;
 	
-	// Begin EdGraphSchema interface
+public:
+
+	// EdGraphSchema interface
 	virtual void GetAssetsGraphHoverMessage(const TArray<FAssetData>& Assets, const UEdGraph* HoverGraph, FString& OutTooltipText, bool& OutOkIcon) const override;
-	void GetGraphContextActions(FGraphContextMenuBuilder& ContextMenuBuilder) const override;
-	const FPinConnectionResponse CanCreateConnection(const UEdGraphPin* A, const UEdGraphPin* B) const override;
-	FLinearColor GetPinTypeColor(const FEdGraphPinType& PinType) const override;
-	bool ShouldHidePinDefaultValue(UEdGraphPin* Pin) const override;
+	virtual void GetGraphContextActions(FGraphContextMenuBuilder& ContextMenuBuilder) const override;
+	virtual const FPinConnectionResponse CanCreateConnection(const UEdGraphPin* A, const UEdGraphPin* B) const override;
+	virtual FLinearColor GetPinTypeColor(const FEdGraphPinType& PinType) const override;
+	virtual bool ShouldHidePinDefaultValue(UEdGraphPin* Pin) const override;
 	virtual void DroppedAssetsOnGraph(const TArray<FAssetData>& Assets, const FVector2D& GraphPosition, UEdGraph* Graph) const override;
 
 	FLinearColor GetPinTypeColor(const FName& PinType) const;
 
-	void GetContextMenuActions(class UToolMenu* Menu, class UGraphNodeContextMenuContext* Context) const override;
+	void GetContextMenuActionsReconstructAllChildNodes(UToolMenu* Menu, TWeakObjectPtr<UGraphNodeContextMenuContext> Context) const;
+	virtual void GetContextMenuActions(UToolMenu* Menu, UGraphNodeContextMenuContext* Context) const override;
 
 	virtual void BreakPinLinks(UEdGraphPin& TargetPin, bool bSendsNodeNotification) const override;
 	virtual void BreakSinglePinLink(UEdGraphPin* SourcePin, UEdGraphPin* TargetPin) const override;
-	// End EdGraphSchema interface
+
+	//  Own interface
 
 	/** Utility to create a node creating action */
 	static TSharedPtr<FCustomizableObjectSchemaAction_NewNode> AddNewNodeAction(FGraphActionListBuilderBase& ContextMenuBuilder, const FString& Category, const FText& MenuDesc, const FText& Tooltip, const int32 Grouping = 0, const FString& Keywords = FString());
 
 	/** */
 	void GetBreakLinkToSubMenuActions( class FMenuBuilder& MenuBuilder, UEdGraphPin* InGraphPin ) const;
-
+	
 private: 
 /** Enum containing all the object types that we are able to convert onto a node when dragging
  * and dropping and asset of that type onto the CO graph.
diff --git a/Engine/Plugins/Experimental/Mutable/Source/CustomizableObjectEditor/Private/MuCOE/CustomizableObjectEditor.cpp b/Engine/Plugins/Experimental/Mutable/Source/CustomizableObjectEditor/Private/MuCOE/CustomizableObjectEditor.cpp
--- a/Engine/Plugins/Experimental/Mutable/Source/CustomizableObjectEditor/Private/MuCOE/CustomizableObjectEditor.cpp
+++ b/Engine/Plugins/Experimental/Mutable/Source/CustomizableObjectEditor/Private/MuCOE/CustomizableObjectEditor.cpp
@@ -620,10 +620,6 @@ TSharedRef<SGraphEditor> FCustomizableObjectEditor::CreateGraphEditorWidget(UEdG
 			FCanExecuteAction::CreateSP(this, &FCustomizableObjectEditor::CanDuplicateSelectedNodes)
 			);
 
-		GraphEditorCommands->MapAction(FCustomizableObjectEditorNodeContextCommands::Get().RefreshMaterialNodesInAllChildren,
-			FExecuteAction::CreateSP(this, &FCustomizableObjectEditor::RefreshMaterialNodesInAllChildrenCallback)
-			);
-
 		GraphEditorCommands->MapAction(FCustomizableObjectEditorNodeContextCommands::Get().CreateComment,
 			FExecuteAction::CreateSP(this, &FCustomizableObjectEditor::CreateCommentBoxFromKey)
 		);
@@ -872,177 +868,47 @@ bool FCustomizableObjectEditor::GroupNodeIsLinkedToParentByName(UCustomizableObj
 }
 
 
-// Deprecated! Remove on MTBL-283
-void StoreNodeLinkedPins(UCustomizableObjectNode* Node, TMap<FString, int32>& MapLinkedPin)
-{
-	TArray<UEdGraphPin*> ArrayPins = Node->GetAllNonOrphanPins();
-	const int32 MaxIndex = ArrayPins.Num();
-
-	for (int32 i = 0; i < MaxIndex; ++i)
-	{
-		if (ArrayPins[i]->LinkedTo.Num() > 0)
-		{
-			MapLinkedPin.Emplace(ArrayPins[i]->GetDisplayName().ToString(), ArrayPins[i]->LinkedTo.Num());
-		}
-	}
-}
-
-
-// Deprecated! Remove on MTBL-283
-bool CompareLinkedPinMaps(const TMap<FString, int32>& Map0, const TMap<FString, int32>& Map1, const UEdGraphNode* Node)
-{
-	// Compare if there are connected pins in Map0 not present in Map1, viceversa, and also if the number of connections is different for Map0 and Map1
-	bool PinIn0NotIn1 = false;
-	bool PinIn1NotIn0 = false;
-	bool DifferentLinkNumber = false;
-
-	for (const TPair<FString, int32>& Element : Map0)
-	{
-		const int32* NumElement1 = Map1.Find(Element.Key);
-		if (NumElement1 == nullptr)
-		{
-			PinIn0NotIn1 = true;
-			continue;
-		}
-
-		int32 NumElement0 = Element.Value;
-		if (NumElement0 != (*NumElement1))
-		{
-			DifferentLinkNumber = true;
-		}
-	}
-
-	for (const TPair<FString, int32>& Element : Map1)
-	{
-		const int32* NumElement0 = Map0.Find(Element.Key);
-		if (NumElement0 == nullptr)
-		{
-			PinIn1NotIn0 = true;
-		}
-	}
-
-	if (PinIn0NotIn1 || PinIn1NotIn0 || DifferentLinkNumber)
-	{
-		FString ErrorLog = "ERROR: Updated UCustomizableObjectNodeMaterial";
-
-		if (PinIn0NotIn1)
-		{
-			ErrorLog += ", has lost at least one connected pin ";
-		}
-
-		if (PinIn1NotIn0)
-		{
-			ErrorLog += ", has at least one connected pin with name not present before update";
-		}
-
-		if (DifferentLinkNumber)
-		{
-			ErrorLog += ", has at least one pin with different number connection after update";
-		}
-
-		ErrorLog += ". Please verify the node works correctly";
-
-		FCustomizableObjectEditorLogger::CreateLog(FText::FromString(ErrorLog))
-		.Severity(EMessageSeverity::Warning)
-		.Node(*CastChecked<UCustomizableObjectNode>(Node))
-		.Log();
-	}
-
-	return (!PinIn0NotIn1 && !PinIn1NotIn0 && !DifferentLinkNumber);
-}
-
-
-void FCustomizableObjectEditor::RefreshMaterialNodesInAllChildrenCallback()
-{
-	FNotificationInfo Info(NSLOCTEXT("CustomizableObject", "RefreshingMaterialNodesInAllChildren", "Refreshing material nodes in all children"));
-	Info.bFireAndForget = true;
-	Info.bUseThrobber = true;
-	Info.FadeOutDuration = 1.0f;
-	Info.ExpireDuration = 1.0f;
-	FSlateNotificationManager::Get().AddNotification(Info);
-
-	LaunchRefreshMaterialInAllChildren = true;
-}
-
-
-// Deprecated! Remove on MTBL-283
-void FCustomizableObjectEditor::RefreshMaterialNodesInAllChildren()
+// TODO GMT Future, use graph traversal abstraction instead of a hardcoded implementation.
+void FCustomizableObjectEditor::ReconstructAllChildNodes(UCustomizableObjectNode& StartNode, const UClass& NodeType)
 {
-	// Verify only one UCustomizableObjectNodeObjectGroup node is selected
-	const FGraphPanelSelectionSet SelectedNodes = GraphEditor->GetSelectedNodes();
-
-	if (SelectedNodes.Num() != 1)
+	UCustomizableObject* Object = CastChecked<UCustomizableObject>(StartNode.GetCustomizableObjectGraph()->GetOuter());
+	const TMultiMap<FGuid, UCustomizableObjectNodeObject*> Mapping = GetNodeGroupObjectNodeMapping(Object);
+	
+	TArray<UCustomizableObjectNode*> NodesToVisit;
+	NodesToVisit.Add(&StartNode);
+	
+	while (!NodesToVisit.IsEmpty())
 	{
-		return;
-	}
+		UCustomizableObjectNode* Node = NodesToVisit.Pop();
 
-	FString ParentGroupName = "";
-	for (FGraphPanelSelectionSet::TConstIterator SelectedIter(SelectedNodes); SelectedIter; ++SelectedIter)
-	{
-		if (UCustomizableObjectNodeObjectGroup* Node = Cast<UCustomizableObjectNodeObjectGroup>(*SelectedIter))
+		if (&NodeType == Node->GetClass())
 		{
-			ParentGroupName = Node->GroupName;
+			Node->UCustomizableObjectNode::ReconstructNode();							
 		}
-	}
-
-	// Loop through all COs and find if a CO has a UCustomizableObjectNodeObject node linked to the CO being analyzed,
-	// and adding the group name for filtering only for those COs of the selected group
-	FAssetRegistryModule& AssetRegistryModule = FModuleManager::LoadModuleChecked<FAssetRegistryModule>("AssetRegistry");
-	TArray<FAssetData> OutAssetData;
-	AssetRegistryModule.Get().GetAssetsByClass(UE_MUTABLE_TOPLEVELASSETPATH(TEXT("/Script/CustomizableObject"), TEXT("CustomizableObject")), OutAssetData);
 
-	for (auto Itr = OutAssetData.CreateIterator(); Itr; ++Itr)
-	{
-		UCustomizableObject* CustomizableObjectTemp = Cast<UCustomizableObject>(Itr->GetAsset());
-
-		if (CustomizableObjectTemp == nullptr)
+		if (const UCustomizableObjectNodeObjectGroup* GroupNode = Cast<UCustomizableObjectNodeObjectGroup>(Node))
 		{
-			continue;
+			TArray<UCustomizableObjectNodeObject*> ObjectNodes;
+			Mapping.MultiFind(GroupNode->NodeGuid, ObjectNodes);
+			
+			for (UCustomizableObjectNodeObject* ObjectNode : ObjectNodes)
+			{
+				NodesToVisit.Add(ObjectNode);	
+			}
 		}
 		
-		bool MultipleBaseObjectsFound;
-		UCustomizableObjectNodeObject* TempParentNodeObject = GetRootNode(CustomizableObjectTemp, MultipleBaseObjectsFound);
-
-		if (MultipleBaseObjectsFound)
-		{
-			continue;
-		}
-
-		bool IsChild = HasCandidateAsParent(TempParentNodeObject, CustomizableObject);
-
-		if (IsChild)
+		for (const UEdGraphPin* Pin : Node->GetAllPins()) // Not using GetAllNonOrphanPins on purpose since we want want to be able to reconstruct nodes that have non-orphan pins.
 		{
-			bool IsLinkedTo = GroupNodeIsLinkedToParentByName(TempParentNodeObject, CustomizableObject, ParentGroupName);
-
-			if (IsLinkedTo)
+			if (Pin->Direction != EGPD_Input)
 			{
-				UEdGraph* Graph = CustomizableObjectTemp->Source;
+				continue;
+			}
 
-				if (Graph != nullptr)
+			for (const UEdGraphPin* ConnectedPin : FollowInputPinArray(*Pin))
+			{
+				if (UCustomizableObjectNode* TypedNode = Cast<UCustomizableObjectNode>(ConnectedPin->GetOwningNode()))
 				{
-					TArray<UCustomizableObjectNodeMaterial*> ArrayNode;
-					Graph->GetNodesOfClass<UCustomizableObjectNodeMaterial>(ArrayNode);
-
-					int32 MaxIndex = ArrayNode.Num();
-					int32 NumNodeReconstructed = 0;
-					TMap<FString, int32> MapLinkedPinsBeforeUpdate;
-					TMap<FString, int32> MapLinkedPinsAfterUpdate;
-
-					for (int32 i = 0; i < MaxIndex; ++i)
-					{
-						StoreNodeLinkedPins(ArrayNode[i], MapLinkedPinsBeforeUpdate);
-						ArrayNode[i]->UCustomizableObjectNode::ReconstructNode();
-						StoreNodeLinkedPins(ArrayNode[i], MapLinkedPinsAfterUpdate);
-						NumNodeReconstructed++;
-							
-						CompareLinkedPinMaps(MapLinkedPinsBeforeUpdate, MapLinkedPinsAfterUpdate, ArrayNode[i]);
-					}
-
-					if (NumNodeReconstructed > 0)
-					{
-						Graph->NotifyGraphChanged();
-						Graph->MarkPackageDirty();
-					}
+					NodesToVisit.Add(TypedNode);						
 				}
 			}
 		}
@@ -1669,7 +1535,7 @@ void FCustomizableObjectEditor::NotifyPostChange( const FPropertyChangedEvent& P
 }
 
 
-bool FCustomizableObjectEditor::IsTickable(void) const
+bool FCustomizableObjectEditor::IsTickable() const
 {
 	return true;
 }
@@ -1826,18 +1692,6 @@ void FCustomizableObjectEditor::Tick( float InDeltaTime )
         }
 	}
 
-	if (LaunchRefreshMaterialInAllChildren)
-	{
-		PendingTimeRefreshMaterialInAllChildren -= InDeltaTime;
-
-		if (PendingTimeRefreshMaterialInAllChildren < 0.0f)
-		{
-			PendingTimeRefreshMaterialInAllChildren = 2.0f;
-			LaunchRefreshMaterialInAllChildren = false;
-			RefreshMaterialNodesInAllChildren();
-		}
-	}
-
 	// Reconstruct marked nodes 
 	if (GraphEditor.IsValid())
 	{
diff --git a/Engine/Plugins/Experimental/Mutable/Source/CustomizableObjectEditor/Private/MuCOE/CustomizableObjectEditor.h b/Engine/Plugins/Experimental/Mutable/Source/CustomizableObjectEditor/Private/MuCOE/CustomizableObjectEditor.h
--- a/Engine/Plugins/Experimental/Mutable/Source/CustomizableObjectEditor/Private/MuCOE/CustomizableObjectEditor.h
+++ b/Engine/Plugins/Experimental/Mutable/Source/CustomizableObjectEditor/Private/MuCOE/CustomizableObjectEditor.h
@@ -189,7 +189,8 @@ public:
 	virtual bool CanPasteNodes() const override;
 	virtual void PasteNodesHere(const FVector2D& Location) override;
 	virtual void SelectNode(const UCustomizableObjectNode* Node) override;
-	
+	virtual void ReconstructAllChildNodes(UCustomizableObjectNode& StartNode, const UClass& NodeType) override;
+
 	/** Called to undo the last action */
 	void UndoGraphAction();
 
@@ -200,22 +201,15 @@ public:
 	* in the Test CO given as parameter */
 	static bool GroupNodeIsLinkedToParentByName(UCustomizableObjectNodeObject* Node, UCustomizableObject* Test, const FString& ParentGroupName);
 
-	/** Callback to flag LaunchRefreshMaterialInAllChildren and start the timing for the Editor notification before actually calling
-	* the method that does the task, RefreshMaterialNodesInAllChildren */
-	void RefreshMaterialNodesInAllChildrenCallback();
-
-	/** Called to refresh all the Customizable Object Material Node nodes children of a Customizalbe Object Node Object Group */
-	void RefreshMaterialNodesInAllChildren();
-
 	// FEditorUndoClient Interface
 	virtual void PostUndo(bool bSuccess) override;
 	virtual void PostRedo(bool bSuccess) override { PostUndo(bSuccess); }
 
-	/** FNotifyHook interface */
+	// FNotifyHook interface
 	virtual void NotifyPostChange( const FPropertyChangedEvent& PropertyChangedEvent, FProperty* PropertyThatChanged ) override;
 
-	/** FTickableGameObject interface */
-	virtual bool IsTickable(void) const override;
+	// FTickableGameObject interface
+	virtual bool IsTickable() const override;
 	virtual void Tick( float InDeltaTime ) override;
 	virtual TStatId GetStatId() const override;
 
@@ -225,6 +219,7 @@ public:
 	void DuplicateSelectedNodes();
 	bool CanDuplicateSelectedNodes() const;
 	void OnSelectedGraphNodesChanged(const FGraphPanelSelectionSet& NewSelection);
+
 	/**
 	 * Called when a node's title is committed for a rename
 	 *
@@ -470,14 +465,6 @@ private:
 	/** Advanced scene preview settings */
 	class UCustomizableObjectEmptyClassForSettings* AdditionalSettings;
 	
-	/** When launching the refresh material in all children command from Customizable Object Node Object Group, if the COs are not loaded, the
-	* Editor will remaing stuck since the command is executed on the game thread. Use this flag to launch the Editor notification */
-	bool LaunchRefreshMaterialInAllChildren = false;
-
-	/** Used together with LaunchRefreshMaterialInAllChildren as a timer to launch the RefreshMaterialNodesInAllChildren method
-	* which will perform the task */
-	float PendingTimeRefreshMaterialInAllChildren = 2.0;
-
 	/** Texture Analyzer table widget which shows the information of the transient textures used in the customizable object instance */
 	TSharedPtr<class SCustomizableObjecEditorTextureAnalyzer> TextureAnalyzer;
 
@@ -490,7 +477,7 @@ private:
 	/** Adds the customizable Object Editor commands to the default toolbar */
 	void ExtendToolbar();
 
-	/** List of ndes to be reconstructed at the end of the tick. */
+	/** List of nodes to be reconstructed at the end of the tick. */
 	TArray<UEdGraphNode*> ReconstructNodes;
 
 	/** URL to open when pressing the documentation button generated by UE */
diff --git a/Engine/Plugins/Experimental/Mutable/Source/CustomizableObjectEditor/Private/MuCOE/CustomizableObjectEditorNodeContextCommands.cpp b/Engine/Plugins/Experimental/Mutable/Source/CustomizableObjectEditor/Private/MuCOE/CustomizableObjectEditorNodeContextCommands.cpp
--- a/Engine/Plugins/Experimental/Mutable/Source/CustomizableObjectEditor/Private/MuCOE/CustomizableObjectEditorNodeContextCommands.cpp
+++ b/Engine/Plugins/Experimental/Mutable/Source/CustomizableObjectEditor/Private/MuCOE/CustomizableObjectEditorNodeContextCommands.cpp
@@ -7,8 +7,6 @@
 
 void FCustomizableObjectEditorNodeContextCommands::RegisterCommands()
 {
-	UI_COMMAND(RefreshMaterialNodesInAllChildren, "Refresh all material nodes in all children", "Refresh all material nodes in all children.", EUserInterfaceActionType::Button, FInputChord());
-	
 	FKey CommentSpawnKey = FKey(TEXT("C"));
 	UI_COMMAND(CreateComment, "Create a Comment node", "Create a Comment node", EUserInterfaceActionType::Button, FInputChord(CommentSpawnKey));
 }
diff --git a/Engine/Plugins/Experimental/Mutable/Source/CustomizableObjectEditor/Private/MuCOE/CustomizableObjectEditorNodeContextCommands.h b/Engine/Plugins/Experimental/Mutable/Source/CustomizableObjectEditor/Private/MuCOE/CustomizableObjectEditorNodeContextCommands.h
--- a/Engine/Plugins/Experimental/Mutable/Source/CustomizableObjectEditor/Private/MuCOE/CustomizableObjectEditorNodeContextCommands.h
+++ b/Engine/Plugins/Experimental/Mutable/Source/CustomizableObjectEditor/Private/MuCOE/CustomizableObjectEditorNodeContextCommands.h
@@ -24,9 +24,6 @@ public:
 	{
 	}
 
-	/** Refresh material nodes in all children */
-	TSharedPtr<FUICommandInfo> RefreshMaterialNodesInAllChildren;
-
 	/** Create comment*/
 	TSharedPtr<FUICommandInfo> CreateComment;
 
diff --git a/Engine/Plugins/Experimental/Mutable/Source/CustomizableObjectEditor/Private/MuCOE/EdGraphSchema_CustomizableObject.cpp b/Engine/Plugins/Experimental/Mutable/Source/CustomizableObjectEditor/Private/MuCOE/EdGraphSchema_CustomizableObject.cpp
--- a/Engine/Plugins/Experimental/Mutable/Source/CustomizableObjectEditor/Private/MuCOE/EdGraphSchema_CustomizableObject.cpp
+++ b/Engine/Plugins/Experimental/Mutable/Source/CustomizableObjectEditor/Private/MuCOE/EdGraphSchema_CustomizableObject.cpp
@@ -227,8 +227,69 @@ const FName UEdGraphSchema_CustomizableObject::PC_MaterialAsset("materialAsset")
 
 
 UEdGraphSchema_CustomizableObject::UEdGraphSchema_CustomizableObject()
-	: Super()
 {
+	NodeTypes.Add(UCustomizableObjectNodeAnimationPose::StaticClass());
+	NodeTypes.Add(UCustomizableObjectNodeColorArithmeticOp::StaticClass());
+	NodeTypes.Add(UCustomizableObjectNodeColorConstant::StaticClass());
+	NodeTypes.Add(UCustomizableObjectNodeColorFromFloats::StaticClass());
+	NodeTypes.Add(UCustomizableObjectNodeColorParameter::StaticClass());
+	NodeTypes.Add(UCustomizableObjectNodeColorSwitch::StaticClass());
+	NodeTypes.Add(UCustomizableObjectNodeColorVariation::StaticClass());
+	NodeTypes.Add(UCustomizableObjectNodeCopyMaterial::StaticClass());
+	NodeTypes.Add(UCustomizableObjectNodeCurve::StaticClass());
+	NodeTypes.Add(UCustomizableObjectNodeEditMaterial::StaticClass());
+	NodeTypes.Add(UCustomizableObjectNodeEnumParameter::StaticClass());
+	NodeTypes.Add(UCustomizableObjectNodeExtendMaterial::StaticClass());
+	NodeTypes.Add(UCustomizableObjectNodeFloatConstant::StaticClass());
+	NodeTypes.Add(UCustomizableObjectNodeFloatParameter::StaticClass());
+	NodeTypes.Add(UCustomizableObjectNodeFloatSwitch::StaticClass());
+	NodeTypes.Add(UCustomizableObjectNodeFloatVariation::StaticClass());
+	NodeTypes.Add(UCustomizableObjectNodeGroupProjectorParameter::StaticClass());
+	NodeTypes.Add(UCustomizableObjectNodeMaterial::StaticClass());
+	NodeTypes.Add(UCustomizableObjectNodeMaterialVariation::StaticClass());
+	NodeTypes.Add(UCustomizableObjectNodeMeshClipDeform::StaticClass());
+	NodeTypes.Add(UCustomizableObjectNodeMeshClipMorph::StaticClass());
+	NodeTypes.Add(UCustomizableObjectNodeMeshClipWithMesh::StaticClass());
+	NodeTypes.Add(UCustomizableObjectNodeMeshGeometryOperation::StaticClass());
+	NodeTypes.Add(UCustomizableObjectNodeMeshMorph::StaticClass());
+	NodeTypes.Add(UCustomizableObjectNodeMeshMorphStackApplication::StaticClass());
+	NodeTypes.Add(UCustomizableObjectNodeMeshMorphStackDefinition::StaticClass());
+	NodeTypes.Add(UCustomizableObjectNodeMeshReshape::StaticClass());
+	NodeTypes.Add(UCustomizableObjectNodeMeshSwitch::StaticClass());
+	NodeTypes.Add(UCustomizableObjectNodeMeshVariation::StaticClass());
+	NodeTypes.Add(UCustomizableObjectNodeMorphMaterial::StaticClass());
+	NodeTypes.Add(UCustomizableObjectNodeObject::StaticClass());
+	NodeTypes.Add(UCustomizableObjectNodeObjectGroup::StaticClass());
+	NodeTypes.Add(UCustomizableObjectNodeProjectorConstant::StaticClass());
+	NodeTypes.Add(UCustomizableObjectNodeProjectorParameter::StaticClass());
+	NodeTypes.Add(UCustomizableObjectNodeRemoveMesh::StaticClass());
+	NodeTypes.Add(UCustomizableObjectNodeRemoveMeshBlocks::StaticClass());
+	NodeTypes.Add(UCustomizableObjectNodeSkeletalMesh::StaticClass());
+	NodeTypes.Add(UCustomizableObjectNodeStaticMesh::StaticClass());
+	NodeTypes.Add(UCustomizableObjectNodeTable::StaticClass());
+	NodeTypes.Add(UCustomizableObjectNodeTexture::StaticClass());
+	NodeTypes.Add(UCustomizableObjectNodeTextureBinarise::StaticClass());
+	NodeTypes.Add(UCustomizableObjectNodeTextureColourMap::StaticClass());
+	NodeTypes.Add(UCustomizableObjectNodeTextureFromChannels::StaticClass());
+	NodeTypes.Add(UCustomizableObjectNodeTextureInterpolate::StaticClass());
+	NodeTypes.Add(UCustomizableObjectNodeTextureInvert::StaticClass());
+	NodeTypes.Add(UCustomizableObjectNodeTextureLayer::StaticClass());
+	NodeTypes.Add(UCustomizableObjectNodeTextureParameter::StaticClass());
+	NodeTypes.Add(UCustomizableObjectNodeTextureProject::StaticClass());
+	NodeTypes.Add(UCustomizableObjectNodeTextureSample::StaticClass());
+	NodeTypes.Add(UCustomizableObjectNodeTextureSaturate::StaticClass());
+	NodeTypes.Add(UCustomizableObjectNodeTextureSwitch::StaticClass());
+	NodeTypes.Add(UCustomizableObjectNodeTextureToChannels::StaticClass());
+	NodeTypes.Add(UCustomizableObjectNodeTextureTransform::StaticClass());
+	NodeTypes.Add(UCustomizableObjectNodeTextureVariation::StaticClass());
+
+	NodeTypes.Sort([](UClass& A, UClass& B) -> bool
+	{
+		const UCustomizableObjectNode* NodeA = CastChecked<UCustomizableObjectNode>(A.GetDefaultObject());
+		const UCustomizableObjectNode* NodeB = CastChecked<UCustomizableObjectNode>(B.GetDefaultObject());
+		
+		return NodeA->GetNodeTitle(ENodeTitleType::ListView).CompareTo(NodeB->GetNodeTitle(ENodeTitleType::ListView)) < 0;
+	});
 }
 
 
@@ -717,7 +778,49 @@ bool UEdGraphSchema_CustomizableObject::ShouldHidePinDefaultValue(UEdGraphPin* P
 }
 
 
-void UEdGraphSchema_CustomizableObject::GetContextMenuActions(class UToolMenu* Menu, class UGraphNodeContextMenuContext* Context) const
+void UEdGraphSchema_CustomizableObject::GetContextMenuActionsReconstructAllChildNodes(UToolMenu* Menu, TWeakObjectPtr<UGraphNodeContextMenuContext> WeakContext) const
+{
+	FToolMenuSection& SubSection = Menu->AddSection("Section");
+	
+	for (UClass* NodeType : NodeTypes)
+	{
+		const UCustomizableObjectNode* Node = Cast<UCustomizableObjectNode>(NodeType->GetDefaultObject());
+
+		auto Call = [](TWeakObjectPtr<UGraphNodeContextMenuContext> WeakContext, const TWeakObjectPtr<UClass>& WeakNodeType)
+		{
+			const UGraphNodeContextMenuContext* Context = WeakContext.Get();
+			if (!Context)
+			{
+				return;
+			}
+
+			UCustomizableObjectNode* Node = Cast<UCustomizableObjectNode>(Context->Node);
+			if (!Node)
+			{
+				return;
+			}
+
+			const UClass* NodeType = WeakNodeType.Get();
+			if (!NodeType)
+			{
+				return;
+			}
+
+			const TSharedRef<ICustomizableObjectEditor> Editor = GetCustomizableObjectEditor(Context->Graph).ToSharedRef();
+			Editor->ReconstructAllChildNodes(*Node, *NodeType);			
+		};
+
+		SubSection.AddMenuEntry(
+			Node->GetFName(),
+			Node->GetNodeTitle(ENodeTitleType::ListView),
+			FText(),
+			FSlateIcon(),
+			FExecuteAction::CreateLambda(Call, WeakContext, MakeWeakObjectPtr(NodeType)));
+	}
+}
+
+
+void UEdGraphSchema_CustomizableObject::GetContextMenuActions(UToolMenu* Menu, UGraphNodeContextMenuContext* Context) const
 {
 	if (Context && Context->Node)
 	{
@@ -733,13 +836,12 @@ void UEdGraphSchema_CustomizableObject::GetContextMenuActions(class UToolMenu* M
 				Section.AddMenuEntry(FGenericCommands::Get().Duplicate);
 				Section.AddMenuEntry(FGraphEditorCommands::Get().ReconstructNodes);
 				Section.AddMenuEntry(FGraphEditorCommands::Get().BreakNodeLinks);
-
-				// In the case of a UCustomizableObjectNodeObjectGroup, add the option to refresh all Customizable Object Material Node nodes of all the children of this node
-				const UCustomizableObjectNodeObjectGroup* TypedNode = Cast<UCustomizableObjectNodeObjectGroup>(Context->Node);
-				if (TypedNode != nullptr)
-				{
-					Section.AddMenuEntry(FCustomizableObjectEditorNodeContextCommands::Get().RefreshMaterialNodesInAllChildren);
-				}
+				
+				Section.AddSubMenu(
+					"ReconstructAllNodes",
+					LOCTEXT("ReconstructChildAllNodes", "Refresh All Child Nodes"),
+					LOCTEXT("ReconstructAllChildNodes_Tooltip", "Refresh all child nodes from the selected ones (inclusive)."),
+					FNewToolMenuDelegate::CreateUObject(this, &UEdGraphSchema_CustomizableObject::GetContextMenuActionsReconstructAllChildNodes, MakeWeakObjectPtr(Context)));
 			}
 
 			struct SCommentUtility
diff --git a/Engine/Plugins/Experimental/Mutable/Source/CustomizableObjectEditor/Private/MuCOE/GraphTraversal.cpp b/Engine/Plugins/Experimental/Mutable/Source/CustomizableObjectEditor/Private/MuCOE/GraphTraversal.cpp
--- a/Engine/Plugins/Experimental/Mutable/Source/CustomizableObjectEditor/Private/MuCOE/GraphTraversal.cpp
+++ b/Engine/Plugins/Experimental/Mutable/Source/CustomizableObjectEditor/Private/MuCOE/GraphTraversal.cpp
@@ -2,6 +2,7 @@
 
 #include "MuCOE/GraphTraversal.h"
 
+#include "AssetRegistry/AssetRegistryModule.h"
 #include "MuCOE/CustomizableObjectPin.h"
 #include "MuCOE/EdGraphSchema_CustomizableObject.h"
 #include "MuCOE/Nodes/CustomizableObjectNodeEnumParameter.h"
@@ -22,6 +23,7 @@
 #include "MuCOE/Nodes/CustomizableObjectNodeStaticMesh.h"
 #include "MuCOE/Nodes/CustomizableObjectNodeTable.h"
 #include "MuCOE/Nodes/CustomizableObjectNodeAnimationPose.h"
+#include "MuCOE/Nodes/CustomizableObjectNodeObjectGroup.h"
 #include "UObject/UObjectIterator.h"
 
 
@@ -501,3 +503,67 @@ const UEdGraphPin* FindMeshBaseSource(const UEdGraphPin& Pin, const bool bOnlyLo
 	
 	return nullptr;
 }
+
+
+void GetNodeGroupObjectNodeMappingImmersive(UCustomizableObject* Object, FAssetRegistryModule& AssetRegistryModule, TSet<UCustomizableObject*>& Visited, TMultiMap<FGuid, UCustomizableObjectNodeObject*>& Mapping)
+{
+	Visited.Add(Object);
+
+	TArray<FName> ArrayReferenceNames;
+	AssetRegistryModule.Get().GetReferencers(*Object->GetOuter()->GetPathName(), ArrayReferenceNames, UE::AssetRegistry::EDependencyCategory::Package, UE::AssetRegistry::EDependencyQuery::Hard);
+
+	FARFilter Filter;
+	for (const FName& ReferenceName : ArrayReferenceNames)
+	{
+		if (!ReferenceName.ToString().StartsWith(TEXT("/TempAutosave")))
+		{
+			Filter.PackageNames.Add(ReferenceName);
+		}
+	}
+
+	Filter.bIncludeOnlyOnDiskAssets = false;
+	
+	TArray<FAssetData> ArrayAssetData;
+	AssetRegistryModule.Get().GetAssets(Filter, ArrayAssetData);
+
+	for (FAssetData& AssetData : ArrayAssetData)
+	{
+		UCustomizableObject* ChildObject = Cast<UCustomizableObject>(AssetData.GetAsset());
+		if (!ChildObject)
+		{
+			continue;			
+		}
+
+		if (ChildObject != Object && !ChildObject->HasAnyFlags(RF_Transient))
+		{
+			bool bMultipleBaseObjectsFound = false;
+			UCustomizableObjectNodeObject* ChildRoot = GetRootNode(ChildObject, bMultipleBaseObjectsFound);
+
+			if (ChildRoot && !bMultipleBaseObjectsFound)
+			{
+				if (ChildRoot->ParentObject == Object)
+				{
+					Mapping.Add(ChildRoot->ParentObjectGroupId, ChildRoot);
+				}
+			}
+		}
+
+		if (!Visited.Contains(ChildObject))
+		{
+			GetNodeGroupObjectNodeMappingImmersive(ChildObject, AssetRegistryModule, Visited, Mapping);
+		}
+	}
+}
+
+
+TMultiMap<FGuid, UCustomizableObjectNodeObject*> GetNodeGroupObjectNodeMapping(UCustomizableObject* Object)
+{
+	FAssetRegistryModule& AssetRegistryModule = FModuleManager::LoadModuleChecked<FAssetRegistryModule>("AssetRegistry");
+
+	TSet<UCustomizableObject*> Visited;
+	TMultiMap<FGuid, UCustomizableObjectNodeObject*> Mapping;
+
+	GetNodeGroupObjectNodeMappingImmersive(Object, AssetRegistryModule, Visited, Mapping);
+	
+	return Mapping;
+}
\ No newline at end of file
diff --git a/Engine/Plugins/Experimental/Mutable/Source/CustomizableObjectEditor/Private/MuCOE/Nodes/CustomizableObjectNodeObject.cpp b/Engine/Plugins/Experimental/Mutable/Source/CustomizableObjectEditor/Private/MuCOE/Nodes/CustomizableObjectNodeObject.cpp
--- a/Engine/Plugins/Experimental/Mutable/Source/CustomizableObjectEditor/Private/MuCOE/Nodes/CustomizableObjectNodeObject.cpp
+++ b/Engine/Plugins/Experimental/Mutable/Source/CustomizableObjectEditor/Private/MuCOE/Nodes/CustomizableObjectNodeObject.cpp
@@ -142,30 +142,30 @@ void UCustomizableObjectNodeObject::AllocateDefaultPins(UCustomizableObjectNodeR
 
 FText UCustomizableObjectNodeObject::GetNodeTitle(ENodeTitleType::Type TitleType) const
 {
-	FFormatNamedArguments Args;	
-	Args.Add(TEXT("ObjectName"), FText::FromString(ObjectName) );
-
-
-	if (TitleType == ENodeTitleType::ListView)
+	if (TitleType == ENodeTitleType::ListView ||
+		ObjectName.IsEmpty())
 	{
-		if (!bIsBase)
-		{						
-			return FText::Format(LOCTEXT("Child_Object_Title_List_Deprecated", "{ObjectName} - Child Object (Deprecated)"), Args);
+		if (bIsBase)
+		{
+			return LOCTEXT("Base_Object", "Base Object");
 		}
 		else
 		{
-			return FText::Format(LOCTEXT("Base_Object_Title_List", "{ObjectName} - Base Object"), Args);
+			return LOCTEXT("Base_Object_Deprecated", "Base Object (Deprecated)");
 		}
 	}
 	else
 	{
-		if (!bIsBase)
+		FFormatNamedArguments Args;	
+		Args.Add(TEXT("ObjectName"), FText::FromString(ObjectName) );
+
+		if (bIsBase)
 		{
-			return FText::Format(LOCTEXT("Child_Object_Title_Deprecated", "{ObjectName}\nChild Object (Deprecated)"), Args);
+			return FText::Format(LOCTEXT("Base_Object_Title", "{ObjectName}\nBase Object"), Args);
 		}
 		else
 		{
-			return FText::Format(LOCTEXT("Base_Object_Title", "{ObjectName}\nBase Object"), Args);
+			return FText::Format(LOCTEXT("Child_Object_Title_Deprecated", "{ObjectName}\nChild Object (Deprecated)"), Args);
 		}
 	}
 }
diff --git a/Engine/Plugins/Experimental/Mutable/Source/CustomizableObjectEditor/Private/MuCOE/Nodes/CustomizableObjectNodeObjectGroup.cpp b/Engine/Plugins/Experimental/Mutable/Source/CustomizableObjectEditor/Private/MuCOE/Nodes/CustomizableObjectNodeObjectGroup.cpp
--- a/Engine/Plugins/Experimental/Mutable/Source/CustomizableObjectEditor/Private/MuCOE/Nodes/CustomizableObjectNodeObjectGroup.cpp
+++ b/Engine/Plugins/Experimental/Mutable/Source/CustomizableObjectEditor/Private/MuCOE/Nodes/CustomizableObjectNodeObjectGroup.cpp
@@ -85,10 +85,18 @@ void UCustomizableObjectNodeObjectGroup::PostEditChangeProperty(FPropertyChanged
 
 FText UCustomizableObjectNodeObjectGroup::GetNodeTitle(ENodeTitleType::Type TitleType) const
 {
-	FFormatNamedArguments Args;
-	Args.Add(TEXT("GroupName"), FText::FromString(GroupName));
+	if (TitleType == ENodeTitleType::ListView ||
+		GroupName.IsEmpty())
+	{
+		return LOCTEXT("Object_Group", "Object Group");;
+	}
+	else
+	{
+		FFormatNamedArguments Args;
+		Args.Add(TEXT("GroupName"), FText::FromString(GroupName));
 
-	return FText::Format(LOCTEXT("Group_Object_Title", "{GroupName}\nGroup"), Args);			
+		return FText::Format(LOCTEXT("Group_Object_Title", "{GroupName}\nObject Group"), Args);
+	}
 }
 
 
diff --git a/Engine/Plugins/Experimental/Mutable/Source/CustomizableObjectEditor/Public/MuCOE/GraphTraversal.h b/Engine/Plugins/Experimental/Mutable/Source/CustomizableObjectEditor/Public/MuCOE/GraphTraversal.h
--- a/Engine/Plugins/Experimental/Mutable/Source/CustomizableObjectEditor/Public/MuCOE/GraphTraversal.h
+++ b/Engine/Plugins/Experimental/Mutable/Source/CustomizableObjectEditor/Public/MuCOE/GraphTraversal.h
@@ -41,7 +41,7 @@ UEdGraphPin* FollowOutputPin(const UEdGraphPin& Pin, bool* CycleDetected = nullp
 UCustomizableObjectNodeObject* GetRootNode(UCustomizableObject* Object, bool& bOutMultipleBaseObjectsFound);
 
 /** Return in ArrayNodeObject the roots nodes in each Customizable Object graph until the whole root node is found (i.e. the one with parent = nullptr)
-		return false if a cycle is found between Customizable Objects */
+		* return false if a cycle is found between Customizable Objects */
 bool GetParentsUntilRoot(UCustomizableObject* Object, TArray<UCustomizableObjectNodeObject*>& ArrayNodeObject, TArray<UCustomizableObject*>& ArrayCustomizableObject);
 
 /** Returns true if the Candidate is parent of the current Customizable Object */
@@ -55,3 +55,7 @@ UCustomizableObjectNodeObject* GetFullGraphRootNodeObject(UCustomizableObjectNod
 
 /** Given an output pin, return the output pin where the mesh is located. */
 const UEdGraphPin* FindMeshBaseSource(const UEdGraphPin& Pin, const bool bOnlyLookForStaticMesh);
+
+/** Return the mapping of Group Object Nodes to Child Object Nodes of the given hierarchy.
+ * @param Object Child or root Object to start the search from. */
+TMultiMap<FGuid, UCustomizableObjectNodeObject*> GetNodeGroupObjectNodeMapping(UCustomizableObject* Object);
diff --git a/Engine/Plugins/Experimental/Mutable/Source/CustomizableObjectEditor/Public/MuCOE/ICustomizableObjectEditor.h b/Engine/Plugins/Experimental/Mutable/Source/CustomizableObjectEditor/Public/MuCOE/ICustomizableObjectEditor.h
--- a/Engine/Plugins/Experimental/Mutable/Source/CustomizableObjectEditor/Public/MuCOE/ICustomizableObjectEditor.h
+++ b/Engine/Plugins/Experimental/Mutable/Source/CustomizableObjectEditor/Public/MuCOE/ICustomizableObjectEditor.h
@@ -8,6 +8,9 @@ class UEdGraphNode;
 
 class UCustomizableObject;
 class UCustomizableObjectInstance;
+class UClass;
+class UCustomizableObjectNode;
+
 
 // Public interface to Customizable Object Editor
 class ICustomizableObjectEditor : public ICustomizableObjectInstanceEditor
@@ -20,18 +23,23 @@ public:
 	virtual bool CanPasteNodes() const {return false;}
 
 	// Paste nodes at a specific location
-	virtual void PasteNodesHere(const FVector2D& Location) {};
+	virtual void PasteNodesHere(const FVector2D& Location) {}
 
-	virtual class UEdGraphNode* CreateCommentBox(const FVector2D& NodePos) { return nullptr; };
+	virtual UEdGraphNode* CreateCommentBox(const FVector2D& NodePos) { return nullptr; }
 
-	virtual void UpdateGraphNodeProperties() {};
+	virtual void UpdateGraphNodeProperties() {}
 
-	virtual void UpdateObjectProperties() {};
+	virtual void UpdateObjectProperties() {}
 
 	// Make sure the graph editor selects a specific node
-	virtual void SelectNode(const class UCustomizableObjectNode* Node) {}
+	virtual void SelectNode(const UCustomizableObjectNode* Node) {}
+
+	virtual void ReconstructNode(UEdGraphNode* Node) {}
 
-	virtual void ReconstructNode(UEdGraphNode* Node) {};
+	/** Reconstructs all child the nodes that match the given type.
+	 * @param StartNode Root node to start the graph traversal. This one also will be reconstructed.
+	 * @param NodeType Node types to reconstruct. */
+	virtual void ReconstructAllChildNodes(UCustomizableObjectNode& StartNode, const UClass& NodeType) {}
 };
 
 
