[Mutable] Fixed MIP tasks holding time critical Update tasks.

- Created a new FMutableTaskGraph holding all functoins and members related to the Mutable tasks system.
- Added new Low Priority Mutable Tasks.
- Allow canceling low priority tasks.
- Moved ClearMutableTaskIfDone from AdvanceCurrentOperation to System's Tick.

#rnx

diff --git a/Engine/Plugins/Experimental/Mutable/Source/CustomizableObject/Private/MuCO/CustomizableObjectMipDataProvider.cpp b/Engine/Plugins/Experimental/Mutable/Source/CustomizableObject/Private/MuCO/CustomizableObjectMipDataProvider.cpp
--- a/Engine/Plugins/Experimental/Mutable/Source/CustomizableObject/Private/MuCO/CustomizableObjectMipDataProvider.cpp
+++ b/Engine/Plugins/Experimental/Mutable/Source/CustomizableObject/Private/MuCO/CustomizableObjectMipDataProvider.cpp
@@ -240,14 +240,12 @@ int32 FMutableTextureMipDataProvider::GetMips(const FTextureUpdateContext& Conte
 			OperationData->RescheduleCallback = SyncOptions.RescheduleCallback;
 
 			TSharedPtr<FMutableImageOperationData> LocalOperationData = OperationData;
-			UpdateImageMutableTaskEvent = CustomizableObjectSystem->AddMutableThreadTask(
+			MutableTaskId = CustomizableObjectSystem->MutableTaskGraph.AddMutableThreadTaskLowPriority(
 				TEXT("Mutable_MipUpdate"),
 				[LocalOperationData]()
 				{
 					impl::Task_Mutable_UpdateImage(LocalOperationData);
-				},
-				UE::Tasks::ETaskPriority::BackgroundHigh
-					);
+				});
 		}
 
 		AdvanceTo(ETickState::PollMips, ETickThread::Async);
@@ -323,6 +321,11 @@ void FMutableTextureMipDataProvider::Cancel(const FTextureUpdateSyncOptions& Syn
 	bRequestAborted = true;
 
 	CancelCounterSafely();
+
+	if (UCustomizableObjectSystem* System = UCustomizableObjectSystem::GetInstance())
+	{
+		System->GetPrivate()->MutableTaskGraph.CancelMutableThreadTaskLowPriority(MutableTaskId);	
+	}
 }
 
 
@@ -338,6 +341,11 @@ void FMutableTextureMipDataProvider::AbortPollMips()
 	bRequestAborted = true;
 
 	CancelCounterSafely();
+
+	if (UCustomizableObjectSystem* System = UCustomizableObjectSystem::GetInstance())
+	{
+		System->GetPrivate()->MutableTaskGraph.CancelMutableThreadTaskLowPriority(MutableTaskId);	
+	}
 }
 
 
diff --git a/Engine/Plugins/Experimental/Mutable/Source/CustomizableObject/Private/MuCO/CustomizableObjectSystem.cpp b/Engine/Plugins/Experimental/Mutable/Source/CustomizableObject/Private/MuCO/CustomizableObjectSystem.cpp
--- a/Engine/Plugins/Experimental/Mutable/Source/CustomizableObject/Private/MuCO/CustomizableObjectSystem.cpp
+++ b/Engine/Plugins/Experimental/Mutable/Source/CustomizableObject/Private/MuCO/CustomizableObjectSystem.cpp
@@ -396,7 +396,7 @@ void UCustomizableObjectSystem::BeginDestroy()
 		Private->PendingTasks.Empty();
 
 		// Complete pending taskgraph tasks
-		Private->WaitForMutableTasks();
+		Private->MutableTaskGraph.WaitForMutableTasks();
 
 		// Clear the ongoing operation
 		Private->CurrentMutableOperation = nullptr;
@@ -737,6 +737,8 @@ FAutoConsoleVariableRef CVarMaxTextureSizeToGenerate(
 
 void FinishUpdateGlobal(UCustomizableObjectInstance* Instance, EUpdateResult UpdateResult, FInstanceUpdateDelegate* UpdateCallback, const FDescriptorRuntimeHash InUpdatedHash)
 {
+	check(IsInGameThread())
+	
 	// Callbacks. Must be done at the end.
 	if (Instance)
 	{
@@ -764,6 +766,8 @@ void FinishUpdateGlobal(UCustomizableObjectInstance* Instance, EUpdateResult Upd
 		
 		UpdateCallback->ExecuteIfBound(Context);
 	}
+
+	UCustomizableObjectSystem::GetInstance()->GetPrivate()->MutableTaskGraph.AllowLaunchingMutableTaskLowPriority(true, false);
 }
 
 
@@ -1742,6 +1746,7 @@ namespace impl
 		OperationData->MutableRuntimeCycles += EndCycles - StartCycles;
 #endif
 
+		UCustomizableObjectSystem::GetInstance()->GetPrivate()->MutableTaskGraph.AllowLaunchingMutableTaskLowPriority(true, true);
 	}
 
 
@@ -1955,10 +1960,9 @@ namespace impl
 		//-------------------------------------------------------------		
 		{
 			mu::SystemPtr MutableSystem = CustomizableObjectSystemPrivateData->MutableSystem;
-			CustomizableObjectSystemPrivateData->AddMutableThreadTask(
+			CustomizableObjectSystemPrivateData->MutableTaskGraph.AddMutableThreadTask(
 				TEXT("Task_Mutable_ReleaseInstance"),
-				[OperationData, MutableSystem]() {Task_Mutable_ReleaseInstance(OperationData, MutableSystem); },
-				UE::Tasks::ETaskPriority::BackgroundNormal);
+				[OperationData, MutableSystem]() {Task_Mutable_ReleaseInstance(OperationData, MutableSystem); });
 		}
 
 
@@ -1969,10 +1973,12 @@ namespace impl
 			TSharedPtr<FMutableReleasePlatformOperationData> ReleaseOperationData = MakeShared<FMutableReleasePlatformOperationData>();
 			check(ReleaseOperationData);
 			ReleaseOperationData->ImageToPlatformDataMap = MoveTemp(OperationData->ImageToPlatformDataMap);
-			CustomizableObjectSystemPrivateData->AddAnyThreadTask(
+			CustomizableObjectSystemPrivateData->MutableTaskGraph.AddAnyThreadTask(
 				TEXT("Mutable_ReleasePlatformData"),
-				[ReleaseOperationData]() { Task_Game_ReleasePlatformData(ReleaseOperationData); },
-				UE::Tasks::ETaskPriority::BackgroundNormal
+				[ReleaseOperationData]()
+				{
+					Task_Game_ReleasePlatformData(ReleaseOperationData);
+				}
 			);
 
 
@@ -2113,13 +2119,12 @@ namespace impl
 			TSharedPtr<mu::Model, ESPMode::ThreadSafe> Model = CustomizableObject->GetPrivate()->GetModel();
 			int32 State = ObjectInstance->GetState();
 
-			Mutable_GetImagesTask = SystemPrivateData->AddMutableThreadTask(
+			Mutable_GetImagesTask = SystemPrivateData->MutableTaskGraph.AddMutableThreadTask(
 					TEXT("Task_Mutable_GetImages"),
 					[OperationData, Parameters, Model, State]()
 					{
 						impl::Task_Mutable_Update_GetImages(OperationData, Model, Parameters, State);
-					},
-					UE::Tasks::ETaskPriority::BackgroundHigh);
+					});
 		}
 
 
@@ -2178,13 +2183,12 @@ namespace impl
 			TSharedPtr<FMutableOperation> CurrentMutableOperation = SystemPrivateData->CurrentMutableOperation;
 			check(CurrentMutableOperation);
 
-			Mutable_GetMeshTask = SystemPrivateData->AddMutableThreadTask(
+			Mutable_GetMeshTask = SystemPrivateData->MutableTaskGraph.AddMutableThreadTask(
 				TEXT("Task_Mutable_ReleaseInstanceID"),
 				[CurrentOperationData, MutableSystem]()
 				{
 					impl::Task_Mutable_ReleaseInstanceID(CurrentOperationData, MutableSystem);
-				},
-				UE::Tasks::ETaskPriority::BackgroundHigh);
+				});
 		}
 	}
 
@@ -2206,6 +2210,9 @@ namespace impl
 	{
 		MUTABLE_CPUPROFILER_SCOPE(Task_Game_StartUpdate)
 
+		UCustomizableObjectSystem::GetInstance()->GetPrivate()->MutableTaskGraph.AllowLaunchingMutableTaskLowPriority(false, false);
+
+		
 		check(Operation);
 		check(Operation->Type == FMutableOperation::EOperationType::Update);
 
@@ -2385,13 +2392,12 @@ namespace impl
 			TSharedPtr<mu::Model, ESPMode::ThreadSafe> Model = CustomizableObject->GetPrivate()->GetModel();
 			int32 State = CustomizableObjectInstancePtr->GetState();
 
-			Mutable_GetMeshTask = SystemPrivateData->AddMutableThreadTask(
+			Mutable_GetMeshTask = SystemPrivateData->MutableTaskGraph.AddMutableThreadTask(
 				TEXT("Task_Mutable_Update_GetMesh"),
 				[CurrentOperationData, bBuildParameterRelevancy, Parameters, Model, State]()
 				{
 					impl::Task_Mutable_Update_GetMesh(CurrentOperationData, Model, Parameters, bBuildParameterRelevancy, State);
-				},
-				UE::Tasks::ETaskPriority::BackgroundHigh);
+				});
 		}
 
 
@@ -2425,9 +2431,6 @@ void UCustomizableObjectSystem::AdvanceCurrentOperation()
 
 	check(Private != nullptr);
 
-	// See if we can clear the last reference to the mutable-thread task
-	Private->ClearMutableTaskIfDone();
-
 	// See if we have a game-thread task to process
 	FMutableTask* PendingTask = Private->PendingTasks.Peek();
 	if (PendingTask)
@@ -2564,10 +2567,7 @@ bool UCustomizableObjectSystem::Tick(float DeltaTime)
 		return true; // Assets are still being loaded, so subobjects won't be found, compiled objects incomplete and thus updates wrong
 	}
 #endif
-
-
-	MUTABLE_CPUPROFILER_SCOPE(TickCustomizableObjectSystem)
-
+	
 	// Get a new operation if we aren't working on one
 	if (!Private->CurrentMutableOperation)
 	{
@@ -2631,9 +2631,7 @@ bool UCustomizableObjectSystem::Tick(float DeltaTime)
 	{
 		AdvanceCurrentOperation();
 	}
-
-
-
+	
 	// TODO: T2927
 	if (LogBenchmarkUtil::isLoggingActive())
 	{
@@ -2654,6 +2652,8 @@ bool UCustomizableObjectSystem::Tick(float DeltaTime)
 	TickRecompileCustomizableObjects();
 #endif
 
+	Private->MutableTaskGraph.Tick();
+
 	return true;
 }
 
diff --git a/Engine/Plugins/Experimental/Mutable/Source/CustomizableObject/Private/MuCO/CustomizableObjectSystemPrivate.h b/Engine/Plugins/Experimental/Mutable/Source/CustomizableObject/Private/MuCO/CustomizableObjectSystemPrivate.h
--- a/Engine/Plugins/Experimental/Mutable/Source/CustomizableObject/Private/MuCO/CustomizableObjectSystemPrivate.h
+++ b/Engine/Plugins/Experimental/Mutable/Source/CustomizableObject/Private/MuCO/CustomizableObjectSystemPrivate.h
@@ -15,6 +15,7 @@
 #include "MuR/Image.h"
 #include "UObject/GCObject.h"
 #include "WorldCollision.h"
+#include "MuCO/FMutableTaskGraph.h"
 
 // This define could come from MuR/System.h
 #ifdef MUTABLE_USE_NEW_TASKGRAPH
@@ -569,109 +570,10 @@ public:
 	/** */
 	inline void AddGameThreadTask(const FMutableTask& Task)
 	{
+		check(IsInGameThread())
 		PendingTasks.Enqueue(Task);
 	}
 
-	/** */
-#ifdef MUTABLE_USE_NEW_TASKGRAPH
-	template<typename TaskBodyType>
-	inline UE::Tasks::FTask AddMutableThreadTask(const TCHAR* DebugName, TaskBodyType&& TaskBody, UE::Tasks::ETaskPriority Priority)
-	{
-		// This could be called from the game thread or from a worker thread for the texture streaming system
-		//check(IsInGameThread());
-		FScopeLock Lock(&MutableTaskLock);
-
-		if (LastMutableTask.IsValid())
-		{
-			LastMutableTask = UE::Tasks::Launch(DebugName, MoveTemp(TaskBody), LastMutableTask, Priority);
-		}
-		else
-		{
-			LastMutableTask = UE::Tasks::Launch(DebugName, MoveTemp(TaskBody), Priority);
-		}
-		return LastMutableTask;
-	}
-#else
-	template<typename TaskBodyType>
-	inline FGraphEventRef AddMutableThreadTask(const TCHAR* DebugName, TaskBodyType&& TaskBody, UE::Tasks::ETaskPriority Priority)
-	{
-		// This could be called from the game thread or from a worker thread for the texture streaming system
-		//check(IsInGameThread());
-		FScopeLock Lock(&MutableTaskLock);
-
-		// Chain to pending tasks of the mutable thread if any. This ensures exclusive access to mutable.
-		FGraphEventArray Prerequisites;
-		if (LastMutableTask)
-		{
-			Prerequisites.Add(LastMutableTask);
-		}
-
-		LastMutableTask = FFunctionGraphTask::CreateAndDispatchWhenReady(
-			MoveTemp(TaskBody),
-			TStatId{},
-			&Prerequisites, 
-			ENamedThreads::AnyBackgroundHiPriTask);
-		LastMutableTask->SetDebugName(DebugName);
-
-		return LastMutableTask;
-	}
-#endif
-
-	/** This should only be used when shutting down. */
-	inline void WaitForMutableTasks()
-	{
-		FScopeLock Lock(&MutableTaskLock);
-
-		if (LastMutableTask.IsValid())
-		{
-#ifdef MUTABLE_USE_NEW_TASKGRAPH
-			LastMutableTask.Wait();
-#else
-			LastMutableTask->Wait();
-#endif
-			LastMutableTask = {};
-		}
-	}
-
-	inline void ClearMutableTaskIfDone()
-	{
-		FScopeLock Lock(&MutableTaskLock);
-
-#ifdef MUTABLE_USE_NEW_TASKGRAPH
-		if (LastMutableTask.IsValid() && LastMutableTask.IsCompleted())
-		{
-			LastMutableTask = {};
-		}
-#else
-		if (LastMutableTask.IsValid() && LastMutableTask->IsComplete())
-		{
-			LastMutableTask = nullptr;
-		}
-#endif
-	}
-
-
-	template<typename TaskBodyType>
-	inline void AddAnyThreadTask(const TCHAR* DebugName, TaskBodyType&& TaskBody, UE::Tasks::ETaskPriority Priority)
-	{
-		check(IsInGameThread());
-
-#ifdef MUTABLE_USE_NEW_TASKGRAPH
-		UE::Tasks::Launch(DebugName, MoveTemp(TaskBody), Priority);
-#else
-		FGraphEventRef Task = FFunctionGraphTask::CreateAndDispatchWhenReady(
-			MoveTemp(TaskBody),
-			TStatId{},
-			nullptr,
-			ENamedThreads::AnyThread);
-		Task->SetDebugName(TEXT("Mutable_Anythread"));
-#endif
-	}
-
-
-	//
-	//TSharedPtr<FMutableOperationData> CurrentOperationData;
-
 
 	/** FSerializableObject interface */
 	virtual void AddReferencedObjects(FReferenceCollector& Collector) override;
@@ -841,17 +743,7 @@ public:
 
 	bool IsMutableAnimInfoDebuggingEnabled() const;
 
-private:
-
-	/** Access to this must be protected with this. */
-	mutable FCriticalSection MutableTaskLock;
-
-#ifdef MUTABLE_USE_NEW_TASKGRAPH
-	UE::Tasks::FTask LastMutableTask;
-#else
-	FGraphEventRef LastMutableTask = nullptr;
-#endif
-
-
+	/** Mutable TaskGraph system (Mutable Thread). */
+	FMutableTaskGraph MutableTaskGraph;
 };
 
diff --git a/Engine/Plugins/Experimental/Mutable/Source/CustomizableObject/Private/MuCO/FMutableTaskGraph.cpp b/Engine/Plugins/Experimental/Mutable/Source/CustomizableObject/Private/MuCO/FMutableTaskGraph.cpp
new file mode 100644
--- /dev/null
+++ b/Engine/Plugins/Experimental/Mutable/Source/CustomizableObject/Private/MuCO/FMutableTaskGraph.cpp
@@ -0,0 +1,247 @@
+// Copyright Epic Games, Inc. All Rights Reserved.
+
+#include "MuCO/FMutableTaskGraph.h"
+
+#include "MuCO/CustomizableObject.h"
+#include "MuR/MutableTrace.h"
+
+
+#ifdef MUTABLE_USE_NEW_TASKGRAPH
+const LowLevelTasks::ETaskPriority TASKGRAPH_PRIORITY = LowLevelTasks::ETaskPriority::BackgroundHigh;
+#else
+const ENamedThreads::Type TASKGRAPH_PRIORITY = ENamedThreads::AnyBackgroundHiPriTask;
+#endif
+
+
+static TAutoConsoleVariable<float> CVarMutableTaskLowPriorityMaxWaitTime(
+	TEXT("mutable.MutableTaskLowPriorityMaxWaitTime"),
+	3.0f,
+	TEXT("Max time a Mutable Task with Low priority can wait. Once this time has passed it will be launched unconditionally."),
+	ECVF_Scalability);
+
+static TAutoConsoleVariable<bool> CVarEnableMutableTaskLowPriority(
+	TEXT("mutable.EnableMutableTaskLowPriority"),
+	true,
+	TEXT("Enable or disable Mutable Tasks with Low priority. If disabled, all task will have the same priority. "),
+	ECVF_Scalability);
+
+
+FMutableTaskGraph::TaskType FMutableTaskGraph::AddMutableThreadTask(const TCHAR* DebugName, TUniqueFunction<void()>&& TaskBody)
+{
+	FScopeLock Lock(&MutableTaskLock);
+
+#ifdef MUTABLE_USE_NEW_TASKGRAPH
+	if (LastMutableTask.IsValid())
+	{
+		LastMutableTask = UE::Tasks::Launch(DebugName, MoveTemp(TaskBody), LastMutableTask, TASKGRAPH_PRIORITY);
+	}
+	else
+	{
+		LastMutableTask = UE::Tasks::Launch(DebugName, MoveTemp(TaskBody), TASKGRAPH_PRIORITY);
+	}
+#else
+	FGraphEventArray Prerequisites;
+	if (LastMutableTask)
+	{
+		Prerequisites.Add(LastMutableTask);
+	}
+
+	LastMutableTask = FFunctionGraphTask::CreateAndDispatchWhenReady(
+		MoveTemp(TaskBody),
+		TStatId{},
+		&Prerequisites, 
+		TASKGRAPH_PRIORITY);
+	
+	LastMutableTask->SetDebugName(DebugName);
+#endif
+	
+	return LastMutableTask;
+}
+
+
+uint32 FMutableTaskGraph::AddMutableThreadTaskLowPriority(const TCHAR* DebugName, TFunction<void()>&& TaskBody)
+{
+	FScopeLock Lock(&MutableTaskLock);
+
+	if (CVarEnableMutableTaskLowPriority.GetValueOnAnyThread())
+	{
+		const uint32 Id = ++TaskIdGenerator;
+		const FTask Task { Id, DebugName, MoveTemp(TaskBody)};
+		QueueMutableTasksLowPriority.Add(Task);		
+
+		TryLaunchMutableTaskLowPriority(false);
+
+		return Id;
+	}
+	else
+	{
+		AddMutableThreadTask(DebugName, TaskBody);
+		return INVALID_ID;	
+	}
+}
+
+
+void FMutableTaskGraph::CancelMutableThreadTaskLowPriority(uint32 Id)
+{
+	FScopeLock Lock(&MutableTaskLock);
+
+	for (int32 Index = 0; Index < QueueMutableTasksLowPriority.Num(); ++Index)
+	{
+		if (QueueMutableTasksLowPriority[Index].Id == Id)
+		{
+			QueueMutableTasksLowPriority.RemoveAt(Index);
+			break;
+		}
+	}
+}
+
+
+void FMutableTaskGraph::AddAnyThreadTask(const TCHAR* DebugName, TUniqueFunction<void()>&& TaskBody) const
+{
+#ifdef MUTABLE_USE_NEW_TASKGRAPH
+	UE::Tasks::Launch(DebugName, MoveTemp(TaskBody));
+#else
+	FGraphEventRef Task = FFunctionGraphTask::CreateAndDispatchWhenReady(
+		MoveTemp(TaskBody),
+		TStatId{},
+		nullptr,
+		TASKGRAPH_PRIORITY);
+
+	Task->SetDebugName(DebugName);
+#endif
+}
+	
+
+void FMutableTaskGraph::WaitForMutableTasks()
+{
+	FScopeLock Lock(&MutableTaskLock);
+
+	if (LastMutableTask.IsValid())
+	{
+#ifdef MUTABLE_USE_NEW_TASKGRAPH
+		LastMutableTask.Wait();
+#else
+		LastMutableTask->Wait();
+#endif
+		LastMutableTask = {};
+	}
+}
+
+
+void FMutableTaskGraph::TryLaunchMutableTaskLowPriority(bool bFromMutableTask)
+{
+	MUTABLE_CPUPROFILER_SCOPE(TryLaunchMutableTaskLowPriority)
+
+	if (QueueMutableTasksLowPriority.IsEmpty())
+	{
+		return;		
+	}
+	
+	if (!IsTaskCompleted(LastMutableTaskLowPriority)) // At any time only a single Low Priority task can be launched.
+	{
+		return;
+	}
+
+	float TimeLimit;
+	float TimeElapsed;
+	
+	bool bTimeLimit;
+	
+	{
+		FScopeLock Lock(&MutableTaskLock);
+
+		if (QueueMutableTasksLowPriority.IsEmpty()) // Also checked inside the lock since we will be writing it
+		{
+			return;		
+		}
+		
+		if (!IsTaskCompleted(LastMutableTaskLowPriority)) // Also check inside the lock since we will be writing it
+		{
+			return;
+		}
+
+		FTask& NextTask = QueueMutableTasksLowPriority[0];
+
+		TimeLimit = CVarMutableTaskLowPriorityMaxWaitTime.GetValueOnAnyThread();
+		TimeElapsed = FPlatformTime::Seconds() - NextTask.CreationTime;
+
+		bTimeLimit = TimeElapsed >= TimeLimit;
+		if (!bTimeLimit && (
+				!bAllowLaunchMutableTaskLowPriority ||
+				(bFromMutableTask && IsTaskCompleted(LastMutableTask))))
+		{			
+			return;
+		}
+
+		LastMutableTaskLowPriority = AddMutableThreadTask(*NextTask.DebugName, [this, Task = MoveTemp(NextTask)]() // Moves the task, not the pointer.
+		{
+			MUTABLE_CPUPROFILER_SCOPE(LowPriorityTaskBody)
+			Task.Body();
+
+			{
+				FScopeLock Lock(&MutableTaskLock);
+				LastMutableTaskLowPriority = {};
+				
+				TryLaunchMutableTaskLowPriority(true);
+			}
+		});
+
+		QueueMutableTasksLowPriority.RemoveAt(0);
+	}
+
+	if (bTimeLimit)
+	{
+		UE_LOG(LogMutable, Verbose, TEXT("Low Priority Mutable Task launched due to time limit (%f)! Waited for: %f"), TimeLimit, TimeElapsed)
+	}
+}
+
+
+bool FMutableTaskGraph::IsTaskCompleted(const TaskType& Task) const
+{
+#ifdef MUTABLE_USE_NEW_TASKGRAPH
+	return Task.IsCompleted();
+#else
+	return !Task.IsValid() || Task->IsComplete();
+#endif
+}
+
+
+void FMutableTaskGraph::AllowLaunchingMutableTaskLowPriority(bool bAllow, bool bFromMutableTask)
+{
+	FScopeLock Lock(&MutableTaskLock);
+
+	bAllowLaunchMutableTaskLowPriority = bAllow;
+
+	if (bAllowLaunchMutableTaskLowPriority)
+	{
+		TryLaunchMutableTaskLowPriority(bFromMutableTask);
+	}
+}
+
+
+void FMutableTaskGraph::Tick()
+{
+	check(IsInGameThread())
+	
+	TryLaunchMutableTaskLowPriority(false);
+
+	// See if we can clear the last reference to the mutable-thread task
+	if (IsTaskCompleted(LastMutableTask))
+	{
+		FScopeLock Lock(&MutableTaskLock);
+		if (IsTaskCompleted(LastMutableTask))
+		{
+			LastMutableTask = {};
+		}
+	}
+
+	if (IsTaskCompleted(LastMutableTaskLowPriority))
+	{
+		FScopeLock Lock(&MutableTaskLock);
+		if (IsTaskCompleted(LastMutableTaskLowPriority))
+		{
+			LastMutableTaskLowPriority = {};
+		}
+	}
+}
+
diff --git a/Engine/Plugins/Experimental/Mutable/Source/CustomizableObject/Public/MuCO/CustomizableObjectMipDataProvider.h b/Engine/Plugins/Experimental/Mutable/Source/CustomizableObject/Public/MuCO/CustomizableObjectMipDataProvider.h
--- a/Engine/Plugins/Experimental/Mutable/Source/CustomizableObject/Public/MuCO/CustomizableObjectMipDataProvider.h
+++ b/Engine/Plugins/Experimental/Mutable/Source/CustomizableObject/Public/MuCO/CustomizableObjectMipDataProvider.h
@@ -2,15 +2,13 @@
 
 #pragma once
 
+#include "FMutableTaskGraph.h"
 #include "Async/TaskGraphInterfaces.h"
 #include "Engine/TextureMipDataProviderFactory.h"
 #include "MuCO/CustomizableObjectSystem.h"
 #include "MuR/Image.h"
 #include "Streaming/TextureMipDataProvider.h"
 
- // For the MUTABLE_USE_NEW_TASKGRAPH define
-#include "MuR/System.h"
-
 #include "CustomizableObjectMipDataProvider.generated.h"
 
 enum EPixelFormat : uint8;
@@ -100,11 +98,8 @@ public:
 	bool bRequestAborted = false;
 
 	TSharedPtr<FMutableImageOperationData> OperationData;
-#ifdef MUTABLE_USE_NEW_TASKGRAPH
-	UE::Tasks::FTask UpdateImageMutableTaskEvent;
-#else
-	FGraphEventRef UpdateImageMutableTaskEvent;
-#endif
+
+	uint32 MutableTaskId = FMutableTaskGraph::INVALID_ID;
 };
 
 
diff --git a/Engine/Plugins/Experimental/Mutable/Source/CustomizableObject/Public/MuCO/FMutableTaskGraph.h b/Engine/Plugins/Experimental/Mutable/Source/CustomizableObject/Public/MuCO/FMutableTaskGraph.h
new file mode 100644
--- /dev/null
+++ b/Engine/Plugins/Experimental/Mutable/Source/CustomizableObject/Public/MuCO/FMutableTaskGraph.h
@@ -0,0 +1,103 @@
+// Copyright Epic Games, Inc. All Rights Reserved.
+
+#pragma once
+
+#include "MuR/System.h" // For the MUTABLE_USE_NEW_TASKGRAPH define
+
+#include "Tasks/Task.h"
+#include "Async/TaskGraphInterfaces.h"
+
+
+/** Mutable Tasks System.
+ *
+ * Allows launching tasks in different threads:
+ * - Mutable Thread
+ * - Game Thread (code currently in UCustomizableObjectSystem)
+ * - Any Thread
+ *
+ * Concurrency between tasks in the Mutable Thread is forbidden by chaining all tasks through their prerequisites.
+ *
+ *
+ * LOW PRIORITY MUTABLE TASKS:
+ *
+ * - Only one Low Priority task can be launched at the same time.
+ *   This is because once a Task Graph task is launched, it can not be canceled.
+ *   To allow canceling them, the system holds them until it can ensure that its execution will be imminent (no other tasks running).
+ *
+ * - A Low Priority task will not be launched if one of the follow conditions is true (in order):
+ *     1. There is a task Low Priority task running.
+ *     2. Flag bLaunchMutableTaskLowPriority is false (unless time limit).
+ *     3. There is a Normal Priority task running (unless time limit).
+ */
+class FMutableTaskGraph
+{
+	struct FTask
+	{
+		uint32 Id;
+		FString DebugName;
+		TFunction<void()> Body;
+		float CreationTime = FPlatformTime::Seconds();
+	};
+	
+public:
+	using TaskType = 
+#ifdef MUTABLE_USE_NEW_TASKGRAPH
+		UE::Tasks::FTask;
+#else
+		FGraphEventRef;
+#endif
+
+	/** Create and launch a task on the Mutable Thread with Normal priority. */
+	TaskType AddMutableThreadTask(const TCHAR* DebugName, TUniqueFunction<void()>&& TaskBody);
+
+	/** Create and launch a task on the Mutable Thread with Low priority. */
+	uint32 AddMutableThreadTaskLowPriority(const TCHAR* DebugName, TFunction<void()>&& TaskBody);
+
+	/** Cancel, if not already launched, a Mutable Thread with Low priority. */
+	void CancelMutableThreadTaskLowPriority(uint32 Id);
+	
+	/** Create and launch a task on Any Thread. */
+	void AddAnyThreadTask(const TCHAR* DebugName, TUniqueFunction<void()>&& TaskBody) const;
+
+	/** Wait for all Mutable Thread tasks. */
+	void WaitForMutableTasks();
+
+	/** Allow or disallow launching Mutable Tasks with Low priority.
+	 @param bFromMutableTask true if called from a Mutable Task. */
+	void AllowLaunchingMutableTaskLowPriority(bool bAllow, bool bFromMutableTask);
+	
+	void Tick();
+	
+private:
+	/** A Mutable Task Low Priority will only be launched if:
+	 * - No other low priority task is running.
+	 * - Is allowed to launch low priority tasks.
+	 *
+	 * @param bFromMutableTask true if called from a Mutable Task. */
+	void TryLaunchMutableTaskLowPriority(bool bFromMutableTask);
+
+	/** Return true if the task is completed (or is no longer valid). */
+	bool IsTaskCompleted(const TaskType& Task) const;
+
+	mutable FCriticalSection MutableTaskLock;
+
+	/** Allow or disallow launching low priority tasks. */
+	bool bAllowLaunchMutableTaskLowPriority = true;
+
+	/** Queue of low priority tasks. FIFO. */
+	TArray<FTask> QueueMutableTasksLowPriority;
+
+public:
+	inline static constexpr uint32 INVALID_ID = 0;
+
+private:
+	/** Incremental task ID generator. */
+	uint32 TaskIdGenerator = INVALID_ID;
+
+	/** Last Mutable Task Low Priority launched to the TaskGraph system. */
+	TaskType LastMutableTaskLowPriority = {};
+
+	/** Last Mutable Task launched to the TaskGraph system. Low and normal priority. */
+	TaskType LastMutableTask = {};
+};
+
