[Mutable] Fix Non-Unity build issues.

#rnx

diff --git a/Engine/Plugins/Experimental/Mutable/Source/CustomizableObjectEditor/Public/MuCOE/CustomizableObjectEditor_Deprecated.h b/Engine/Plugins/Experimental/Mutable/Source/CustomizableObjectEditor/Public/MuCOE/CustomizableObjectEditor_Deprecated.h
--- a/Engine/Plugins/Experimental/Mutable/Source/CustomizableObjectEditor/Public/MuCOE/CustomizableObjectEditor_Deprecated.h
+++ b/Engine/Plugins/Experimental/Mutable/Source/CustomizableObjectEditor/Public/MuCOE/CustomizableObjectEditor_Deprecated.h
@@ -2,10 +2,14 @@
 
 #pragma once
 
+#include "Containers/Array.h"
+#include "UObject/SoftObjectPtr.h"
 
 #include "CustomizableObjectEditor_Deprecated.generated.h"
 
 class UTexture2D;
+class UEdGraphPin_Deprecated;
+struct FEdGraphPinReference;
 
 
 // Place to hide all deprecated data structures. They are still needed for deserialization backwards compatibility.
