[Mutable] Fix crash when copying Node Skeletal Mesh.

#rnx

diff --git a/Engine/Plugins/Experimental/Mutable/Source/CustomizableObject/Public/MuCO/CustomizableObjectCustomVersion.h b/Engine/Plugins/Experimental/Mutable/Source/CustomizableObject/Public/MuCO/CustomizableObjectCustomVersion.h
--- a/Engine/Plugins/Experimental/Mutable/Source/CustomizableObject/Public/MuCO/CustomizableObjectCustomVersion.h
+++ b/Engine/Plugins/Experimental/Mutable/Source/CustomizableObject/Public/MuCO/CustomizableObjectCustomVersion.h
@@ -76,6 +76,8 @@ struct CUSTOMIZABLEOBJECT_API FCustomizableObjectCustomVersion
 
 		AutomaticNodeSkeletalMesh,
 
+		AutomaticNodeSkeletalMeshPinDataOuter,
+		
 		// -----<new versions can be added above this line>-------------------------------------------------
 		VersionPlusOne,
 		LatestVersion = VersionPlusOne - 1
diff --git a/Engine/Plugins/Experimental/Mutable/Source/CustomizableObjectEditor/Private/MuCOE/Nodes/CustomizableObjectNodeSkeletalMesh.cpp b/Engine/Plugins/Experimental/Mutable/Source/CustomizableObjectEditor/Private/MuCOE/Nodes/CustomizableObjectNodeSkeletalMesh.cpp
--- a/Engine/Plugins/Experimental/Mutable/Source/CustomizableObjectEditor/Private/MuCOE/Nodes/CustomizableObjectNodeSkeletalMesh.cpp
+++ b/Engine/Plugins/Experimental/Mutable/Source/CustomizableObjectEditor/Private/MuCOE/Nodes/CustomizableObjectNodeSkeletalMesh.cpp
@@ -76,7 +76,7 @@ void UCustomizableObjectNodeSkeletalMesh::AllocateDefaultPins(UCustomizableObjec
 				{
 					FString MeshName = FString::Printf(TEXT("LOD %i - Section %i - Mesh"), LODIndex, SectionIndex);
 
-					UCustomizableObjectNodeSkeletalMeshPinDataMesh* PinData = NewObject<UCustomizableObjectNodeSkeletalMeshPinDataMesh>();
+					UCustomizableObjectNodeSkeletalMeshPinDataMesh* PinData = NewObject<UCustomizableObjectNodeSkeletalMeshPinDataMesh>(this);
 					PinData->Init(LODIndex, SectionIndex);
 					
 					UEdGraphPin* Pin = CustomCreatePin(EGPD_Output, Schema->PC_Mesh, FName(*MeshName), PinData);
@@ -93,7 +93,7 @@ void UCustomizableObjectNodeSkeletalMesh::AllocateDefaultPins(UCustomizableObjec
 						{
 							FString PinName = FString::Printf(TEXT("LOD %i - Section %i - UV %i"), LODIndex, SectionIndex, UVIndex);
 							
-							UCustomizableObjectNodeSkeletalMeshPinDataLayout* PinData = NewObject<UCustomizableObjectNodeSkeletalMeshPinDataLayout>();
+							UCustomizableObjectNodeSkeletalMeshPinDataLayout* PinData = NewObject<UCustomizableObjectNodeSkeletalMeshPinDataLayout>(this);
 							PinData->Init(LODIndex, SectionIndex, UVIndex);
 							
 							UEdGraphPin* Pin = CustomCreatePin(EGPD_Input, Schema->PC_Layout, FName(*PinName), PinData);
@@ -121,7 +121,7 @@ void UCustomizableObjectNodeSkeletalMesh::AllocateDefaultPins(UCustomizableObjec
 						FString ImageNameStr = *ImageInfo.Name.ToString();
 						FString PinName = FString::Printf(TEXT("LOD %i - Section %i - Texture Parameter %s"), LODIndex, SectionIndex, *ImageNameStr);
 						
-						UCustomizableObjectNodeSkeletalMeshPinDataImage* PinData = NewObject<UCustomizableObjectNodeSkeletalMeshPinDataImage>();
+						UCustomizableObjectNodeSkeletalMeshPinDataImage* PinData = NewObject<UCustomizableObjectNodeSkeletalMeshPinDataImage>(this);
 						PinData->Init(LODIndex, SectionIndex, ImageId);
 						
 						UEdGraphPin* Pin = CustomCreatePin(EGPD_Output, Schema->PC_Image, FName(*PinName), PinData);
@@ -142,6 +142,12 @@ UCustomizableObjectNodeRemapPins* UCustomizableObjectNodeSkeletalMesh::CreateRem
 }
 
 
+void UCustomizableObjectNodeSkeletalMesh::RemapPinsData(const TMap<UEdGraphPin*, UEdGraphPin*>& PinsToRemap)
+{
+	// Do not copy pin data from old pins
+}
+
+
 FText UCustomizableObjectNodeSkeletalMesh::GetNodeTitle(ENodeTitleType::Type TitleType) const
 {
 	if (SkeletalMesh)
@@ -543,7 +549,7 @@ void UCustomizableObjectNodeSkeletalMesh::BackwardsCompatibleFixup()
 				const FCustomizableObjectNodeSkeletalMeshMaterial& Section = LOD.Materials[SectionIndex];
 
 				{
-					UCustomizableObjectNodeSkeletalMeshPinDataMesh* PinData = NewObject<UCustomizableObjectNodeSkeletalMeshPinDataMesh>();
+					UCustomizableObjectNodeSkeletalMeshPinDataMesh* PinData = NewObject<UCustomizableObjectNodeSkeletalMeshPinDataMesh>(this);
 					PinData->Init(LODIndex, SectionIndex);
 				
 					AddPinData(*Section.MeshPinRef.Get(), *PinData);
@@ -572,7 +578,7 @@ void UCustomizableObjectNodeSkeletalMesh::BackwardsCompatibleFixup()
 								}
 							}
 							
-							UCustomizableObjectNodeSkeletalMeshPinDataImage* PinData = NewObject<UCustomizableObjectNodeSkeletalMeshPinDataImage>();
+							UCustomizableObjectNodeSkeletalMeshPinDataImage* PinData = NewObject<UCustomizableObjectNodeSkeletalMeshPinDataImage>(this);
 							PinData->Init(LODIndex, SectionIndex, TextureParameterId);
 						
 							AddPinData(*ImagePin, *PinData);
@@ -582,7 +588,7 @@ void UCustomizableObjectNodeSkeletalMesh::BackwardsCompatibleFixup()
 
 				for (int32 LayoutIndex = 0; LayoutIndex < Section.LayoutPinsRef.Num(); ++LayoutIndex)
 				{
-					UCustomizableObjectNodeSkeletalMeshPinDataLayout* PinData = NewObject<UCustomizableObjectNodeSkeletalMeshPinDataLayout>();
+					UCustomizableObjectNodeSkeletalMeshPinDataLayout* PinData = NewObject<UCustomizableObjectNodeSkeletalMeshPinDataLayout>(this);
 					PinData->Init(LODIndex, SectionIndex, LayoutIndex);
 				
 					AddPinData(*Section.LayoutPinsRef[LayoutIndex].Get(), *PinData);
@@ -592,6 +598,11 @@ void UCustomizableObjectNodeSkeletalMesh::BackwardsCompatibleFixup()
 			
 		ReconstructNode();
 	}
+
+	if (CustomizableObjectCustomVersion < FCustomizableObjectCustomVersion::AutomaticNodeSkeletalMeshPinDataOuter)
+	{
+		ReconstructNode();
+	}
 }
 
 
