[Mutable] Missing Remap pins function definition.

#rnx

diff --git a/Engine/Plugins/Experimental/Mutable/Source/CustomizableObjectEditor/Classes/MuCOE/Nodes/CustomizableObjectNodeSkeletalMesh.h b/Engine/Plugins/Experimental/Mutable/Source/CustomizableObjectEditor/Classes/MuCOE/Nodes/CustomizableObjectNodeSkeletalMesh.h
--- a/Engine/Plugins/Experimental/Mutable/Source/CustomizableObjectEditor/Classes/MuCOE/Nodes/CustomizableObjectNodeSkeletalMesh.h
+++ b/Engine/Plugins/Experimental/Mutable/Source/CustomizableObjectEditor/Classes/MuCOE/Nodes/CustomizableObjectNodeSkeletalMesh.h
@@ -208,7 +208,9 @@ public:
 	virtual void BackwardsCompatibleFixup() override;
 	virtual void AllocateDefaultPins(UCustomizableObjectNodeRemapPins* RemapPins) override;
 	virtual UCustomizableObjectNodeRemapPins* CreateRemapPinsDefault() const override;
+	virtual void RemapPinsData(const TMap<UEdGraphPin*, UEdGraphPin*>& PinsToRemap) override;
 
+public:
 	// UCustomizableObjectNodeMesh interface
 	virtual UTexture2D* FindTextureForPin(const UEdGraphPin* Pin) const override;
 	virtual void GetUVChannelForPin(const UEdGraphPin* Pin, TArray<FVector2f>& OutSegments, int32 UVIndex) const override;
diff --git a/Engine/Plugins/Experimental/Mutable/Source/CustomizableObjectEditor/Private/MuCOE/Nodes/CustomizableObjectNode.cpp b/Engine/Plugins/Experimental/Mutable/Source/CustomizableObjectEditor/Private/MuCOE/Nodes/CustomizableObjectNode.cpp
--- a/Engine/Plugins/Experimental/Mutable/Source/CustomizableObjectEditor/Private/MuCOE/Nodes/CustomizableObjectNode.cpp
+++ b/Engine/Plugins/Experimental/Mutable/Source/CustomizableObjectEditor/Private/MuCOE/Nodes/CustomizableObjectNode.cpp
@@ -58,7 +58,10 @@ UEdGraphPin* UCustomizableObjectNode::CustomCreatePin(EEdGraphPinDirection Direc
 UEdGraphPin* UCustomizableObjectNode::CustomCreatePin(EEdGraphPinDirection Direction, const FName& Type, const FName& Name, UCustomizableObjectNodePinData* PinData)
 {
 	UEdGraphPin* Pin = CreatePin(Direction, Type, Name);
-	PinsDataId.Add(Pin->PinId, PinData);
+	if (Pin && PinData)
+	{
+		AddPinData(*Pin, *PinData);		
+	}
 	
 	return Pin;
 }
@@ -455,6 +458,7 @@ void UCustomizableObjectNode::RemapPinsData(const TMap<UEdGraphPin*, UEdGraphPin
 
 void UCustomizableObjectNode::AddPinData(const UEdGraphPin& Pin, UCustomizableObjectNodePinData& PinData)
 {
+	check(PinData.GetOuter())
 	PinsDataId.Add(Pin.PinId, &PinData);
 }
 
