Fix missing <string> include

#rnx

diff --git a/Engine/Plugins/Experimental/Mutable/Source/MutableTools/Private/MuT/Platform.h b/Engine/Plugins/Experimental/Mutable/Source/MutableTools/Private/MuT/Platform.h
--- a/Engine/Plugins/Experimental/Mutable/Source/MutableTools/Private/MuT/Platform.h
+++ b/Engine/Plugins/Experimental/Mutable/Source/MutableTools/Private/MuT/Platform.h
@@ -7,6 +7,7 @@
 #include <vector>
 #include <map>
 #include <set>
+#include <string>
 #include <unordered_map>
 #include <sys/stat.h>
 
