Restore CustomizableObjectPopulationModule existing in non-editor builds
#rnx

diff --git a/Engine/Plugins/Experimental/Mutable/Source/CustomizableObjectPopulation/Private/MuCOP/CustomizableObjectPopulationModule.cpp b/Engine/Plugins/Experimental/Mutable/Source/CustomizableObjectPopulation/Private/MuCOP/CustomizableObjectPopulationModule.cpp
--- a/Engine/Plugins/Experimental/Mutable/Source/CustomizableObjectPopulation/Private/MuCOP/CustomizableObjectPopulationModule.cpp
+++ b/Engine/Plugins/Experimental/Mutable/Source/CustomizableObjectPopulation/Private/MuCOP/CustomizableObjectPopulationModule.cpp
@@ -1,5 +1,7 @@
 // Copyright Epic Games, Inc. All Rights Reserved.
 
+#include "Modules/ModuleManager.h"
+
 #if WITH_EDITOR
 #include "AssetRegistry/ARFilter.h"
 #include "AssetRegistry/AssetData.h"
@@ -10,12 +12,17 @@
 #include "MuCOE/ICustomizableObjectPopulationModule.h"
 #include "MuCOP/CustomizableObjectPopulation.h"
 #include "MuCOP/CustomizableObjectPopulationGenerator.h"
-
+#endif
 
 /**
  * MovieSceneCore module implementation (private)
  */
-class FCustomizableObjectPopulationModule : public ICustomizableObjectPopulationModule
+class FCustomizableObjectPopulationModule : 
+#if WITH_EDITOR
+	public ICustomizableObjectPopulationModule
+#else
+	public IModuleInterface
+#endif
 {
 
 public:
@@ -24,10 +31,10 @@ public:
 	void StartupModule() override{}
 	void ShutdownModule() override{}
 
+#if WITH_EDITOR
 	// ICustomizableObjectPopulationModule interface
 	FString GetPluginVersion() const override;
 
-#if WITH_EDITOR
 	// Recompiles all the populations referenced in the Customizabled Object
 	virtual void RecompilePopulations(UCustomizableObject* Object);
 #endif
@@ -38,6 +45,7 @@ IMPLEMENT_MODULE( FCustomizableObjectPopulationModule, CustomizableObjectPopulat
 
 //-------------------------------------------------------------------------------------------------
 
+#if WITH_EDITOR
 FString FCustomizableObjectPopulationModule::GetPluginVersion() const
 {
 	FString Version = "x.x";
