[mutable] Fix for check fail due to invocation of mu::UpdateLLMStats() in bad context
- From now on the manual update of the LLM data recollection system will only be performed if Mutable is running as part of a program and not a game.

diff --git a/Engine/Plugins/Experimental/Mutable/Source/MutableRuntime/Private/MuR/SystemPrivate.h b/Engine/Plugins/Experimental/Mutable/Source/MutableRuntime/Private/MuR/SystemPrivate.h
--- a/Engine/Plugins/Experimental/Mutable/Source/MutableRuntime/Private/MuR/SystemPrivate.h
+++ b/Engine/Plugins/Experimental/Mutable/Source/MutableRuntime/Private/MuR/SystemPrivate.h
@@ -23,7 +23,7 @@ namespace mu
 	{
 		// This code will only be compiled (and ran) if the global definition to enable LLM tracking is set to 1 for the host program
 		// Ex : 			GlobalDefinitions.Add("LLM_ENABLED_IN_CONFIG=1");
-#if ENABLE_LOW_LEVEL_MEM_TRACKER && 0
+#if ENABLE_LOW_LEVEL_MEM_TRACKER && IS_PROGRAM
 		FLowLevelMemTracker& MemTracker = FLowLevelMemTracker::Get();
 		if (MemTracker.IsEnabled())
 		{
