[Mutable] Fix Layout or block index error.

- Fix random compilation crashes.

#rnx

diff --git a/Engine/Plugins/Experimental/Mutable/Source/CustomizableObject/Public/MuCO/CustomizableObjectCustomVersion.h b/Engine/Plugins/Experimental/Mutable/Source/CustomizableObject/Public/MuCO/CustomizableObjectCustomVersion.h
--- a/Engine/Plugins/Experimental/Mutable/Source/CustomizableObject/Public/MuCO/CustomizableObjectCustomVersion.h
+++ b/Engine/Plugins/Experimental/Mutable/Source/CustomizableObject/Public/MuCO/CustomizableObjectCustomVersion.h
@@ -77,6 +77,8 @@ struct CUSTOMIZABLEOBJECT_API FCustomizableObjectCustomVersion
 		AutomaticNodeSkeletalMesh,
 
 		AutomaticNodeSkeletalMeshPinDataOuter,
+
+		AutomaticNodeSkeletalMeshPinDataUProperty,
 		
 		// -----<new versions can be added above this line>-------------------------------------------------
 		VersionPlusOne,
diff --git a/Engine/Plugins/Experimental/Mutable/Source/CustomizableObjectEditor/Classes/MuCOE/Nodes/CustomizableObjectNodeSkeletalMesh.h b/Engine/Plugins/Experimental/Mutable/Source/CustomizableObjectEditor/Classes/MuCOE/Nodes/CustomizableObjectNodeSkeletalMesh.h
--- a/Engine/Plugins/Experimental/Mutable/Source/CustomizableObjectEditor/Classes/MuCOE/Nodes/CustomizableObjectNodeSkeletalMesh.h
+++ b/Engine/Plugins/Experimental/Mutable/Source/CustomizableObjectEditor/Classes/MuCOE/Nodes/CustomizableObjectNodeSkeletalMesh.h
@@ -107,7 +107,10 @@ protected:
 	virtual bool Equals(const UCustomizableObjectNodePinData& Other) const override;
 
 private:
+	UPROPERTY()
 	int32 LODIndex = -1;
+
+	UPROPERTY()
 	int32 SectionIndex = -1;
 };
 
@@ -135,6 +138,7 @@ protected:
 	virtual bool Equals(const UCustomizableObjectNodePinData& Other) const override;
 
 private:
+	UPROPERTY()
 	FGuid TextureParameterId;
 };
 
@@ -154,6 +158,7 @@ protected:
 	virtual bool Equals(const UCustomizableObjectNodePinData& Other) const override;
 
 private:
+	UPROPERTY()
 	int32 UVIndex = -1;
 };
 
diff --git a/Engine/Plugins/Experimental/Mutable/Source/CustomizableObjectEditor/Private/MuCOE/GenerateMutableSource/GenerateMutableSourceMesh.cpp b/Engine/Plugins/Experimental/Mutable/Source/CustomizableObjectEditor/Private/MuCOE/GenerateMutableSource/GenerateMutableSourceMesh.cpp
--- a/Engine/Plugins/Experimental/Mutable/Source/CustomizableObjectEditor/Private/MuCOE/GenerateMutableSource/GenerateMutableSourceMesh.cpp
+++ b/Engine/Plugins/Experimental/Mutable/Source/CustomizableObjectEditor/Private/MuCOE/GenerateMutableSource/GenerateMutableSourceMesh.cpp
@@ -2917,11 +2917,13 @@ mu::NodeMeshPtr GenerateMutableSourceMesh(const UEdGraphPin * Pin,
 			// Layouts
 			// If we didn't find a layout, but we are generating LODs and this LOD is automatic, reuse the first valid LOD layout
 			bool LayoutFound = false;
-			
+
+			FSkeletalMeshModel* ImportedModel = TypedNodeSkel->SkeletalMesh->GetImportedModel();
+
 			int32 LayoutLOD = LOD;
 			while (!LayoutFound && LayoutLOD >= 0)
 			{
-				const int32 NumLayouts = TypedNodeSkel->SkeletalMesh->GetImportedModel()->LODModels[LayoutLOD].Sections.Num();
+				const int32 NumLayouts = ImportedModel->LODModels[LayoutLOD].NumTexCoords;
 				
 				MeshNode->SetLayoutCount(NumLayouts);
 				for (int32 LayoutIndex = 0; LayoutIndex < NumLayouts; ++LayoutIndex)
diff --git a/Engine/Plugins/Experimental/Mutable/Source/CustomizableObjectEditor/Private/MuCOE/Nodes/CustomizableObjectNodeSkeletalMesh.cpp b/Engine/Plugins/Experimental/Mutable/Source/CustomizableObjectEditor/Private/MuCOE/Nodes/CustomizableObjectNodeSkeletalMesh.cpp
--- a/Engine/Plugins/Experimental/Mutable/Source/CustomizableObjectEditor/Private/MuCOE/Nodes/CustomizableObjectNodeSkeletalMesh.cpp
+++ b/Engine/Plugins/Experimental/Mutable/Source/CustomizableObjectEditor/Private/MuCOE/Nodes/CustomizableObjectNodeSkeletalMesh.cpp
@@ -184,7 +184,7 @@ UTexture2D* UCustomizableObjectNodeSkeletalMesh::FindTextureForPin(const UEdGrap
 	}
 
 	const FSkeletalMeshModel* ImportedModel = SkeletalMesh->GetImportedModel();
-	if (ImportedModel)
+	if (!ImportedModel)
 	{
 		return nullptr;
 	}	
@@ -601,7 +601,12 @@ void UCustomizableObjectNodeSkeletalMesh::BackwardsCompatibleFixup()
 
 	if (CustomizableObjectCustomVersion < FCustomizableObjectCustomVersion::AutomaticNodeSkeletalMeshPinDataOuter)
 	{
-		ReconstructNode();
+		ReconstructNode(); // Pins did not have Pin Data. Reconstruct them.
+	}
+
+	if (CustomizableObjectCustomVersion < FCustomizableObjectCustomVersion::AutomaticNodeSkeletalMeshPinDataUProperty)
+	{
+		ReconstructNode(CreateRemapPinsByName()); // Correct pins but incorrect Pin Data. Reconstruct and remap pins only by name, no Pin Data.
 	}
 }
 
