[Mutable] Fix potential invalid access when getting the name of a material.

#rnx

diff --git a/Engine/Plugins/Experimental/Mutable/Source/CustomizableObjectEditor/Private/MuCOE/Nodes/CustomizableObjectNodeSkeletalMesh.cpp b/Engine/Plugins/Experimental/Mutable/Source/CustomizableObjectEditor/Private/MuCOE/Nodes/CustomizableObjectNodeSkeletalMesh.cpp
--- a/Engine/Plugins/Experimental/Mutable/Source/CustomizableObjectEditor/Private/MuCOE/Nodes/CustomizableObjectNodeSkeletalMesh.cpp
+++ b/Engine/Plugins/Experimental/Mutable/Source/CustomizableObjectEditor/Private/MuCOE/Nodes/CustomizableObjectNodeSkeletalMesh.cpp
@@ -634,7 +634,10 @@ bool UCustomizableObjectNodeSkeletalMesh::CheckIsValidLayout(const UEdGraphPin*
 	int32 SectionIndex;
 	GetPinSection(*ConnectedPin, LODIndex, SectionIndex, LayoutIndex);
 
-	MaterialName = GetMaterialInterfaceFor(LODIndex, SectionIndex)->GetName();
+	if (const UMaterialInterface* MaterialInterface = GetMaterialInterfaceFor(LODIndex, SectionIndex))
+	{
+		MaterialName = MaterialInterface->GetName();
+	}
 	
 	if (LayoutIndex == 0)
 	{
