[Mutable] Fix crashes when generating LODs on meshes created by Mutable.
- Fix broken skinning after regenerating LODs from the imported model data.
- Remove unnecessary regeneration of the imported model.
- Unify bake helpers.

#rnx

diff --git a/Engine/Plugins/Experimental/Mutable/Source/CustomizableObject/Private/MuCO/UnrealBakeHelpers.cpp b/Engine/Plugins/Experimental/Mutable/Source/CustomizableObject/Private/MuCO/UnrealBakeHelpers.cpp
--- a/Engine/Plugins/Experimental/Mutable/Source/CustomizableObject/Private/MuCO/UnrealBakeHelpers.cpp
+++ b/Engine/Plugins/Experimental/Mutable/Source/CustomizableObject/Private/MuCO/UnrealBakeHelpers.cpp
@@ -2,9 +2,21 @@
 
 #include "MuCO/UnrealBakeHelpers.h"
 
+#include "Engine/SkeletalMesh.h"
+#include "Engine/Texture2D.h"
+
+#if WITH_EDITORONLY_DATA
+#include "MuCO/CustomizableObject.h"
+#include "MuCO/CustomizableObjectMipDataProvider.h"
+
+#include "AssetRegistry/AssetRegistryModule.h"
+#include "AssetToolsModule.h"
+#include "IAssetTools.h"
+#include "Serialization/ArchiveReplaceObjectRef.h"
+
+#include "TextureResource.h"
 #include "ClothLODData.h"
 #include "ClothingAsset.h"
-#include "Engine/SkeletalMesh.h"
 #include "Engine/SkeletalMeshLODSettings.h"
 #include "Engine/SkinnedAssetCommon.h"
 #include "Rendering/SkeletalMeshLODModel.h"
@@ -12,6 +24,612 @@
 #include "Rendering/SkeletalMeshRenderData.h"
 
 
+//-------------------------------------------------------------------------------------------------
+//-------------------------------------------------------------------------------------------------
+//-------------------------------------------------------------------------------------------------
+inline uint32_t PackRGBA(uint8_t r, uint8_t g, uint8_t b, uint8_t a)
+{
+	return r | (g << 8) | (b << 16) | (a << 24);
+}
+
+
+//-------------------------------------------------------------------------------------------------
+inline float Int8ToFloat_SNORM(const uint8_t input)
+{
+	return (float)((int8_t)input) / 127.0f;
+}
+
+
+//-------------------------------------------------------------------------------------------------
+inline float Int8ToFloat_UNORM(const uint8_t input)
+{
+	return (float)input / 255.0f;
+}
+
+
+//-------------------------------------------------------------------------------------------------
+inline void Decompress16x3bitIndices(const uint8_t* packed, uint8_t* unpacked)
+{
+	uint32_t tmp, block, i;
+
+	for (block = 0; block < 2; ++block) {
+		tmp = 0;
+
+		// Read three bytes
+		for (i = 0; i < 3; ++i) {
+			tmp |= ((uint32_t)packed[i]) << (i * 8);
+		}
+
+		// Unpack 8x3 bit from last 3 byte block
+		for (i = 0; i < 8; ++i) {
+			unpacked[i] = (tmp >> (i * 3)) & 0x7;
+		}
+
+		packed += 3;
+		unpacked += 8;
+	}
+}
+
+
+//-------------------------------------------------------------------------------------------------
+inline void DecompressBlockBC1Internal(const uint8_t* block,
+	unsigned char* output, uint32_t outputStride, const uint8_t* alphaValues)
+{
+	uint32_t temp, code;
+
+	uint16 color0, color1;
+	uint8_t r0, g0, b0, r1, g1, b1;
+
+	int i, j;
+
+	color0 = *(const uint16*)(block);
+	color1 = *(const uint16*)(block + 2);
+
+	temp = (color0 >> 11) * 255 + 16;
+	r0 = (uint8_t)((temp / 32 + temp) / 32);
+	temp = ((color0 & 0x07E0) >> 5) * 255 + 32;
+	g0 = (uint8_t)((temp / 64 + temp) / 64);
+	temp = (color0 & 0x001F) * 255 + 16;
+	b0 = (uint8_t)((temp / 32 + temp) / 32);
+
+	temp = (color1 >> 11) * 255 + 16;
+	r1 = (uint8_t)((temp / 32 + temp) / 32);
+	temp = ((color1 & 0x07E0) >> 5) * 255 + 32;
+	g1 = (uint8_t)((temp / 64 + temp) / 64);
+	temp = (color1 & 0x001F) * 255 + 16;
+	b1 = (uint8_t)((temp / 32 + temp) / 32);
+
+	code = *(const uint32_t*)(block + 4);
+
+	if (color0 > color1) {
+		for (j = 0; j < 4; ++j) {
+			for (i = 0; i < 4; ++i) {
+				uint32_t finalColor, positionCode;
+				uint8_t alpha;
+
+				alpha = alphaValues[j * 4 + i];
+
+				finalColor = 0;
+				positionCode = (code >> 2 * (4 * j + i)) & 0x03;
+
+				switch (positionCode) {
+				case 0:
+					finalColor = PackRGBA(r0, g0, b0, alpha);
+					break;
+				case 1:
+					finalColor = PackRGBA(r1, g1, b1, alpha);
+					break;
+				case 2:
+					finalColor = PackRGBA((2 * r0 + r1) / 3, (2 * g0 + g1) / 3, (2 * b0 + b1) / 3, alpha);
+					break;
+				case 3:
+					finalColor = PackRGBA((r0 + 2 * r1) / 3, (g0 + 2 * g1) / 3, (b0 + 2 * b1) / 3, alpha);
+					break;
+				}
+
+				*(uint32_t*)(output + j * outputStride + i * sizeof(uint32_t)) = finalColor;
+			}
+		}
+	}
+	else {
+		for (j = 0; j < 4; ++j) {
+			for (i = 0; i < 4; ++i) {
+				uint32_t finalColor, positionCode;
+				uint8_t alpha;
+
+				alpha = alphaValues[j * 4 + i];
+
+				finalColor = 0;
+				positionCode = (code >> 2 * (4 * j + i)) & 0x03;
+
+				switch (positionCode) {
+				case 0:
+					finalColor = PackRGBA(r0, g0, b0, alpha);
+					break;
+				case 1:
+					finalColor = PackRGBA(r1, g1, b1, alpha);
+					break;
+				case 2:
+					finalColor = PackRGBA((r0 + r1) / 2, (g0 + g1) / 2, (b0 + b1) / 2, alpha);
+					break;
+				case 3:
+					finalColor = PackRGBA(0, 0, 0, alpha);
+					break;
+				}
+
+				*(uint32_t*)(output + j * outputStride + i * sizeof(uint32_t)) = finalColor;
+			}
+		}
+	}
+}
+
+
+inline void DecompressBlockBC1(uint32_t x, uint32_t y, uint32_t stride, const uint8_t* blockStorage, unsigned char* image)
+{
+	static const uint8_t const_alpha[] = {
+		255, 255, 255, 255,
+		255, 255, 255, 255,
+		255, 255, 255, 255,
+		255, 255, 255, 255
+	};
+
+	DecompressBlockBC1Internal(blockStorage, image + x * sizeof(uint32_t) + (y * stride), stride, const_alpha);
+}
+
+
+//-------------------------------------------------------------------------------------------------
+inline void DecompressBlockBC3(uint32_t x, uint32_t y, uint32_t stride, const uint8_t* blockStorage, unsigned char* image)
+{
+	uint8_t alpha0, alpha1;
+	uint8_t alphaIndices[16];
+
+	uint16 color0, color1;
+	uint8_t r0, g0, b0, r1, g1, b1;
+
+	int i, j;
+
+	uint32_t temp, code;
+
+	alpha0 = *(blockStorage);
+	alpha1 = *(blockStorage + 1);
+
+	Decompress16x3bitIndices(blockStorage + 2, alphaIndices);
+
+	color0 = *(const uint16*)(blockStorage + 8);
+	color1 = *(const uint16*)(blockStorage + 10);
+
+	temp = (color0 >> 11) * 255 + 16;
+	r0 = (uint8_t)((temp / 32 + temp) / 32);
+	temp = ((color0 & 0x07E0) >> 5) * 255 + 32;
+	g0 = (uint8_t)((temp / 64 + temp) / 64);
+	temp = (color0 & 0x001F) * 255 + 16;
+	b0 = (uint8_t)((temp / 32 + temp) / 32);
+
+	temp = (color1 >> 11) * 255 + 16;
+	r1 = (uint8_t)((temp / 32 + temp) / 32);
+	temp = ((color1 & 0x07E0) >> 5) * 255 + 32;
+	g1 = (uint8_t)((temp / 64 + temp) / 64);
+	temp = (color1 & 0x001F) * 255 + 16;
+	b1 = (uint8_t)((temp / 32 + temp) / 32);
+
+	code = *(const uint32_t*)(blockStorage + 12);
+
+	for (j = 0; j < 4; j++) {
+		for (i = 0; i < 4; i++) {
+			uint8_t finalAlpha;
+			int alphaCode;
+			uint8_t colorCode;
+			uint32_t finalColor;
+
+			alphaCode = alphaIndices[4 * j + i];
+
+			if (alphaCode == 0) {
+				finalAlpha = alpha0;
+			}
+			else if (alphaCode == 1) {
+				finalAlpha = alpha1;
+			}
+			else {
+				if (alpha0 > alpha1) {
+					finalAlpha = (uint8_t)(((8 - alphaCode) * alpha0 + (alphaCode - 1) * alpha1) / 7);
+				}
+				else {
+					if (alphaCode == 6) {
+						finalAlpha = 0;
+					}
+					else if (alphaCode == 7) {
+						finalAlpha = 255;
+					}
+					else {
+						finalAlpha = (uint8_t)(((6 - alphaCode) * alpha0 + (alphaCode - 1) * alpha1) / 5);
+					}
+				}
+			}
+
+			colorCode = (code >> 2 * (4 * j + i)) & 0x03;
+			finalColor = 0;
+
+			switch (colorCode) {
+			case 0:
+				finalColor = PackRGBA(r0, g0, b0, finalAlpha);
+				break;
+			case 1:
+				finalColor = PackRGBA(r1, g1, b1, finalAlpha);
+				break;
+			case 2:
+				finalColor = PackRGBA((2 * r0 + r1) / 3, (2 * g0 + g1) / 3, (2 * b0 + b1) / 3, finalAlpha);
+				break;
+			case 3:
+				finalColor = PackRGBA((r0 + 2 * r1) / 3, (g0 + 2 * g1) / 3, (b0 + 2 * b1) / 3, finalAlpha);
+				break;
+			}
+
+
+			*(uint32_t*)(image + sizeof(uint32_t) * (i + x) + (stride * (y + j))) = finalColor;
+		}
+	}
+}
+
+
+//-------------------------------------------------------------------------------------------------
+inline void DecompressBlockBC2(uint32_t x, uint32_t y, uint32_t stride, const uint8_t* blockStorage, unsigned char* image)
+{
+	int i;
+
+	uint8_t alphaValues[16] = { 0 };
+
+	for (i = 0; i < 4; ++i) {
+		const uint16* alphaData = (const uint16*)(blockStorage);
+
+		alphaValues[i * 4 + 0] = (((*alphaData) >> 0) & 0xF) * 17;
+		alphaValues[i * 4 + 1] = (((*alphaData) >> 4) & 0xF) * 17;
+		alphaValues[i * 4 + 2] = (((*alphaData) >> 8) & 0xF) * 17;
+		alphaValues[i * 4 + 3] = (((*alphaData) >> 12) & 0xF) * 17;
+
+		blockStorage += 2;
+	}
+
+	DecompressBlockBC1Internal(blockStorage, image + x * sizeof(uint32_t) + (y * stride), stride, alphaValues);
+}
+
+
+//-------------------------------------------------------------------------------------------------
+inline void DecompressBlockBC4Internal(const uint8_t* block, uint8_t* output, uint32_t outputStride, const float* colorTable)
+{
+	uint8_t indices[16];
+	int x, y;
+
+	Decompress16x3bitIndices(block + 2, indices);
+
+	for (y = 0; y < 4; ++y)
+	{
+		for (x = 0; x < 4; ++x)
+		{
+			uint8_t v = (uint8_t)FMath::Clamp(colorTable[indices[y * 4 + x]] * 255.0f, 0.0f, 255.0f);
+			*(output + x) = v;
+		}
+
+		output += outputStride;
+	}
+}
+
+
+enum class BC4Mode : uint8_t { BC4_UNORM, BC4_SNORM };
+
+
+//-------------------------------------------------------------------------------------------------
+inline void DecompressBlockBC4(uint32_t x, uint32_t y, uint32_t stride, BC4Mode mode, const uint8_t* blockStorage, uint8_t* image)
+{
+	float colorTable[8];
+	float r0, r1;
+
+	if (mode == BC4Mode::BC4_UNORM)
+	{
+		r0 = Int8ToFloat_UNORM(blockStorage[0]);
+		r1 = Int8ToFloat_UNORM(blockStorage[1]);
+
+		colorTable[0] = r0;
+		colorTable[1] = r1;
+
+		if (r0 > r1)
+		{
+			// 6 interpolated color values
+			colorTable[2] = (6 * r0 + 1 * r1) / 7.0f; // bit code 010
+			colorTable[3] = (5 * r0 + 2 * r1) / 7.0f; // bit code 011
+			colorTable[4] = (4 * r0 + 3 * r1) / 7.0f; // bit code 100
+			colorTable[5] = (3 * r0 + 4 * r1) / 7.0f; // bit code 101
+			colorTable[6] = (2 * r0 + 5 * r1) / 7.0f; // bit code 110
+			colorTable[7] = (1 * r0 + 6 * r1) / 7.0f; // bit code 111
+		}
+		else
+		{
+			// 4 interpolated color values
+			colorTable[2] = (4 * r0 + 1 * r1) / 5.0f; // bit code 010
+			colorTable[3] = (3 * r0 + 2 * r1) / 5.0f; // bit code 011
+			colorTable[4] = (2 * r0 + 3 * r1) / 5.0f; // bit code 100
+			colorTable[5] = (1 * r0 + 4 * r1) / 5.0f; // bit code 101
+			colorTable[6] = 0.0f;               // bit code 110
+			colorTable[7] = 1.0f;               // bit code 111
+		}
+	}
+	else if (mode == BC4Mode::BC4_SNORM)
+	{
+		r0 = Int8ToFloat_SNORM(blockStorage[0]);
+		r1 = Int8ToFloat_SNORM(blockStorage[1]);
+
+		colorTable[0] = r0;
+		colorTable[1] = r1;
+
+		if (r0 > r1)
+		{
+			// 6 interpolated color values
+			colorTable[2] = (6 * r0 + 1 * r1) / 7.0f; // bit code 010
+			colorTable[3] = (5 * r0 + 2 * r1) / 7.0f; // bit code 011
+			colorTable[4] = (4 * r0 + 3 * r1) / 7.0f; // bit code 100
+			colorTable[5] = (3 * r0 + 4 * r1) / 7.0f; // bit code 101
+			colorTable[6] = (2 * r0 + 5 * r1) / 7.0f; // bit code 110
+			colorTable[7] = (1 * r0 + 6 * r1) / 7.0f; // bit code 111
+		}
+		else
+		{
+			// 4 interpolated color values
+			colorTable[2] = (4 * r0 + 1 * r1) / 5.0f; // bit code 010
+			colorTable[3] = (3 * r0 + 2 * r1) / 5.0f; // bit code 011
+			colorTable[4] = (2 * r0 + 3 * r1) / 5.0f; // bit code 100
+			colorTable[5] = (1 * r0 + 4 * r1) / 5.0f; // bit code 101
+			colorTable[6] = -1.0f;              // bit code 110
+			colorTable[7] = 1.0f;              // bit code 111
+		}
+	}
+
+	DecompressBlockBC4Internal(blockStorage, image + x + (y * stride), stride, colorTable);
+}
+
+
+//-------------------------------------------------------------------------------------------------
+inline void DecompressBlockBC5(uint32_t x, uint32_t y, uint32_t stride, BC4Mode mode, const uint8_t* blockStorage, uint8_t* image)
+{
+	// We decompress the two channels separately and interleave them when
+	// writing to the output
+	uint8_t c0[16];
+	uint8_t c1[16];
+
+	int dx, dy;
+
+	DecompressBlockBC4(0, 0, 4, mode, blockStorage, (uint8_t*)c0);
+	DecompressBlockBC4(0, 0, 4, mode, blockStorage + 8, (uint8_t*)c1);
+
+	for (dy = 0; dy < 4; ++dy)
+	{
+		for (dx = 0; dx < 4; ++dx)
+		{
+			uint8_t* pDest = image + stride * (y + dy) + (x + dx) * 4;
+			pDest[0] = c0[dy * 4 + dx];
+			pDest[1] = c1[dy * 4 + dx];
+			pDest[2] = 255;
+			pDest[3] = 255;
+		}
+	}
+}
+#endif
+
+//-------------------------------------------------------------------------------------------------
+UObject* FUnrealBakeHelpers::BakeHelper_DuplicateAsset(UObject* Object, const FString& ObjName, const FString& PkgName, bool ResetDuplicatedFlags, TMap<UObject*, UObject*>& ReplacementMap, bool OverwritePackage)
+{
+#if WITH_EDITORONLY_DATA
+	FString FinalObjName = ObjName;
+	FString FinalPkgName = PkgName;
+
+	if (!OverwritePackage)
+	{
+		FAssetToolsModule& AssetToolsModule = FModuleManager::GetModuleChecked<FAssetToolsModule>("AssetTools");
+		AssetToolsModule.Get().CreateUniqueAssetName(PkgName, "", FinalPkgName, FinalObjName);
+	}
+
+	UPackage* Package = CreatePackage(*FinalPkgName);
+	Package->FullyLoad();
+
+	FObjectDuplicationParameters Params = InitStaticDuplicateObjectParams(Object, Package, *FinalObjName, RF_AllFlags, nullptr, EDuplicateMode::Normal);
+
+	UObject* DupObject = StaticDuplicateObjectEx(Params);
+	if (DupObject)
+	{
+		DupObject->SetFlags(RF_Public | RF_Standalone);
+		DupObject->ClearFlags(RF_Transient);
+
+		// The garbage collector is called in the middle of the bake process, and this can destroy this temporary objects. 
+		// We add them to the garbage root to prevent this. This will avoid them being unloaded while the editor is running, but this
+		// action is not used often.
+		DupObject->AddToRoot();
+		DupObject->MarkPackageDirty();
+
+		// Notify the asset registry
+		FAssetRegistryModule::AssetCreated(DupObject);
+
+		// Replace all references
+		ReplacementMap.Add(Object, DupObject);
+
+		constexpr EArchiveReplaceObjectFlags ReplaceFlags = (EArchiveReplaceObjectFlags::IgnoreOuterRef | EArchiveReplaceObjectFlags::IgnoreArchetypeRef);
+		FArchiveReplaceObjectRef<UObject> ReplaceAr(DupObject, ReplacementMap, ReplaceFlags);
+	}
+
+	return DupObject;
+#else
+	return nullptr;
+#endif
+}
+
+
+//-------------------------------------------------------------------------------------------------
+UTexture2D* FUnrealBakeHelpers::BakeHelper_CreateAssetTexture(UTexture2D* SrcTex, const FString& TexObjName, const FString& TexPkgName, const UTexture* OrgTex, bool ResetDuplicatedFlags, TMap<UObject*, UObject*>& ReplacementMap, bool OverwritePackage)
+{
+#if WITH_EDITORONLY_DATA
+	const bool bIsMutableTexture = !SrcTex->Source.IsValid();
+	if (bIsMutableTexture)
+	{
+		int32 sx = SrcTex->GetPlatformData()->SizeX;
+		int32 sy = SrcTex->GetPlatformData()->SizeY;
+		EPixelFormat SrcPixelFormat = SrcTex->GetPlatformData()->PixelFormat;
+		ETextureSourceFormat PixelFormat = SrcPixelFormat == PF_BC4 || SrcPixelFormat == PF_G8 ? TSF_G8 : TSF_BGRA8;
+
+		// Begin duplicate texture
+		FString FinalObjName = TexObjName;
+		FString FinalPkgName = TexPkgName;
+
+		if (!OverwritePackage)
+		{
+			FAssetToolsModule& AssetToolsModule = FModuleManager::GetModuleChecked<FAssetToolsModule>("AssetTools");
+			AssetToolsModule.Get().CreateUniqueAssetName(TexPkgName, "", FinalPkgName, FinalObjName);
+		}
+
+		UPackage* Package = CreatePackage(*FinalPkgName);
+		Package->FullyLoad();
+
+		UTexture2D* DupTex = NewObject<UTexture2D>(Package, *FinalObjName, RF_Public | RF_Standalone);
+
+		ReplacementMap.Add(SrcTex, DupTex);
+
+		// The garbage collector is called in the middle of the bake process, and this can destroy this temporary object. 
+		// We add them to the garbage root to prevent this. This will avoid them being unloaded while the editor is running, but this
+		// action is not used often.
+		DupTex->AddToRoot();
+		DupTex->MarkPackageDirty();
+
+		// Notify the asset registry
+		FAssetRegistryModule::AssetCreated(DupTex);
+
+		// Replace all references
+		constexpr EArchiveReplaceObjectFlags ReplaceFlags = (EArchiveReplaceObjectFlags::IgnoreOuterRef | EArchiveReplaceObjectFlags::IgnoreArchetypeRef);
+		FArchiveReplaceObjectRef<UObject> ReplaceAr(DupTex, ReplacementMap, ReplaceFlags);
+
+		// End duplicate texture
+
+		CopyTextureProperties(DupTex, SrcTex);
+
+		DupTex->RemoveUserDataOfClass(UMutableTextureMipDataProviderFactory::StaticClass());
+		if (OrgTex)
+		{
+			DupTex->CompressionSettings = OrgTex->CompressionSettings;
+		}
+
+		// Mutable textures only have platform data. We need to build the source data for them to be assets.
+		DupTex->Source.Init(sx, sy, 1, 1, PixelFormat);
+
+		int32 MipCount = SrcTex->GetPlatformData()->Mips.Num();
+		if (!MipCount)
+		{
+			UE_LOG(LogMutable, Warning, TEXT("Bake Instances: Empty texture found [%s]."), *SrcTex->GetName());
+			return DupTex;
+		}
+
+		const uint8_t* src = reinterpret_cast<const uint8_t*>(SrcTex->GetPlatformData()->Mips[0].BulkData.LockReadOnly());
+		check(src); // Mutable texture should always contain platform data
+
+		uint8_t* dst = DupTex->Source.LockMip(0);
+		check(dst);
+
+		switch (SrcTex->GetPlatformData()->PixelFormat)
+		{
+		case PF_R8G8B8A8:
+			FMemory::Memcpy(dst, src, sx * sy * 4);
+			break;
+
+		case PF_G8:
+			FMemory::Memcpy(dst, src, sx * sy * 1);
+			break;
+
+		case PF_DXT1:
+			for (int y = 0; y < sy; y += 4)
+				for (int x = 0; x < sx; x += 4)
+				{
+					DecompressBlockBC1(x, y, sx * 4, src, dst);
+					src += 8;
+				}
+			break;
+
+		case PF_DXT3:
+			for (int y = 0; y < sy; y += 4)
+				for (int x = 0; x < sx; x += 4)
+				{
+					DecompressBlockBC2(x, y, sx * 4, src, dst);
+					src += 16;
+				}
+			break;
+
+		case PF_DXT5:
+			for (int y = 0; y < sy; y += 4)
+				for (int x = 0; x < sx; x += 4)
+				{
+					DecompressBlockBC3(x, y, sx * 4, src, dst);
+					src += 16;
+				}
+			break;
+
+		case PF_BC4:
+			for (int y = 0; y < sy; y += 4)
+				for (int x = 0; x < sx; x += 4)
+				{
+					DecompressBlockBC4(x, y, sx * 1, BC4Mode::BC4_UNORM, src, dst);
+					src += 8;
+				}
+			break;
+
+		case PF_BC5:
+			for (int y = 0; y < sy; y += 4)
+				for (int x = 0; x < sx; x += 4)
+				{
+					DecompressBlockBC5(x, y, sx * 4, BC4Mode::BC4_UNORM, src, dst);
+					src += 16;
+				}
+			break;
+
+		default:
+			// Not implemented...
+			check(false);
+			break;
+		}
+
+		const bool bNeedsRBSwizzle = PixelFormat == TSF_BGRA8;
+		if (bNeedsRBSwizzle)
+		{
+			for (int x = 0; x < sx * sy; ++x)
+			{
+				uint8_t temp = dst[0];
+				dst[0] = dst[2];
+				dst[2] = temp;
+				dst += 4;
+			}
+		}
+
+		if (PixelFormat == TSF_G8 || PixelFormat == TSF_G16)
+		{
+			// If compression settings are not set to TC_Grayscale the texture will get a DXT format
+			// instead of G8 or G16.
+			FTextureFormatSettings Settings;
+			Settings.CompressionSettings = TC_Grayscale;
+			DupTex->SetLayerFormatSettings(0, Settings);
+
+			DupTex->CompressionSettings = TC_Grayscale;
+		}
+
+		SrcTex->GetPlatformData()->Mips[0].BulkData.Unlock();
+
+		DupTex->Source.UnlockMip(0);
+
+		DupTex->UpdateResource();
+
+		return DupTex;
+	}
+	else
+	{
+		return Cast<UTexture2D>(BakeHelper_DuplicateAsset(SrcTex, TexObjName, TexPkgName, ResetDuplicatedFlags, ReplacementMap, OverwritePackage));
+	}
+#else
+	return nullptr;
+#endif
+}
+
+
 //-------------------------------------------------------------------------------------------------
 void FUnrealBakeHelpers::BakeHelper_RegenerateImportedModel(USkeletalMesh* SkeletalMesh)
 {
diff --git a/Engine/Plugins/Experimental/Mutable/Source/CustomizableObject/Public/MuCO/UnrealBakeHelpers.h b/Engine/Plugins/Experimental/Mutable/Source/CustomizableObject/Public/MuCO/UnrealBakeHelpers.h
--- a/Engine/Plugins/Experimental/Mutable/Source/CustomizableObject/Public/MuCO/UnrealBakeHelpers.h
+++ b/Engine/Plugins/Experimental/Mutable/Source/CustomizableObject/Public/MuCO/UnrealBakeHelpers.h
@@ -3,6 +3,12 @@
 #pragma once
 
 #include "HAL/Platform.h"
+#include "Containers/Map.h"
+
+class UObject;
+class USkeletalMesh;
+class UTexture2D;
+class UTexture;
 
 //-------------------------------------------------------------------------------------------------
 
@@ -10,7 +16,17 @@ class CUSTOMIZABLEOBJECT_API FUnrealBakeHelpers
 {
 public:
 
-	static void BakeHelper_RegenerateImportedModel(class USkeletalMesh* SkeletalMesh);
+	static UObject* BakeHelper_DuplicateAsset(UObject* Object, const FString& ObjName, const FString& PkgName, bool ResetDuplicatedFlags, TMap<UObject*, UObject*>& ReplacementMap, bool OverwritePackage);
+
+	/**
+	 * Duplicates a texture asset. Duplicates Mutable and non Mutable textures.
+	 *
+	 * @param OrgTex Original source texture from which a Mutable texture has been generated. Only required when SrcTex is a Mutable texture.
+	 */
+	static UTexture2D* BakeHelper_CreateAssetTexture(UTexture2D* SrcTex, const FString& TexObjName, const FString& TexPkgName, const UTexture* OrgTex, bool ResetDuplicatedFlags, TMap<UObject*, UObject*>& ReplacementMap, bool OverwritePackage);
+
+
+	static void BakeHelper_RegenerateImportedModel(USkeletalMesh* SkeletalMesh);
 };
 
 //-------------------------------------------------------------------------------------------------
diff --git a/Engine/Plugins/Experimental/Mutable/Source/CustomizableObjectEditor/Private/MuCOE/CustomizableObjectBakeHelpers.cpp b/Engine/Plugins/Experimental/Mutable/Source/CustomizableObjectEditor/Private/MuCOE/CustomizableObjectBakeHelpers.cpp
deleted file mode 100644
diff --git a/Engine/Plugins/Experimental/Mutable/Source/CustomizableObjectEditor/Private/MuCOE/CustomizableObjectBakeHelpers.h b/Engine/Plugins/Experimental/Mutable/Source/CustomizableObjectEditor/Private/MuCOE/CustomizableObjectBakeHelpers.h
deleted file mode 100644
diff --git a/Engine/Plugins/Experimental/Mutable/Source/CustomizableObjectEditor/Private/MuCOE/CustomizableObjectEditor.cpp b/Engine/Plugins/Experimental/Mutable/Source/CustomizableObjectEditor/Private/MuCOE/CustomizableObjectEditor.cpp
--- a/Engine/Plugins/Experimental/Mutable/Source/CustomizableObjectEditor/Private/MuCOE/CustomizableObjectEditor.cpp
+++ b/Engine/Plugins/Experimental/Mutable/Source/CustomizableObjectEditor/Private/MuCOE/CustomizableObjectEditor.cpp
@@ -20,7 +20,6 @@
 #include "MuCO/CustomizableObjectSystem.h"
 #include "MuCO/CustomizableSkeletalComponent.h"
 #include "MuCO/UnrealPortabilityHelpers.h"
-#include "MuCOE/CustomizableObjectBakeHelpers.h"
 #include "MuCOE/CustomizableObjectCustomSettings.h"
 #include "MuCOE/CustomizableObjectEditorActions.h"
 #include "MuCOE/CustomizableObjectEditorLogger.h"
@@ -2541,11 +2540,6 @@ void FCustomizableObjectEditor::OnUpdatePreviewInstance()
 		TextureAnalyzer->RefreshTextureAnalyzerTable(PreviewInstance);
 	}
 
-	for (int32 ComponentIndex = 0; ComponentIndex < PreviewInstance->SkeletalMeshes.Num(); ++ComponentIndex)
-	{
-		BakeHelper_RegenerateImportedModel(PreviewInstance->SkeletalMeshes[ComponentIndex]);
-	}
-
 	// Do postponed work.
 	for (const TFunction<void()>& Work : OnUpdatePreviewInstanceWork)
 	{
diff --git a/Engine/Plugins/Experimental/Mutable/Source/CustomizableObjectEditor/Private/MuCOE/CustomizableObjectEditorViewportClient.cpp b/Engine/Plugins/Experimental/Mutable/Source/CustomizableObjectEditor/Private/MuCOE/CustomizableObjectEditorViewportClient.cpp
--- a/Engine/Plugins/Experimental/Mutable/Source/CustomizableObjectEditor/Private/MuCOE/CustomizableObjectEditorViewportClient.cpp
+++ b/Engine/Plugins/Experimental/Mutable/Source/CustomizableObjectEditor/Private/MuCOE/CustomizableObjectEditorViewportClient.cpp
@@ -27,7 +27,6 @@
 #include "MuCO/CustomizableObjectSystem.h"
 #include "MuCO/CustomizableObjectMipDataProvider.h"
 #include "MuCO/UnrealBakeHelpers.h"
-#include "MuCOE/CustomizableObjectBakeHelpers.h"
 #include "MuCOE/CustomizableObjectPreviewScene.h"
 #include "MuCOE/CustomizableObjectWidget.h"
 #include "MuCOE/ICustomizableObjectInstanceEditor.h"
@@ -2431,7 +2430,7 @@ void FCustomizableObjectEditorViewportClient::BakeInstance(UCustomizableObjectIn
 
 										PackageName = FolderDlg->GetAssetPath() + FString("/") + ResourceName;
 										TMap<UObject*, UObject*> FakeReplacementMap;
-										UTexture2D* DupTex = BakeHelper_CreateAssetTexture(SrcTex, ResourceName, PackageName, OriginalTexture, true, FakeReplacementMap, BakingOverwritePermission);
+										UTexture2D* DupTex = FUnrealBakeHelpers::BakeHelper_CreateAssetTexture(SrcTex, ResourceName, PackageName, OriginalTexture, true, FakeReplacementMap, BakingOverwritePermission);
 										ArrayCachedElement.Add(ResourceName);
 										ArrayCachedObject.Add(DupTex);
 										PackagesToSave.Add(DupTex->GetPackage());
@@ -2478,7 +2477,7 @@ void FCustomizableObjectEditorViewportClient::BakeInstance(UCustomizableObjectIn
 
 										PackageName = FolderDlg->GetAssetPath() + FString("/") + ResourceName;
 										TMap<UObject*, UObject*> FakeReplacementMap;
-										DuplicatedObject = BakeHelper_DuplicateAsset(Texture, ResourceName, PackageName, true, FakeReplacementMap, BakingOverwritePermission);
+										DuplicatedObject = FUnrealBakeHelpers::BakeHelper_DuplicateAsset(Texture, ResourceName, PackageName, true, FakeReplacementMap, BakingOverwritePermission);
 										ArrayCachedElement.Add(ResourceName);
 										ArrayCachedObject.Add(DuplicatedObject);
 										PackagesToSave.Add(DuplicatedObject->GetPackage());
@@ -2515,7 +2514,7 @@ void FCustomizableObjectEditorViewportClient::BakeInstance(UCustomizableObjectIn
 
 							PackageName = FolderDlg->GetAssetPath() + FString("/") + ResourceName;
 							TMap<UObject*, UObject*> FakeReplacementMap;
-							DuplicatedObject = BakeHelper_DuplicateAsset(Material, ResourceName, PackageName, false, FakeReplacementMap, BakingOverwritePermission);
+							DuplicatedObject = FUnrealBakeHelpers::BakeHelper_DuplicateAsset(Material, ResourceName, PackageName, false, FakeReplacementMap, BakingOverwritePermission);
 							ArrayCachedElement.Add(ResourceName);
 							ArrayCachedObject.Add(DuplicatedObject);
 							ReplacementMap.Add(Interface, DuplicatedObject);
@@ -2611,7 +2610,7 @@ void FCustomizableObjectEditorViewportClient::BakeInstance(UCustomizableObjectIn
 						}
 
 						FString MatPkgName = FolderDlg->GetAssetPath() + FString("/") + MatObjName;
-						UObject* DupMat = BakeHelper_DuplicateAsset(Interface, MatObjName, MatPkgName, false, ReplacementMap, BakingOverwritePermission);
+						UObject* DupMat = FUnrealBakeHelpers::BakeHelper_DuplicateAsset(Interface, MatObjName, MatPkgName, false, ReplacementMap, BakingOverwritePermission);
 						ArrayCachedObject.Add(DupMat);
 						ArrayCachedElement.Add(MatObjName);
 						PackagesToSave.Add(DupMat->GetPackage());
@@ -2649,7 +2648,7 @@ void FCustomizableObjectEditorViewportClient::BakeInstance(UCustomizableObjectIn
 
 										FString TexPkgName = FolderDlg->GetAssetPath() + FString("/") + TexObjName;
 										TMap<UObject*, UObject*> FakeReplacementMap;
-										UTexture2D* DupTex = BakeHelper_CreateAssetTexture(SrcTex, TexObjName, TexPkgName, nullptr, false, FakeReplacementMap, BakingOverwritePermission);
+										UTexture2D* DupTex = FUnrealBakeHelpers::BakeHelper_CreateAssetTexture(SrcTex, TexObjName, TexPkgName, nullptr, false, FakeReplacementMap, BakingOverwritePermission);
 										ArrayCachedObject.Add(DupTex);
 										ArrayCachedElement.Add(TexObjName);
 										PackagesToSave.Add(DupTex->GetPackage());
@@ -2677,7 +2676,7 @@ void FCustomizableObjectEditorViewportClient::BakeInstance(UCustomizableObjectIn
 						}
 
 						FString SkeletonPkgName = FolderDlg->GetAssetPath() + FString("/") + SkeletonName;
-						UObject* DuplicatedSkeleton = BakeHelper_DuplicateAsset(Mesh->GetSkeleton(), SkeletonName, SkeletonPkgName, false, ReplacementMap, BakingOverwritePermission);
+						UObject* DuplicatedSkeleton = FUnrealBakeHelpers::BakeHelper_DuplicateAsset(Mesh->GetSkeleton(), SkeletonName, SkeletonPkgName, false, ReplacementMap, BakingOverwritePermission);
 
 						ArrayCachedObject.Add(DuplicatedSkeleton);
 						PackagesToSave.Add(DuplicatedSkeleton->GetPackage());
@@ -2695,7 +2694,7 @@ void FCustomizableObjectEditorViewportClient::BakeInstance(UCustomizableObjectIn
 				}
 
 				FString PkgName = FolderDlg->GetAssetPath() + FString("/") + ObjectName;
-				UObject* DupObject = BakeHelper_DuplicateAsset(Mesh, ObjectName, PkgName, false, ReplacementMap, BakingOverwritePermission);
+				UObject* DupObject = FUnrealBakeHelpers::BakeHelper_DuplicateAsset(Mesh, ObjectName, PkgName, false, ReplacementMap, BakingOverwritePermission);
 				ArrayCachedObject.Add(DupObject);
 				PackagesToSave.Add(DupObject->GetPackage());
 
diff --git a/Engine/Plugins/Experimental/Mutable/Source/CustomizableObjectEditor/Private/MuCOE/CustomizableObjectInstanceEditor.cpp b/Engine/Plugins/Experimental/Mutable/Source/CustomizableObjectEditor/Private/MuCOE/CustomizableObjectInstanceEditor.cpp
--- a/Engine/Plugins/Experimental/Mutable/Source/CustomizableObjectEditor/Private/MuCOE/CustomizableObjectInstanceEditor.cpp
+++ b/Engine/Plugins/Experimental/Mutable/Source/CustomizableObjectEditor/Private/MuCOE/CustomizableObjectInstanceEditor.cpp
@@ -17,7 +17,6 @@
 #include "MuCO/CustomizableObjectInstance.h"
 #include "MuCO/CustomizableObjectSystem.h"
 #include "MuCO/CustomizableSkeletalComponent.h"
-#include "MuCOE/CustomizableObjectBakeHelpers.h"
 #include "MuCOE/CustomizableObjectCompiler.h"
 #include "MuCOE/CustomizableObjectEditorActions.h"
 #include "MuCOE/CustomizableObjectEditorModule.h"
@@ -776,11 +775,6 @@ void FCustomizableObjectInstanceEditor::OnUpdatePreviewInstance()
 	{
 		TextureAnalyzer.Get()->RefreshTextureAnalyzerTable(CustomizableObjectInstance);
 	}
-	
-	for (int32 ComponentIndex = 0; ComponentIndex < CustomizableObjectInstance->SkeletalMeshes.Num(); ++ComponentIndex)
-	{
-		BakeHelper_RegenerateImportedModel(CustomizableObjectInstance->GetSkeletalMesh(ComponentIndex));
-	}
 }
 
 
