[mutable] Fixed all cases where table nodes could send a nullptr cell to the mutable core and generate a crash.

#rnx

diff --git a/Engine/Plugins/Experimental/Mutable/Source/CustomizableObjectEditor/Private/MuCOE/GenerateMutableSource/GenerateMutableSourceColor.cpp b/Engine/Plugins/Experimental/Mutable/Source/CustomizableObjectEditor/Private/MuCOE/GenerateMutableSource/GenerateMutableSourceColor.cpp
--- a/Engine/Plugins/Experimental/Mutable/Source/CustomizableObjectEditor/Private/MuCOE/GenerateMutableSource/GenerateMutableSourceColor.cpp
+++ b/Engine/Plugins/Experimental/Mutable/Source/CustomizableObjectEditor/Private/MuCOE/GenerateMutableSource/GenerateMutableSourceColor.cpp
@@ -47,6 +47,8 @@ mu::NodeColourPtr GenerateMutableSourceColor(const UEdGraphPin* Pin, FMutableGra
 		return static_cast<mu::NodeColour*>(Generated->Node.get());
 	}
 
+	bool bDoNotAddToGeneratedCache = false;
+
 	mu::NodeColourPtr Result;
 
 	if (const UCustomizableObjectNodeColorConstant* TypedNodeColorConst = Cast<UCustomizableObjectNodeColorConstant>(Node))
@@ -269,46 +271,81 @@ mu::NodeColourPtr GenerateMutableSourceColor(const UEdGraphPin* Pin, FMutableGra
 
 	else if (const UCustomizableObjectNodeTable* TypedNodeTable = Cast<UCustomizableObjectNodeTable>(Node))
 	{
-		mu::NodeColourTablePtr ColorTableNode = new mu::NodeColourTable();
-		Result = ColorTableNode;
+		//This node will add a white color in case of error
+		mu::NodeColourConstantPtr WhiteColorNode = new mu::NodeColourConstant();
+		WhiteColorNode->SetValue(1.0f, 1.0f, 1.0f, 1.0f);
+
+		Result = WhiteColorNode;
+
+		bool bSuccess = true;
+
+		if (Pin->PinType.PinCategory == Schema->PC_MaterialAsset)
+		{
+			// Material pins have to skip the cache of nodes or they will return always the same column node
+			bDoNotAddToGeneratedCache = true;
+		}
 
-		mu::TablePtr Table;
 		if (TypedNodeTable->Table)
 		{
 			FString ColumnName = Pin->PinFriendlyName.ToString();
+			FProperty* Property = TypedNodeTable->Table->FindTableProperty(FName(*ColumnName));
 
-			if (Pin->PinType.PinCategory == Schema->PC_MaterialAsset)
+			if (!Property)
 			{
-				ColumnName = GenerationContext.CurrentMaterialTableParameterId;
-			}
-
-			// Generating a new data table if not exists
-			Table = GenerateMutableSourceTable(TypedNodeTable->Table->GetName(), Pin, GenerationContext);
+				FString Msg = FString::Printf(TEXT("Couldn't find the column [%s] in the data table's struct."), *ColumnName);
+				GenerationContext.Compiler->CompilerLog(FText::FromString(Msg), Node);
 
-			// Generating a new Color column if not exists
-			if (Table && Table->FindColumn(StringCast<ANSICHAR>(*ColumnName).Get()) == INDEX_NONE)
-			{
-				GenerateTableColumn(TypedNodeTable, Pin, Table, ColumnName, GenerationContext.CurrentLOD, GenerationContext);
+				bSuccess = false;
 			}
 
-			ColorTableNode->SetTable(Table);
-			ColorTableNode->SetColumn(StringCast<ANSICHAR>(*ColumnName).Get());
-			ColorTableNode->SetParameterName(StringCast<ANSICHAR>(*TypedNodeTable->ParameterName).Get());
-
-			GenerationContext.AddParameterNameUnique(Node, TypedNodeTable->ParameterName);
-
-			if (Table->FindColumn(StringCast<ANSICHAR>(*ColumnName).Get()) == INDEX_NONE)
+			if (bSuccess)
 			{
-				FString Msg = FString::Printf(TEXT("Couldn't find pin column with name %s"), *ColumnName);
-				GenerationContext.Compiler->CompilerLog(FText::FromString(Msg), Node);
-			}
+				// Generating a new data table if not exists
+				mu::TablePtr Table;
+				Table = GenerateMutableSourceTable(TypedNodeTable->Table->GetName(), Pin, GenerationContext);
 
+				if (Table)
+				{
+					mu::NodeColourTablePtr ColorTableNode = new mu::NodeColourTable();
+
+					if (Pin->PinType.PinCategory == Schema->PC_MaterialAsset)
+					{
+						// Materials use the parameter id as column names
+						ColumnName = GenerationContext.CurrentMaterialTableParameterId;
+					}
+
+					// Generating a new Color column if not exists
+					if (Table->FindColumn(StringCast<ANSICHAR>(*ColumnName).Get()) == INDEX_NONE)
+					{
+						bSuccess = GenerateTableColumn(TypedNodeTable, Pin, Table, ColumnName, GenerationContext.CurrentLOD, GenerationContext);
+
+						if (!bSuccess)
+						{
+							FString Msg = FString::Printf(TEXT("Failed to generate the mutable table column [%s]"), *ColumnName);
+							GenerationContext.Compiler->CompilerLog(FText::FromString(Msg), Node);
+						}
+					}
+
+					if (bSuccess)
+					{
+						Result = ColorTableNode;
+
+						ColorTableNode->SetTable(Table);
+						ColorTableNode->SetColumn(StringCast<ANSICHAR>(*ColumnName).Get());
+						ColorTableNode->SetParameterName(StringCast<ANSICHAR>(*TypedNodeTable->ParameterName).Get());
+
+						GenerationContext.AddParameterNameUnique(Node, TypedNodeTable->ParameterName);
+					}
+				}
+				else
+				{
+					FString Msg = FString::Printf(TEXT("Couldn't generate a mutable table."), *ColumnName);
+					GenerationContext.Compiler->CompilerLog(FText::FromString(Msg), Node);
+				}
+			}
 		}
 		else
 		{
-			Table = new mu::Table();
-			ColorTableNode->SetTable(Table);
-
 			GenerationContext.Compiler->CompilerLog(LOCTEXT("ColorTableError", "Couldn't find the data table of the node."), Node);
 		}
 	}
@@ -318,9 +355,12 @@ mu::NodeColourPtr GenerateMutableSourceColor(const UEdGraphPin* Pin, FMutableGra
 		GenerationContext.Compiler->CompilerLog(LOCTEXT("UnimplementedNode", "Node type not implemented yet."), Node);
 	}
 
-	FGeneratedData CacheData = FGeneratedData(Node, Result);
-	GenerationContext.Generated.Add(Key, CacheData);
-	GenerationContext.GeneratedNodes.Add(Node);
+	if (!bDoNotAddToGeneratedCache)
+	{
+		FGeneratedData CacheData = FGeneratedData(Node, Result);
+		GenerationContext.Generated.Add(Key, CacheData);
+		GenerationContext.GeneratedNodes.Add(Node);
+	}
 	
 	if (Result)
 	{
diff --git a/Engine/Plugins/Experimental/Mutable/Source/CustomizableObjectEditor/Private/MuCOE/GenerateMutableSource/GenerateMutableSourceFloat.cpp b/Engine/Plugins/Experimental/Mutable/Source/CustomizableObjectEditor/Private/MuCOE/GenerateMutableSource/GenerateMutableSourceFloat.cpp
--- a/Engine/Plugins/Experimental/Mutable/Source/CustomizableObjectEditor/Private/MuCOE/GenerateMutableSource/GenerateMutableSourceFloat.cpp
+++ b/Engine/Plugins/Experimental/Mutable/Source/CustomizableObjectEditor/Private/MuCOE/GenerateMutableSource/GenerateMutableSourceFloat.cpp
@@ -74,6 +74,8 @@ mu::NodeScalarPtr GenerateMutableSourceFloat(const UEdGraphPin* Pin, FMutableGra
 		return static_cast<mu::NodeScalar*>(Generated->Node.get());
 	}
 
+	bool bDoNotAddToGeneratedCache = false;
+
 	mu::NodeScalarPtr Result;
 	
 	if (const UCustomizableObjectNodeFloatConstant* FloatConstantNode = Cast<UCustomizableObjectNodeFloatConstant>(Node))
@@ -295,47 +297,81 @@ mu::NodeScalarPtr GenerateMutableSourceFloat(const UEdGraphPin* Pin, FMutableGra
 
 	else if (const UCustomizableObjectNodeTable* TypedNodeTable = Cast<UCustomizableObjectNodeTable>(Node))
 	{
-		mu::NodeScalarTablePtr ScalarTableNode = new mu::NodeScalarTable();
-		Result = ScalarTableNode;
+		//This node will add a default value in case of error
+		mu::NodeScalarConstantPtr ConstantValue = new mu::NodeScalarConstant();
+		ConstantValue->SetValue(1.0f);
+
+		Result = ConstantValue;
+
+		if (Pin->PinType.PinCategory == Schema->PC_MaterialAsset)
+		{
+			// Material pins have to skip the cache of nodes or they will return always the same column node
+			bDoNotAddToGeneratedCache = true;
+		}
+
+		bool bSuccess = true;
 
-		mu::TablePtr Table;
 		if (TypedNodeTable->Table)
 		{
 			FString ColumnName = Pin->PinFriendlyName.ToString();
+			FProperty* Property = TypedNodeTable->Table->FindTableProperty(FName(*ColumnName));
 
-			if (Pin->PinType.PinCategory == Schema->PC_MaterialAsset)
+			if (!Property)
 			{
-				ColumnName = GenerationContext.CurrentMaterialTableParameterId;
-			}
-			
-			// Generating a new data table if not exists
-			Table = GenerateMutableSourceTable(TypedNodeTable->Table->GetName(), Pin, GenerationContext);
+				FString Msg = FString::Printf(TEXT("Couldn't find the column [%s] in the data table's struct."), *ColumnName);
+				GenerationContext.Compiler->CompilerLog(FText::FromString(Msg), Node);
 
-			// Generating a new Float column if not exists
-			if (Table && Table->FindColumn(StringCast<ANSICHAR>(*ColumnName).Get()) == INDEX_NONE)
-			{
-				GenerateTableColumn(TypedNodeTable, Pin, Table, ColumnName, GenerationContext.CurrentLOD, GenerationContext);
+				bSuccess = false;
 			}
-			
-			ScalarTableNode->SetTable(Table);
-			ScalarTableNode->SetColumn(StringCast<ANSICHAR>(*ColumnName).Get());
-			ScalarTableNode->SetParameterName(StringCast<ANSICHAR>(*TypedNodeTable->ParameterName).Get());
-
-			GenerationContext.AddParameterNameUnique(Node, TypedNodeTable->ParameterName);
-			
-			if (Table->FindColumn(StringCast<ANSICHAR>(*ColumnName).Get()) == INDEX_NONE)
+
+			if (bSuccess)
 			{
-				Table = new mu::Table();
+				// Generating a new data table if not exists
+				mu::TablePtr Table;
+				Table = GenerateMutableSourceTable(TypedNodeTable->Table->GetName(), Pin, GenerationContext);
 
-				FString Msg = FString::Printf(TEXT("Couldn't find pin column with name %s"), *ColumnName);
-				GenerationContext.Compiler->CompilerLog(FText::FromString(Msg), Node);
+				if (Table)
+				{
+					mu::NodeScalarTablePtr ScalarTableNode = new mu::NodeScalarTable();
+
+					if (Pin->PinType.PinCategory == Schema->PC_MaterialAsset)
+					{
+						// Materials use the parameter id as column names
+						ColumnName = GenerationContext.CurrentMaterialTableParameterId;
+					}
+
+					// Generating a new Float column if not exists
+					if (Table->FindColumn(StringCast<ANSICHAR>(*ColumnName).Get()) == INDEX_NONE)
+					{
+						bSuccess = GenerateTableColumn(TypedNodeTable, Pin, Table, ColumnName, GenerationContext.CurrentLOD, GenerationContext);
+
+						if (!bSuccess)
+						{
+							FString Msg = FString::Printf(TEXT("Failed to generate the mutable table column [%s]"), *ColumnName);
+							GenerationContext.Compiler->CompilerLog(FText::FromString(Msg), Node);
+						}
+					}
+
+					if (bSuccess)
+					{
+						Result = ScalarTableNode;
+
+						ScalarTableNode->SetTable(Table);
+						ScalarTableNode->SetColumn(StringCast<ANSICHAR>(*ColumnName).Get());
+						ScalarTableNode->SetParameterName(StringCast<ANSICHAR>(*TypedNodeTable->ParameterName).Get());
+
+						GenerationContext.AddParameterNameUnique(Node, TypedNodeTable->ParameterName);
+					}
+				}
+				else
+				{
+					FString Msg = FString::Printf(TEXT("Couldn't generate a mutable table."), *ColumnName);
+					GenerationContext.Compiler->CompilerLog(FText::FromString(Msg), Node);
+				}
 			}
 		}
 		else
 		{
-			Table = new mu::Table();
-			ScalarTableNode->SetTable(Table);
-
 			GenerationContext.Compiler->CompilerLog(LOCTEXT("ScalarTableError", "Couldn't find the data table of the node."), Node);
 		}
 	}
@@ -345,8 +381,11 @@ mu::NodeScalarPtr GenerateMutableSourceFloat(const UEdGraphPin* Pin, FMutableGra
 		GenerationContext.Compiler->CompilerLog(LOCTEXT("UnimplementedNode", "Node type not implemented yet."), Node);
 	}
 
-	GenerationContext.Generated.Add(Key, FGeneratedData(Node, Result));
-	GenerationContext.GeneratedNodes.Add(Node);
+	if (!bDoNotAddToGeneratedCache)
+	{
+		GenerationContext.Generated.Add(Key, FGeneratedData(Node, Result));
+		GenerationContext.GeneratedNodes.Add(Node);
+	}
 
 	if (Result)
 	{
diff --git a/Engine/Plugins/Experimental/Mutable/Source/CustomizableObjectEditor/Private/MuCOE/GenerateMutableSource/GenerateMutableSourceImage.cpp b/Engine/Plugins/Experimental/Mutable/Source/CustomizableObjectEditor/Private/MuCOE/GenerateMutableSource/GenerateMutableSourceImage.cpp
--- a/Engine/Plugins/Experimental/Mutable/Source/CustomizableObjectEditor/Private/MuCOE/GenerateMutableSource/GenerateMutableSourceImage.cpp
+++ b/Engine/Plugins/Experimental/Mutable/Source/CustomizableObjectEditor/Private/MuCOE/GenerateMutableSource/GenerateMutableSourceImage.cpp
@@ -141,6 +141,8 @@ mu::NodeImagePtr GenerateMutableSourceImage(const UEdGraphPin* Pin, FMutableGrap
 		return static_cast<mu::NodeImage*>(Generated->Node.get());
 	}
 
+	bool bDoNotAddToGeneratedCache = false;
+
 	mu::NodeImagePtr Result;
 	
 	if (const UCustomizableObjectNodeTexture* TypedNodeTex = Cast<UCustomizableObjectNodeTexture>(Node))
@@ -921,7 +923,6 @@ mu::NodeImagePtr GenerateMutableSourceImage(const UEdGraphPin* Pin, FMutableGrap
 		if (Pin->PinType.PinCategory == Schema->PC_Color)
 		{
 			mu::NodeColourPtr ColorNode = GenerateMutableSourceColor(Pin, GenerationContext);
-			
 			mu::NodeImagePlainColourPtr ImageNode = new mu::NodeImagePlainColour;
 
 			ImageNode->SetSize(16, 16);
@@ -929,49 +930,89 @@ mu::NodeImagePtr GenerateMutableSourceImage(const UEdGraphPin* Pin, FMutableGrap
 
 			Result = ImageNode;
 		}
-
 		else
 		{
-			mu::NodeImageTablePtr ImageTableNode = new mu::NodeImageTable();
-			Result = ImageTableNode;
+			//This node will add a checker texture in case of error
+			mu::NodeImageConstantPtr EmptyNode = new mu::NodeImageConstant();
+			Result = EmptyNode;
+
+			bool bSuccess = true;
 
-			mu::TablePtr Table = nullptr;
+			if (Pin->PinType.PinCategory == Schema->PC_MaterialAsset)
+			{
+				// Material pins have to skip the cache of nodes or they will return always the same column node
+				bDoNotAddToGeneratedCache = true;
+			}
 
 			if (TypedNodeTable->Table)
 			{
 				FString ColumnName = Pin->PinFriendlyName.ToString();
+				FProperty* Property = TypedNodeTable->Table->FindTableProperty(FName(*ColumnName));
 
-				if (Pin->PinType.PinCategory == Schema->PC_MaterialAsset)
+				if (!Property)
 				{
-					ColumnName = GenerationContext.CurrentMaterialTableParameterId;
-				}
+					FString Msg = FString::Printf(TEXT("Couldn't find the column [%s] in the data table's struct."), *ColumnName);
+					GenerationContext.Compiler->CompilerLog(FText::FromString(Msg), Node);
 
-				// Generating a new data table if not exists
-				Table = GenerateMutableSourceTable(TypedNodeTable->Table->GetName(), Pin, GenerationContext);
+					bSuccess = false;
+				}
 
-				// Generating a new Texture column if not exists
-				if (Table && Table->FindColumn(StringCast<ANSICHAR>(*ColumnName).Get()) == INDEX_NONE)
+				if (bSuccess && Pin->PinType.PinCategory != Schema->PC_MaterialAsset && !TypedNodeTable->GetColumnDefaultAssetByType<UTexture2D>(Pin))
 				{
-					GenerateTableColumn(TypedNodeTable, Pin, Table, ColumnName, GenerationContext.CurrentLOD, GenerationContext);
+					FString Msg = FString::Printf(TEXT("Couldn't find a default value in the data table's struct for the column [%s]. The default value is null or not a Texture2D"), *ColumnName);
+					GenerationContext.Compiler->CompilerLog(FText::FromString(Msg), Node);
+					
+					bSuccess = false;
 				}
-				
-				ImageTableNode->SetTable(Table);
-				ImageTableNode->SetColumn(StringCast<ANSICHAR>(*ColumnName).Get());
-				ImageTableNode->SetParameterName(StringCast<ANSICHAR>(*TypedNodeTable->ParameterName).Get());
 
-				GenerationContext.AddParameterNameUnique(Node, TypedNodeTable->ParameterName);
-				
-				if (Table->FindColumn(StringCast<ANSICHAR>(*ColumnName).Get()) == INDEX_NONE)
+				if (bSuccess)
 				{
-					FString Msg = FString::Printf(TEXT("Couldn't find pin column with name %s"), *ColumnName);
-					GenerationContext.Compiler->CompilerLog(FText::FromString(Msg), Node);
+					// Generating a new data table if not exists
+					mu::TablePtr Table = nullptr;
+					Table = GenerateMutableSourceTable(TypedNodeTable->Table->GetName(), Pin, GenerationContext);
+
+					if (Table)
+					{
+						mu::NodeImageTablePtr ImageTableNode = new mu::NodeImageTable();
+
+						if (Pin->PinType.PinCategory == Schema->PC_MaterialAsset)
+						{
+							// Material parameters use the parameter id as column names
+							ColumnName = GenerationContext.CurrentMaterialTableParameterId;
+						}
+
+						// Generating a new Texture column if not exists
+						if (Table->FindColumn(StringCast<ANSICHAR>(*ColumnName).Get()) == INDEX_NONE)
+						{
+							bSuccess = GenerateTableColumn(TypedNodeTable, Pin, Table, ColumnName, GenerationContext.CurrentLOD, GenerationContext);
+
+							if (!bSuccess)
+							{
+								FString Msg = FString::Printf(TEXT("Failed to generate the mutable table column [%s]"), *ColumnName);
+								GenerationContext.Compiler->CompilerLog(FText::FromString(Msg), Node);
+							}
+						}
+
+						if (bSuccess)
+						{
+							Result = ImageTableNode;
+
+							ImageTableNode->SetTable(Table);
+							ImageTableNode->SetColumn(StringCast<ANSICHAR>(*ColumnName).Get());
+							ImageTableNode->SetParameterName(StringCast<ANSICHAR>(*TypedNodeTable->ParameterName).Get());
+
+							GenerationContext.AddParameterNameUnique(Node, TypedNodeTable->ParameterName);
+						}
+					}
+					else
+					{
+						FString Msg = FString::Printf(TEXT("Couldn't generate a mutable table."));
+						GenerationContext.Compiler->CompilerLog(FText::FromString(Msg), Node);
+					}
 				}
 			}
 			else
 			{
-				Table = new mu::Table();
-				ImageTableNode->SetTable(Table);
-
 				GenerationContext.Compiler->CompilerLog(LOCTEXT("ImageTableError", "Couldn't find the data table of the node."), Node);
 			}
 		}
@@ -982,8 +1023,11 @@ mu::NodeImagePtr GenerateMutableSourceImage(const UEdGraphPin* Pin, FMutableGrap
 		GenerationContext.Compiler->CompilerLog(LOCTEXT("UnimplementedNode", "Node type not implemented yet."), Node);
 	}
 
-	GenerationContext.Generated.Add(Key, FGeneratedData(Node, Result));
-	GenerationContext.GeneratedNodes.Add(Node);
+	if (!bDoNotAddToGeneratedCache)
+	{
+		GenerationContext.Generated.Add(Key, FGeneratedData(Node, Result));
+		GenerationContext.GeneratedNodes.Add(Node);
+	}
 
 	if (Result)
 	{
diff --git a/Engine/Plugins/Experimental/Mutable/Source/CustomizableObjectEditor/Private/MuCOE/GenerateMutableSource/GenerateMutableSourceMesh.cpp b/Engine/Plugins/Experimental/Mutable/Source/CustomizableObjectEditor/Private/MuCOE/GenerateMutableSource/GenerateMutableSourceMesh.cpp
--- a/Engine/Plugins/Experimental/Mutable/Source/CustomizableObjectEditor/Private/MuCOE/GenerateMutableSource/GenerateMutableSourceMesh.cpp
+++ b/Engine/Plugins/Experimental/Mutable/Source/CustomizableObjectEditor/Private/MuCOE/GenerateMutableSource/GenerateMutableSourceMesh.cpp
@@ -3486,145 +3486,179 @@ mu::NodeMeshPtr GenerateMutableSourceMesh(const UEdGraphPin * Pin,
 
 	else if (const UCustomizableObjectNodeTable* TypedNodeTable = Cast<UCustomizableObjectNodeTable>(Node))
 	{
-		mu::NodeMeshTablePtr MeshTableNode = new mu::NodeMeshTable();
-		Result = MeshTableNode;
-
-		mu::TablePtr Table = nullptr;
+		//This node will add a checker texture in case of error
+		mu::NodeMeshConstantPtr EmptyNode = new mu::NodeMeshConstant();
+		Result = EmptyNode;
+		bool bSuccess = true;
 
 		if (TypedNodeTable->Table)
-		{			
-			USkeletalMesh* SkeletalMesh = TypedNodeTable->GetColumnDefaultAssetByType<USkeletalMesh>(Pin);
-			
-			int32 CurrentLOD = 0;
-			int32 MaterialIndex = 0;
-			
-			TypedNodeTable->GetPinLODAndMaterial(Pin, CurrentLOD, MaterialIndex);
+		{
+			// Getting the real name of the data table column
+			FString DataTableColumnName = TypedNodeTable->GetColumnNameByPin(Pin);
+			FProperty* Property = TypedNodeTable->Table->FindTableProperty(FName(*DataTableColumnName));
 
-			if (GenerationContext.CurrentAutoLODStrategy == ECustomizableObjectAutomaticLODStrategy::AutomaticFromMesh)
+			if (!Property)
 			{
-				if (SkeletalMesh && SkeletalMesh->GetImportedModel())
-				{
-					FSkeletalMeshModel* ImportedModel = SkeletalMesh->GetImportedModel();
-
-					CurrentLOD += GenerationContext.CurrentLOD;
+				FString Msg = FString::Printf(TEXT("Couldn't find the column [%s] in the data table's struct."), *DataTableColumnName);
+				GenerationContext.Compiler->CompilerLog(FText::FromString(Msg), Node);
 
-					// Checking if the current LOD is valid
-					if (ImportedModel->LODModels.Num() <= CurrentLOD
-						||
-						ImportedModel->LODModels[CurrentLOD].Sections.Num() <= MaterialIndex)
-					{
-						// Find the closest valid LOD 
-						for (int32 LODIndex = ImportedModel->LODModels.Num() - 1; LODIndex >= 0; --LODIndex)
-						{
-							if (ImportedModel->LODModels[LODIndex].Sections.Num() > MaterialIndex)
-							{
-								CurrentLOD = LODIndex;
-								break;
-							}
-						}
-					}
-				}
+				bSuccess = false;
 			}
 
-			// Generating a new data table if not exists
-			Table = GenerateMutableSourceTable(TypedNodeTable->Table->GetName(), Pin, GenerationContext);
+			if (bSuccess && !TypedNodeTable->GetColumnDefaultAssetByType<USkeletalMesh>(Pin))
+			{
+				FString Msg = FString::Printf(TEXT("Couldn't find a default value in the data table's struct for the column [%s]."), *DataTableColumnName);
+				GenerationContext.Compiler->CompilerLog(FText::FromString(Msg), Node);
 
-			// Getting mutable and data table column names
-			FString MutableColumnName = TypedNodeTable->GetMutableColumnName(Pin, CurrentLOD);
-			FString DataTableColumnName = TypedNodeTable->GetColumnNameByPin(Pin);
+				bSuccess = false;
+			}
 
-			// Generating a new Mesh column if not exists
-			if (Table && Table->FindColumn(StringCast<ANSICHAR>(*MutableColumnName).Get()) == INDEX_NONE)
+			if (bSuccess)
 			{
-				GenerateTableColumn(TypedNodeTable, Pin, Table, DataTableColumnName, CurrentLOD, GenerationContext);
-			}
+				// Generating a new data table if not exists
+				mu::TablePtr Table = nullptr;
+				Table = GenerateMutableSourceTable(TypedNodeTable->Table->GetName(), Pin, GenerationContext);
 
-			MeshTableNode->SetTable(Table);
-			MeshTableNode->SetColumn(StringCast<ANSICHAR>(*MutableColumnName).Get());
-			MeshTableNode->SetParameterName(StringCast<ANSICHAR>(*TypedNodeTable->ParameterName).Get());
+				if (Table)
+				{
+					mu::NodeMeshTablePtr MeshTableNode = new mu::NodeMeshTable();
+					USkeletalMesh* SkeletalMesh = TypedNodeTable->GetColumnDefaultAssetByType<USkeletalMesh>(Pin);
 
-			GenerationContext.AddParameterNameUnique(Node, TypedNodeTable->ParameterName);
+					int32 CurrentLOD = 0;
+					int32 MaterialIndex = 0;
 
-			if (SkeletalMesh)
-			{
-				// TODO: this should be made for all the meshes of the Column to support meshes with different values
-				// Filling Mesh Data
-				FSkeletalMeshModel* importModel = SkeletalMesh->GetImportedModel();
-				MeshData.bHasVertexColors = SkeletalMesh->GetHasVertexColors();
-				MeshData.NumTexCoordChannels = importModel->LODModels[CurrentLOD].NumTexCoords;
-				MeshData.MaxBoneIndexTypeSizeBytes = importModel->LODModels[CurrentLOD].RequiredBones.Num() > 256 ? 2 : 1;
-				MeshData.MaxNumBonesPerVertex = importModel->LODModels[CurrentLOD].GetMaxBoneInfluences();
+					TypedNodeTable->GetPinLODAndMaterial(Pin, CurrentLOD, MaterialIndex);
 
-				// When mesh data is combined we will get an upper and lower bound of the number of triangles.
-				MeshData.MaxNumTriangles = importModel->LODModels[CurrentLOD].Sections[MaterialIndex].NumTriangles;
-				MeshData.MinNumTriangles = importModel->LODModels[CurrentLOD].Sections[MaterialIndex].NumTriangles;
-			}
+					if (GenerationContext.CurrentAutoLODStrategy == ECustomizableObjectAutomaticLODStrategy::AutomaticFromMesh)
+					{
+						if (SkeletalMesh && SkeletalMesh->GetImportedModel())
+						{
+							FSkeletalMeshModel* ImportedModel = SkeletalMesh->GetImportedModel();
 
-			TArray<UCustomizableObjectLayout*> Layouts = TypedNodeTable->GetLayouts(Pin);
-			MeshTableNode->SetLayoutCount(Layouts.Num());
+							CurrentLOD += GenerationContext.CurrentLOD;
 
-			if (Layouts.Num())
-			{
-				// Generating node Layouts
-				for (int32 i = 0; i < Layouts.Num(); ++i)
-				{
-					mu::NodeLayoutBlocksPtr LayoutNode = new mu::NodeLayoutBlocks;
+							// Checking if the current LOD is valid
+							if (ImportedModel->LODModels.Num() <= CurrentLOD
+								||
+								ImportedModel->LODModels[CurrentLOD].Sections.Num() <= MaterialIndex)
+							{
+								// Find the closest valid LOD 
+								for (int32 LODIndex = ImportedModel->LODModels.Num() - 1; LODIndex >= 0; --LODIndex)
+								{
+									if (ImportedModel->LODModels[LODIndex].Sections.Num() > MaterialIndex)
+									{
+										CurrentLOD = LODIndex;
+										break;
+									}
+								}
+							}
+						}
+					}
 
-					LayoutNode->SetGridSize(Layouts[i]->GetGridSize().X, Layouts[i]->GetGridSize().Y);
-					LayoutNode->SetMaxGridSize(Layouts[i]->GetMaxGridSize().X, Layouts[i]->GetMaxGridSize().Y);
-					LayoutNode->SetBlockCount(Layouts[i]->Blocks.Num() ? Layouts[i]->Blocks.Num() : 1);
-					LayoutNode->SetLayoutPackingStrategy(Layouts[i]->GetPackingStrategy() == ECustomizableObjectTextureLayoutPackingStrategy::Fixed ? mu::EPackStrategy::FIXED_LAYOUT : mu::EPackStrategy::RESIZABLE_LAYOUT);
-					LayoutNode->SetBlockReductionMethod(Layouts[i]->GetBlockReductionMethod() == ECustomizableObjectLayoutBlockReductionMethod::Halve ? mu::EReductionMethod::HALVE_REDUCTION : mu::EReductionMethod::UNITARY_REDUCTION);
+					// Getting the mutable table mesh column name
+					FString MutableColumnName = TypedNodeTable->GetMutableColumnName(Pin, CurrentLOD);
 
-					if (bLinkedToExtendMaterial)
+					// Generating a new Mesh column if not exists
+					if (Table->FindColumn(StringCast<ANSICHAR>(*MutableColumnName).Get()) == INDEX_NONE)
 					{
-						// Layout warnings can be safely ignored in this case. Vertices that do not belong to any layout block will be removed (Extend Materials only)
-						LayoutNode->SetIgnoreWarningsLOD(0);
-					}
+						bSuccess = GenerateTableColumn(TypedNodeTable, Pin, Table, DataTableColumnName, CurrentLOD, GenerationContext);
 
-					if (Layouts[i]->Blocks.Num())
-					{
-						for (int BlockIndex = 0; BlockIndex < Layouts[i]->Blocks.Num(); ++BlockIndex)
+						if (!bSuccess)
 						{
-							LayoutNode->SetBlock(BlockIndex,
-								Layouts[i]->Blocks[BlockIndex].Min.X,
-								Layouts[i]->Blocks[BlockIndex].Min.Y,
-								Layouts[i]->Blocks[BlockIndex].Max.X - Layouts[i]->Blocks[BlockIndex].Min.X,
-								Layouts[i]->Blocks[BlockIndex].Max.Y - Layouts[i]->Blocks[BlockIndex].Min.Y);
-
-							LayoutNode->SetBlockOptions(BlockIndex, Layouts[i]->Blocks[BlockIndex].Priority, Layouts[i]->Blocks[BlockIndex].bUseSymmetry);
+							FString Msg = FString::Printf(TEXT("Failed to generate the mutable table column [%s]"), *MutableColumnName);
+							GenerationContext.Compiler->CompilerLog(FText::FromString(Msg), Node);
 						}
 					}
-					else
+
+					if (bSuccess)
 					{
-						FString msg = "Mesh Column [" + MutableColumnName + "] Layout doesn't has any block. A grid sized block will be used instead.";
-						GenerationContext.Compiler->CompilerLog(FText::FromString(msg), Node, EMessageSeverity::Warning);
+						Result = MeshTableNode;
 
-						LayoutNode->SetBlock(0, 0, 0, Layouts[i]->GetGridSize().X, Layouts[i]->GetGridSize().Y);
-						LayoutNode->SetBlockOptions(0, 0, false);
-					}
+						MeshTableNode->SetTable(Table);
+						MeshTableNode->SetColumn(StringCast<ANSICHAR>(*MutableColumnName).Get());
+						MeshTableNode->SetParameterName(StringCast<ANSICHAR>(*TypedNodeTable->ParameterName).Get());
 
-					MeshTableNode->SetLayout(i, LayoutNode);
-				}
-			}
+						GenerationContext.AddParameterNameUnique(Node, TypedNodeTable->ParameterName);
 
-			// Applying Mesh Morph Nodes
-			if (GenerationContext.MeshMorphStack.Num())
-			{
-				MorphResult = GenerateMorphMesh(Pin, GenerationContext.MeshMorphStack, 0, Result, GenerationContext, MeshData, MutableColumnName);
-			}
+						if (SkeletalMesh)
+						{
+							// TODO: this should be made for all the meshes of the Column to support meshes with different values
+							// Filling Mesh Data
+							FSkeletalMeshModel* importModel = SkeletalMesh->GetImportedModel();
+							MeshData.bHasVertexColors = SkeletalMesh->GetHasVertexColors();
+							MeshData.NumTexCoordChannels = importModel->LODModels[CurrentLOD].NumTexCoords;
+							MeshData.MaxBoneIndexTypeSizeBytes = importModel->LODModels[CurrentLOD].RequiredBones.Num() > 256 ? 2 : 1;
+							MeshData.MaxNumBonesPerVertex = importModel->LODModels[CurrentLOD].GetMaxBoneInfluences();
+
+							// When mesh data is combined we will get an upper and lower bound of the number of triangles.
+							MeshData.MaxNumTriangles = importModel->LODModels[CurrentLOD].Sections[MaterialIndex].NumTriangles;
+							MeshData.MinNumTriangles = importModel->LODModels[CurrentLOD].Sections[MaterialIndex].NumTriangles;
+						}
 
-			if (Table->FindColumn(StringCast<ANSICHAR>(*MutableColumnName).Get()) == INDEX_NONE)
-			{
-				FString Msg = FString::Printf(TEXT("Couldn't find pin column with name %s"), *MutableColumnName);
-				GenerationContext.Compiler->CompilerLog(FText::FromString(Msg), Node);
+						TArray<UCustomizableObjectLayout*> Layouts = TypedNodeTable->GetLayouts(Pin);
+						MeshTableNode->SetLayoutCount(Layouts.Num());
+
+						if (Layouts.Num())
+						{
+							// Generating node Layouts
+							for (int32 i = 0; i < Layouts.Num(); ++i)
+							{
+								mu::NodeLayoutBlocksPtr LayoutNode = new mu::NodeLayoutBlocks;
+
+								LayoutNode->SetGridSize(Layouts[i]->GetGridSize().X, Layouts[i]->GetGridSize().Y);
+								LayoutNode->SetMaxGridSize(Layouts[i]->GetMaxGridSize().X, Layouts[i]->GetMaxGridSize().Y);
+								LayoutNode->SetBlockCount(Layouts[i]->Blocks.Num() ? Layouts[i]->Blocks.Num() : 1);
+								LayoutNode->SetLayoutPackingStrategy(Layouts[i]->GetPackingStrategy() == ECustomizableObjectTextureLayoutPackingStrategy::Fixed ? mu::EPackStrategy::FIXED_LAYOUT : mu::EPackStrategy::RESIZABLE_LAYOUT);
+								LayoutNode->SetBlockReductionMethod(Layouts[i]->GetBlockReductionMethod() == ECustomizableObjectLayoutBlockReductionMethod::Halve ? mu::EReductionMethod::HALVE_REDUCTION : mu::EReductionMethod::UNITARY_REDUCTION);
+
+								if (bLinkedToExtendMaterial)
+								{
+									// Layout warnings can be safely ignored in this case. Vertices that do not belong to any layout block will be removed (Extend Materials only)
+									LayoutNode->SetIgnoreWarningsLOD(0);
+								}
+
+								if (Layouts[i]->Blocks.Num())
+								{
+									for (int BlockIndex = 0; BlockIndex < Layouts[i]->Blocks.Num(); ++BlockIndex)
+									{
+										LayoutNode->SetBlock(BlockIndex,
+											Layouts[i]->Blocks[BlockIndex].Min.X,
+											Layouts[i]->Blocks[BlockIndex].Min.Y,
+											Layouts[i]->Blocks[BlockIndex].Max.X - Layouts[i]->Blocks[BlockIndex].Min.X,
+											Layouts[i]->Blocks[BlockIndex].Max.Y - Layouts[i]->Blocks[BlockIndex].Min.Y);
+
+										LayoutNode->SetBlockOptions(BlockIndex, Layouts[i]->Blocks[BlockIndex].Priority, Layouts[i]->Blocks[BlockIndex].bUseSymmetry);
+									}
+								}
+								else
+								{
+									FString msg = "Mesh Column [" + MutableColumnName + "] Layout doesn't has any block. A grid sized block will be used instead.";
+									GenerationContext.Compiler->CompilerLog(FText::FromString(msg), Node, EMessageSeverity::Warning);
+
+									LayoutNode->SetBlock(0, 0, 0, Layouts[i]->GetGridSize().X, Layouts[i]->GetGridSize().Y);
+									LayoutNode->SetBlockOptions(0, 0, false);
+								}
+
+								MeshTableNode->SetLayout(i, LayoutNode);
+							}
+						}
+
+						// Applying Mesh Morph Nodes
+						if (GenerationContext.MeshMorphStack.Num())
+						{
+							MorphResult = GenerateMorphMesh(Pin, GenerationContext.MeshMorphStack, 0, Result, GenerationContext, MeshData, MutableColumnName);
+						}
+					}
+				}
+				else
+				{
+					FString Msg = FString::Printf(TEXT("Couldn't generate a mutable table."));
+					GenerationContext.Compiler->CompilerLog(FText::FromString(Msg), Node);
+				}
 			}
 		}
 		else
 		{
-			Table = new mu::Table();
-			MeshTableNode->SetTable(Table);
-
 			GenerationContext.Compiler->CompilerLog(LOCTEXT("ImageTableError", "Couldn't find the data table of the node."), Node);
 		}
 	}
diff --git a/Engine/Plugins/Experimental/Mutable/Source/CustomizableObjectEditor/Private/MuCOE/GenerateMutableSource/GenerateMutableSourceSurface.cpp b/Engine/Plugins/Experimental/Mutable/Source/CustomizableObjectEditor/Private/MuCOE/GenerateMutableSource/GenerateMutableSourceSurface.cpp
--- a/Engine/Plugins/Experimental/Mutable/Source/CustomizableObjectEditor/Private/MuCOE/GenerateMutableSource/GenerateMutableSourceSurface.cpp
+++ b/Engine/Plugins/Experimental/Mutable/Source/CustomizableObjectEditor/Private/MuCOE/GenerateMutableSource/GenerateMutableSourceSurface.cpp
@@ -413,48 +413,58 @@ mu::NodeSurfacePtr GenerateMutableSourceSurface(const UEdGraphPin * Pin, FMutabl
 			{
 				FString ColumnName = ConnectedPin->PinFriendlyName.ToString();
 
-				// Generating a new data table if not exists
-				mu::TablePtr Table = GenerateMutableSourceTable(TypedNodeTable->Table->GetName(), ConnectedPin, GenerationContext);
-				
-				// Generating a column for each modified parameter(texture, color & float) if not exists
-				if (Table && Table->FindColumn(StringCast<ANSICHAR>(*ColumnName).Get()) == INDEX_NONE)
+				if (!TypedNodeTable->GetColumnDefaultAssetByType<UMaterialInstance>(ConnectedPin))
 				{
-					GenerateTableColumn(TypedNodeTable, ConnectedPin, Table, ColumnName, GenerationContext.CurrentLOD, GenerationContext);
-				}
+					FString Msg = FString::Printf(TEXT("Couldn't find a default value in the data table's struct for the column [%s]. The default value is null or not a Material Instance."), *ColumnName);
+					GenerationContext.Compiler->CompilerLog(FText::FromString(Msg), Node);
 
-				// Checking if this material has some parameters modified by the table node linked to it
-				if (GenerationContext.GeneratedParametersInTables.Contains(TypedNodeTable))
+					bMaterialPinConnected = false;
+				}
+				else
 				{
-					NodeTableParametersGenerated = GenerationContext.GeneratedParametersInTables[TypedNodeTable];
+					// Generating a new data table if not exists
+					mu::TablePtr Table = GenerateMutableSourceTable(TypedNodeTable->Table->GetName(), ConnectedPin, GenerationContext);
 
-					UMaterialInstance* TableMaterial = TypedNodeTable->GetColumnDefaultAssetByType<UMaterialInstance>(ConnectedPin);
+					// Generating a column for each modified parameter(texture, color & float) if not exists
+					if (Table && Table->FindColumn(StringCast<ANSICHAR>(*ColumnName).Get()) == INDEX_NONE)
+					{
+						GenerateTableColumn(TypedNodeTable, ConnectedPin, Table, ColumnName, GenerationContext.CurrentLOD, GenerationContext);
+					}
 
-					if (TableMaterial && TypedNodeMat->Material)
+					// Checking if this material has some parameters modified by the table node linked to it
+					if (GenerationContext.GeneratedParametersInTables.Contains(TypedNodeTable))
 					{
-						// Checking if the reference material of the Table Node has the same parent as the material of the Material Node 
-						if (TableMaterial->GetMaterial() != TypedNodeMat->Material->GetMaterial())
+						NodeTableParametersGenerated = GenerationContext.GeneratedParametersInTables[TypedNodeTable];
+
+						UMaterialInstance* TableMaterial = TypedNodeTable->GetColumnDefaultAssetByType<UMaterialInstance>(ConnectedPin);
+
+						if (TableMaterial && TypedNodeMat->Material)
+						{
+							// Checking if the reference material of the Table Node has the same parent as the material of the Material Node 
+							if (TableMaterial->GetMaterial() != TypedNodeMat->Material->GetMaterial())
+							{
+								bMaterialPinConnected = false;
+
+								FString msg = FString::Printf(TEXT("Material from NodeTable is an instance from a different Parent"));
+								GenerationContext.Compiler->CompilerLog(FText::FromString(msg), TypedNodeMat);
+							}
+						}
+						else
 						{
 							bMaterialPinConnected = false;
 
-							FString msg = FString::Printf(TEXT("Material from NodeTable is an instance from a different Parent"));
-							GenerationContext.Compiler->CompilerLog(FText::FromString(msg), TypedNodeMat);
+							if (!TypedNodeMat->Material)
+							{
+								FString msg = FString::Printf(TEXT("Need to select a Material to work with Node Table Materials."));
+								GenerationContext.Compiler->CompilerLog(FText::FromString(msg), TypedNodeMat);
+							}
 						}
 					}
 					else
 					{
 						bMaterialPinConnected = false;
-
-						if (!TypedNodeMat->Material)
-						{
-							FString msg = FString::Printf(TEXT("Need to select a Material to work with Node Table Materials."));
-							GenerationContext.Compiler->CompilerLog(FText::FromString(msg), TypedNodeMat);
-						}
 					}
 				}
-				else
-				{
-					bMaterialPinConnected = false;
-				}
 			}
 		}
 
diff --git a/Engine/Plugins/Experimental/Mutable/Source/CustomizableObjectEditor/Private/MuCOE/GenerateMutableSource/GenerateMutableSourceTable.cpp b/Engine/Plugins/Experimental/Mutable/Source/CustomizableObjectEditor/Private/MuCOE/GenerateMutableSource/GenerateMutableSourceTable.cpp
--- a/Engine/Plugins/Experimental/Mutable/Source/CustomizableObjectEditor/Private/MuCOE/GenerateMutableSource/GenerateMutableSourceTable.cpp
+++ b/Engine/Plugins/Experimental/Mutable/Source/CustomizableObjectEditor/Private/MuCOE/GenerateMutableSource/GenerateMutableSourceTable.cpp
@@ -20,7 +20,7 @@
 #define LOCTEXT_NAMESPACE "CustomizableObjectEditor"
 
 
-void FillTableColumn(const UCustomizableObjectNodeTable* TableNode,	mu::TablePtr MutableTable,	const FString& ColumnName,	const FString& RowName,	const int32 RowIdx,	uint8* CellData, FProperty* Property, const int32 LOD, FMutableGraphGenerationContext& GenerationContext)
+bool FillTableColumn(const UCustomizableObjectNodeTable* TableNode,	mu::TablePtr MutableTable,	const FString& ColumnName,	const FString& RowName,	const int32 RowIdx,	uint8* CellData, FProperty* Property, const int32 LOD, FMutableGraphGenerationContext& GenerationContext)
 {
 	int32 CurrentColumn;
 
@@ -29,13 +29,15 @@ void FillTableColumn(const UCustomizableObjectNodeTable* TableNode,	mu::TablePtr
 	{
 		UObject* Object = SoftObjectProperty->GetPropertyValue(CellData).LoadSynchronous();
 
-		if (!Object)
+		if (SoftObjectProperty->PropertyClass->IsChildOf(USkeletalMesh::StaticClass()))
 		{
-			return;
-		}
+			USkeletalMesh* SkeletalMesh = Cast<USkeletalMesh>(Object);
+
+			if(!SkeletalMesh)
+			{
+				return false;
+			}
 
-		if (USkeletalMesh* SkeletalMesh = Cast<USkeletalMesh>(Object))
-		{
 			// Getting Animation Blueprint and Animation Slot
 			FString AnimBP, AnimSlot, GameplayTag, AnimBPAssetTag;
 			TArray<FGameplayTag> GameplayTags;
@@ -153,7 +155,7 @@ void FillTableColumn(const UCustomizableObjectNodeTable* TableNode,	mu::TablePtr
 				FString msg = FString::Printf(TEXT("Reference Skeletal Mesh not found for column [%s]."), *ColumnName);
 				GenerationContext.Compiler->CompilerLog(FText::FromString(msg), TableNode);
 
-				return;
+				return false;
 			}
 
 			// Parameter used for LOD differences
@@ -256,8 +258,15 @@ void FillTableColumn(const UCustomizableObjectNodeTable* TableNode,	mu::TablePtr
 			}
 		}
 
-		else if (UStaticMesh* StaticMesh = Cast<UStaticMesh>(Object))
+		else if (SoftObjectProperty->PropertyClass->IsChildOf(UStaticMesh::StaticClass()))
 		{
+			UStaticMesh* StaticMesh = Cast<UStaticMesh>(Object);
+
+			if (!StaticMesh)
+			{
+				return false;
+			}
+
 			// Getting reference Mesh
 			UStaticMesh* ReferenceStaticMesh = TableNode->GetColumnDefaultAssetByType<UStaticMesh>(ColumnName);
 
@@ -266,7 +275,7 @@ void FillTableColumn(const UCustomizableObjectNodeTable* TableNode,	mu::TablePtr
 				FString msg = FString::Printf(TEXT("Reference Static Mesh not found for column [%s]."), *ColumnName);
 				GenerationContext.Compiler->CompilerLog(FText::FromString(msg), TableNode);
 
-				return;
+				return false;
 			}
 
 			// Parameter used for LOD differences
@@ -324,8 +333,19 @@ void FillTableColumn(const UCustomizableObjectNodeTable* TableNode,	mu::TablePtr
 			}
 		}
 
-		else if (UTexture2D* Texture = Cast<UTexture2D>(Object))
+		else if (SoftObjectProperty->PropertyClass->IsChildOf(UTexture::StaticClass()))
 		{
+			UTexture2D* Texture = Cast<UTexture2D>(Object);
+
+			if (!Texture)
+			{
+				Texture = TableNode->GetColumnDefaultAssetByType<UTexture2D>(ColumnName);
+				FString Message = Cast<UObject>(Object) ? "not a Texture2D" : "null";
+
+				FString msg = FString::Printf(TEXT("Texture from column [%s] row [%s] is %s. The default texture will be used instead."), *ColumnName, *RowName, *Message);
+				GenerationContext.Compiler->CompilerLog(FText::FromString(msg), TableNode);
+			}
+
 			// Getting column index from column name
 			CurrentColumn = MutableTable->FindColumn(StringCast<ANSICHAR>(*ColumnName).Get());
 
@@ -357,8 +377,18 @@ void FillTableColumn(const UCustomizableObjectNodeTable* TableNode,	mu::TablePtr
 			}
 		}
 
-		else if (UMaterialInstance* Material = Cast<UMaterialInstance>(Object))
+		else if (SoftObjectProperty->PropertyClass->IsChildOf(UMaterialInstance::StaticClass()))
 		{
+			UMaterialInstance* Material = Cast<UMaterialInstance>(Object);
+
+			if (!Material)
+			{
+				Material = TableNode->GetColumnDefaultAssetByType<UMaterialInstance>(ColumnName);
+
+				FString msg = FString::Printf(TEXT("Material Instance from column [%s] row [%s] is null. The default Material Instance will be used instead."), *ColumnName, *RowName);
+				GenerationContext.Compiler->CompilerLog(FText::FromString(msg), TableNode);
+			}
+
 			//Adding an empty column for searching purposes
 			if (MutableTable.get()->FindColumn(StringCast<ANSICHAR>(*ColumnName).Get()) == -1)
 			{
@@ -379,7 +409,7 @@ void FillTableColumn(const UCustomizableObjectNodeTable* TableNode,	mu::TablePtr
 				FString msg = FString::Printf(TEXT("Reference Material not found for column [%s]."), *ColumnName);
 				GenerationContext.Compiler->CompilerLog(FText::FromString(msg), TableNode);
 
-				return;
+				return false;
 			}
 
 			if (ReferenceMaterial->GetMaterial() == Material->GetMaterial())
@@ -395,7 +425,7 @@ void FillTableColumn(const UCustomizableObjectNodeTable* TableNode,	mu::TablePtr
 				
 				for (FTextureParameterValue ReferenceTexture : ReferenceMaterial->TextureParameterValues)
 				{
-					if (Cast<UTexture2D>(ReferenceTexture.ParameterValue))
+					if (UTexture2D* Texture = Cast<UTexture2D>(ReferenceTexture.ParameterValue))
 					{
 						int32 ColumnIndex;
 
@@ -422,36 +452,40 @@ void FillTableColumn(const UCustomizableObjectNodeTable* TableNode,	mu::TablePtr
 							ColumnIndex = MutableTable->AddColumn(StringCast<ANSICHAR>(*ParameterGuid).Get(), mu::TABLE_COLUMN_TYPE::TCT_IMAGE);
 						}
 
-						UTexture* OutTexture = nullptr;
-						UTexture2D* TextureToConvert = nullptr;
-						
-						// Getting the parameter value from the parent material
-						if (Material->GetMaterial()->GetTextureParameterValue(ReferenceTexture.ParameterInfo.Name, OutTexture))
-						{
-							TextureToConvert = Cast<UTexture2D>(OutTexture);
-						}
-
 						// Getting the parameter value from the instance if it has been modified
 						for (FTextureParameterValue InstanceTexture : Material->TextureParameterValues)
 						{
 							if (InstanceTexture.ParameterInfo.Name == ReferenceTexture.ParameterInfo.Name)
 							{
-								TextureToConvert = Cast<UTexture2D>(InstanceTexture.ParameterValue);
+								if (UTexture2D* ParameterTexture = Cast<UTexture2D>(InstanceTexture.ParameterValue))
+								{
+									Texture = ParameterTexture;
+								}
+								else
+								{
+									FString ParamName = ReferenceTexture.ParameterInfo.Name.ToString();
+									FString Message = Cast<UObject>(InstanceTexture.ParameterValue) ? "not a Texture2D" : "null";
+
+									FString msg = FString::Printf(TEXT("Parameter [%s] from material instance of column [%s] row [%s] is %s. The parameter texture of the default material will be used instead."), *ParamName,  *ColumnName, *RowName, *Message);
+									GenerationContext.Compiler->CompilerLog(FText::FromString(msg), TableNode);
+								}
+
 								ModifiedTextureParameters++;
 								
 								break;
 							}
 						}
 
-						if (TextureToConvert)
-						{
-							GenerationContext.ArrayTextureUnrealToMutableTask.Add(FTextureUnrealToMutableTask(MutableTable, TextureToConvert, TableNode, ColumnIndex, RowIdx));
-						}
-						else
-						{
-							FString msg = FString::Printf(TEXT("Didn't find Material Texture Parameter [%s] from parent Material of column [%s] row [%s]."), *TextureParameterName, *ColumnName, *RowName);
-							GenerationContext.Compiler->CompilerLog(FText::FromString(msg), TableNode);
-						}
+						check(Texture);
+						GenerationContext.ArrayTextureUnrealToMutableTask.Add(FTextureUnrealToMutableTask(MutableTable, Texture, TableNode, ColumnIndex, RowIdx));
+					}
+					else
+					{
+						FString ParamName = ReferenceTexture.ParameterInfo.Name.ToString();
+						FString Message = Cast<UObject>(ReferenceTexture.ParameterValue) ? "not a Texture2D" : "null";
+
+						FString msg = FString::Printf(TEXT("Parameter [%s] from Default Material Instance of column [%s] is %s. This parameter will be ignored."), *ParamName, *ColumnName, *Message);
+						GenerationContext.Compiler->CompilerLog(FText::FromString(msg), TableNode);
 					}
 				}
 
@@ -475,8 +509,10 @@ void FillTableColumn(const UCustomizableObjectNodeTable* TableNode,	mu::TablePtr
 
 				for (FVectorParameterValue ReferenceVector : ReferenceMaterial->VectorParameterValues)
 				{
-					int32 ColumnIndex;
+					// Getting the parameter value from the default material
+					FLinearColor VectorValue = ReferenceVector.ParameterValue;
 
+					int32 ColumnIndex;
 					FString ParameterGuid;
 
 					for (int32 i = 0; i < ParameterInfos.Num(); ++i)
@@ -499,11 +535,6 @@ void FillTableColumn(const UCustomizableObjectNodeTable* TableNode,	mu::TablePtr
 						ColumnIndex = MutableTable->AddColumn(StringCast<ANSICHAR>(*ParameterGuid).Get(), mu::TABLE_COLUMN_TYPE::TCT_COLOUR);
 					}
 
-					// Getting the parameter value from the parent material
-					FLinearColor VectorValue;
-
-					Material->GetMaterial()->GetVectorParameterValue(ReferenceVector.ParameterInfo.Name, VectorValue);
-					
 					// Getting the parameter value from the instance if it has been modified
 					for (FVectorParameterValue InstanceVector : Material->VectorParameterValues)
 					{
@@ -541,8 +572,10 @@ void FillTableColumn(const UCustomizableObjectNodeTable* TableNode,	mu::TablePtr
 
 				for (FScalarParameterValue ReferenceScalar : ReferenceMaterial->ScalarParameterValues)
 				{
-					int32 ColumnIndex;
+					// Getting the parameter value from the parent material
+					float ScalarValue = ReferenceScalar.ParameterValue;
 
+					int32 ColumnIndex;
 					FString ParameterGuid;
 
 					for (int32 i = 0; i < ParameterInfos.Num(); ++i)
@@ -565,10 +598,6 @@ void FillTableColumn(const UCustomizableObjectNodeTable* TableNode,	mu::TablePtr
 						ColumnIndex = MutableTable->AddColumn(StringCast<ANSICHAR>(*ParameterGuid).Get(), mu::TABLE_COLUMN_TYPE::TCT_SCALAR);
 					}
 
-					// Getting the parameter value from the parent material
-					float ScalarValue;
-					Material->GetMaterial()->GetScalarParameterValue(ReferenceScalar.ParameterInfo.Name, ScalarValue);
-					
 					// Getting the parameter value from the instance if it has been modified
 					for (FScalarParameterValue InstanceScalar : Material->ScalarParameterValues)
 					{
@@ -602,6 +631,12 @@ void FillTableColumn(const UCustomizableObjectNodeTable* TableNode,	mu::TablePtr
 				GenerationContext.Compiler->CompilerLog(FText::FromString(msg), TableNode);
 			}
 		}
+
+		else
+		{
+			// Unsuported Variable Type
+			return false;
+		}
 	}
 
 	else if (const FStructProperty* StructProperty = CastField<FStructProperty>(Property))
@@ -619,6 +654,12 @@ void FillTableColumn(const UCustomizableObjectNodeTable* TableNode,	mu::TablePtr
 			FLinearColor Value = *(FLinearColor*)CellData;
 			MutableTable->SetCell(CurrentColumn, RowIdx, Value.R, Value.G, Value.B, Value.A);
 		}
+		
+		else
+		{
+			// Unsuported Variable Type
+			return false;
+		}
 	}
 
 	else if (const FNumericProperty* FloatNumProperty = CastField<FFloatProperty>(Property))
@@ -648,59 +689,61 @@ void FillTableColumn(const UCustomizableObjectNodeTable* TableNode,	mu::TablePtr
 		float Value = DoubleNumProperty->GetFloatingPointPropertyValue(CellData);
 		MutableTable->SetCell(CurrentColumn, RowIdx, Value);
 	}
+
+	else
+	{
+		// Unsuported Variable Type
+		return false;
+	}
+
+	return true;
 }
 
 
-void GenerateTableColumn(const UCustomizableObjectNodeTable* TableNode, const UEdGraphPin* Pin, mu::TablePtr MutableTable, const FString& DataTableColumnName, const int32 LOD, FMutableGraphGenerationContext& GenerationContext)
+bool GenerateTableColumn(const UCustomizableObjectNodeTable* TableNode, const UEdGraphPin* Pin, mu::TablePtr MutableTable, const FString& DataTableColumnName, const int32 LOD, FMutableGraphGenerationContext& GenerationContext)
 {
 	SCOPED_PIN_DATA(GenerationContext, Pin)
 
-	if (!TableNode)
+	bool bSuccess = false;
+
+	if (!TableNode || !TableNode->Table || !TableNode->Table->GetRowStruct())
 	{
-		return;
+		return false;
 	}
 
-	UDataTable* DataTable = TableNode->Table;
+	FProperty* ColumnProperty = TableNode->Table->FindTableProperty(FName(*DataTableColumnName));
 
-	if (!DataTable)
+	if (!ColumnProperty)
 	{
-		return;
+		return false;
 	}
+	
+	// Getting names of the rows to access the information
+	TArray<FName> RowNames = TableNode->GetRowNames();
 
-	const UScriptStruct* TableStruct = DataTable->GetRowStruct();
-
-	if (TableStruct)
+	for (int32 RowIndex = 0; RowIndex < RowNames.Num(); ++RowIndex)
 	{
-		// Getting names of the rows to access the information
-		TArray<FName> RowNames = TableNode->GetRowNames();
+		// Getting Row Data
+		uint8* RowData = TableNode->Table->FindRowUnchecked(RowNames[RowIndex]);
 
-		for (TFieldIterator<FProperty> It(TableStruct); It; ++It)
+		if (RowData)
 		{
-			FProperty* ColumnProperty = *It;
-
-			if (!ColumnProperty || DataTableColumnName != DataTableUtils::GetPropertyExportName(ColumnProperty))
-			{
-				continue;
-			}
+			// Getting Cell Data
+			uint8* CellData = ColumnProperty->ContainerPtrToValuePtr<uint8>(RowData, 0);
 
-			for (int32 RowIndex = 0; RowIndex < RowNames.Num(); ++RowIndex)
+			if (CellData)
 			{
-				// Getting Row Data
-				uint8* RowData = DataTable->FindRowUnchecked(RowNames[RowIndex]);
-
-				if (RowData)
+				bool bColumnGenerated = FillTableColumn(TableNode, MutableTable, DataTableColumnName, RowNames[RowIndex].ToString(), RowIndex, CellData, ColumnProperty, LOD, GenerationContext);
+				
+				if (!bSuccess)
 				{
-					// Getting Cell Data
-					uint8* CellData = ColumnProperty->ContainerPtrToValuePtr<uint8>(RowData, 0);
-
-					if (CellData)
-					{
-						FillTableColumn(TableNode, MutableTable, DataTableColumnName, RowNames[RowIndex].ToString(), RowIndex, CellData, ColumnProperty, LOD, GenerationContext);
-					}
+					bSuccess = bColumnGenerated;
 				}
 			}
 		}
 	}
+
+	return bSuccess;
 }
 
 
@@ -728,7 +771,7 @@ mu::TablePtr GenerateMutableSourceTable(const FString& TableName, const UEdGraph
 			FString msg = "Couldn't find the Data Table asset in the Node.";
 			GenerationContext.Compiler->CompilerLog(FText::FromString(msg), Node);
 
-			return MutableTable;
+			return nullptr;
 		}
 
 		const UScriptStruct* TableStruct = Table->GetRowStruct();
@@ -764,11 +807,19 @@ mu::TablePtr GenerateMutableSourceTable(const FString& TableName, const UEdGraph
 
 			GenerationContext.ParameterUIDataMap.Add(TypedTable->ParameterName, ParameterUIData);
 		}
+		else
+		{
+			FString msg = "Couldn't find the Data Table's Struct asset in the Node.";
+			GenerationContext.Compiler->CompilerLog(FText::FromString(msg), Node);
+			
+			return nullptr;
+		}
 	}
-	
 	else
 	{
 		GenerationContext.Compiler->CompilerLog(LOCTEXT("UnimplementedNode", "Node type not implemented yet."), Node);
+
+		return nullptr;
 	}
 
 	GenerationContext.GeneratedTables.Add(TableName, MutableTable);
diff --git a/Engine/Plugins/Experimental/Mutable/Source/CustomizableObjectEditor/Private/MuCOE/GenerateMutableSource/GenerateMutableSourceTable.h b/Engine/Plugins/Experimental/Mutable/Source/CustomizableObjectEditor/Private/MuCOE/GenerateMutableSource/GenerateMutableSourceTable.h
--- a/Engine/Plugins/Experimental/Mutable/Source/CustomizableObjectEditor/Private/MuCOE/GenerateMutableSource/GenerateMutableSourceTable.h
+++ b/Engine/Plugins/Experimental/Mutable/Source/CustomizableObjectEditor/Private/MuCOE/GenerateMutableSource/GenerateMutableSourceTable.h
@@ -11,9 +11,9 @@ class UCustomizableObjectNodeTable;
 class UEdGraphPin;
 struct FMutableGraphGenerationContext;
 
-void GenerateTableColumn(const UCustomizableObjectNodeTable* TableNode, const UEdGraphPin* Pin, mu::TablePtr MutableTable, const FString& DataTableColumnName, const int32 LOD, FMutableGraphGenerationContext& GenerationContext);
+bool GenerateTableColumn(const UCustomizableObjectNodeTable* TableNode, const UEdGraphPin* Pin, mu::TablePtr MutableTable, const FString& DataTableColumnName, const int32 LOD, FMutableGraphGenerationContext& GenerationContext);
 
-void FillTableColumn(const UCustomizableObjectNodeTable* TableNode, mu::TablePtr MutableTable, const FString& ColumnName,
+bool FillTableColumn(const UCustomizableObjectNodeTable* TableNode, mu::TablePtr MutableTable, const FString& ColumnName,
 	const FString& RowName, const int32 RowIdx, uint8* CellData, FProperty* Property, const int32 LOD, FMutableGraphGenerationContext& GenerationContext);
 
 mu::TablePtr GenerateMutableSourceTable(const FString& TableName, const UEdGraphPin* Pin, FMutableGraphGenerationContext& GenerationContext);
