#rnx

diff --git a/Engine/Plugins/Experimental/Mutable/Source/CustomizableObject/Private/MuCO/UnrealConversionUtils.cpp b/Engine/Plugins/Experimental/Mutable/Source/CustomizableObject/Private/MuCO/UnrealConversionUtils.cpp
--- a/Engine/Plugins/Experimental/Mutable/Source/CustomizableObject/Private/MuCO/UnrealConversionUtils.cpp
+++ b/Engine/Plugins/Experimental/Mutable/Source/CustomizableObject/Private/MuCO/UnrealConversionUtils.cpp
@@ -489,22 +489,22 @@ namespace UnrealConversionUtils
 			  int32 NumBoneInfluences = SrcSkinWeightBuffer.GetDataVertexBuffer()->GetMaxBoneInfluences();
 			 int32 NumBones = SrcSkinWeightBuffer.GetDataVertexBuffer()->GetNumBoneWeights();
 
-			 const_cast<FSkinWeightDataVertexBuffer*>(DestSkinWeightBuffer.GetDataVertexBuffer())->SetMaxBoneInfluences(NumBoneInfluences);
-			 const_cast<FSkinWeightDataVertexBuffer*>(DestSkinWeightBuffer.GetDataVertexBuffer())->Init(NumBones, NumVertices);
+			 DestSkinWeightBuffer.SetUse16BitBoneIndex(SrcSkinWeightBuffer.Use16BitBoneIndex());
+			 FSkinWeightDataVertexBuffer* SkinWeightDataVertexBuffer = DestSkinWeightBuffer.GetDataVertexBuffer();
+			 SkinWeightDataVertexBuffer->SetMaxBoneInfluences(NumBoneInfluences);
+			 SkinWeightDataVertexBuffer->Init(NumBones, NumVertices);
 
 			 if (NumVertices)
 			 {
 				 DestSkinWeightBuffer.SetNeedsCPUAccess(bAllowCPUAccess);
 
 				 const uint8* SrcData = SrcSkinWeightBuffer.GetDataVertexBuffer()->GetWeightData();
-				 uint8* Data = DestSkinWeightBuffer.GetDataVertexBuffer()->GetWeightData();
+				 uint8* Data = SkinWeightDataVertexBuffer->GetWeightData();
 				 check(SrcData);
 				 check(Data);
 
 				 FMemory::Memcpy(Data, SrcData, DestSkinWeightBuffer.GetVertexDataSize());
 			 }
-
-			 DestSkinWeightBuffer.SetUse16BitBoneIndex(SrcSkinWeightBuffer.Use16BitBoneIndex());
 		 }
 
 		 // Copying Skin Weight Profiles Buffers
