[Mutable] Fix issue with OpMeshMerge not applying the right format for bone index buffers

#rnx

diff --git a/Engine/Plugins/Experimental/Mutable/Source/MutableRuntime/Private/MuR/OpMeshMerge.h b/Engine/Plugins/Experimental/Mutable/Source/MutableRuntime/Private/MuR/OpMeshMerge.h
--- a/Engine/Plugins/Experimental/Mutable/Source/MutableRuntime/Private/MuR/OpMeshMerge.h
+++ b/Engine/Plugins/Experimental/Mutable/Source/MutableRuntime/Private/MuR/OpMeshMerge.h
@@ -762,23 +762,48 @@ namespace mu
 			bool bChangeBoneIndicesFormat = false;
 			MESH_BUFFER_FORMAT BoneIndexFormat = MaxBoneMapIndices > MAX_uint8 ? MBF_UINT16 : MBF_UINT8;
 
+			// Iterate all vertex buffers to check if we need to format bone indices
 			{
-				// Iterate all vertex buffers 
-				FMeshBufferSet& VertexBuffers = pResult->GetVertexBuffers();
-				for (int32 VertexBufferIndex = 0; !bChangeBoneIndicesFormat &&  VertexBufferIndex < VertexBuffers.m_buffers.Num(); ++VertexBufferIndex)
 				{
-					MESH_BUFFER& Buffer = VertexBuffers.m_buffers[VertexBufferIndex];
+					// Check if pFirst requires a reformat of the bone index buffers
+					const FMeshBufferSet& VertexBuffers = pFirst->GetVertexBuffers();
+					for (int32 VertexBufferIndex = 0; !bChangeBoneIndicesFormat && VertexBufferIndex < VertexBuffers.m_buffers.Num(); ++VertexBufferIndex)
+					{
+						const MESH_BUFFER& Buffer = VertexBuffers.m_buffers[VertexBufferIndex];
 
-					const int32 elemSize = VertexBuffers.GetElementSize(VertexBufferIndex);
-					const int32 firstSize = FirstCount * elemSize;
+						const int32 elemSize = VertexBuffers.GetElementSize(VertexBufferIndex);
+						const int32 firstSize = FirstCount * elemSize;
 
-					const int32 ChannelsCount = VertexBuffers.GetBufferChannelCount(VertexBufferIndex);
-					for (int32 ChannelIndex = 0; ChannelIndex < ChannelsCount; ++ChannelIndex)
+						const int32 ChannelsCount = VertexBuffers.GetBufferChannelCount(VertexBufferIndex);
+						for (int32 ChannelIndex = 0; ChannelIndex < ChannelsCount; ++ChannelIndex)
+						{
+							if (Buffer.m_channels[ChannelIndex].m_semantic == MBS_BONEINDICES)
+							{
+								bChangeBoneIndicesFormat = Buffer.m_channels[ChannelIndex].m_format != BoneIndexFormat;
+								break;
+							}
+						}
+					}
+				}
+
+				{
+					// Check if pSecond requires a reformat of the bone index buffers
+					const FMeshBufferSet& VertexBuffers = pSecond->GetVertexBuffers();
+					for (int32 VertexBufferIndex = 0; !bChangeBoneIndicesFormat && VertexBufferIndex < VertexBuffers.m_buffers.Num(); ++VertexBufferIndex)
 					{
-						if (Buffer.m_channels[ChannelIndex].m_semantic == MBS_BONEINDICES)
+						const MESH_BUFFER& Buffer = VertexBuffers.m_buffers[VertexBufferIndex];
+
+						const int32 elemSize = VertexBuffers.GetElementSize(VertexBufferIndex);
+						const int32 firstSize = FirstCount * elemSize;
+
+						const int32 ChannelsCount = VertexBuffers.GetBufferChannelCount(VertexBufferIndex);
+						for (int32 ChannelIndex = 0; ChannelIndex < ChannelsCount; ++ChannelIndex)
 						{
-							bChangeBoneIndicesFormat = Buffer.m_channels[ChannelIndex].m_format != BoneIndexFormat;
-							break;
+							if (Buffer.m_channels[ChannelIndex].m_semantic == MBS_BONEINDICES)
+							{
+								bChangeBoneIndicesFormat = Buffer.m_channels[ChannelIndex].m_format != BoneIndexFormat;
+								break;
+							}
 						}
 					}
 				}
