[Mutable] FixPin data incorrect check.

#rnx

diff --git a/Engine/Plugins/Experimental/Mutable/Source/CustomizableObjectEditor/Private/MuCOE/Nodes/CustomizableObjectNode.cpp b/Engine/Plugins/Experimental/Mutable/Source/CustomizableObjectEditor/Private/MuCOE/Nodes/CustomizableObjectNode.cpp
--- a/Engine/Plugins/Experimental/Mutable/Source/CustomizableObjectEditor/Private/MuCOE/Nodes/CustomizableObjectNode.cpp
+++ b/Engine/Plugins/Experimental/Mutable/Source/CustomizableObjectEditor/Private/MuCOE/Nodes/CustomizableObjectNode.cpp
@@ -458,7 +458,7 @@ void UCustomizableObjectNode::RemapPinsData(const TMap<UEdGraphPin*, UEdGraphPin
 
 void UCustomizableObjectNode::AddPinData(const UEdGraphPin& Pin, UCustomizableObjectNodePinData& PinData)
 {
-	check(PinData.GetOuter())
+	check(PinData.GetOuter() != GetTransientPackage())
 	PinsDataId.Add(Pin.PinId, &PinData);
 }
 
