[Mutable] Fix issues with Texture Parameters.

This CL solves two main issues:
- Texture Parameters where choosing the incorrect texture. Solved by not reusing ids.
- Texture Parameter were released to early. Solved by added a reference counter system.

Changes:
- Image Provider now cache the images automatically when there is an update. The user still can force to cache/uncache images at any given time.
- Add Core FString serialization functions.
- Changed the type of EXTERNAL_IMAGE_ID.
- Used changed hardcoded types to EXTERNAL_IMAGE_ID.
- Removed SetTextureParameterSelectedOptionT from CO API.

Lifetime of a Texture Parameter:
1. UpdateSkeletalMeshAsync called
2. All Texture Parameters inside Parameters are cached (System reference = 1)
3. MIP Data Provider is created (System reference = 2)
4. FinishGlobalUpdate (System reference = 1)
5. MIP Data Provider is destoryed (System reference = 0). Uncaches Texture Parameters.

#rnx

diff --git a/Engine/Plugins/Experimental/Mutable/Source/CustomizableObject/Classes/MuCO/CustomizableObject.h b/Engine/Plugins/Experimental/Mutable/Source/CustomizableObject/Classes/MuCO/CustomizableObject.h
--- a/Engine/Plugins/Experimental/Mutable/Source/CustomizableObject/Classes/MuCO/CustomizableObject.h
+++ b/Engine/Plugins/Experimental/Mutable/Source/CustomizableObject/Classes/MuCO/CustomizableObject.h
@@ -1380,9 +1380,9 @@ public:
 	
 	/** Get the default value of a parameter of type Texture.
 	 * @param InParameterName The name of the Projector parameter to get the default value of.
-	 * @return An uint64 (EXTERNAL_IMAGE_ID) pointing at the default parameter's texture.
+	 * @return An id representing the default parameter's texture.
 	 */
-	uint64 GetTextureParameterDefaultValue (const FString& InParameterName) const;
+	FString GetTextureParameterDefaultValue (const FString& InParameterName) const;
 
 	/** Return true or false depending if the parameter at the index provided is multidimensional or not.
 	 * @param InParamIndex The index of the parameter to check.
@@ -1413,7 +1413,7 @@ private:
 	// This is a manual version number for the binary blobs in this asset.
 	// Increasing it invalidates all the previously compiled models.
 	// Warning: If while merging code both versions have changed, take the highest+1.
-	static const int32 CurrentSupportedVersion = 385;
+	static const int32 CurrentSupportedVersion = 386;
 
 public:
 
diff --git a/Engine/Plugins/Experimental/Mutable/Source/CustomizableObject/Classes/MuCO/CustomizableObjectInstance.h b/Engine/Plugins/Experimental/Mutable/Source/CustomizableObject/Classes/MuCO/CustomizableObjectInstance.h
--- a/Engine/Plugins/Experimental/Mutable/Source/CustomizableObject/Classes/MuCO/CustomizableObjectInstance.h
+++ b/Engine/Plugins/Experimental/Mutable/Source/CustomizableObject/Classes/MuCO/CustomizableObjectInstance.h
@@ -8,8 +8,8 @@
 
 #include "CustomizableObjectInstance.generated.h"
 
+class FCustomizableObjectSystemPrivate;
 class USkeletalMesh;
-
 class AActor;
 class FProperty;
 class UAnimInstance;
@@ -359,16 +359,14 @@ public:
 	void SetFloatParameterSelectedOption(const FString& FloatParamName, float FloatValue, int32 RangeIndex = -1);
 	
 	/** Gets the value of a texture parameter with name "TextureParamName". */
-	uint64 GetTextureParameterSelectedOption(const FString& TextureParamName, int32 RangeIndex = -1) const;
-
-	/** Gets the texture of a texture parameter with name "TextureParamName". */
 	UFUNCTION(BlueprintCallable, Category = CustomizableObjectInstance)
-	UTexture2D* GetTextureParameterSelectedOptionT(const FString& TextureParamName, int32 RangeIndex = -1) const;
+	FString GetTextureParameterSelectedOption(const FString& TextureParamName, int32 RangeIndex = -1) const;
 
 	/** Sets the texture value "TextureValue" of a texture parameter with index "TextureParamIndex". */
-	void SetTextureParameterSelectedOption(const FString& TextureParamName, uint64 TextureValue, int32 RangeIndex = -1);
+	UFUNCTION(BlueprintCallable, Category = CustomizableObjectInstance)
+	void SetTextureParameterSelectedOption(const FString& TextureParamName, const FString& TextureValue, int32 RangeIndex = -1);
 
-	/** Sets the texture "Texture" of a texture parameter with index "TextureParamIndex". */
+	/** @deprecated Use SetTextureParameterSelectedOption instead. */
 	UFUNCTION(BlueprintCallable, Category = CustomizableObjectInstance)
 	void SetTextureParameterSelectedOptionT(const FString& TextureParamName, UTexture2D* TextureValue, int32 RangeIndex = -1);
 	
@@ -706,7 +704,7 @@ private:
 
 	UPROPERTY( Transient )
 	TObjectPtr<UCustomizableInstancePrivateData> PrivateData;
-
+	
 	/** Hash of the UCustomizableObjectInstance::Descriptor on the last update request. */
 	FDescriptorRuntimeHash UpdateDescriptorRuntimeHash;
 	
diff --git a/Engine/Plugins/Experimental/Mutable/Source/CustomizableObject/Classes/MuCO/CustomizableObjectSystem.h b/Engine/Plugins/Experimental/Mutable/Source/CustomizableObject/Classes/MuCO/CustomizableObjectSystem.h
--- a/Engine/Plugins/Experimental/Mutable/Source/CustomizableObject/Classes/MuCO/CustomizableObjectSystem.h
+++ b/Engine/Plugins/Experimental/Mutable/Source/CustomizableObject/Classes/MuCO/CustomizableObjectSystem.h
@@ -7,6 +7,8 @@
 #include "AssetRegistry/AssetData.h"
 #include "MuCO/CustomizableObjectInstance.h"
 #include "MuCO/CustomizableObjectInstanceDescriptor.h"
+#include "MuR/Parameters.h"
+#include "MuR/Types.h"
 
 #if WITH_EDITOR
 #include "Framework/Notifications/NotificationManager.h"
@@ -134,11 +136,11 @@ struct FCustomizableObjectExternalTexture
 	FCustomizableObjectExternalTexture() = default;
 
 	FString Name;
-	int64_t Value = 0;
+	FString Value;
 };
 
 
-//
+/** Base class for Image provider. */
 UCLASS(abstract)
 class CUSTOMIZABLEOBJECT_API UCustomizableSystemImageProvider : public UObject
 {
@@ -165,17 +167,17 @@ public:
 
 	// Query that Mutable will run to find out if a texture will be provided as an Unreal UTexture2D,
 	// or as a raw data blob.
-	virtual ValueType HasTextureParameterValue(int64 ID) { return ValueType::None; }
+	virtual ValueType HasTextureParameterValue(FString ID) { return ValueType::None; }
 
 	// In case IsTextureParameterValueUnreal returns false, this will be used to query the texture size data.
-	virtual FIntVector GetTextureParameterValueSize(int64 ID) { return FIntVector(0, 0, 0); }
+	virtual FIntVector GetTextureParameterValueSize(FString ID) { return FIntVector(0, 0, 0); }
 
 	// In case IsTextureParameterValueUnreal returns false, this will be used to query the texture data that must
 	// be copied in the preallocated buffer. The pixel format is assumed to be 4-channel RGBA, uint8_t per channel.
-	virtual void GetTextureParameterValueData(int64 ID, uint8* OutData) {}
+	virtual void GetTextureParameterValueData(FString ID, uint8* OutData) {}
 
 	// In case IsTextureParameterValueUnreal returns true, this will be used to query the texture.
-	virtual UTexture2D* GetTextureParameterValue(int64 ID) { return nullptr; }
+	virtual UTexture2D* GetTextureParameterValue(FString ID) { return nullptr; }
 
 	// Used in the editor to show the list of available options.
 	// Only necessary if the images are required in editor previews.
@@ -184,7 +186,7 @@ public:
 
 
 /** End a Customizable Object Instance Update. All code paths of an update have to end here. */
-void FinishUpdateGlobal(UCustomizableObjectInstance* Instance, EUpdateResult UpdateResult, FInstanceUpdateDelegate* UpdateCallback, const FDescriptorRuntimeHash InUpdatedHash = FDescriptorRuntimeHash());
+void FinishUpdateGlobal(UCustomizableObjectInstance* Instance, EUpdateResult UpdateResult, FInstanceUpdateDelegate* UpdateCallback, const mu::ParametersPtrConst Parameters = nullptr, const FDescriptorRuntimeHash InUpdatedHash = FDescriptorRuntimeHash());
 
 
 UCLASS(Blueprintable, BlueprintType)
@@ -196,10 +198,11 @@ public:
 	UCustomizableObjectSystem() = default;
 	void InitSystem();
 
-	// Get the singleton object. It will be created if it doesn't exist yet.
+	/** Get the singleton object. It will be created if it doesn't exist yet.
+	 * @param bCreate Create a system if it does not has been created yet. */
 	UFUNCTION(BlueprintCallable, BlueprintPure, Category = Status)
 	static UCustomizableObjectSystem* GetInstance();
-
+	
 	// Return true if the singleton has been created. It is different than GetInstance in that GetInstance will create it if it doesn't exist.
 	static bool IsCreated();
 
@@ -285,18 +288,12 @@ public:
 		memory is used and the real texel data is loaded when needed during an update and then immediately discarded */
 
 	/** [Texture Parameters] Cache an image which has to have been previously registered by an Image provider with the parameter id. */
-	void CacheImage(uint64 ImageId);
+	void CacheImage(FString ImageId);
 	/** [Texture Parameters] Remove an image from the cache. */
-	void UnCacheImage(uint64 ImageId);
-	/** [Texture Parameters] Cache all images which have been previously registered by all registered Image provider with the parameter id
-		that was used in the provider. If bClearPreviousCacheImages is true, then all previous cache state is cleared */
-	void CacheAllImagesInAllProviders(bool bClearPreviousCacheImages);
+	void UnCacheImage(FString ImageId);
 	/** [Texture Parameters] Remove all images from the cache. */
 	void ClearImageCache();
 
-	/** Get the default image provider. Returns null if it was not already initialized. */
-	TObjectPtr<UDefaultImageProvider> GetDefaultImageProvider() const;
-
 	/** Initialize (if was not already) and get the default image provider. */
 	UDefaultImageProvider& GetOrCreateDefaultImageProvider();
 
diff --git a/Engine/Plugins/Experimental/Mutable/Source/CustomizableObject/Private/MuCO/CustomizableInstancePrivateData.h b/Engine/Plugins/Experimental/Mutable/Source/CustomizableObject/Private/MuCO/CustomizableInstancePrivateData.h
--- a/Engine/Plugins/Experimental/Mutable/Source/CustomizableObject/Private/MuCO/CustomizableInstancePrivateData.h
+++ b/Engine/Plugins/Experimental/Mutable/Source/CustomizableObject/Private/MuCO/CustomizableInstancePrivateData.h
@@ -193,6 +193,11 @@ public:
 	// Only used in LiveUpdateMode to reuse core instances between updates and their temp data to speed up updates, but spend way more memory
 	mu::Instance::ID LiveUpdateModeInstanceID = 0;
 
+	/** Cached Texture Parameters from the last update.
+	 * This cache is required since the Instance can have a LOD Update at any time.
+	 * So we need to make sure that the initially provided Texture Parameters by the user will be available until the user decides to change them. */
+	TArray<FString> UpdateTextureParameters;
+	
 	FCustomizableInstanceComponentData* GetComponentData(int32 ComponentIndex);
 	const FCustomizableInstanceComponentData* GetComponentData(int32 ComponentIndex) const;
 
diff --git a/Engine/Plugins/Experimental/Mutable/Source/CustomizableObject/Private/MuCO/CustomizableObject.cpp b/Engine/Plugins/Experimental/Mutable/Source/CustomizableObject/Private/MuCO/CustomizableObject.cpp
--- a/Engine/Plugins/Experimental/Mutable/Source/CustomizableObject/Private/MuCO/CustomizableObject.cpp
+++ b/Engine/Plugins/Experimental/Mutable/Source/CustomizableObject/Private/MuCO/CustomizableObject.cpp
@@ -1619,7 +1619,7 @@ FCustomizableObjectProjector UCustomizableObject::GetProjectorParameterDefaultVa
 }
 
 
-uint64 UCustomizableObject::GetTextureParameterDefaultValue(const FString& InParameterName) const
+FString UCustomizableObject::GetTextureParameterDefaultValue(const FString& InParameterName) const
 {
 	const int32 ParameterIndex = FindParameter(InParameterName);
 	if (ParameterIndex == INDEX_NONE)
diff --git a/Engine/Plugins/Experimental/Mutable/Source/CustomizableObject/Private/MuCO/CustomizableObjectInstance.cpp b/Engine/Plugins/Experimental/Mutable/Source/CustomizableObject/Private/MuCO/CustomizableObjectInstance.cpp
--- a/Engine/Plugins/Experimental/Mutable/Source/CustomizableObject/Private/MuCO/CustomizableObjectInstance.cpp
+++ b/Engine/Plugins/Experimental/Mutable/Source/CustomizableObject/Private/MuCO/CustomizableObjectInstance.cpp
@@ -239,15 +239,7 @@ void UCustomizableObjectInstance::PreEditChange(FProperty* PropertyAboutToChange
 		
 		for (TObjectPtr<UTexture2D> Texture : TextureParameterDeclarations)
 		{
-			if (Texture)
-			{
-				const int32 TextureId = DefaultImageProvider.Get(Texture);
-
-				if (TextureId != INDEX_NONE)
-				{
-					DefaultImageProvider.Keep(TextureId, false);
-				}
-			}
+			DefaultImageProvider.Remove(Texture);
 		}
 	}
 }
@@ -267,11 +259,7 @@ void UCustomizableObjectInstance::PostEditChangeProperty(FPropertyChangedEvent&
 		
 		for (TObjectPtr<UTexture2D> Texture : TextureParameterDeclarations)
 		{
-			if (Texture)
-			{
-				const int32 TextureId = DefaultImageProvider.GetOrAdd(Texture);
-				DefaultImageProvider.Keep(TextureId, true);
-			}
+			DefaultImageProvider.Add(Texture);
 		}
 
 		UpdateSkeletalMeshAsync(true, true);
@@ -329,7 +317,7 @@ void UCustomizableObjectInstance::BeginDestroy()
 	}
 	
 	ReleaseMutableResources(true);
-
+	
 	Super::BeginDestroy();
 }
 
@@ -372,6 +360,14 @@ void UCustomizableObjectInstance::ReleaseMutableResources(bool bCalledFromBeginD
 				}
 			}
 		}
+
+		// Remove all references to cached Texture Parameters
+		FUnrealMutableImageProvider* ImageProvider = UCustomizableObjectSystem::GetInstance()->GetPrivate()->ImageProvider;
+
+		for (const FString& TextureParameter : PrivateData->UpdateTextureParameters)
+		{
+			ImageProvider->UnCacheImage(TextureParameter, false);				
+		}
 	}
 
 	PrivateData->GeneratedTextures.Empty();
@@ -2463,19 +2459,13 @@ void UCustomizableObjectInstance::SetFloatParameterSelectedOption(const FString&
 }
 
 
-uint64 UCustomizableObjectInstance::GetTextureParameterSelectedOption(const FString& TextureParamName, const int32 RangeIndex) const
+FString UCustomizableObjectInstance::GetTextureParameterSelectedOption(const FString& TextureParamName, const int32 RangeIndex) const
 {
 	return Descriptor.GetTextureParameterSelectedOption(TextureParamName, RangeIndex);
 }
 
 
-UTexture2D* UCustomizableObjectInstance::GetTextureParameterSelectedOptionT(const FString& TextureParamName, const int32 RangeIndex) const
-{
-	return Descriptor.GetTextureParameterSelectedOptionT(TextureParamName, RangeIndex);
-}
-
-
-void UCustomizableObjectInstance::SetTextureParameterSelectedOption(const FString& TextureParamName, const uint64 TextureValue, const int32 RangeIndex)
+void UCustomizableObjectInstance::SetTextureParameterSelectedOption(const FString& TextureParamName, const FString& TextureValue, const int32 RangeIndex)
 {
 	Descriptor.SetTextureParameterSelectedOption(TextureParamName, TextureValue, RangeIndex);
 }
@@ -5543,8 +5533,6 @@ void UCustomizableInstancePrivateData::BuildMaterials(const TSharedPtr<FMutableO
 
 	UCustomizableObject* CustomizableObject = Public->GetCustomizableObject();
 
-	const int32 LODCount = OperationData->InstanceUpdateData.LODs.Num();
-
 	// Find skipped LODs. The following valid LOD will be copied into them. 
 	TArray<bool> LODsSkipped;
 	LODsSkipped.SetNum(OperationData->CurrentMaxLOD);
@@ -5554,12 +5542,11 @@ void UCustomizableInstancePrivateData::BuildMaterials(const TSharedPtr<FMutableO
 	GeneratedMaterials.Reset();
 
 	// Prepare the data to store in order to regenerate resources for this instance (usually texture mips).
-	TSharedPtr<FMutableUpdateContext> UpdateContext = MakeShared<FMutableUpdateContext>();
-	UpdateContext->System = UCustomizableObjectSystem::GetInstance()->GetPrivate()->MutableSystem;
-	UpdateContext->Model = CustomizableObject->GetModel();
-	UpdateContext->Parameters = OperationData->MutableParameters;
-	UpdateContext->State = OperationData->State;
-
+	TSharedPtr<FMutableUpdateContext> UpdateContext = MakeShared<FMutableUpdateContext>(UCustomizableObjectSystem::GetInstance()->GetPrivate()->MutableSystem,
+		CustomizableObject->GetModel(),
+		OperationData->MutableParameters,
+	    OperationData->State);
+	
 	const bool bReuseTextures = OperationData->bReuseInstanceTextures;
 
 	const FInstanceUpdateData::FLOD& FirstLOD = OperationData->InstanceUpdateData.LODs[OperationData->CurrentMinLOD];
diff --git a/Engine/Plugins/Experimental/Mutable/Source/CustomizableObject/Private/MuCO/CustomizableObjectInstanceDescriptor.cpp b/Engine/Plugins/Experimental/Mutable/Source/CustomizableObject/Private/MuCO/CustomizableObjectInstanceDescriptor.cpp
--- a/Engine/Plugins/Experimental/Mutable/Source/CustomizableObject/Private/MuCO/CustomizableObjectInstanceDescriptor.cpp
+++ b/Engine/Plugins/Experimental/Mutable/Source/CustomizableObject/Private/MuCO/CustomizableObjectInstanceDescriptor.cpp
@@ -181,8 +181,8 @@ void FCustomizableObjectInstanceDescriptor::SaveDescriptor(FArchive& Ar)
 
 		case EMutableParameterType::Texture:
 		{
-			uint64 Value = 0;
-			TArray<uint64> RangeValues;
+			FString Value;
+			TArray<FString> RangeValues;
 
 			for (const FCustomizableObjectTextureParameterValue& P : TextureParameters)
 			{
@@ -379,8 +379,8 @@ void FCustomizableObjectInstanceDescriptor::LoadDescriptor(FArchive& Ar)
 
 		case EMutableParameterType::Texture:
 		{
-			uint64 Value = 0;
-			TArray<uint64> RangeValues;
+			FString Value;
+			TArray<FString> RangeValues;
 			Ar << Value;
 			Ar << RangeValues;
 
@@ -1346,7 +1346,7 @@ void FCustomizableObjectInstanceDescriptor::SetFloatParameterSelectedOption(cons
 }
 
 
-uint64 FCustomizableObjectInstanceDescriptor::GetTextureParameterSelectedOption(const FString& TextureParamName, const int32 RangeIndex) const
+FString FCustomizableObjectInstanceDescriptor::GetTextureParameterSelectedOption(const FString& TextureParamName, const int32 RangeIndex) const
 {
 	check(CustomizableObject);
 
@@ -1376,16 +1376,7 @@ uint64 FCustomizableObjectInstanceDescriptor::GetTextureParameterSelectedOption(
 }
 
 
-UTexture2D* FCustomizableObjectInstanceDescriptor::GetTextureParameterSelectedOptionT(const FString& TextureParamName, const int32 RangeIndex) const
-{
-	UDefaultImageProvider& ImageProvider = UCustomizableObjectSystem::GetInstance()->GetOrCreateDefaultImageProvider();
-
-	const uint64 TextureValue = GetTextureParameterSelectedOption(TextureParamName, RangeIndex);
-	return ImageProvider.Get(TextureValue);
-}
-
-
-void FCustomizableObjectInstanceDescriptor::SetTextureParameterSelectedOption(const FString& TextureParamName, const uint64 TextureValue, const int32 RangeIndex)
+void FCustomizableObjectInstanceDescriptor::SetTextureParameterSelectedOption(const FString& TextureParamName, const FString& TextureValue, const int32 RangeIndex)
 {
 	check(CustomizableObject);
 
@@ -1418,12 +1409,12 @@ void FCustomizableObjectInstanceDescriptor::SetTextureParameterSelectedOption(co
 }
 
 
-void FCustomizableObjectInstanceDescriptor::SetTextureParameterSelectedOptionT(const FString& TextureParamName, UTexture2D* TextureValue, const int32 RangeIndex)
+void FCustomizableObjectInstanceDescriptor::SetTextureParameterSelectedOptionT(const FString& TextureParamName,	UTexture2D* TextureValue, int32 RangeIndex)
 {
-	UDefaultImageProvider& ImageProvider = UCustomizableObjectSystem::GetInstance()->GetOrCreateDefaultImageProvider();
-
-	const uint64 NewUid = ImageProvider.GetOrAdd(TextureValue);
-	SetTextureParameterSelectedOption(TextureParamName, NewUid, RangeIndex);
+	UDefaultImageProvider& UDefaultImageProvider = UCustomizableObjectSystem::GetInstance()->GetOrCreateDefaultImageProvider();
+	const FString Id = UDefaultImageProvider.Add(TextureValue);
+	
+	SetTextureParameterSelectedOption(TextureParamName, Id, RangeIndex);
 }
 
 
diff --git a/Engine/Plugins/Experimental/Mutable/Source/CustomizableObject/Private/MuCO/CustomizableObjectMipDataProvider.cpp b/Engine/Plugins/Experimental/Mutable/Source/CustomizableObject/Private/MuCO/CustomizableObjectMipDataProvider.cpp
--- a/Engine/Plugins/Experimental/Mutable/Source/CustomizableObject/Private/MuCO/CustomizableObjectMipDataProvider.cpp
+++ b/Engine/Plugins/Experimental/Mutable/Source/CustomizableObject/Private/MuCO/CustomizableObjectMipDataProvider.cpp
@@ -16,6 +16,59 @@ UMutableTextureMipDataProviderFactory::UMutableTextureMipDataProviderFactory(con
 }
 
 
+FMutableUpdateContext::FMutableUpdateContext(mu::Ptr<mu::System> InSystem,
+	TSharedPtr<mu::Model, ESPMode::ThreadSafe> InModel, mu::Ptr<const mu::Parameters> InParameters, int32 InState):
+	System(InSystem),
+	Model(InModel),
+	Parameters(InParameters),
+	State(InState)
+{
+	if (Parameters)
+	{
+		UCustomizableObjectSystem::GetInstance()->GetPrivate()->ImageProvider->CacheImages(*Parameters);	
+	}
+}
+
+
+FMutableUpdateContext::~FMutableUpdateContext()
+{
+	if (Parameters)
+	{
+		UCustomizableObjectSystem::GetInstance()->GetPrivate()->ImageProvider->UnCacheImages(*Parameters);
+	}
+}
+
+
+mu::Ptr<mu::System> FMutableUpdateContext::GetSystem() const
+{
+	return System;
+}
+
+
+TSharedPtr<mu::Model, ESPMode::ThreadSafe> FMutableUpdateContext::GetModel() const
+{
+	return Model;
+}
+
+
+mu::Ptr<const mu::Parameters> FMutableUpdateContext::GetParameters() const
+{
+	return Parameters;
+}
+
+
+int32 FMutableUpdateContext::GetState() const
+{
+	return State;
+}
+
+
+const TArray<mu::Ptr<const mu::Image>>& FMutableUpdateContext::GetImageParameterValues() const
+{
+	return ImageParameterValues;
+}
+
+
 FMutableTextureMipDataProvider::FMutableTextureMipDataProvider(const UTexture* Texture, UCustomizableObjectInstance* InCustomizableObjectInstance, const FMutableImageReference& InImageRef)
 	: FTextureMipDataProvider(Texture, ETickState::Init, ETickThread::Async),
 	CustomizableObjectInstance(InCustomizableObjectInstance), ImageRef(InImageRef)
@@ -59,14 +112,14 @@ namespace impl
 
 		// This runs in a worker thread.
 		check(OperationData.IsValid());
-		check(OperationData->UpdateContext->System.get());
-		check(OperationData->UpdateContext->Model);
-		check(OperationData->UpdateContext->Parameters.get());
+		check(OperationData->UpdateContext->GetSystem().get());
+		check(OperationData->UpdateContext->GetModel());
+		check(OperationData->UpdateContext->GetParameters().get());
 
 		if (OperationData.IsValid())
 		{
-			mu::SystemPtr System = OperationData->UpdateContext->System;
-			const TSharedPtr<mu::Model, ESPMode::ThreadSafe> Model = OperationData->UpdateContext->Model;
+			mu::SystemPtr System = OperationData->UpdateContext->GetSystem();
+			const TSharedPtr<mu::Model, ESPMode::ThreadSafe> Model = OperationData->UpdateContext->GetModel();
 
 			// For now, we are forcing the recreation of mutable-side instances with every update.
 			mu::Instance::ID InstanceID = System->NewInstance(Model);
@@ -79,7 +132,7 @@ namespace impl
 				// LOD mask, set to all ones to build all LODs
 				uint32 LODMask = 0xFFFFFFFF;
 
-				Instance = System->BeginUpdate(InstanceID, OperationData->UpdateContext->Parameters, OperationData->UpdateContext->State, LODMask);
+				Instance = System->BeginUpdate(InstanceID, OperationData->UpdateContext->GetParameters(), OperationData->UpdateContext->GetState(), LODMask);
 				check(Instance);
 			}
 
diff --git a/Engine/Plugins/Experimental/Mutable/Source/CustomizableObject/Private/MuCO/CustomizableObjectSystem.cpp b/Engine/Plugins/Experimental/Mutable/Source/CustomizableObject/Private/MuCO/CustomizableObjectSystem.cpp
--- a/Engine/Plugins/Experimental/Mutable/Source/CustomizableObject/Private/MuCO/CustomizableObjectSystem.cpp
+++ b/Engine/Plugins/Experimental/Mutable/Source/CustomizableObject/Private/MuCO/CustomizableObjectSystem.cpp
@@ -126,7 +126,7 @@ void FMutableQueue::Enqueue(const FMutableQueueElem& TaskToEnqueue)
 		{
 			UCustomizableObjectInstance* Instance = QueueElem.Operation->CustomizableObjectInstance.Get();
 			FInstanceUpdateDelegate* Callback = &QueueElem.Operation->UpdateCallback;
-			FinishUpdateGlobal(Instance, EUpdateResult::ErrorReplaced, Callback);
+			FinishUpdateGlobal(Instance, EUpdateResult::ErrorReplaced, Callback, QueueElem.Operation->GetParameters());
 		
 			QueueElem.Operation = TaskToEnqueue.Operation;
 			QueueElem.PriorityType = FMath::Min(QueueElem.PriorityType, TaskToEnqueue.PriorityType);
@@ -735,7 +735,7 @@ FAutoConsoleVariableRef CVarMaxTextureSizeToGenerate(
 		"If a texture doesn't have small enough mips, mip 0 will be the last mip available."));
 
 
-void FinishUpdateGlobal(UCustomizableObjectInstance* Instance, EUpdateResult UpdateResult, FInstanceUpdateDelegate* UpdateCallback, const FDescriptorRuntimeHash InUpdatedHash)
+void FinishUpdateGlobal(UCustomizableObjectInstance* Instance, EUpdateResult UpdateResult, FInstanceUpdateDelegate* UpdateCallback, mu::ParametersPtrConst Parameters, const FDescriptorRuntimeHash InUpdatedHash)
 {
 	check(IsInGameThread())
 	
@@ -767,12 +767,17 @@ void FinishUpdateGlobal(UCustomizableObjectInstance* Instance, EUpdateResult Upd
 		UpdateCallback->ExecuteIfBound(Context);
 	}
 
+	if (Parameters)
+	{
+		UCustomizableObjectSystem::GetInstance()->GetPrivate()->ImageProvider->UnCacheImages(*Parameters);
+	}
+
 	UCustomizableObjectSystem::GetInstance()->GetPrivate()->MutableTaskGraph.AllowLaunchingMutableTaskLowPriority(true, false);
 }
 
 
 /** Update the given Instance Skeletal Meshes and call its callbacks. */
-void UpdateSkeletalMesh(UCustomizableObjectInstance& CustomizableObjectInstance, const FDescriptorRuntimeHash& UpdatedDescriptorRuntimeHash, EUpdateResult UpdateResult, FInstanceUpdateDelegate* UpdateCallback)
+void UpdateSkeletalMesh(UCustomizableObjectInstance& CustomizableObjectInstance, const FDescriptorRuntimeHash& UpdatedDescriptorRuntimeHash, EUpdateResult UpdateResult, FInstanceUpdateDelegate* UpdateCallback, mu::ParametersPtrConst Parameters = nullptr)
 {
 	MUTABLE_CPUPROFILER_SCOPE(UpdateSkeletalMesh);
 
@@ -833,7 +838,7 @@ void UpdateSkeletalMesh(UCustomizableObjectInstance& CustomizableObjectInstance,
 		CustomizableObjectInstancePrivateData->TexturesToRelease.Empty();
 	}
 
-	FinishUpdateGlobal(&CustomizableObjectInstance, UpdateResult, UpdateCallback, UpdatedDescriptorRuntimeHash);
+	FinishUpdateGlobal(&CustomizableObjectInstance, UpdateResult, UpdateCallback, Parameters, UpdatedDescriptorRuntimeHash);
 }
 
 
@@ -897,11 +902,14 @@ void FCustomizableObjectSystemPrivate::InitUpdateSkeletalMesh(UCustomizableObjec
 		UpdateSkeletalMesh(Instance, Instance.GetDescriptorRuntimeHash(), EUpdateResult::Success, UpdateCallback);
 	}
 	else
-	{
+	{	
 		Instance.SkeletalMeshStatus = ESkeletalMeshState::AsyncUpdatePending;
 
 		const TSharedPtr<FMutableOperation> Operation = MakeShared<FMutableOperation>(FMutableOperation::CreateInstanceUpdate(&Instance, bNeverStream, MipsToSkip, UpdateCallback));
 		MutableOperationQueue.Enqueue(FMutableQueueElem::Create(Operation, Priority, Instance.GetPrivate()->MinSquareDistFromComponentToPlayer));
+
+		// Cache Texture Parameters being used during the update:
+		ImageProvider->CacheImages(*Operation->GetParameters());
 	}
 }
 
@@ -1821,7 +1829,7 @@ namespace impl
 		UCustomizableObjectSystem* System = UCustomizableObjectSystem::GetInstance();
 		if (!System || !System->IsValidLowLevel() || System->HasAnyFlags(RF_BeginDestroyed))
 		{
-			FinishUpdateGlobal(CustomizableObjectInstancePtr.Get(), EUpdateResult::Error, &OperationData->UpdateCallback);
+			FinishUpdateGlobal(CustomizableObjectInstancePtr.Get(), EUpdateResult::Error, &OperationData->UpdateCallback, OperationData->MutableParameters);
 			return;
 		}
 
@@ -1831,7 +1839,7 @@ namespace impl
 		if (!CustomizableObjectInstance || !CustomizableObjectInstance->IsValidLowLevel() )
 		{
 			System->ClearCurrentMutableOperation();
-			FinishUpdateGlobal(CustomizableObjectInstance, EUpdateResult::Error, &OperationData->UpdateCallback);
+			FinishUpdateGlobal(CustomizableObjectInstance, EUpdateResult::Error, &OperationData->UpdateCallback, OperationData->MutableParameters);
 			return;
 		}
 
@@ -1840,7 +1848,7 @@ namespace impl
 
 		// Actual work
 		// TODO MTBL-391: Review This hotfix
-		UpdateSkeletalMesh(*CustomizableObjectInstance, CustomizableObjectSystemPrivateData->CurrentMutableOperation->InstanceDescriptorRuntimeHash, OperationData->UpdateResult, &OperationData->UpdateCallback);
+		UpdateSkeletalMesh(*CustomizableObjectInstance, CustomizableObjectSystemPrivateData->CurrentMutableOperation->InstanceDescriptorRuntimeHash, OperationData->UpdateResult, &OperationData->UpdateCallback, OperationData->MutableParameters);
 
 		// TODO: T2927
 		if (LogBenchmarkUtil::isLoggingActive())
@@ -1872,7 +1880,7 @@ namespace impl
 		UCustomizableObjectSystem* System = UCustomizableObjectSystem::GetInstance();
 		if (!System || !System->IsValidLowLevel() || System->HasAnyFlags(RF_BeginDestroyed))
 		{
-			FinishUpdateGlobal(CustomizableObjectInstancePtr.Get(), EUpdateResult::Error, &OperationData->UpdateCallback);
+			FinishUpdateGlobal(CustomizableObjectInstancePtr.Get(), EUpdateResult::Error, &OperationData->UpdateCallback, OperationData->MutableParameters);
 			return;
 		}
 
@@ -2005,7 +2013,7 @@ namespace impl
 		}
 		else
 		{
-			FinishUpdateGlobal(CustomizableObjectInstance, EUpdateResult::Error, &OperationData->UpdateCallback);
+			FinishUpdateGlobal(CustomizableObjectInstance, EUpdateResult::Error, &OperationData->UpdateCallback, OperationData->MutableParameters);
 		}
 	}
 
@@ -2029,7 +2037,7 @@ namespace impl
 		if (!ObjectInstance)
 		{
 			System->ClearCurrentMutableOperation();
-			FinishUpdateGlobal(ObjectInstance, EUpdateResult::Error, &OperationData->UpdateCallback);
+			FinishUpdateGlobal(ObjectInstance, EUpdateResult::Error, &OperationData->UpdateCallback, OperationData->MutableParameters);
 			return;
 		}
 
@@ -2051,7 +2059,7 @@ namespace impl
 		if (!CustomizableObject)
 		{
 			System->ClearCurrentMutableOperation();
-			FinishUpdateGlobal(CustomizableObjectInstancePtr.Get(), EUpdateResult::Error, &OperationData->UpdateCallback);
+			FinishUpdateGlobal(CustomizableObjectInstancePtr.Get(), EUpdateResult::Error, &OperationData->UpdateCallback, OperationData->MutableParameters);
 			return;
 		}
 
@@ -2063,9 +2071,27 @@ namespace impl
 			ObjectInstancePrivateData->RelevantParameters = OperationData->RelevantParametersInProgress;
 		}
 
-		if (const TObjectPtr<UDefaultImageProvider> DefaultImageProvider = System->GetDefaultImageProvider())
+		// Cache Texture Parameters
+		FUnrealMutableImageProvider* ImageProvider = UCustomizableObjectSystem::GetInstance()->GetPrivate()->ImageProvider;
+		
+		// Uncache old Texture Parameters
+		for (const FString& TextureParameter : ObjectInstancePrivateData->UpdateTextureParameters)
 		{
-			DefaultImageProvider->CacheTextures(*OperationData->MutableParameters);
+			ImageProvider->UnCacheImage(TextureParameter, false);
+		}
+		ObjectInstancePrivateData->UpdateTextureParameters.Reset();
+
+		// Cache new Texture Parameters and update which ones are currently are being used
+		for (const FCustomizableObjectTextureParameterValue& TextureParameters : ObjectInstance->GetDescriptor().GetTextureParameters())
+		{
+			ImageProvider->CacheImage(TextureParameters.ParameterValue, false);				
+			ObjectInstancePrivateData->UpdateTextureParameters.Add(TextureParameters.ParameterValue);
+		
+			for (const FString& TextureParameter: TextureParameters.ParameterRangeValues)
+			{
+				ImageProvider->CacheImage(TextureParameter, false);
+				ObjectInstancePrivateData->UpdateTextureParameters.Add(TextureParameter);
+			}
 		}
 		
 		// Selectively lock the resource cache for the object used by this instance to avoid the destruction of resources that we may want to reuse.
@@ -2222,7 +2248,7 @@ namespace impl
 		if (!Operation->CustomizableObjectInstance.IsValid() || !Operation->CustomizableObjectInstance->IsValidLowLevel()) // Only start if it hasn't been already destroyed (i.e. GC after finish PIE)
 		{
 			System->ClearCurrentMutableOperation();
-			FinishUpdateGlobal(nullptr, EUpdateResult::Error, &Operation->UpdateCallback);
+			FinishUpdateGlobal(nullptr, EUpdateResult::Error, &Operation->UpdateCallback, Operation->GetParameters());
 			return;
 		}
 
@@ -2232,7 +2258,7 @@ namespace impl
 		if (!CandidateInstancePrivateData)
 		{
 			System->ClearCurrentMutableOperation();
-			FinishUpdateGlobal(nullptr, EUpdateResult::Error, &Operation->UpdateCallback);
+			FinishUpdateGlobal(nullptr, EUpdateResult::Error, &Operation->UpdateCallback, Operation->GetParameters());
 			return;
 		}
 		
@@ -2274,7 +2300,7 @@ namespace impl
 		if (Operation->InstanceDescriptorRuntimeHash.IsSubset(CandidateInstance->GetDescriptorRuntimeHash()))
 		{
 			CandidateInstance->SkeletalMeshStatus = ESkeletalMeshState::Correct;
-			UpdateSkeletalMesh(*CandidateInstance, CandidateInstance->GetDescriptorRuntimeHash(), EUpdateResult::Success, &Operation->UpdateCallback);
+			UpdateSkeletalMesh(*CandidateInstance, CandidateInstance->GetDescriptorRuntimeHash(), EUpdateResult::Success, &Operation->UpdateCallback, Operation->GetParameters());
 			bCancel = true;
 		}
 
@@ -2290,7 +2316,7 @@ namespace impl
 
 			System->ClearCurrentMutableOperation();
 
-			FinishUpdateGlobal(CandidateInstance, EUpdateResult::Error, &Operation->UpdateCallback);
+			FinishUpdateGlobal(CandidateInstance, EUpdateResult::Error, &Operation->UpdateCallback, Operation->GetParameters());
 			return;
 		}
 
@@ -2694,11 +2720,6 @@ void UCustomizableObjectSystem::UnregisterImageProvider(UCustomizableSystemImage
 }
 
 
-TObjectPtr<UDefaultImageProvider> UCustomizableObjectSystem::GetDefaultImageProvider() const
-{
-	return DefaultImageProvider;
-}
-
 
 UDefaultImageProvider& UCustomizableObjectSystem::GetOrCreateDefaultImageProvider()
 {
@@ -3120,27 +3141,19 @@ void UCustomizableObjectSystem::OnEnableImageCacheChanged(IConsoleVariable* CVar
 }
 
 
-void UCustomizableObjectSystem::CacheImage(uint64 ImageId)
+void UCustomizableObjectSystem::CacheImage(FString ImageId)
 {
 	check(GetPrivate() != nullptr);
 	check(GetPrivate()->ImageProvider != nullptr);
-	GetPrivate()->ImageProvider->CacheImage(ImageId);
+	GetPrivate()->ImageProvider->CacheImage(ImageId, true);
 }
 
 
-void UCustomizableObjectSystem::UnCacheImage(uint64 ImageId)
+void UCustomizableObjectSystem::UnCacheImage(FString ImageId)
 {
 	check(GetPrivate() != nullptr); 
 	check(GetPrivate()->ImageProvider != nullptr);
-	GetPrivate()->ImageProvider->UnCacheImage(ImageId);
-}
-
-
-void UCustomizableObjectSystem::CacheAllImagesInAllProviders(bool bClearPreviousCacheImages)
-{
-	check(GetPrivate() != nullptr);
-	check(GetPrivate()->ImageProvider != nullptr);
-	GetPrivate()->ImageProvider->CacheAllImagesInAllProviders(bClearPreviousCacheImages);
+	GetPrivate()->ImageProvider->UnCacheImage(ImageId, true);
 }
 
 
@@ -3148,7 +3161,7 @@ void UCustomizableObjectSystem::ClearImageCache()
 {
 	check(GetPrivate() != nullptr);
 	check(GetPrivate()->ImageProvider != nullptr);
-	GetPrivate()->ImageProvider->ClearCache();
+	GetPrivate()->ImageProvider->ClearCache(true);
 }
 
 
diff --git a/Engine/Plugins/Experimental/Mutable/Source/CustomizableObject/Private/MuCO/DefaultImageProvider.cpp b/Engine/Plugins/Experimental/Mutable/Source/CustomizableObject/Private/MuCO/DefaultImageProvider.cpp
--- a/Engine/Plugins/Experimental/Mutable/Source/CustomizableObject/Private/MuCO/DefaultImageProvider.cpp
+++ b/Engine/Plugins/Experimental/Mutable/Source/CustomizableObject/Private/MuCO/DefaultImageProvider.cpp
@@ -3,8 +3,6 @@
 #include "MuCO/DefaultImageProvider.h" 
 
 #include "MuCO/CustomizableObject.h"
-#include "MuCO/CustomizableObjectInstance.h"
-#include "UObject/UObjectIterator.h"
 
 #include "MuR/Parameters.h"
 
@@ -36,235 +34,55 @@ namespace UDefaultImageProviderCVars
 }
 
 
-UCustomizableSystemImageProvider::ValueType UDefaultImageProvider::HasTextureParameterValue(int64 ID)
+UCustomizableSystemImageProvider::ValueType UDefaultImageProvider::HasTextureParameterValue(FString ID)
 {
-	const int32 TextureIndex = ToIndex(ID);
+	const TObjectPtr<UTexture2D>* Texture = Textures.Find(ID);
 	
-	return Textures.IsValidIndex(TextureIndex) && Textures[TextureIndex] ?
+	return Texture ?
 		static_cast<ValueType>(UDefaultImageProviderCVars::ImageMode) :
 		ValueType::None;
 }
 
 
-UTexture2D* UDefaultImageProvider::GetTextureParameterValue(int64 ID)
+UTexture2D* UDefaultImageProvider::GetTextureParameterValue(FString ID)
 {
-	return Textures[ToIndex(ID)];
+	return Textures[ID];
 }
 
 
 void UDefaultImageProvider::GetTextureParameterValues(TArray<FCustomizableObjectExternalTexture>& OutValues)
 {
-	for (int32 TextureIndex = 0; TextureIndex < Textures.Num(); ++TextureIndex)
+	for (TTuple<FString, TObjectPtr<UTexture2D>>& Pair : Textures)
 	{
-		if (const UTexture2D* Texture = Textures[TextureIndex])
-		{
-			FCustomizableObjectExternalTexture Data;
-			Data.Name = Texture->GetName();
-			Data.Value = ToTextureId(TextureIndex);
-			OutValues.Add(Data);
-		}
+		FCustomizableObjectExternalTexture Data;
+		Data.Name = Pair.Value.GetName();
+		Data.Value = Pair.Key;
+		OutValues.Add(Data);
 	}
 }
 
 
-UTexture2D* UDefaultImageProvider::Get(const uint64 TextureId) const
-{
-	return Textures[ToIndex(TextureId)];
-}
-
-
-int64 UDefaultImageProvider::Get(UTexture2D* Texture) const
-{
-	for (int32 TextureIndex = 0; TextureIndex < Textures.Num(); ++TextureIndex)
-	{
-		if (Textures[TextureIndex] == Texture)
-		{
-			return ToTextureId(TextureIndex);
-		}
-	}
-
-	return INDEX_NONE;
-}
-
-
-uint64 UDefaultImageProvider::GetOrAdd(UTexture2D* Texture)
+FString UDefaultImageProvider::Add(UTexture2D* Texture)
 {
 	if (!Texture)
 	{
 		return FCustomizableObjectTextureParameterValue::DEFAULT_PARAMETER_VALUE;
 	}
-	
-	int32 Hole = INDEX_NONE;
-
-	for (int32 TextureIndex = 0; TextureIndex < Textures.Num(); ++TextureIndex)
-	{
-		const TObjectPtr<UTexture2D> UsedTexture = Textures[TextureIndex];
-		if (UsedTexture == Texture)
-		{
-			return ToTextureId(TextureIndex);
-		}
-		else if (!UsedTexture)
-		{
-			Hole = TextureIndex;
-			break;
-		}
-	}
-	
-	if (Hole == INDEX_NONE)
-	{
-		Hole = Textures.Num();
-
-		const int32 NumElements = Hole + 1;
-		check(NumElements < MAX_IDS); // Max number of TextureIds reached.
-
-		Textures.SetNumUninitialized(NumElements);
-		KeepTextures.SetNumUninitialized(NumElements);
-	} 
-	
-	Textures[Hole] = Texture;
-	KeepTextures[Hole] = false;
-	
-	return ToTextureId(Hole);
-}
-
-
-void UDefaultImageProvider::CacheTextures(const mu::Parameters& Parameters)
-{
-	GarbageCollectTextureIds(Parameters);
-
-	UCustomizableObjectSystem* System = UCustomizableObjectSystem::GetInstance();
-
-	const int32 NumParams = Parameters.GetCount();
-	for (int32 ParamIndex = 0; ParamIndex < NumParams; ++ParamIndex)
-	{
-		if (Parameters.GetType(ParamIndex) != mu::PARAMETER_TYPE::T_IMAGE)
-		{
-			continue;
-		}
-	
-		{
-			const mu::EXTERNAL_IMAGE_ID TextureId = Parameters.GetImageValue(ParamIndex);
-			if (const int32 TextureIndex = ToIndex(TextureId);
-				Textures.IsValidIndex(TextureIndex) && Textures[TextureIndex])
-			{
-				System->CacheImage(TextureId);
-			}
-		}
-
-		const int32 NumValues = Parameters.GetValueCount(ParamIndex);
-		for (int32 ValueIndex = 0; ValueIndex < NumValues; ++ValueIndex)
-		{
-			mu::Ptr<const mu::RangeIndex> Range = Parameters.GetValueIndex(ParamIndex, ValueIndex);
-			const mu::EXTERNAL_IMAGE_ID TextureId = Parameters.GetImageValue(ParamIndex, Range);
-
-			if (const int32 TextureIndex = ToIndex(TextureId);
-				Textures.IsValidIndex(TextureIndex) && Textures[TextureIndex])
-			{
-				System->CacheImage(TextureId);
-			}
-		}
-	}
-}
 
+	const FString Id = Texture->GetFullName();
+	Textures.Add(Id, Texture);	
 
-void UDefaultImageProvider::Keep(const uint64 TextureId, const bool bKeep)
-{
-	KeepTextures[ToIndex(TextureId)] = bKeep;
+	return Id;
 }
 
 
-void UDefaultImageProvider::GarbageCollectTextureIds(const mu::Parameters& Parameters)
+void UDefaultImageProvider::Remove(UTexture2D* Texture)
 {
-	TBitArray IdUsed;
-	IdUsed.Init(false, Textures.Num());
-
-	for (TObjectIterator<UCustomizableObjectInstance> It; It; ++It)
-	{
-		for (const FCustomizableObjectTextureParameterValue& TextureParameter :  It->GetDescriptor().GetTextureParameters())
-		{			
-			if (const int32 TextureIndex = ToIndex(TextureParameter.ParameterValue);
-				IdUsed.IsValidIndex(TextureIndex))
-			{
-				IdUsed[TextureIndex] = true;
-			}
-
-			for (const uint64 RangeValue : TextureParameter.ParameterRangeValues)
-			{
-				if (const int32 TextureIndex = ToIndex(RangeValue);
-					IdUsed.IsValidIndex(TextureIndex))
-				{
-					IdUsed[TextureIndex] = true;
-				}
-			}
-		}
-	}
-
-	const int32 NumParams = Parameters.GetCount();
-	for (int32 ParamIndex = 0; ParamIndex < NumParams; ++ParamIndex)
+	if (!Texture)
 	{
-		if (Parameters.GetType(ParamIndex) != mu::PARAMETER_TYPE::T_IMAGE)
-		{
-			continue;
-		}
-
-		{
-			const mu::EXTERNAL_IMAGE_ID TextureId = Parameters.GetImageValue(ParamIndex);
-			if (const int32 TextureIndex = ToIndex(TextureId);
-				Textures.IsValidIndex(TextureIndex) && Textures[TextureIndex])
-			{
-				IdUsed[TextureIndex] = true;
-			}
-		}
-
-		const int32 NumValues = Parameters.GetValueCount(ParamIndex);
-		for (int32 ValueIndex = 0; ValueIndex < NumValues; ++ValueIndex)
-		{
-			mu::Ptr<const mu::RangeIndex> Range = Parameters.GetValueIndex(ParamIndex, ValueIndex);
-			const mu::EXTERNAL_IMAGE_ID TextureId = Parameters.GetImageValue(ParamIndex, Range);
-
-			if (const int32 TextureIndex = ToIndex(TextureId);
-				Textures.IsValidIndex(TextureIndex) && Textures[TextureIndex])
-			{
-				IdUsed[TextureIndex] = true;
-			}
-		}
-
+		return;
 	}
-
-	UCustomizableObjectSystem* System = UCustomizableObjectSystem::GetInstance();
-
-	int32 LastToKeep = -1;
 	
-	for (int32 TextureIndex = 0; TextureIndex < Textures.Num(); ++TextureIndex)
-	{
-		if (!IdUsed[TextureIndex] && !KeepTextures[TextureIndex])
-		{
-			Textures[TextureIndex] = nullptr;
-			System->UnCacheImage(ToTextureId(TextureIndex));
-		}
-		else
-		{
-			LastToKeep = TextureIndex;
-		}
-	}
-
-	if (const int32 NumToKeep = LastToKeep + 1;
-		NumToKeep < Textures.Num())
-	{
-		Textures.SetNumUninitialized(NumToKeep);
-		KeepTextures.SetNumUninitialized(NumToKeep);
-	}
-}
-
-
-int32 UDefaultImageProvider::ToIndex(const uint64 TextureId) const
-{
-	return TextureId - BASE_ID;
-}
-
-
-uint64 UDefaultImageProvider::ToTextureId(const int32 TextureIndex) const
-{
-	return TextureIndex + BASE_ID;
+	Textures.Remove(Texture->GetFullName());
 }
 
diff --git a/Engine/Plugins/Experimental/Mutable/Source/CustomizableObject/Private/MuCO/UnrealMutableImageProvider.cpp b/Engine/Plugins/Experimental/Mutable/Source/CustomizableObject/Private/MuCO/UnrealMutableImageProvider.cpp
--- a/Engine/Plugins/Experimental/Mutable/Source/CustomizableObject/Private/MuCO/UnrealMutableImageProvider.cpp
+++ b/Engine/Plugins/Experimental/Mutable/Source/CustomizableObject/Private/MuCO/UnrealMutableImageProvider.cpp
@@ -5,6 +5,7 @@
 #include "MuCO/CustomizableObjectSystem.h"
 #include "MuCO/CustomizableObject.h"
 #include "TextureResource.h"
+#include "MuR/Parameters.h"
 
 
 //-------------------------------------------------------------------------------------------------
@@ -462,109 +463,186 @@ mu::FImageDesc FUnrealMutableImageProvider::GetImageDesc(mu::EXTERNAL_IMAGE_ID i
 }
 
 
-void FUnrealMutableImageProvider::CacheImage(mu::EXTERNAL_IMAGE_ID id)
+void FUnrealMutableImageProvider::CacheImage(const mu::EXTERNAL_IMAGE_ID& Id, bool bUser)
 {
 	check( IsInGameThread() );
 
-	mu::ImagePtr pResult;
-	UTexture2D* UnrealDeferredTexture = nullptr;
+	if (Id == FCustomizableObjectTextureParameterValue::DEFAULT_PARAMETER_VALUE)
+	{
+		return;	
+	}	
 
-	// See if any provider provides this id.
-	for (int32 p = 0; !pResult && p < ImageProviders.Num(); ++p)
 	{
-		TWeakObjectPtr<UCustomizableSystemImageProvider> Provider = ImageProviders[p];
+		FScopeLock Lock(&ExternalImagesLock);
 
-		if (Provider.IsValid())
+		if (FUnrealMutableImageInfo* Result = GlobalExternalImages.Find(Id))
 		{
-			switch (Provider->HasTextureParameterValue(id))
+			if (bUser)
 			{
-
-			case UCustomizableSystemImageProvider::ValueType::Raw:
+				Result->ReferencesUser = true;			
+			}
+			else
 			{
-				FIntVector desc = Provider->GetTextureParameterValueSize(id);
-				pResult = new mu::Image(desc[0], desc[1], 1, mu::EImageFormat::IF_RGBA_UBYTE);
-				Provider->GetTextureParameterValueData(id, pResult->GetData());
-				break;
+				Result->ReferencesSystem++;							
 			}
+		}
+		else
+		{
+			mu::ImagePtr pResult;
+			UTexture2D* UnrealDeferredTexture = nullptr;
 
-			case UCustomizableSystemImageProvider::ValueType::Unreal:
+			// See if any provider provides this id.
+			for (int32 p = 0; !pResult && p < ImageProviders.Num(); ++p)
 			{
-				UTexture2D* UnrealTexture = Provider->GetTextureParameterValue(id);
-				pResult = ConvertTextureUnrealToMutable(UnrealTexture, 0);
-				break;
+				TWeakObjectPtr<UCustomizableSystemImageProvider> Provider = ImageProviders[p];
+
+				if (Provider.IsValid())
+				{
+					switch (Provider->HasTextureParameterValue(Id))
+					{
+
+					case UCustomizableSystemImageProvider::ValueType::Raw:
+					{
+						FIntVector desc = Provider->GetTextureParameterValueSize(Id);
+						pResult = new mu::Image(desc[0], desc[1], 1, mu::EImageFormat::IF_RGBA_UBYTE);
+						Provider->GetTextureParameterValueData(Id, pResult->GetData());
+						break;
+					}
+
+					case UCustomizableSystemImageProvider::ValueType::Unreal:
+					{
+						UTexture2D* UnrealTexture = Provider->GetTextureParameterValue(Id);
+						pResult = ConvertTextureUnrealToMutable(UnrealTexture, 0);
+						break;
+					}
+
+					case UCustomizableSystemImageProvider::ValueType::Unreal_Deferred:
+					{
+						UnrealDeferredTexture = Provider->GetTextureParameterValue(Id);
+						break;
+					}
+
+					default:
+						break;
+					}
+				}
 			}
 
-			case UCustomizableSystemImageProvider::ValueType::Unreal_Deferred:
+			if (!pResult && !UnrealDeferredTexture)
 			{
-				UnrealDeferredTexture = Provider->GetTextureParameterValue(id);
-				break;
+				UE_LOG(LogMutable, Warning, TEXT("Failed to cache external image %s."), *Id);
+				pResult = CreateDummy();
 			}
+			
+			FUnrealMutableImageInfo ImageInfo(pResult, UnrealDeferredTexture);
 
-			default:
-				break;
+			if (bUser)
+			{
+				ImageInfo.ReferencesUser = true;			
+			}
+			else
+			{
+				ImageInfo.ReferencesSystem++;							
 			}
+			
+			GlobalExternalImages.Add(Id, ImageInfo);
 		}
 	}
-
-	if (!pResult && !UnrealDeferredTexture)
-	{
-		UE_LOG(LogMutable, Warning, TEXT("Failed to catch external image."));
-		pResult = CreateDummy();
-	}
-
-	{
-		FScopeLock Lock(&ExternalImagesLock);
-
-		GlobalExternalImages.Add(id, FUnrealMutableImageInfo(pResult, UnrealDeferredTexture));
-	}
 }
 
 
-void FUnrealMutableImageProvider::UnCacheImage(mu::EXTERNAL_IMAGE_ID id)
+void FUnrealMutableImageProvider::UnCacheImage(const mu::EXTERNAL_IMAGE_ID& Id, bool bUser)
 {
 	check(IsInGameThread());
 
+	if (Id == FCustomizableObjectTextureParameterValue::DEFAULT_PARAMETER_VALUE)
+	{
+		return;	
+	}
+	
 	FScopeLock Lock(&ExternalImagesLock);
-	GlobalExternalImages.Remove(id);
+	if (FUnrealMutableImageInfo* Result = GlobalExternalImages.Find(Id))
+	{
+		if (bUser)
+		{
+			Result->ReferencesUser = false;			
+		}
+		else
+		{
+			--Result->ReferencesSystem;							
+		}
+		
+		if (Result->ReferencesUser + Result->ReferencesSystem == 0)
+		{
+			GlobalExternalImages.Remove(Id);		
+		}
+	}
 }
 
 
-void FUnrealMutableImageProvider::CacheAllImagesInAllProviders(bool bClearPreviousCacheImages)
+void FUnrealMutableImageProvider::ClearCache(bool bUser)
 {
 	check(IsInGameThread());
 
 	FScopeLock Lock(&ExternalImagesLock);
-
-	if (bClearPreviousCacheImages)
+	for (TTuple<FString, FUnrealMutableImageInfo> Tuple : GlobalExternalImages)
 	{
-		ClearCache();
-	}
+		UnCacheImage(Tuple.Key, bUser);
+	}	
+}
+
 
-	// See if any provider provides this id.
-	for (int32 p = 0; p < ImageProviders.Num(); ++p)
+void FUnrealMutableImageProvider::CacheImages(const mu::Parameters& Parameters)
+{
+	const int32 NumParams = Parameters.GetCount();
+	for (int32 ParamIndex = 0; ParamIndex < NumParams; ++ParamIndex)
 	{
-		TWeakObjectPtr<UCustomizableSystemImageProvider> Provider = ImageProviders[p];
+		if (Parameters.GetType(ParamIndex) != mu::PARAMETER_TYPE::T_IMAGE)
+		{
+			continue;
+		}
+	
+		{
+			const FString TextureId = Parameters.GetImageValue(ParamIndex);
+			CacheImage(TextureId, false);
+		}
 
-		if (Provider.IsValid())
+		const int32 NumValues = Parameters.GetValueCount(ParamIndex);
+		for (int32 ValueIndex = 0; ValueIndex < NumValues; ++ValueIndex)
 		{
-			TArray<FCustomizableObjectExternalTexture> OutValues;
-			Provider->GetTextureParameterValues(OutValues);
+			mu::Ptr<const mu::RangeIndex> Range = Parameters.GetValueIndex(ParamIndex, ValueIndex);
+			const FString TextureId = Parameters.GetImageValue(ParamIndex, Range);
 
-			for (const FCustomizableObjectExternalTexture& Value : OutValues)
-			{
-				CacheImage(Value.Value);
-			}
+			CacheImage(TextureId, false);
 		}
 	}
 }
 
 
-void FUnrealMutableImageProvider::ClearCache()
+void FUnrealMutableImageProvider::UnCacheImages(const mu::Parameters& Parameters)
 {
-	check(IsInGameThread());
+	const int32 NumParams = Parameters.GetCount();
+	for (int32 ParamIndex = 0; ParamIndex < NumParams; ++ParamIndex)
+	{
+		if (Parameters.GetType(ParamIndex) != mu::PARAMETER_TYPE::T_IMAGE)
+		{
+			continue;
+		}
+	
+		{
+			const FString TextureId = Parameters.GetImageValue(ParamIndex);
+			UnCacheImage(TextureId, false);				
+		}
 
-	FScopeLock Lock(&ExternalImagesLock);
-	GlobalExternalImages.Empty();
+		const int32 NumValues = Parameters.GetValueCount(ParamIndex);
+		for (int32 ValueIndex = 0; ValueIndex < NumValues; ++ValueIndex)
+		{
+			mu::Ptr<const mu::RangeIndex> Range = Parameters.GetValueIndex(ParamIndex, ValueIndex);
+			const FString TextureId = Parameters.GetImageValue(ParamIndex, Range);
+
+			UnCacheImage(TextureId, false);
+		}
+	}
 }
 
 
@@ -609,7 +687,7 @@ mu::FImageDesc FUnrealMutableImageProvider::CreateDummyDesc()
 
 void FUnrealMutableImageProvider::AddReferencedObjects(FReferenceCollector& Collector)
 {
-	for (TPair<uint64, FUnrealMutableImageInfo>& Image : GlobalExternalImages)
+	for (TPair<mu::EXTERNAL_IMAGE_ID, FUnrealMutableImageInfo>& Image : GlobalExternalImages)
 	{
 		if (Image.Value.TextureToLoad)
 		{
diff --git a/Engine/Plugins/Experimental/Mutable/Source/CustomizableObject/Private/MuCO/UnrealMutableImageProvider.h b/Engine/Plugins/Experimental/Mutable/Source/CustomizableObject/Private/MuCO/UnrealMutableImageProvider.h
--- a/Engine/Plugins/Experimental/Mutable/Source/CustomizableObject/Private/MuCO/UnrealMutableImageProvider.h
+++ b/Engine/Plugins/Experimental/Mutable/Source/CustomizableObject/Private/MuCO/UnrealMutableImageProvider.h
@@ -31,13 +31,23 @@ public:
 	virtual mu::FImageDesc GetImageDesc(mu::EXTERNAL_IMAGE_ID Id, uint8 MipmapsToSkip) override;
 
 	
-	// Own interface
+	// Own interface	
 	// Thread: Game
-	void CacheImage(mu::EXTERNAL_IMAGE_ID id);
-	void UnCacheImage(mu::EXTERNAL_IMAGE_ID id);
-	void CacheAllImagesInAllProviders(bool bClearPreviousCacheImages);
-	void ClearCache();
+	/** Add a reference to the image. If it was not cached it caches it.
+	 * @param bUser if true, adds a reference to the user reference counter. If false, adds a reference to the system reference counter. */
+	void CacheImage(const mu::EXTERNAL_IMAGE_ID& Id, bool bUser);
 
+	/** Removes a reference to the image. If all references are removed, it uncaches the image.
+	 * @param bUser if true, removes a reference from the user reference counter. If false, removes a reference from the system reference counter. */
+	void UnCacheImage(const mu::EXTERNAL_IMAGE_ID& Id, bool bUser);
+
+	/** Removes a reference to all images. All images which no longer have references will be uncached.
+	 * @param bUser if true, removes a references from the user reference counter. If false, removes a reference from the system reference counter. */
+	void ClearCache(bool bUser);
+
+	void CacheImages(const mu::Parameters& Parameters);
+	void UnCacheImages(const mu::Parameters& Parameters);
+	
 	/** List of actual image providers that have been registered to the CustomizableObjectSystem. */
 	TArray< TWeakObjectPtr<class UCustomizableSystemImageProvider> > ImageProviders;
 
@@ -59,7 +69,13 @@ private:
 		/** If the above Image has not been loaded in the game thread, the TextureToLoad bulk data will be loaded
 		* from the Mutable thread when it's needed
 		*/
-		TObjectPtr<UTexture2D> TextureToLoad = nullptr;
+		UTexture2D* TextureToLoad = nullptr;
+
+		/** true of the reference maintained by the user. */
+		bool ReferencesUser = false;
+		
+		/** Number of reference maintained by the system. */
+		int32 ReferencesSystem = 0;
 	};
 
 	static inline const mu::FImageDesc DUMMY_IMAGE_DESC = 
@@ -73,10 +89,9 @@ private:
 	* This is only safely written from the game thread protected by the following critical section, and it
 	* is safely read from the mutable thread during the update of the instance or texture mips
 	*/
-	TMap<uint64, FUnrealMutableImageInfo> GlobalExternalImages;
+	TMap<FString, FUnrealMutableImageInfo> GlobalExternalImages;
 	
 	/** Access to GlobalExternalImages must be protected with this because it may be accessed concurrently from the 
 	Game thread to modify it and from the Mutable thread to read it. */
 	FCriticalSection ExternalImagesLock;
-
 };
diff --git a/Engine/Plugins/Experimental/Mutable/Source/CustomizableObject/Public/MuCO/CustomizableObjectInstanceDescriptor.h b/Engine/Plugins/Experimental/Mutable/Source/CustomizableObject/Public/MuCO/CustomizableObjectInstanceDescriptor.h
--- a/Engine/Plugins/Experimental/Mutable/Source/CustomizableObject/Public/MuCO/CustomizableObjectInstanceDescriptor.h
+++ b/Engine/Plugins/Experimental/Mutable/Source/CustomizableObject/Public/MuCO/CustomizableObjectInstanceDescriptor.h
@@ -115,17 +115,14 @@ struct CUSTOMIZABLEOBJECT_API FCustomizableObjectInstanceDescriptor
 	void SetFloatParameterSelectedOption(const FString& FloatParamName, float FloatValue, int32 RangeIndex = -1);
 
 	/** Gets the value of a texture parameter with name "TextureParamName". */
-	uint64 GetTextureParameterSelectedOption(const FString& TextureParamName, int32 RangeIndex) const;
-
-	/** Gets the texture of a texture parameter with name "TextureParamName". */
-	UTexture2D* GetTextureParameterSelectedOptionT(const FString& TextureParamName, int32 RangeIndex) const;
+	FString GetTextureParameterSelectedOption(const FString& TextureParamName, int32 RangeIndex) const;
 
 	/** Sets the texture value "TextureValue" of a texture parameter with index "TextureParamIndex". */
-	void SetTextureParameterSelectedOption(const FString& TextureParamName, uint64 TextureValue, int32 RangeIndex);
+	void SetTextureParameterSelectedOption(const FString& TextureParamName, const FString& TextureValue, int32 RangeIndex);
 
-	/** Sets the texture "Texture" of a texture parameter with index "TextureParamIndex". */
+	/** @deprecated Use SetTextureParameterSelectedOption instead. */
 	void SetTextureParameterSelectedOptionT(const FString& TextureParamName, UTexture2D* TextureValue, int32 RangeIndex);
-
+	
 	/** Gets the value of a color parameter with name "ColorParamName". */
 	FLinearColor GetColorParameterSelectedOption(const FString& ColorParamName) const;
 
diff --git a/Engine/Plugins/Experimental/Mutable/Source/CustomizableObject/Public/MuCO/CustomizableObjectMipDataProvider.h b/Engine/Plugins/Experimental/Mutable/Source/CustomizableObject/Public/MuCO/CustomizableObjectMipDataProvider.h
--- a/Engine/Plugins/Experimental/Mutable/Source/CustomizableObject/Public/MuCO/CustomizableObjectMipDataProvider.h
+++ b/Engine/Plugins/Experimental/Mutable/Source/CustomizableObject/Public/MuCO/CustomizableObjectMipDataProvider.h
@@ -9,6 +9,11 @@
 #include "MuR/Image.h"
 #include "Streaming/TextureMipDataProvider.h"
 
+#include "DefaultImageProvider.h"
+#include "MuCO/CustomizableObjectSystemPrivate.h"
+#include "MuCO/UnrealMutableImageProvider.h"
+#include "MuR/System.h"
+
 #include "CustomizableObjectMipDataProvider.generated.h"
 
 enum EPixelFormat : uint8;
@@ -23,8 +28,26 @@ class UTexture;
 /** This struct stores the data relevant for the construction of a specific texture. 
 * This includes all the data required to rebuild the image (or any of its mips).
 */
-struct FMutableUpdateContext
+class FMutableUpdateContext
 {
+public:
+	FMutableUpdateContext() = default;
+	
+	FMutableUpdateContext(mu::Ptr<mu::System> InSystem, TSharedPtr<mu::Model, ESPMode::ThreadSafe> InModel, mu::Ptr<const mu::Parameters> InParameters, int32 InState);
+
+	~FMutableUpdateContext();
+
+	mu::Ptr<mu::System> GetSystem() const;
+	
+	TSharedPtr<mu::Model, ESPMode::ThreadSafe> GetModel() const;
+
+	mu::Ptr<const mu::Parameters> GetParameters() const;
+	
+	int32 GetState() const;
+
+	const TArray<mu::Ptr<const mu::Image>>& GetImageParameterValues() const;
+	
+private:
 	mu::Ptr<mu::System> System;
 	TSharedPtr<mu::Model, ESPMode::ThreadSafe> Model;
 	mu::Ptr<const mu::Parameters> Parameters;
diff --git a/Engine/Plugins/Experimental/Mutable/Source/CustomizableObject/Public/MuCO/CustomizableObjectParameterTypeDefinitions.h b/Engine/Plugins/Experimental/Mutable/Source/CustomizableObject/Public/MuCO/CustomizableObjectParameterTypeDefinitions.h
--- a/Engine/Plugins/Experimental/Mutable/Source/CustomizableObject/Public/MuCO/CustomizableObjectParameterTypeDefinitions.h
+++ b/Engine/Plugins/Experimental/Mutable/Source/CustomizableObject/Public/MuCO/CustomizableObjectParameterTypeDefinitions.h
@@ -2,6 +2,8 @@
 
 #pragma once
 
+#include "MuR/Types.h"
+
 #include "CustomizableObjectParameterTypeDefinitions.generated.h"
 
 class UTexture2D;
@@ -157,19 +159,19 @@ struct FCustomizableObjectTextureParameterValue
 {
 	GENERATED_USTRUCT_BODY()
 
-	static constexpr uint64 DEFAULT_PARAMETER_VALUE = 0;
+	inline static const FString DEFAULT_PARAMETER_VALUE = {}; // Means no Texture Parameter assigned.
 	
 	UPROPERTY(Category = CustomizableObjectTextureParameterValue, VisibleAnywhere)
 	FString ParameterName;
 
 	UPROPERTY(Category = CustomizableObjectTextureParameterValue, VisibleAnywhere)
-	uint64 ParameterValue = DEFAULT_PARAMETER_VALUE;
+	FString ParameterValue = DEFAULT_PARAMETER_VALUE;
 
 	UPROPERTY(Category = CustomizableObjectTextureParameterValue, VisibleAnywhere)
 	FString Uid;
 
 	UPROPERTY(Category = CustomizableObjectTextureParameterValue, VisibleAnywhere)
-	TArray<uint64> ParameterRangeValues;
+	TArray<FString> ParameterRangeValues;
 };
 
 
@@ -178,7 +180,7 @@ inline uint32 GetTypeHash(const FCustomizableObjectTextureParameterValue& Key)
 	uint32 Hash = GetTypeHash(Key.ParameterName);
 	Hash = HashCombine(Hash, GetTypeHash(Key.ParameterValue));
 
-	for (const uint64 Value : Key.ParameterRangeValues)
+	for (const FString& Value : Key.ParameterRangeValues)
     {
     	Hash = HashCombine(Hash, GetTypeHash(Value));
     }
diff --git a/Engine/Plugins/Experimental/Mutable/Source/CustomizableObject/Public/MuCO/DefaultImageProvider.h b/Engine/Plugins/Experimental/Mutable/Source/CustomizableObject/Public/MuCO/DefaultImageProvider.h
--- a/Engine/Plugins/Experimental/Mutable/Source/CustomizableObject/Public/MuCO/DefaultImageProvider.h
+++ b/Engine/Plugins/Experimental/Mutable/Source/CustomizableObject/Public/MuCO/DefaultImageProvider.h
@@ -16,9 +16,7 @@ namespace mu
 
 /** Simple image provider that translates UTexture2D to Mutable IDs.
  *
- * Allows the reuse of UTexture2D.
- * IDs will be garbage collected if they are referenced inside an Instance (GetTextureParameters()) or are not mark explicitly to keep (Keep(...)).
- */
+ * Allows the reuse of UTexture2D. */
 UCLASS()
 class CUSTOMIZABLEOBJECT_API UDefaultImageProvider : public UCustomizableSystemImageProvider
 {
@@ -26,39 +24,19 @@ class CUSTOMIZABLEOBJECT_API UDefaultImageProvider : public UCustomizableSystemI
 
 public:
 	// UCustomizableSystemImageProvider interface
-	virtual ValueType HasTextureParameterValue(int64 ID) override;
-	virtual UTexture2D* GetTextureParameterValue(int64 ID) override;
+	virtual ValueType HasTextureParameterValue(FString ID) override;
+	virtual UTexture2D* GetTextureParameterValue(FString ID) override;
 	virtual void GetTextureParameterValues(TArray<FCustomizableObjectExternalTexture>& OutValues) override;
-
-	// Own interface
-
-	UTexture2D* Get(uint64 TextureId) const;
-
-	int64 Get(UTexture2D* Texture) const;
-
-	/** If the TextureId does not get assigned to an Instance Texture Parameter it will be GC on the next Instance update. */
-	uint64 GetOrAdd(UTexture2D* Texture); 
-
-	void CacheTextures(const mu::Parameters&);
-
-	/** Keep the TextureId from being collected by the GC. */
-	void Keep(uint64 TextureId, bool bKeep);
 	
-	void GarbageCollectTextureIds(const mu::Parameters&);
+	/** Add a Texture to the provider. */
+	FString Add(UTexture2D* Texture); 
 
-private:
-	int32 ToIndex(uint64 TextureId) const;
-
-	uint64 ToTextureId(int32 TextureIndex) const;
-	
-	/** Change it to avoid TextureId collisions with other Image Providers. */
-	const int32 BASE_ID = 1;
-	const int32 MAX_IDS = 1000;
-
-	UPROPERTY()
-	TArray<TObjectPtr<UTexture2D>> Textures;
+	/** Remove a Texture from the provider. */
+	void Remove(UTexture2D* Texture); 
 
+private:
+	// Always contains valid pointers
 	UPROPERTY()
-	TArray<bool> KeepTextures;
+	TMap<FString, TObjectPtr<UTexture2D>> Textures;
 };
 
diff --git a/Engine/Plugins/Experimental/Mutable/Source/CustomizableObjectEditor/Private/MuCOE/SCustomizableInstanceProperties.cpp b/Engine/Plugins/Experimental/Mutable/Source/CustomizableObjectEditor/Private/MuCOE/SCustomizableInstanceProperties.cpp
--- a/Engine/Plugins/Experimental/Mutable/Source/CustomizableObjectEditor/Private/MuCOE/SCustomizableInstanceProperties.cpp
+++ b/Engine/Plugins/Experimental/Mutable/Source/CustomizableObjectEditor/Private/MuCOE/SCustomizableInstanceProperties.cpp
@@ -773,7 +773,7 @@ void SCustomizableInstanceProperties::AddParameter(int32 ParamIndexInObject)
 		const TArray<FCustomizableObjectTextureParameterValue>& TextureParameters = CustomInstance->GetTextureParameters();
 		TSharedPtr<FString> InitiallySelected = TextureParameterValueNames[0]; // First index is always the None option.
 
-		const uint64 ParameterValue = CustomInstance->GetTextureParameterSelectedOption(ParamName);
+		const FString ParameterValue = CustomInstance->GetTextureParameterSelectedOption(ParamName);
 		
 		// Look for the value index
 		for (int32 ValueIndex = 0; ValueIndex < TextureParameterValueNames.Num(); ++ValueIndex)
diff --git a/Engine/Plugins/Experimental/Mutable/Source/CustomizableObjectEditor/Private/MuCOE/StressTest.cpp b/Engine/Plugins/Experimental/Mutable/Source/CustomizableObjectEditor/Private/MuCOE/StressTest.cpp
--- a/Engine/Plugins/Experimental/Mutable/Source/CustomizableObjectEditor/Private/MuCOE/StressTest.cpp
+++ b/Engine/Plugins/Experimental/Mutable/Source/CustomizableObjectEditor/Private/MuCOE/StressTest.cpp
@@ -666,7 +666,7 @@ void FRunningStressTest::PrintParameterData(const MostExpensiveInstanceData& Dat
 		{
 			Line = FString::Printf(TEXT("\t\t\t %s"), *it.ParameterName);
 			AppendWhitespace(Line, MaxLength - Line.Len());
-			Line += FString::Printf(TEXT("%d"), it.ParameterValue);
+			Line += it.ParameterValue;
 			UE_LOG(LogMutable, Log, TEXT("%s"), *Line);
 		}
 	}
diff --git a/Engine/Plugins/Experimental/Mutable/Source/CustomizableObjectEditor/Public/MuCOE/SCustomizableInstanceProperties.h b/Engine/Plugins/Experimental/Mutable/Source/CustomizableObjectEditor/Public/MuCOE/SCustomizableInstanceProperties.h
--- a/Engine/Plugins/Experimental/Mutable/Source/CustomizableObjectEditor/Public/MuCOE/SCustomizableInstanceProperties.h
+++ b/Engine/Plugins/Experimental/Mutable/Source/CustomizableObjectEditor/Public/MuCOE/SCustomizableInstanceProperties.h
@@ -2,6 +2,7 @@
 
 #pragma once
 
+#include "MuR/Types.h"
 #include "Widgets/Input/SSpinBox.h"
 
 enum class ECheckBoxState : uint8;
@@ -97,8 +98,8 @@ private:
 
 	// These arrays store the textures available for texture parameters of the model.
 	// These come from the texture generators registered in the CustomizableObjectSystem
-	TArray< TSharedPtr<FString> > TextureParameterValueNames;
-	TArray< uint64_t > TextureParameterValues;
+	TArray<TSharedPtr<FString>> TextureParameterValueNames;
+	TArray<FString> TextureParameterValues;
 
 	struct FSliderData
 	{
diff --git a/Engine/Plugins/Experimental/Mutable/Source/MutableRuntime/Private/MuR/Model.cpp b/Engine/Plugins/Experimental/Mutable/Source/MutableRuntime/Private/MuR/Model.cpp
--- a/Engine/Plugins/Experimental/Mutable/Source/MutableRuntime/Private/MuR/Model.cpp
+++ b/Engine/Plugins/Experimental/Mutable/Source/MutableRuntime/Private/MuR/Model.cpp
@@ -443,7 +443,7 @@ namespace mu
         if (!m_pD->m_program.m_parameters.IsValidIndex(Index) ||
             m_pD->m_program.m_parameters[Index].m_type != PARAMETER_TYPE::T_IMAGE)
         {
-            return 0;
+            return {};
         }
 		
         return m_pD->m_program.m_parameters[Index].m_defaultValue.Get<ParamImageType>();
diff --git a/Engine/Plugins/Experimental/Mutable/Source/MutableRuntime/Private/MuR/Parameters.cpp b/Engine/Plugins/Experimental/Mutable/Source/MutableRuntime/Private/MuR/Parameters.cpp
--- a/Engine/Plugins/Experimental/Mutable/Source/MutableRuntime/Private/MuR/Parameters.cpp
+++ b/Engine/Plugins/Experimental/Mutable/Source/MutableRuntime/Private/MuR/Parameters.cpp
@@ -770,7 +770,7 @@ namespace mu
              ||
              GetType(index) != PARAMETER_TYPE::T_IMAGE )
         {
-            return 0;
+            return {};
         }
 		
 		// Single value case
diff --git a/Engine/Plugins/Experimental/Mutable/Source/MutableRuntime/Private/MuR/Serialisation.cpp b/Engine/Plugins/Experimental/Mutable/Source/MutableRuntime/Private/MuR/Serialisation.cpp
--- a/Engine/Plugins/Experimental/Mutable/Source/MutableRuntime/Private/MuR/Serialisation.cpp
+++ b/Engine/Plugins/Experimental/Mutable/Source/MutableRuntime/Private/MuR/Serialisation.cpp
@@ -26,7 +26,26 @@ namespace mu
     MUTABLE_IMPLEMENT_POD_SERIALISABLE( uint8 );  
     MUTABLE_IMPLEMENT_POD_SERIALISABLE( uint16 ); 
     MUTABLE_IMPLEMENT_POD_SERIALISABLE( uint32 ); 
-    MUTABLE_IMPLEMENT_POD_SERIALISABLE( uint64 );
+    MUTABLE_IMPLEMENT_POD_SERIALISABLE( uint64 )
+
+
+	template <>
+    void operator<<<FString>(OutputArchive& arch, const FString& t)
+    {
+	    const TArray<TCHAR> Data = t.GetCharArray();
+    	arch << Data;
+    }
+
+
+    template <>
+    void operator>><FString>(InputArchive& arch, FString& t)
+    {
+	    TArray<TCHAR> Data;
+    	arch >> Data;
+    	
+    	t = FString(Data);
+    }
+                                 
 
 	MUTABLE_IMPLEMENT_POD_VECTOR_SERIALISABLE(float);                                 
 	MUTABLE_IMPLEMENT_POD_VECTOR_SERIALISABLE(double);                                
@@ -48,8 +67,9 @@ namespace mu
 	MUTABLE_IMPLEMENT_POD_VECTOR_SERIALISABLE(vec3f);                                 
 	MUTABLE_IMPLEMENT_POD_VECTOR_SERIALISABLE(mat3f);                                 
 	MUTABLE_IMPLEMENT_POD_VECTOR_SERIALISABLE(mat4f);                                 
-	MUTABLE_IMPLEMENT_POD_VECTOR_SERIALISABLE(vec2<int>);                             
-                                                                                      
+	MUTABLE_IMPLEMENT_POD_VECTOR_SERIALISABLE(vec2<int>);
+	MUTABLE_IMPLEMENT_POD_VECTOR_SERIALISABLE(TCHAR);
+	
 	// Unreal POD Serializables                                                       
 	MUTABLE_IMPLEMENT_POD_SERIALISABLE(FUintVector2);                                 
 	MUTABLE_IMPLEMENT_POD_SERIALISABLE(UE::Math::TIntVector2<uint16>);                
diff --git a/Engine/Plugins/Experimental/Mutable/Source/MutableRuntime/Private/MuR/SerialisationPrivate.h b/Engine/Plugins/Experimental/Mutable/Source/MutableRuntime/Private/MuR/SerialisationPrivate.h
--- a/Engine/Plugins/Experimental/Mutable/Source/MutableRuntime/Private/MuR/SerialisationPrivate.h
+++ b/Engine/Plugins/Experimental/Mutable/Source/MutableRuntime/Private/MuR/SerialisationPrivate.h
@@ -358,6 +358,7 @@ namespace mu
 	MUTABLE_DEFINE_POD_VECTOR_SERIALISABLE(mat3f);
 	MUTABLE_DEFINE_POD_VECTOR_SERIALISABLE(mat4f);
 	MUTABLE_DEFINE_POD_VECTOR_SERIALISABLE(vec2<int>);
+	MUTABLE_DEFINE_POD_VECTOR_SERIALISABLE(TCHAR);
 
 	// Unreal POD Serializables
 	MUTABLE_DEFINE_POD_SERIALISABLE(FUintVector2);
diff --git a/Engine/Plugins/Experimental/Mutable/Source/MutableRuntime/Public/MuR/Serialisation.h b/Engine/Plugins/Experimental/Mutable/Source/MutableRuntime/Public/MuR/Serialisation.h
--- a/Engine/Plugins/Experimental/Mutable/Source/MutableRuntime/Public/MuR/Serialisation.h
+++ b/Engine/Plugins/Experimental/Mutable/Source/MutableRuntime/Public/MuR/Serialisation.h
@@ -337,6 +337,14 @@ namespace mu
     MUTABLE_DEFINE_POD_SERIALISABLE( uint16 );
     MUTABLE_DEFINE_POD_SERIALISABLE( uint32 );
     MUTABLE_DEFINE_POD_SERIALISABLE( uint64 );
+
+	
+	//---------------------------------------------------------------------------------------------
+	template<>
+	void operator<< <FString>(OutputArchive& arch, const FString& t);
+	
+	template<>
+	void operator>> <FString>(InputArchive& arch, FString& t);
 	
 	
 	//---------------------------------------------------------------------------------------------
diff --git a/Engine/Plugins/Experimental/Mutable/Source/MutableRuntime/Public/MuR/Types.h b/Engine/Plugins/Experimental/Mutable/Source/MutableRuntime/Public/MuR/Types.h
--- a/Engine/Plugins/Experimental/Mutable/Source/MutableRuntime/Public/MuR/Types.h
+++ b/Engine/Plugins/Experimental/Mutable/Source/MutableRuntime/Public/MuR/Types.h
@@ -12,7 +12,7 @@ namespace mu
     typedef uint32            RESOURCE_ID;
 
     //!
-    typedef uint64            EXTERNAL_IMAGE_ID;
+    typedef FString           EXTERNAL_IMAGE_ID;
 
 
 }
diff --git a/Engine/Plugins/Experimental/Mutable/Source/MutableTools/Private/MuT/CodeGenerator_Image.cpp b/Engine/Plugins/Experimental/Mutable/Source/MutableTools/Private/MuT/CodeGenerator_Image.cpp
--- a/Engine/Plugins/Experimental/Mutable/Source/MutableTools/Private/MuT/CodeGenerator_Image.cpp
+++ b/Engine/Plugins/Experimental/Mutable/Source/MutableTools/Private/MuT/CodeGenerator_Image.cpp
@@ -290,7 +290,7 @@ namespace mu
 			op->parameter.m_name = node.m_name;
 			op->parameter.m_uid = node.m_uid;
             op->parameter.m_type = PARAMETER_TYPE::T_IMAGE;
-        	op->parameter.m_defaultValue.Set<ParamImageType>(0);
+        	op->parameter.m_defaultValue.Set<ParamImageType>(mu::EXTERNAL_IMAGE_ID());
 
 			// Generate the code for the ranges
 			for (int32 a = 0; a < node.m_ranges.Num(); ++a)
