[mutable] Fixed LayoutBlocksEditor not marking outer as dirty when modifying layout blocks (scale or position)

diff --git a/Engine/Plugins/Experimental/Mutable/Source/CustomizableObjectEditor/Private/MuCOE/SCustomizableObjectNodeLayoutBlocksEditor.cpp b/Engine/Plugins/Experimental/Mutable/Source/CustomizableObjectEditor/Private/MuCOE/SCustomizableObjectNodeLayoutBlocksEditor.cpp
--- a/Engine/Plugins/Experimental/Mutable/Source/CustomizableObjectEditor/Private/MuCOE/SCustomizableObjectNodeLayoutBlocksEditor.cpp
+++ b/Engine/Plugins/Experimental/Mutable/Source/CustomizableObjectEditor/Private/MuCOE/SCustomizableObjectNodeLayoutBlocksEditor.cpp
@@ -505,6 +505,8 @@ void SCustomizableObjectNodeLayoutBlocksEditor::OnBlockChanged( FGuid BlockId, F
 			{
 				B.Min = Block.Min;
 				B.Max = Block.Max;
+
+				CurrentLayout->MarkPackageDirty();
 			}
 		}
 	}
