[mutable] Implemented validator to indirectly check data validity for the referenced COs of the checked out resources.
- Added Validation plugin as a dependancy for Mutable.
- Created MutableValidation module (MuV) as an Editor module.
- New UAssetValidator_ReferencedCustomizableObjects Validator will mark the asset to be validated as invalid if any of its referenced COs return an invalid validation result.
- Moved validation code from UCustomizableObject to UAssetValidator_CustomizableObjects. It allows us to have more information about the context of validation (reason why it is being performed) and only validate if not saving or running a commandlet.
- Reduced amount of compilations per root Co to 7 from 15 (6 lod biases and default).
- Added setting to disable and enable indirect CO validation.

diff --git a/Engine/Plugins/Experimental/Mutable/Mutable.uplugin b/Engine/Plugins/Experimental/Mutable/Mutable.uplugin
--- a/Engine/Plugins/Experimental/Mutable/Mutable.uplugin
+++ b/Engine/Plugins/Experimental/Mutable/Mutable.uplugin
@@ -24,6 +24,11 @@
 			"Type": "Editor",
 			"LoadingPhase": "PreDefault"
 		},
+		{
+			"Name": "MutableValidation",
+			"Type": "Editor",
+			"LoadingPhase": "PreDefault"
+		},
 		{
 			"Name": "CustomizableObjectEditor",
 			"Type": "Editor",
@@ -48,6 +53,10 @@
 		{
 			"Name": "StructUtils",
 			"Enabled": true
+		},
+		{
+			"Name": "DataValidation",
+			"Enabled": true
 		}
 	],
 
diff --git a/Engine/Plugins/Experimental/Mutable/Source/CustomizableObject/Classes/MuCO/CustomizableObject.h b/Engine/Plugins/Experimental/Mutable/Source/CustomizableObject/Classes/MuCO/CustomizableObject.h
--- a/Engine/Plugins/Experimental/Mutable/Source/CustomizableObject/Classes/MuCO/CustomizableObject.h
+++ b/Engine/Plugins/Experimental/Mutable/Source/CustomizableObject/Classes/MuCO/CustomizableObject.h
@@ -1630,36 +1630,6 @@ public:
 	/** Return the names used by mutable to identify which mu::Image should be considered of LowPriority. */
 	void GetLowPriorityTextureNames(TArray<FString>& OutTextureNames);
 
-#if WITH_EDITOR
-
-	// UObject Interface -> Data validation
-	virtual EDataValidationResult IsDataValid(class FDataValidationContext& Context) override;
-	// End of UObject Interface
-
-private:
-
-	/** Cached handle to be able later to remove the bound method from the FEditorDelegates::OnPostAssetValidation delegate */
-	inline static FDelegateHandle OnPostCOValidationHandle;
-	
-	/** Collection with all root objects tested during this IsDataValidRun. Shared with all COs */
-	inline static TArray<UCustomizableObject*> AlreadyValidatedRootObjects;
-	
-	/** Method invoked once the validation of all assets has been completed. */
-	static void OnPostCOsValidation();
-
-public:
-	// UObject Interface -> Asset saving
-	virtual void PreSaveRoot(FObjectPreSaveRootContext ObjectSaveContext) override;
-	// End of UObject Interface
-
-private:
-	/** Flag that tells us if the validation of the asset has been triggered by the saving of it.
-	 * It gets consumed by IsDataValid and it is that same function the one that resets the value of this flag.
-	 */
-	bool bIsValidationTriggeredBySave = false;
-
-#endif
-
 	/** Used to prevent GC of MaskOutCache and keep it in memory while it's needed */
 	UPROPERTY(Transient)
 	TObjectPtr<UMutableMaskOutCache> MaskOutCache_HardRef;
diff --git a/Engine/Plugins/Experimental/Mutable/Source/CustomizableObject/Private/MuCO/CustomizableObject.cpp b/Engine/Plugins/Experimental/Mutable/Source/CustomizableObject/Private/MuCO/CustomizableObject.cpp
--- a/Engine/Plugins/Experimental/Mutable/Source/CustomizableObject/Private/MuCO/CustomizableObject.cpp
+++ b/Engine/Plugins/Experimental/Mutable/Source/CustomizableObject/Private/MuCO/CustomizableObject.cpp
@@ -1729,178 +1729,6 @@ void UCustomizableObject::GetLowPriorityTextureNames(TArray<FString>& OutTexture
 }
 
 
-#if WITH_EDITOR
-void UCustomizableObject::PreSaveRoot(FObjectPreSaveRootContext ObjectSaveContext)
-{
-	UObject::PreSaveRoot(ObjectSaveContext);
-
-	// Tell the validation system on this object that the validation that is going to be next invoked is due to
-	// this asset being saved.
-	// This value will be set to false by the validation method on this object so subsequent validation attempts
-	// get treated as expected.
-	bIsValidationTriggeredBySave = true;
-}
-
-
-EDataValidationResult UCustomizableObject::IsDataValid(FDataValidationContext& Context)
-{
-	// This method seems to be designed to check data errors (like variables with unexpected values).
-	// Currently it does not check if the root that we are compiling has already been compiled during another validation.
-	
-	EDataValidationResult Result = EDataValidationResult::NotValidated;
-
-	// If validation is invoked by the saving of the asset just skip it. It is too expensive.
-	if (bIsValidationTriggeredBySave)
-    {
-		// Reenable the validation for this object after running the saving process and skipping this validation run
-		bIsValidationTriggeredBySave = false;
-    	return Result;
-    }
-
-	// Skip validation when cooking the assets. The validation of the CO is designed to be used explicitly by the user
-	// and not during automated operations like saving or cooking or any other automated action.
-	if (IsRunningCommandlet())
-	{
-		return Result;
-	}
-
-	UE_LOG(LogMutable,Verbose,TEXT("Running data validation checks for %s CO."),*this->GetName());
-
-	// Bind the post validation method to the post validation delegate if not bound already to be able to know when the validation
-	// operation (for all assets) concludes
-	if (!UCustomizableObject::OnPostCOValidationHandle.IsValid())
-	{
-		UCustomizableObject::OnPostCOValidationHandle = FEditorDelegates::OnPostAssetValidation.AddStatic(OnPostCOsValidation);	
-	}
-	
-	// Request a compiler to be able to locate the root and to compile it
-	const TUniquePtr<FCustomizableObjectCompilerBase> Compiler =
-		TUniquePtr<FCustomizableObjectCompilerBase>(UCustomizableObjectSystem::GetInstance()->GetNewCompiler());
-	
-	// Find out witch is the root for this CO (it may be itself but that is OK)
-	UCustomizableObject* RootObject = Compiler->GetRootObject(this);
-	check (RootObject);
-	
-	// Check that the object to be compiled has not already been compiled
-	if (UCustomizableObject::AlreadyValidatedRootObjects.Contains(RootObject))
-	{
-		return Result;
-	}
-
-	// Root Object not yet tested -> Proceed with the testing
-	
-	// Collection of configurations to be tested with the located root object
-	constexpr int32 MaxBias = 15;
-	TArray<FCompilationOptions> CompilationOptionsToTest;
-	for (int32 LodBias = 0; LodBias < MaxBias; LodBias++)
-	{
-		FCompilationOptions ModifiedCompilationOptions = this->CompileOptions;
-		ModifiedCompilationOptions.bForceLargeLODBias = true;	
-		ModifiedCompilationOptions.DebugBias = LodBias;	
-
-		// Add one configuration object for each bias setting
-		CompilationOptionsToTest.Add(ModifiedCompilationOptions);
-	}
-	
-	// Add current configuration to be tested as well.
-	CompilationOptionsToTest.Add(this->CompileOptions);
-
-	
-	// Caches with all the data produced by the subsequent compilations of the root of this CO
-	TArray<FText> CachedValidationErrors;
-	TArray<FText> CachedValidationWarnings;
-	TArray<ECustomizableObjectCompilationState> CachedCompilationEndStates;
-	
-	// Iterate over the compilation options that we want to test and perform the compilation
-	for	(const FCompilationOptions& Options : CompilationOptionsToTest)
-	{
-		// Run Sync compilation -> Warning : Potentially long operation -------------
-		Compiler->Compile(*RootObject, Options, false);
-		// --------------------------------------------------------------------------
-		
-		// Get compilation errors and warnings
-		TArray<FText> CompilationErrors;
-		TArray<FText> CompilationWarnings;
-		Compiler->GetCompilationMessages(CompilationWarnings, CompilationErrors);
-		
-		// Cache the messages returned by the compiler
-		for ( const FText& FoundError : CompilationErrors)
-		{
-			// Add message if not already present
-			if (!CachedValidationErrors.ContainsByPredicate([&FoundError](const FText& ArrayEntry)
-				{ return FoundError.EqualTo(ArrayEntry);}))
-			{
-				CachedValidationErrors.Add(FoundError);
-			}
-		}
-		for ( const FText& FoundWarning : CompilationWarnings)
-		{
-			if (!CachedValidationWarnings.ContainsByPredicate([&FoundWarning](const FText& ArrayEntry)
-				{ return FoundWarning.EqualTo(ArrayEntry);}))
-			{
-				CachedValidationWarnings.Add(FoundWarning);
-			}
-		}
-	
-		CachedCompilationEndStates.Add(Compiler->GetCompilationState());
-	}
-	
-	// Cache root object to avoid processing it again when processing another CO related with the same root CO
-	AlreadyValidatedRootObjects.Add(RootObject);
-	
-	// Wrapping up : Fill message output caches and determine if the compilation was successful or not
-
-	// Provide the warning and log messages to the context object (so it can later notify the user using the UI)
-	for (const FText& ValidationError : CachedValidationErrors)
-	{
-		Context.AddError(ValidationError);
-	}
-	for (const FText& ValidationWarning : CachedValidationWarnings)
-	{
-		Context.AddWarning(ValidationWarning);
-	}
-	
-	// Return informed guess about what the validation state of this object should be
-
-	// If one or more tests failed to ran then the result must be invalid
-	if (CachedCompilationEndStates.Contains(ECustomizableObjectCompilationState::Failed))
-	{
-		// Early CO compilation error (before starting mutable compilation) -> Output is invalid
-		Result = EDataValidationResult::Invalid;
-		UE_LOG(LogMutable, Error,
-			   TEXT("Compilation of %s failed : Check previous log messages to get more information."),
-			   *this->GetName())
-	}
-	// If it contains invalid states then notify about it too:
-	// ECustomizableObjectCompilationState::None would mean the resource is locked (and should not be)
-	// ECustomizableObjectCompilationState::InProgress should not be possible since we are compiling synchronously.
-	else if (CachedCompilationEndStates.Contains(ECustomizableObjectCompilationState::InProgress) ||
-		CachedCompilationEndStates.Contains(ECustomizableObjectCompilationState::None))
-	{
-		checkNoEntry();
-	}
-	// All compilations completed successfully
-	else 
-	{
-		// If a warning or error was found then this object failed the validation process
-		Result = (CachedValidationWarnings.IsEmpty() && CachedValidationErrors.IsEmpty()) ? EDataValidationResult::Valid : EDataValidationResult::Invalid;
-	}
-	
-	return Result;
-}
-
-void UCustomizableObject::OnPostCOsValidation()
-{
-	// Unbound this method from the validation end delegate
-	UCustomizableObject::OnPostCOValidationHandle.Reset();
-
-	// Clear collection with the already processed COs once the validation system has completed its operation
-	UCustomizableObject::AlreadyValidatedRootObjects.Empty();
-}
-
-#endif
-
-
 FGuid UCustomizableObject::GetCompilationGuid() const
 {
 	return CompilationGuid;
diff --git a/Engine/Plugins/Experimental/Mutable/Source/CustomizableObjectEditor/Private/MuCOE/CustomizableObjectCompiler.cpp b/Engine/Plugins/Experimental/Mutable/Source/CustomizableObjectEditor/Private/MuCOE/CustomizableObjectCompiler.cpp
--- a/Engine/Plugins/Experimental/Mutable/Source/CustomizableObjectEditor/Private/MuCOE/CustomizableObjectCompiler.cpp
+++ b/Engine/Plugins/Experimental/Mutable/Source/CustomizableObjectEditor/Private/MuCOE/CustomizableObjectCompiler.cpp
@@ -465,17 +465,16 @@ UCustomizableObject* FCustomizableObjectCompiler::GetRootObject( UCustomizableOb
 	// Grab a node to start the search -> Get the root since it should be always present
 	bool bMultipleBaseObjectsFound = false;
 	UCustomizableObjectNodeObject* ObjectRootNode = GetRootNode(InObject, bMultipleBaseObjectsFound);
-	
-	if (ObjectRootNode->ParentObject)
+
+	if (ObjectRootNode && ObjectRootNode->ParentObject)
 	{
 		TArray<UCustomizableObject*> VisitedNodes;
 		return GetFullGraphRootObject(ObjectRootNode,VisitedNodes);
 	}
-	else
-	{
-		// No parent object found, return input as the parent of the graph
-		return InObject;
-	}
+
+	// No parent object found, return input as the parent of the graph
+	// This can also mean the ObjectRootNode does not exist because it has not been opened yet (so no nodes have been generated)
+	return InObject;
 }
 
 
diff --git a/Engine/Plugins/Experimental/Mutable/Source/MutableValidation/MutableValidation.Build.cs b/Engine/Plugins/Experimental/Mutable/Source/MutableValidation/MutableValidation.Build.cs
new file mode 100644
--- /dev/null
+++ b/Engine/Plugins/Experimental/Mutable/Source/MutableValidation/MutableValidation.Build.cs
@@ -0,0 +1,31 @@
+// Copyright Epic Games, Inc. All Rights Reserved.
+
+namespace UnrealBuildTool.Rules
+{
+	/// <summary>
+	/// Module designed to serve as the home for all validation systems running in the engine. 
+	/// </summary>
+	public class MutableValidation : ModuleRules
+	{
+		public MutableValidation(ReadOnlyTargetRules Target) : base(Target)
+		{
+			ShortName = "MuV";
+
+			DefaultBuildSettings = BuildSettingsVersion.V2;
+
+			PublicDependencyModuleNames.AddRange(new string[] { "Settings" });
+			
+			PrivateDependencyModuleNames.AddRange(
+				new string[] {
+					"Core",
+					"CoreUObject",
+					"Engine",
+					"UnrealEd",
+					
+					"DataValidation",
+					"CustomizableObject", 
+				}
+			);
+		}
+	}
+}
\ No newline at end of file
diff --git a/Engine/Plugins/Experimental/Mutable/Source/MutableValidation/Private/MuV/AssetValidator_CustomizableObjects.cpp b/Engine/Plugins/Experimental/Mutable/Source/MutableValidation/Private/MuV/AssetValidator_CustomizableObjects.cpp
new file mode 100644
--- /dev/null
+++ b/Engine/Plugins/Experimental/Mutable/Source/MutableValidation/Private/MuV/AssetValidator_CustomizableObjects.cpp
@@ -0,0 +1,233 @@
+// Copyright Epic Games, Inc. All Rights Reserved.
+
+#include "MuV/AssetValidator_CustomizableObjects.h"
+
+#include "DataValidationModule.h"
+#include "Editor.h"
+#include "AssetRegistry/AssetData.h"
+#include "Delegates/DelegateSignatureImpl.inl"
+#include "Engine/SkeletalMesh.h"
+#include "Engine/StaticMesh.h"
+#include "Engine/Texture.h"
+#include "Materials/Material.h"
+#include "MuCO/CustomizableObject.h"
+#include "MuCO/CustomizableObjectSystem.h"
+#include "UObject/NameTypes.h"
+#include "UObject/Object.h"
+
+#define LOCTEXT_NAMESPACE "CustomizableObjectsValidator"
+
+UAssetValidator_CustomizableObjects::UAssetValidator_CustomizableObjects() : Super()
+{
+	bIsEnabled = true;
+}
+
+
+bool UAssetValidator_CustomizableObjects::CanValidate_Implementation(const EDataValidationUsecase InUsecase) const
+{
+	// Do not run if saving or running a commandlet (we do not want CIS failing due to our warnings and errors)
+	return !(InUsecase == EDataValidationUsecase::Save || InUsecase == EDataValidationUsecase::Commandlet);
+}
+
+
+bool UAssetValidator_CustomizableObjects::CanValidateAsset_Implementation(UObject* InAsset) const
+{
+	return (InAsset ? InAsset->IsA(UCustomizableObject::StaticClass())  : false) ;
+}
+
+
+EDataValidationResult UAssetValidator_CustomizableObjects::ValidateLoadedAsset_Implementation(UObject* InAsset,
+	TArray<FText>& ValidationErrors)
+{
+	check(InAsset);
+
+	UCustomizableObject* CustomizableObjectToValidate = Cast<UCustomizableObject>(InAsset);
+	check (CustomizableObjectToValidate);
+	
+	// Validate that CO and if it fails then mark it as failed. Do not stop until running the validation over all COs
+	TArray<FText> CoValidationWarnings;
+	TArray<FText> CoValidationErrors;
+	const EDataValidationResult COValidationResult = IsCustomizableObjectValid(CustomizableObjectToValidate,CoValidationErrors,CoValidationWarnings);
+	
+	// Process the validation of the CO's output
+	if (COValidationResult == EDataValidationResult::Invalid )
+	{
+		// Cache warning logs
+		for	(const FText& WarningMessage : CoValidationWarnings)
+		{
+			AssetWarning(InAsset,WarningMessage);
+		}
+		
+		// Cache error logs -> They will tag the asset validation as failed
+		for (const FText& ErrorMessage : CoValidationErrors)
+		{
+			AssetFails(InAsset,ErrorMessage,ValidationErrors);
+		}
+
+		const FText ErrorMessage = FText::Format(LOCTEXT("CustomizableObjectsValidator", "Validation compilation of {0} CO failed."),  FText::FromString( CustomizableObjectToValidate->GetName()));
+		AssetFails(InAsset,ErrorMessage,ValidationErrors);
+	}
+	else
+	{
+		AssetPasses(InAsset);
+	}
+		
+	return GetValidationResult();
+}
+
+
+EDataValidationResult UAssetValidator_CustomizableObjects::IsCustomizableObjectValid(UCustomizableObject* InCustomizableObject, TArray<FText>& OutValidationErrors, TArray<FText>& OutValidationWarnings)
+{
+	EDataValidationResult Result = EDataValidationResult::NotValidated;
+
+	UE_LOG(LogMutable,Verbose,TEXT("Running data validation checks for %s CO."),*InCustomizableObject->GetName());
+
+	// Bind the post validation method to the post validation delegate if not bound already to be able to know when the validation
+	// operation (for all assets) concludes
+	if (!OnPostCOValidationHandle.IsValid())
+	{
+		OnPostCOValidationHandle = FEditorDelegates::OnPostAssetValidation.AddStatic(OnPostCOsValidation);	
+	}
+	
+	// Request a compiler to be able to locate the root and to compile it
+	const TUniquePtr<FCustomizableObjectCompilerBase> Compiler =
+		TUniquePtr<FCustomizableObjectCompilerBase>(UCustomizableObjectSystem::GetInstance()->GetNewCompiler());
+	
+	// Find out which is the root for this CO (it may be itself but that is OK)
+	UCustomizableObject* RootObject = Compiler->GetRootObject(InCustomizableObject);
+	check (RootObject);
+	
+	// Check that the object to be compiled has not already been compiled
+	if (AlreadyValidatedRootObjects.Contains(RootObject))
+	{
+		return Result;
+	}
+	
+	// Root Object not yet tested -> Proceed with the testing
+	
+	// Collection of configurations to be tested with the located root object
+	TArray<FCompilationOptions> CompilationOptionsToTest;
+
+	// Current configuration
+	CompilationOptionsToTest.Add(InCustomizableObject->CompileOptions);
+
+	// Configuration with LOD bias applied
+	// constexpr int32 MaxBias = 15;	
+	constexpr int32 MaxBias = 6;		// Reduced amount since we have this value in GenerateMutableSource for the max lod bias provided to mutable
+	for (int32 LodBias = 1; LodBias <= MaxBias; LodBias++)
+	{
+		FCompilationOptions ModifiedCompilationOptions = InCustomizableObject->CompileOptions;
+		ModifiedCompilationOptions.bForceLargeLODBias = true;	
+		ModifiedCompilationOptions.DebugBias = LodBias;	
+	
+		// Add one configuration object for each bias setting
+		CompilationOptionsToTest.Add(ModifiedCompilationOptions);
+	}
+	
+	// Caches with all the data produced by the subsequent compilations of the root of this CO
+	TArray<FText> CachedValidationErrors;
+	TArray<FText> CachedValidationWarnings;
+	
+	// Map with all the possible compilation states. We use this so at the end we can know if any of those states was returned by any of the compilation runs
+	TMap<ECustomizableObjectCompilationState,bool>PossibleEndCompilationStates;
+	PossibleEndCompilationStates.Add(ECustomizableObjectCompilationState::Completed,false);
+	PossibleEndCompilationStates.Add(ECustomizableObjectCompilationState::None,false);
+	PossibleEndCompilationStates.Add(ECustomizableObjectCompilationState::Failed,false);
+	PossibleEndCompilationStates.Add(ECustomizableObjectCompilationState::InProgress,false);
+	PossibleEndCompilationStates.Shrink();
+	
+	// Iterate over the compilation options that we want to test and perform the compilation
+	for	(const FCompilationOptions& Options : CompilationOptionsToTest)
+	{
+		// Run Sync compilation -> Warning : Potentially long operation -------------
+		Compiler->Compile(*RootObject, Options, false);
+		// --------------------------------------------------------------------------
+		
+		// Get compilation errors and warnings
+		TArray<FText> CompilationErrors;
+		TArray<FText> CompilationWarnings;
+		Compiler->GetCompilationMessages(CompilationWarnings, CompilationErrors);
+		
+		// Cache the messages returned by the compiler
+		for ( const FText& FoundError : CompilationErrors)
+		{
+			// Add message if not already present
+			if (!CachedValidationErrors.ContainsByPredicate([&FoundError](const FText& ArrayEntry)
+				{ return FoundError.EqualTo(ArrayEntry);}))
+			{
+				CachedValidationErrors.Add(FoundError);
+			}
+		}
+		for ( const FText& FoundWarning : CompilationWarnings)
+		{
+			if (!CachedValidationWarnings.ContainsByPredicate([&FoundWarning](const FText& ArrayEntry)
+				{ return FoundWarning.EqualTo(ArrayEntry);}))
+			{
+				CachedValidationWarnings.Add(FoundWarning);
+			}
+		}
+
+		// Flag the array with end results to have the current output as true since it was produced by this execution
+		ECustomizableObjectCompilationState CompilationEndResult = Compiler->GetCompilationState();
+		bool* Value = PossibleEndCompilationStates.Find(CompilationEndResult);
+		check (Value);		// If this fails it may mean we are getting a compilation state we are not considering. 
+		*Value = true;
+	}
+	
+	// Cache root object to avoid processing it again when processing another CO related with the same root CO
+	AlreadyValidatedRootObjects.Add(RootObject);
+	
+	// Wrapping up : Fill message output caches and determine if the compilation was successful or not
+	
+	// Provide the warning and log messages to the context object (so it can later notify the user using the UI)
+	const FString ReferencedCOName = FString(TEXT("\"" +  InCustomizableObject->GetName() + "\""));
+	for (const FText& ValidationError : CachedValidationErrors)
+	{
+		FText ComposedMessage = FText::Format(LOCTEXT("CustomizableObjectsValidator", "Customizable Object : {0} {1}"),  FText::FromString(ReferencedCOName),ValidationError );
+		OutValidationErrors.Add(ComposedMessage);
+	}
+	for (const FText& ValidationWarning : CachedValidationWarnings)
+	{
+		FText ComposedMessage = FText::Format(LOCTEXT("CustomizableObjectsValidator", "Customizable Object : {0} {1}"),  FText::FromString(ReferencedCOName),ValidationWarning );
+		OutValidationWarnings.Add(ComposedMessage);
+	}
+	
+	// Return informed guess about what the validation state of this object should be
+
+	// If it contains invalid states then notify about it too:
+	// ECustomizableObjectCompilationState::InProgress should not be possible since we are compiling synchronously.
+	check (*PossibleEndCompilationStates.Find(ECustomizableObjectCompilationState::InProgress) == false);
+	// ECustomizableObjectCompilationState::None would mean the resource is locked (and should not be)
+	check (*PossibleEndCompilationStates.Find(ECustomizableObjectCompilationState::None) == false);
+	
+	// If one or more tests failed to ran then the result must be invalid
+	if (*PossibleEndCompilationStates.Find(ECustomizableObjectCompilationState::Failed) == true)
+	{
+		// Early CO compilation error (before starting mutable compilation) -> Output is invalid
+		Result = EDataValidationResult::Invalid;
+	}
+	// All compilations completed successfully
+	else 
+	{
+		// If a warning or error was found then this object failed the validation process
+		Result = (CachedValidationWarnings.IsEmpty() && CachedValidationErrors.IsEmpty()) ? EDataValidationResult::Valid : EDataValidationResult::Invalid;
+	}
+	
+	return Result;
+}
+
+
+void UAssetValidator_CustomizableObjects::OnPostCOsValidation()
+{
+	// Unbound this method from the validation end delegate
+	check (OnPostCOValidationHandle.IsValid());
+	FEditorDelegates::OnPostAssetValidation.Remove(OnPostCOValidationHandle);
+	OnPostCOValidationHandle.Reset();
+
+	// Clear collection with the already processed COs once the validation system has completed its operation
+	AlreadyValidatedRootObjects.Empty();
+}
+
+
+#undef LOCTEXT_NAMESPACE
+
diff --git a/Engine/Plugins/Experimental/Mutable/Source/MutableValidation/Private/MuV/AssetValidator_ReferencedCustomizableObjects.cpp b/Engine/Plugins/Experimental/Mutable/Source/MutableValidation/Private/MuV/AssetValidator_ReferencedCustomizableObjects.cpp
new file mode 100644
--- /dev/null
+++ b/Engine/Plugins/Experimental/Mutable/Source/MutableValidation/Private/MuV/AssetValidator_ReferencedCustomizableObjects.cpp
@@ -0,0 +1,187 @@
+// Copyright Epic Games, Inc. All Rights Reserved.
+
+#include "MuV/AssetValidator_ReferencedCustomizableObjects.h"
+
+#include "DataValidationModule.h"
+#include "MutableValidationSettings.h"
+#include "AssetRegistry/AssetData.h"
+#include "AssetRegistry/AssetRegistryModule.h"
+#include "Engine/SkeletalMesh.h"
+#include "Engine/StaticMesh.h"
+#include "Engine/Texture.h"
+#include "Materials/Material.h"
+#include "MuCO/CustomizableObject.h"
+#include "MuV/AssetValidator_CustomizableObjects.h"
+#include "UObject/NameTypes.h"
+#include "UObject/Object.h"
+
+#define LOCTEXT_NAMESPACE "ReferencedCustomizableObjectsValidator"
+
+
+UAssetValidator_ReferencedCustomizableObjects::UAssetValidator_ReferencedCustomizableObjects() : Super()
+{
+	bIsEnabled = true;
+}
+
+
+bool UAssetValidator_ReferencedCustomizableObjects::CanValidate_Implementation(const EDataValidationUsecase InUsecase) const
+{
+	// Use module settings to decide if it needs to run or not.
+    if (const UMutableValidationSettings* ValidationSettings = GetDefault<UMutableValidationSettings>())
+    {
+    	if (!ValidationSettings->bEnableIndirectValidation)
+    	{
+    		return false;
+    	}
+    }
+
+	// Do not run if saving or running a commandlet (we do not want CIS failing due to our warnings and errors)
+	return !(InUsecase == EDataValidationUsecase::Save || InUsecase == EDataValidationUsecase::Commandlet);
+}
+
+
+bool UAssetValidator_ReferencedCustomizableObjects::CanValidateAsset_Implementation(UObject* InAsset) const
+{
+	// Use module settings to decide if it needs to run or not.
+	if (const UMutableValidationSettings* ValidationSettings = GetDefault<UMutableValidationSettings>())
+	{
+		if (!ValidationSettings->bEnableIndirectValidation)
+		{
+			return false;
+		}
+	}
+
+	return (InAsset ?
+		InAsset->IsA(UMaterial::StaticClass()) ||
+		InAsset->IsA(UTexture::StaticClass()) ||
+		InAsset->IsA(USkeletalMesh::StaticClass()) || 
+		InAsset->IsA(UStaticMesh::StaticClass()) : false) ;
+}
+
+
+EDataValidationResult UAssetValidator_ReferencedCustomizableObjects::ValidateLoadedAsset_Implementation(UObject* InAsset,
+	TArray<FText>& ValidationErrors)
+{
+	check(InAsset);
+
+	FAssetRegistryModule& AssetRegistryModule = FModuleManager::LoadModuleChecked<FAssetRegistryModule>("AssetRegistry");
+	IAssetRegistry& AssetRegistry = AssetRegistryModule.GetRegistry();
+
+	// Locate all referencers to the provided asset.
+	const TSet<FName> FoundReferencers = GetAllAssetReferencers(InAsset,AssetRegistry);
+
+	// Grab all Customizable Objects
+	const TSet<UCustomizableObject*> FoundCustomizableObjects = FindCustomizableObjects(FoundReferencers, AssetRegistry);
+
+	// Validate all Customizable Objects we have found
+	ValidateCustomizableObjects(InAsset,FoundCustomizableObjects,ValidationErrors);
+	
+	// Compute InAsset validation status
+	if (GetValidationResult() != EDataValidationResult::Invalid)
+	{
+		AssetPasses(InAsset);
+	}
+		
+	return GetValidationResult();
+}
+
+
+TSet<FName> UAssetValidator_ReferencedCustomizableObjects::GetAllAssetReferencers(const UObject* InAsset, const IAssetRegistry& InAssetRegistry) const
+{
+	TSet<FName> FoundReferencers;
+	
+	TArray<FName> PackagesToProcess;
+	PackagesToProcess.Add(InAsset->GetOutermost()->GetFName());
+
+	do
+	{
+		TArray<FName> NextPackagesToProcess;
+		for (FName PackageToProcess : PackagesToProcess)
+		{
+			TArray<FName> Referencers;
+			InAssetRegistry.GetReferencers(PackageToProcess, Referencers, UE::AssetRegistry::EDependencyCategory::Package, UE::AssetRegistry::EDependencyQuery::NoRequirements);
+			for (FName Referencer : Referencers)
+			{
+				// If referencer not already found then add it
+				if (!FoundReferencers.Contains(Referencer))
+				{
+					// Cache the referencer so we can later check for COs
+					FoundReferencers.Add(Referencer);
+						
+					NextPackagesToProcess.Add(Referencer);
+				}
+			}
+		}
+		PackagesToProcess = MoveTemp(NextPackagesToProcess);
+			
+	} while (PackagesToProcess.Num() > 0);
+
+	FoundReferencers.Shrink();
+	return FoundReferencers;
+}
+
+
+TSet< UCustomizableObject*> UAssetValidator_ReferencedCustomizableObjects::FindCustomizableObjects(
+	const TSet<FName>& InPackagesToCheck,const IAssetRegistry& InAssetRegistry) const
+{
+	TSet<UCustomizableObject*> FoundCustomizableObjects;
+	
+	for (const FName& ReferencerPackage : InPackagesToCheck)
+	{
+		TArray<FAssetData> PackageAssets;
+        InAssetRegistry.GetAssetsByPackageName(ReferencerPackage, PackageAssets, true);
+        for (const FAssetData& ReferencerAssetData : PackageAssets)
+        {
+        	// We have found a referenced CustomizableObject
+            if (ReferencerAssetData.GetClass() == UCustomizableObject::StaticClass())
+            {
+				UObject* ReferencedAsset = ReferencerAssetData.GetAsset();
+            	if (ReferencedAsset)
+            	{
+            		UCustomizableObject* CastedCustomizableObject = Cast<UCustomizableObject>(ReferencedAsset);
+            		check(CastedCustomizableObject);
+            	
+            		FoundCustomizableObjects.FindOrAdd(CastedCustomizableObject);
+            	}
+            }
+        }
+	}
+
+	FoundCustomizableObjects.Shrink();
+	return FoundCustomizableObjects;
+}
+
+
+void UAssetValidator_ReferencedCustomizableObjects::ValidateCustomizableObjects(UObject* InAsset, const TSet<UCustomizableObject*>& InCustomizableObjectsToValidate, TArray<FText>& InValidationErrors)
+{
+	for (UCustomizableObject* CustomizableObjectToValidate : InCustomizableObjectsToValidate)
+	{
+		// Validate that CO and if it fails then mark it as failed. Do not stop until running the validation over all COs
+		TArray<FText> CoValidationWarnings;
+		TArray<FText> CoValidationErrors;
+		const EDataValidationResult COValidationResult = UAssetValidator_CustomizableObjects::IsCustomizableObjectValid(CustomizableObjectToValidate,CoValidationErrors,CoValidationWarnings);
+		
+		// Process the validation of the CO's output
+		if (COValidationResult == EDataValidationResult::Invalid )
+		{
+			// Cache warning logs
+			for	(const FText& WarningMessage : CoValidationWarnings)
+			{
+				AssetWarning(InAsset,WarningMessage);
+			}
+		
+			// Cache error logs
+			for (const FText& ErrorMessage : CoValidationErrors)
+			{
+				AssetFails(InAsset,ErrorMessage,InValidationErrors);
+			}
+			
+			// If we say it failed the asset will be marked as bad (containing bad data) and the validator will mark the overall result as failed
+			const FText ErrorMessage = FText::Format(LOCTEXT("RelatedToCustomizableObjectValidator", "The referenced ""\"{0}""\" Mutable Customizable Object is invalid."),  FText::FromString(CustomizableObjectToValidate->GetPathName()));
+			AssetFails(InAsset,ErrorMessage,InValidationErrors);
+		}
+	}
+}
+
+
+#undef LOCTEXT_NAMESPACE
diff --git a/Engine/Plugins/Experimental/Mutable/Source/MutableValidation/Private/MuV/MutableValidationSettings.h b/Engine/Plugins/Experimental/Mutable/Source/MutableValidation/Private/MuV/MutableValidationSettings.h
new file mode 100644
--- /dev/null
+++ b/Engine/Plugins/Experimental/Mutable/Source/MutableValidation/Private/MuV/MutableValidationSettings.h
@@ -0,0 +1,17 @@
+// Copyright Epic Games, Inc. All Rights Reserved.
+
+#pragma once
+
+#include "MutableValidationSettings.generated.h"
+
+UCLASS(config = Engine)
+class UMutableValidationSettings : public UObject
+{
+	GENERATED_BODY()
+
+public:
+
+	/** If true, validation of referenced COs from asset subject to data validation, will be run. */
+	UPROPERTY(config, EditAnywhere, Category = Validation)
+	bool bEnableIndirectValidation = true;
+};
diff --git a/Engine/Plugins/Experimental/Mutable/Source/MutableValidation/Public/MuV/AssetValidator_CustomizableObjects.h b/Engine/Plugins/Experimental/Mutable/Source/MutableValidation/Public/MuV/AssetValidator_CustomizableObjects.h
new file mode 100644
--- /dev/null
+++ b/Engine/Plugins/Experimental/Mutable/Source/MutableValidation/Public/MuV/AssetValidator_CustomizableObjects.h
@@ -0,0 +1,50 @@
+// Copyright Epic Games, Inc. All Rights Reserved.
+
+#pragma once
+
+#include "Editor.h"
+#include "EditorValidatorBase.h"
+#include "AssetRegistry/AssetData.h"
+#include "Delegates/DelegateSignatureImpl.inl"
+
+
+#include "AssetValidator_CustomizableObjects.generated.h"
+
+class FText;
+class UObject;
+class UCustomizableObject;
+
+UCLASS()
+class UAssetValidator_CustomizableObjects : public UEditorValidatorBase
+{
+	GENERATED_BODY()
+
+public:
+	UAssetValidator_CustomizableObjects();
+
+	/** Checks the validity of the provided CustomizableObject by compiling and recalling the compilation status and the raised compilation warning and error messages.
+	 * @param InCustomizableObject The mutable customizable object to test.
+	 * @param OutValidationErrors A list of error messages produced during the provided customizable object's compilation.
+	 * @param OutValidationWarnings A list of warning messages produced during the provided customizable object's compilation.
+	 * @return Validation result for the provided object.
+	 */
+	static EDataValidationResult IsCustomizableObjectValid(UCustomizableObject* InCustomizableObject, TArray<FText>& OutValidationErrors, TArray<FText>& OutValidationWarnings);
+	
+protected:
+	// UEditorValidatorBase
+	virtual bool CanValidate_Implementation(const EDataValidationUsecase InUsecase) const override;
+	virtual bool CanValidateAsset_Implementation(UObject* InAsset) const override;
+	virtual EDataValidationResult ValidateLoadedAsset_Implementation(UObject* InAsset, TArray<FText>& ValidationErrors) override;
+	// UEditorValidatorBase
+
+private:
+	
+	/** Cached handle to be able later to remove the bound method from the FEditorDelegates::OnPostAssetValidation delegate */
+	inline static FDelegateHandle OnPostCOValidationHandle;
+	
+	/** Collection with all root objects tested during this IsDataValidRun. Shared with all COs */
+	inline static TSet<UCustomizableObject*> AlreadyValidatedRootObjects;
+	
+	/** Method invoked once the validation of all assets has been completed. */
+	static void OnPostCOsValidation();
+};
diff --git a/Engine/Plugins/Experimental/Mutable/Source/MutableValidation/Public/MuV/AssetValidator_ReferencedCustomizableObjects.h b/Engine/Plugins/Experimental/Mutable/Source/MutableValidation/Public/MuV/AssetValidator_ReferencedCustomizableObjects.h
new file mode 100644
--- /dev/null
+++ b/Engine/Plugins/Experimental/Mutable/Source/MutableValidation/Public/MuV/AssetValidator_ReferencedCustomizableObjects.h
@@ -0,0 +1,55 @@
+// Copyright Epic Games, Inc. All Rights Reserved.
+
+#pragma once
+
+#include "EditorValidatorBase.h"
+#include "AssetRegistry/AssetData.h"
+#include "AssetRegistry/AssetRegistryModule.h"
+
+#include "AssetValidator_ReferencedCustomizableObjects.generated.h"
+
+class FText;
+class FName;
+class UObject;
+class UCustomizableObject;
+
+UCLASS()
+class UAssetValidator_ReferencedCustomizableObjects : public UEditorValidatorBase
+{
+	GENERATED_BODY()
+
+public:
+	UAssetValidator_ReferencedCustomizableObjects();
+	
+protected:
+
+	// UEditorValidatorBase
+	virtual bool CanValidate_Implementation(const EDataValidationUsecase InUsecase) const override;
+	virtual bool CanValidateAsset_Implementation(UObject* InAsset) const override;
+	virtual EDataValidationResult ValidateLoadedAsset_Implementation(UObject* InAsset, TArray<FText>& ValidationErrors) override;
+	// UEditorValidatorBase
+
+private:
+
+	/** Get a set of referencer objects to a provided UObject.
+	 * @param InAsset A Pointer to the asset to get all referencers of.
+	 * @param InAssetRegistry AssetRegistry object to be used to perform the referencers search.
+	 * @return A set with all the referencing packages.
+	 */
+	TSet<FName> GetAllAssetReferencers(const UObject* InAsset,const IAssetRegistry& InAssetRegistry) const;
+
+	/** Returns a set of customizable objects from the provided collection of packages.
+	 * @param InPackagesToCheck List of packages to scan for CustomizableObjects.
+	 * @param InAssetRegistry AssetRegistry object used to perform the asset search.
+	 * @return A collection of unique Customizable Objets found in the provided packages.
+	 */
+	TSet<UCustomizableObject*> FindCustomizableObjects(const TSet<FName>& InPackagesToCheck,const IAssetRegistry& InAssetRegistry) const;
+
+	/** Validates all the customizable objects provided and sets the validator status accordingly. It will not make the validator fail.
+	 * @param InAsset Input asset provided to the validator.
+	 * @param InCustomizableObjectsToValidate Customizable Objects we want to validate with IsDataValid()
+	 * @param InValidationErrors List to fill with the warnings and errors generated during the validation of the Customizable objects
+	 */
+	void ValidateCustomizableObjects(UObject* InAsset, const TSet<UCustomizableObject*>& InCustomizableObjectsToValidate, TArray<FText>& InValidationErrors);
+
+};
diff --git a/Engine/Plugins/Experimental/Mutable/Source/MutableValidation/Public/MuV/MutableValidationModule.cpp b/Engine/Plugins/Experimental/Mutable/Source/MutableValidation/Public/MuV/MutableValidationModule.cpp
new file mode 100644
--- /dev/null
+++ b/Engine/Plugins/Experimental/Mutable/Source/MutableValidation/Public/MuV/MutableValidationModule.cpp
@@ -0,0 +1,79 @@
+// Copyright Epic Games, Inc. All Rights Reserved.
+
+#pragma once
+
+#include "ISettingsModule.h"
+#include "ISettingsSection.h"
+#include "Modules/ModuleManager.h"
+#include "MuV/MutableValidationSettings.h"
+
+
+#define LOCTEXT_NAMESPACE "MutableSettings"
+
+/**
+ * StaticMesh editor module
+ */
+class FMutableValidationModule : public FDefaultModuleImpl
+{
+public:
+	// IModuleInterface interface
+	virtual void StartupModule() override;
+	virtual void ShutdownModule() override;
+
+	bool HandleSettingsSaved() const;
+
+private:
+	ISettingsSectionPtr SettingsSectionPtr = nullptr;
+};
+
+IMPLEMENT_MODULE(FMutableValidationModule, MutableValidation);
+
+bool FMutableValidationModule::HandleSettingsSaved() const
+{
+	UMutableValidationSettings* CustomizableObjectSettings = GetMutableDefault<UMutableValidationSettings>();
+
+	if (CustomizableObjectSettings != nullptr)
+	{
+		CustomizableObjectSettings->SaveConfig();
+	}
+
+    return true;
+}
+
+
+void FMutableValidationModule::StartupModule()
+{
+	ISettingsModule* SettingsModule = FModuleManager::GetModulePtr<ISettingsModule>("Settings");
+	if (SettingsModule != nullptr)
+	{ 		
+		 SettingsSectionPtr = SettingsModule->RegisterSettings("Project", "Plugins", "MutableValidationSettings",
+			LOCTEXT("MutableSettings_Setting", "Mutable Validation"),
+			LOCTEXT("MutableSettings_Setting_Desc", "Mutable resources validation settings"),
+			GetMutableDefault<UMutableValidationSettings>()
+		);
+		
+		if (SettingsSectionPtr.IsValid())
+        {
+        	SettingsSectionPtr->OnModified().BindRaw(this, &FMutableValidationModule::HandleSettingsSaved);
+        }
+	}
+}
+
+
+void FMutableValidationModule::ShutdownModule()
+{
+	// Unbind OnModified delegate
+	if (SettingsSectionPtr)
+	{
+		SettingsSectionPtr->OnModified().Unbind();
+	}
+	
+	ISettingsModule* SettingsModule = FModuleManager::GetModulePtr<ISettingsModule>("Settings");
+	if (SettingsModule != nullptr)
+	{ 		
+		SettingsModule->UnregisterSettings("Project", "Plugins", "MutableValidationSettings");
+	}
+}
+
+
+#undef LOCTEXT_NAMESPACE
