[Mutable] Add support for disabled sections in base lod when using automatic lod strategy.

- Removes pins pointing to disabled skeletal mesh sections.
- Finds the matching sections in the generated lods when base lod sections are disabled. This is only done for the case where the pin lod is the base lod and the processed lod has at least the number of sections of the base skipping disabled sections.

#rnx

diff --git a/Engine/Plugins/Experimental/Mutable/Source/CustomizableObject/Public/MuCO/CustomizableObjectCustomVersion.h b/Engine/Plugins/Experimental/Mutable/Source/CustomizableObject/Public/MuCO/CustomizableObjectCustomVersion.h
--- a/Engine/Plugins/Experimental/Mutable/Source/CustomizableObject/Public/MuCO/CustomizableObjectCustomVersion.h
+++ b/Engine/Plugins/Experimental/Mutable/Source/CustomizableObject/Public/MuCO/CustomizableObjectCustomVersion.h
@@ -80,6 +80,8 @@ struct CUSTOMIZABLEOBJECT_API FCustomizableObjectCustomVersion
 
 		AutomaticNodeSkeletalMeshPinDataUProperty,
 		
+		IgnoreDisabledSections,
+
 		// -----<new versions can be added above this line>-------------------------------------------------
 		VersionPlusOne,
 		LatestVersion = VersionPlusOne - 1
diff --git a/Engine/Plugins/Experimental/Mutable/Source/CustomizableObjectEditor/Private/MuCOE/GenerateMutableSource/GenerateMutableSourceMesh.cpp b/Engine/Plugins/Experimental/Mutable/Source/CustomizableObjectEditor/Private/MuCOE/GenerateMutableSource/GenerateMutableSourceMesh.cpp
--- a/Engine/Plugins/Experimental/Mutable/Source/CustomizableObjectEditor/Private/MuCOE/GenerateMutableSource/GenerateMutableSourceMesh.cpp
+++ b/Engine/Plugins/Experimental/Mutable/Source/CustomizableObjectEditor/Private/MuCOE/GenerateMutableSource/GenerateMutableSourceMesh.cpp
@@ -48,6 +48,81 @@
 
 #define LOCTEXT_NAMESPACE "CustomizableObjectEditor"
 
+int32 FindLODSectionIndexSkippingDisabled(const FSkeletalMeshLODModel& LODModel, int32 TargetSectionIndex)
+{
+	if (LODModel.Sections.IsValidIndex(TargetSectionIndex))
+	{
+		return INDEX_NONE;
+	}
+
+	if (LODModel.Sections[TargetSectionIndex].bDisabled)
+	{
+		return INDEX_NONE;
+	}
+
+	int32 SkipDisabledSectionIndex = 0;
+	const int32 NumSections = LODModel.Sections.Num();
+	for (int32 SectionIndex = 0; SectionIndex < NumSections; ++SectionIndex)
+	{
+		if (TargetSectionIndex == SectionIndex)
+		{
+			return SkipDisabledSectionIndex;
+		}
+
+		const FSkelMeshSection& Section = LODModel.Sections[SectionIndex];
+
+		SkipDisabledSectionIndex += static_cast<int32>(!Section.bDisabled);
+	}
+
+	return INDEX_NONE;
+}
+
+void GetLODAndSectionForAutomaticLODs(const FSkeletalMeshModel* ImportedModel, int32 InPinLODIndex, int32 InCurrentLOD, int32 InPinSectionIndex, int32& OutLODIndex, int32& OutSectionIndex)
+{
+	check(ImportedModel);
+
+	// The disabled sections are not replicated to unreal generated lods and a mismatch can happen when 
+	// selecting the section with automatic from mesh strategy. The fix for that is only parcial and the casuistics 
+	// of section matching between lods are more extense, for instance, when lower resolution lods have 
+	// sections removed in an arbitrary order. To minimize discrepacines with the old behvaior the disabled sections
+	// will only be taken into account if the PinLODIndex is the base. 
+
+	int32 CurrentLODSectionIndex = (InCurrentLOD != InPinLODIndex && InPinLODIndex == 0) 
+		? FindLODSectionIndexSkippingDisabled(ImportedModel->LODModels[InPinLODIndex], InPinSectionIndex)
+		: InPinSectionIndex;
+
+	// if we couldn't find any matching section, fallback with the old behaviour of using 
+	// the pin SectionIndex
+	if (CurrentLODSectionIndex == INDEX_NONE)
+	{
+		CurrentLODSectionIndex = InPinSectionIndex;
+	}
+	
+	// If the mesh has additional LODs and they use the same material, use them.
+	if (ImportedModel->LODModels.Num() > InCurrentLOD
+		&&
+		ImportedModel->LODModels[InCurrentLOD].Sections.Num() > CurrentLODSectionIndex)
+	{
+		OutLODIndex = InCurrentLOD;
+		OutSectionIndex = CurrentLODSectionIndex;
+	}
+	else
+	{
+		// Find the closest valid LOD.
+		// Don't try to account for disabled sections for this case.
+		for (int32 LODIndex = ImportedModel->LODModels.Num() - 1; LODIndex >= 0; --LODIndex)
+		{
+			if (ImportedModel->LODModels[LODIndex].Sections.Num() > InPinSectionIndex)
+			{
+				OutLODIndex = LODIndex;
+				OutSectionIndex = InPinSectionIndex;
+				break;
+			}
+		}
+	}
+
+}
+
 void BuildRemappedBonesArray(const FMutableComponentInfo& InComponentInfo, TObjectPtr<USkeletalMesh> InSkeletalMesh, int32 InLODIndex, const TArray<FBoneIndexType>& InRequiredBones, TArray<FBoneIndexType>& OutRemappedBones)
 {
 	if (!InSkeletalMesh)
@@ -1965,29 +2040,20 @@ mu::MeshPtr BuildMorphedMutableMesh(const UEdGraphPin* BaseSourcePin, const FStr
 		{
 			if (SkeletalMesh && SkeletalMesh->GetImportedModel())
 			{
+				FSkeletalMeshModel* ImportedModel = SkeletalMesh->GetImportedModel();
+				
+				const int32 PinLODIndex = LODIndex;
 				// Add the CurrentLOD being generated to the first LOD used from this SkeletalMesh
-				int32 CurrentLOD = LODIndex + GenerationContext.CurrentLOD;
+				const int32 CurrentLOD = PinLODIndex + GenerationContext.CurrentLOD; 
+				const int32 PinSectionIndex = SectionIndex;
 
-				// If the mesh has additional LODs and they use the same material, use them.
-				FSkeletalMeshModel* ImportedModel = SkeletalMesh->GetImportedModel();
-				if (ImportedModel->LODModels.Num() > CurrentLOD
-					&&
-					ImportedModel->LODModels[CurrentLOD].Sections.Num() > SectionIndex)
-				{
-					LODIndex = CurrentLOD;
-				}
-				else
-				{
-					// Find the closest valid LOD 
-					for (int32 LOD = ImportedModel->LODModels.Num() - 1; LOD >= 0; --LOD)
-					{
-						if (ImportedModel->LODModels[LOD].Sections.Num() > SectionIndex)
-						{
-							LODIndex = LOD;
-							break;
-						}
-					}
-				}
+				int32 OutLOD = -1;
+				int32 OutSectionIndex = -1;
+
+				GetLODAndSectionForAutomaticLODs(ImportedModel, PinLODIndex, CurrentLOD, PinSectionIndex, OutLOD, OutSectionIndex);
+
+				LODIndex = OutLOD;
+				SectionIndex = OutSectionIndex;
 			}
 		}
 
@@ -2725,7 +2791,6 @@ void GenerateMorphTarget(const UCustomizableObjectNode* Node, const UEdGraphPin*
 		GenerationContext.Compiler->CompilerLog(LOCTEXT("MorphGenerationFailed", "Failed to generate morph target."), Node);
 	}
 }
-
 /** Convert a CustomizableObject Source Graph into a mutable source graph  */
 mu::NodeMeshPtr GenerateMutableSourceMesh(const UEdGraphPin * Pin,
 	FMutableGraphGenerationContext & GenerationContext,
@@ -2778,30 +2843,21 @@ mu::NodeMeshPtr GenerateMutableSourceMesh(const UEdGraphPin * Pin,
 			{
 				if (TypedNodeSkel->SkeletalMesh && TypedNodeSkel->SkeletalMesh->GetImportedModel())
 				{
-					// Add the CurrentLOD being generated to the first LOD used from this SkeletalMesh
-					int32 CurrentLOD = LOD + GenerationContext.CurrentLOD; 
-					
-					// If the mesh has additional LODs and they use the same material, use them.
-					FSkeletalMeshModel* ImportedModel = TypedNodeSkel->SkeletalMesh->GetImportedModel();
-					if (ImportedModel->LODModels.Num() > CurrentLOD
-						&&
-						ImportedModel->LODModels[CurrentLOD].Sections.Num() > SectionIndex)
-					{
-						LOD = CurrentLOD;
-					}
-					else
-					{
-						// Find the closest valid LOD 
-						for (int32 LODIndex = ImportedModel->LODModels.Num() - 1; LODIndex >= 0; --LODIndex)
-						{
-							if (ImportedModel->LODModels[LODIndex].Sections.Num() > SectionIndex)
-							{
-								LOD = LODIndex;
-								break;
-							}
-						}
-					}
+					const FSkeletalMeshModel* ImportedModel = TypedNodeSkel->SkeletalMesh->GetImportedModel();
+
+					const int32 PinLODIndex = LOD;
+					const int32 CurrentLOD = PinLODIndex + GenerationContext.CurrentLOD; 
+					const int32 PinSectionIndex = SectionIndex;
+
+					int32 OutLOD = -1;
+					int32 OutSectionIndex = -1;
+
+					GetLODAndSectionForAutomaticLODs(ImportedModel, PinLODIndex, CurrentLOD, PinSectionIndex, OutLOD, OutSectionIndex);
+
+					LOD = OutLOD;
+					SectionIndex = OutSectionIndex;
 				}
+
 			}
 
 			// First process the mesh tags that are going to make the mesh unique and affect whether it's repeated in 
@@ -3684,5 +3740,4 @@ mu::NodeMeshPtr GenerateMutableSourceMesh(const UEdGraphPin * Pin,
 
 	return Result;
 }
-
 #undef LOCTEXT_NAMESPACE
diff --git a/Engine/Plugins/Experimental/Mutable/Source/CustomizableObjectEditor/Private/MuCOE/Nodes/CustomizableObjectNodeSkeletalMesh.cpp b/Engine/Plugins/Experimental/Mutable/Source/CustomizableObjectEditor/Private/MuCOE/Nodes/CustomizableObjectNodeSkeletalMesh.cpp
--- a/Engine/Plugins/Experimental/Mutable/Source/CustomizableObjectEditor/Private/MuCOE/Nodes/CustomizableObjectNodeSkeletalMesh.cpp
+++ b/Engine/Plugins/Experimental/Mutable/Source/CustomizableObjectEditor/Private/MuCOE/Nodes/CustomizableObjectNodeSkeletalMesh.cpp
@@ -68,6 +68,13 @@ void UCustomizableObjectNodeSkeletalMesh::AllocateDefaultPins(UCustomizableObjec
 			const int32 NumSections = ImportedModel->LODModels[LODIndex].Sections.Num();
 			for (int32 SectionIndex = 0; SectionIndex < NumSections; ++SectionIndex)
 			{
+				// Ignore disabled sections.
+				const bool bIsSectionDisabled = ImportedModel ->LODModels[LODIndex].Sections[SectionIndex].bDisabled;
+				if (bIsSectionDisabled)
+				{
+					continue;
+				}
+
 				UMaterialInterface* MaterialInterface = GetMaterialInterfaceFor(LODIndex, SectionIndex, ImportedModel);
 				
 				FString SectionFriendlyName = MaterialInterface ? MaterialInterface->GetName() : FString::Printf(TEXT("Section %i"), SectionIndex);
@@ -401,7 +408,8 @@ bool UCustomizableObjectNodeSkeletalMesh::IsNodeOutDatedAndNeedsRefresh()
 			auto OutdatedSectionPinData = [&](const UCustomizableObjectNodeSkeletalMeshPinDataSection& PinData) -> bool
 			{
 				return !ImportedModel->LODModels.IsValidIndex(PinData.GetLODIndex()) ||
-					!ImportedModel->LODModels[PinData.GetLODIndex()].Sections.IsValidIndex(PinData.GetSectionIndex());
+					!ImportedModel->LODModels[PinData.GetLODIndex()].Sections.IsValidIndex(PinData.GetSectionIndex()) ||
+					ImportedModel->LODModels[PinData.GetLODIndex()].Sections[PinData.GetSectionIndex()].bDisabled;
 			};
 		
 			if (const UCustomizableObjectNodeSkeletalMeshPinDataLayout* LayoutPinData = Cast<UCustomizableObjectNodeSkeletalMeshPinDataLayout>(GetPinData(*Pin)))
@@ -608,6 +616,11 @@ void UCustomizableObjectNodeSkeletalMesh::BackwardsCompatibleFixup()
 	{
 		ReconstructNode(CreateRemapPinsByName()); // Correct pins but incorrect Pin Data. Reconstruct and remap pins only by name, no Pin Data.
 	}
+
+	if (CustomizableObjectCustomVersion < FCustomizableObjectCustomVersion::IgnoreDisabledSections)
+	{
+		ReconstructNode(); // Pins representing disabled sections could be present. 
+	}
 }
 
 
