[mutable] Implement memory budget control for image operations.
Cleanup memory budget interface.
Remove unused image operations.
#prefight 646c958c4422ba05f47ea1e6
#rnx

diff --git a/Engine/Plugins/Experimental/Mutable/Source/CustomizableObject/Classes/MuCO/CustomizableObject.h b/Engine/Plugins/Experimental/Mutable/Source/CustomizableObject/Classes/MuCO/CustomizableObject.h
--- a/Engine/Plugins/Experimental/Mutable/Source/CustomizableObject/Classes/MuCO/CustomizableObject.h
+++ b/Engine/Plugins/Experimental/Mutable/Source/CustomizableObject/Classes/MuCO/CustomizableObject.h
@@ -1413,7 +1413,7 @@ private:
 	// This is a manual version number for the binary blobs in this asset.
 	// Increasing it invalidates all the previously compiled models.
 	// Warning: If while merging code both versions have changed, take the highest+1.
-	static const int32 CurrentSupportedVersion = 386;
+	static const int32 CurrentSupportedVersion = 387;
 
 public:
 
diff --git a/Engine/Plugins/Experimental/Mutable/Source/CustomizableObject/Classes/MuCO/CustomizableObjectSystem.h b/Engine/Plugins/Experimental/Mutable/Source/CustomizableObject/Classes/MuCO/CustomizableObjectSystem.h
--- a/Engine/Plugins/Experimental/Mutable/Source/CustomizableObject/Classes/MuCO/CustomizableObjectSystem.h
+++ b/Engine/Plugins/Experimental/Mutable/Source/CustomizableObject/Classes/MuCO/CustomizableObjectSystem.h
@@ -59,10 +59,9 @@ DECLARE_DWORD_ACCUMULATOR_STAT_EXTERN(TEXT("Abandoned Instance Updates"), STAT_M
 DECLARE_DWORD_ACCUMULATOR_STAT_EXTERN(TEXT("Last Instance Build Time"), STAT_MutableInstanceBuildTime, STATGROUP_Mutable, );
 DECLARE_DWORD_ACCUMULATOR_STAT_EXTERN(TEXT("Avrg Instance Build Time"), STAT_MutableInstanceBuildTimeAvrg, STATGROUP_Mutable, );
 DECLARE_DWORD_ACCUMULATOR_STAT_EXTERN(TEXT("Streaming Ops"), STAT_MutableStreamingOps, STATGROUP_Mutable, );
-DECLARE_DWORD_ACCUMULATOR_STAT_EXTERN(TEXT("Streaming Cache"), STAT_MutableStreamingCache, STATGROUP_Mutable, );
 
 
-extern TAutoConsoleVariable<bool> CVarClearStreamingCacheOnUpdateEnd;
+extern TAutoConsoleVariable<bool> CVarClearWorkingMemoryOnUpdateEnd;
 
 extern TAutoConsoleVariable<bool> CVarEnableImageCache;
 
diff --git a/Engine/Plugins/Experimental/Mutable/Source/CustomizableObject/Private/MuCO/CustomizableObjectInstance.cpp b/Engine/Plugins/Experimental/Mutable/Source/CustomizableObject/Private/MuCO/CustomizableObjectInstance.cpp
--- a/Engine/Plugins/Experimental/Mutable/Source/CustomizableObject/Private/MuCO/CustomizableObjectInstance.cpp
+++ b/Engine/Plugins/Experimental/Mutable/Source/CustomizableObject/Private/MuCO/CustomizableObjectInstance.cpp
@@ -362,7 +362,7 @@ void UCustomizableObjectInstance::ReleaseMutableResources(bool bCalledFromBeginD
 		}
 
 		// Remove all references to cached Texture Parameters
-		FUnrealMutableImageProvider* ImageProvider = UCustomizableObjectSystem::GetInstance()->GetPrivate()->ImageProvider;
+		FUnrealMutableImageProvider* ImageProvider = UCustomizableObjectSystem::GetInstance()->GetPrivate()->ImageProvider.Get();
 
 		for (const FString& TextureParameter : PrivateData->UpdateTextureParameters)
 		{
diff --git a/Engine/Plugins/Experimental/Mutable/Source/CustomizableObject/Private/MuCO/CustomizableObjectMipDataProvider.cpp b/Engine/Plugins/Experimental/Mutable/Source/CustomizableObject/Private/MuCO/CustomizableObjectMipDataProvider.cpp
--- a/Engine/Plugins/Experimental/Mutable/Source/CustomizableObject/Private/MuCO/CustomizableObjectMipDataProvider.cpp
+++ b/Engine/Plugins/Experimental/Mutable/Source/CustomizableObject/Private/MuCO/CustomizableObjectMipDataProvider.cpp
@@ -32,7 +32,7 @@ FMutableUpdateContext::FMutableUpdateContext(mu::Ptr<mu::System> InSystem,
 
 FMutableUpdateContext::~FMutableUpdateContext()
 {
-	if (Parameters)
+	if (Parameters && UCustomizableObjectSystem::IsCreated() && UCustomizableObjectSystem::GetInstance()->GetPrivate()->ImageProvider)
 	{
 		UCustomizableObjectSystem::GetInstance()->GetPrivate()->ImageProvider->UnCacheImages(*Parameters);
 	}
@@ -188,9 +188,9 @@ namespace impl
 				System->EndUpdate(InstanceID);
 				System->ReleaseInstance(InstanceID);
 
-				if (CVarClearStreamingCacheOnUpdateEnd.GetValueOnAnyThread())
+				if (CVarClearWorkingMemoryOnUpdateEnd.GetValueOnAnyThread())
 				{
-					System->ClearStreamingCache();
+					System->ClearWorkingMemory();
 				}
 			}
 
diff --git a/Engine/Plugins/Experimental/Mutable/Source/CustomizableObject/Private/MuCO/CustomizableObjectSystem.cpp b/Engine/Plugins/Experimental/Mutable/Source/CustomizableObject/Private/MuCO/CustomizableObjectSystem.cpp
--- a/Engine/Plugins/Experimental/Mutable/Source/CustomizableObject/Private/MuCO/CustomizableObjectSystem.cpp
+++ b/Engine/Plugins/Experimental/Mutable/Source/CustomizableObject/Private/MuCO/CustomizableObjectSystem.cpp
@@ -71,28 +71,27 @@ DEFINE_STAT(STAT_MutableAbandonedInstanceUpdates);
 DEFINE_STAT(STAT_MutableInstanceBuildTime);
 DEFINE_STAT(STAT_MutableInstanceBuildTimeAvrg);
 DEFINE_STAT(STAT_MutableStreamingOps);
-DEFINE_STAT(STAT_MutableStreamingCache);
 
 DECLARE_CYCLE_STAT(TEXT("MutablePendingRelease Time"), STAT_MutablePendingRelease, STATGROUP_Game);
 DECLARE_CYCLE_STAT(TEXT("MutableTask"), STAT_MutableTask, STATGROUP_Game);
 
 UCustomizableObjectSystem* FCustomizableObjectSystemPrivate::SSystem = nullptr;
 
-static TAutoConsoleVariable<int32> CVarStreamingMemory(
-	TEXT("mutable.StreamingCacheMemory"),
+static TAutoConsoleVariable<int32> CVarWorkingMemory(
+	TEXT("mutable.WorkingMemory"),
 #if !PLATFORM_DESKTOP
-	(3 * 1024),
+	(10 * 1024),
 #else
-	(12 * 1024),
+	(50 * 1024),
 #endif
-	TEXT("Limit the amount of memory (in KB) to use to cache streaming data when building characters. More memory reduces the object construction time. 0 means no restriction. Defaults: Desktop = 12288, Others = 3072"),
+	TEXT("Limit the amount of memory (in KB) to use as working memory when building characters. More memory reduces the object construction time. 0 means no restriction. Defaults: Desktop = 50,000 KB, Others = 10,000 KB"),
 	ECVF_Scalability);
 
 
-TAutoConsoleVariable<bool> CVarClearStreamingCacheOnUpdateEnd(
-	TEXT("mutable.ClearStreamingCacheOnUpdateEnd"),
+TAutoConsoleVariable<bool> CVarClearWorkingMemoryOnUpdateEnd(
+	TEXT("mutable.CVarClearWorkingMemoryOnUpdateEnd"),
 	false,
-	TEXT("Clear the Streaming Cache when releasing an Instance."),
+	TEXT("Clear the wokring memory and cache after every Mutable operation."),
 	ECVF_Scalability);
 
 
@@ -328,22 +327,22 @@ void UCustomizableObjectSystem::InitSystem()
 	SET_DWORD_STAT(STAT_MutableInstanceBuildTimeAvrg, 0);
 #endif
 
-	Private->LastStreamingMemorySize = CVarClearStreamingCacheOnUpdateEnd.GetValueOnGameThread();
+	Private->LastWorkingMemoryBytes = CVarWorkingMemory.GetValueOnGameThread() * 1024;
 
-	mu::SettingsPtr pSettings = new mu::Settings;
+	mu::Ptr<mu::Settings> pSettings = new mu::Settings;
 	check(pSettings);
 	pSettings->SetProfile(false);
-	pSettings->SetStreamingCache(Private->LastStreamingMemorySize);
-	Private->ExtensionDataStreamer = new FUnrealExtensionDataStreamer(Private.ToSharedRef());
+	pSettings->SetWorkingMemoryBytes(Private->LastWorkingMemoryBytes);
+	Private->ExtensionDataStreamer = MakeShared<FUnrealExtensionDataStreamer>(Private.ToSharedRef());
 	Private->MutableSystem = new mu::System(pSettings, Private->ExtensionDataStreamer);
 	check(Private->MutableSystem);
 
-	Private->Streamer = new FUnrealMutableModelBulkStreamer();
+	Private->Streamer = MakeShared<FUnrealMutableModelBulkStreamer>();
 	check(Private->Streamer != nullptr);
 	Private->MutableSystem->SetStreamingInterface(Private->Streamer);
 
 	// Set up the external image provider, for image parameters.
-	FUnrealMutableImageProvider* Provider = new FUnrealMutableImageProvider();
+	TSharedPtr<FUnrealMutableImageProvider> Provider = MakeShared<FUnrealMutableImageProvider>();
 	check(Provider != nullptr);
 	Private->ImageProvider = Provider;
 	Private->MutableSystem->SetImageParameterGenerator(Provider);
@@ -1055,7 +1054,7 @@ bool UCustomizableObjectSystem::CheckIfDiskOrMipUpdateOperationsPending(const UC
 
 	// Ensure that we don't try to handle any further streaming operations for this object
 	check(GetPrivate());
-	if (const FUnrealMutableModelBulkStreamer* Streamer = GetPrivate()->Streamer)
+	if (const FUnrealMutableModelBulkStreamer* Streamer = GetPrivate()->Streamer.Get())
 	{
 		if (Streamer->AreTherePendingStreamingOperationsForObject(&Object))
 		{
@@ -1145,17 +1144,17 @@ void UCustomizableObjectSystem::ClearCurrentMutableOperation()
 }
 
 
-void FCustomizableObjectSystemPrivate::UpdateStreamingLimit()
+void FCustomizableObjectSystemPrivate::UpdateMemoryLimit()
 {
 	// This must run on game thread, and when the mutable thread is not running
 	check(IsInGameThread());
 
-	const uint64 MemoryBytes = CVarStreamingMemory.GetValueOnGameThread() * 1024;
-	if (MemoryBytes != LastStreamingMemorySize)
+	const uint64 MemoryBytes = CVarWorkingMemory.GetValueOnGameThread() * 1024;
+	if (MemoryBytes != LastWorkingMemoryBytes)
 	{
-		LastStreamingMemorySize = MemoryBytes;
+		LastWorkingMemoryBytes = MemoryBytes;
 		check(MutableSystem);
-		MutableSystem->SetStreamingCache(MemoryBytes);
+		MutableSystem->SetWorkingMemoryBytes(MemoryBytes);
 	}
 }
 
@@ -1779,9 +1778,9 @@ namespace impl
 			}
 		}
 
-		if (CVarClearStreamingCacheOnUpdateEnd.GetValueOnAnyThread())
+		if (CVarClearWorkingMemoryOnUpdateEnd.GetValueOnAnyThread())
 		{
-			MutableSystem->ClearStreamingCache();
+			MutableSystem->ClearWorkingMemory();
 		}
 	}
 
@@ -1799,9 +1798,9 @@ namespace impl
 			OperationData->InstanceID = 0;
 		}
 
-		if (CVarClearStreamingCacheOnUpdateEnd.GetValueOnAnyThread())
+		if (CVarClearWorkingMemoryOnUpdateEnd.GetValueOnAnyThread())
 		{
-			MutableSystem->ClearStreamingCache();			
+			MutableSystem->ClearWorkingMemory();
 		}
 	}
 
@@ -1855,15 +1854,12 @@ namespace impl
 		{
 			double deltaSeconds = FPlatformTime::Seconds() - CustomizableObjectSystemPrivateData->CurrentMutableOperation->StartUpdateTime;
 			int32 deltaMs = int32(deltaSeconds * 1000);
-			int64 streamingCache = CustomizableObjectSystemPrivateData->LastStreamingMemorySize / 1024;
 
 			LogBenchmarkUtil::updateStat("customizable_instance_build_time", deltaMs);
-			LogBenchmarkUtil::updateStat("mutable_streaming_cache_memory", (long double)streamingCache / 1024.0);
 			CustomizableObjectSystemPrivateData->TotalBuildMs += deltaMs;
 			CustomizableObjectSystemPrivateData->TotalBuiltInstances++;
 			SET_DWORD_STAT(STAT_MutableInstanceBuildTime, deltaMs);
 			SET_DWORD_STAT(STAT_MutableInstanceBuildTimeAvrg, CustomizableObjectSystemPrivateData->TotalBuildMs / CustomizableObjectSystemPrivateData->TotalBuiltInstances);
-			SET_DWORD_STAT(STAT_MutableStreamingCache, streamingCache);
 		}
 
 		// End Update
@@ -2072,7 +2068,7 @@ namespace impl
 		}
 
 		// Cache Texture Parameters
-		FUnrealMutableImageProvider* ImageProvider = UCustomizableObjectSystem::GetInstance()->GetPrivate()->ImageProvider;
+		FUnrealMutableImageProvider* ImageProvider = UCustomizableObjectSystem::GetInstance()->GetPrivate()->ImageProvider.Get();
 		
 		// Uncache old Texture Parameters
 		for (const FString& TextureParameter : ObjectInstancePrivateData->UpdateTextureParameters)
@@ -2473,7 +2469,7 @@ void UCustomizableObjectSystem::AdvanceCurrentOperation()
 	}
 
 	// It is safe to do this now.
-	Private->UpdateStreamingLimit();
+	Private->UpdateMemoryLimit();
 
 	// If we don't have an ongoing operation, don't do anything.
 	if (!Private->CurrentMutableOperation.IsValid())
@@ -2638,7 +2634,7 @@ bool UCustomizableObjectSystem::Tick(float DeltaTime)
 		}
 
 		// Update the streaming limit if it has changed. It is safe to do this now.
-		Private->UpdateStreamingLimit();
+		Private->UpdateMemoryLimit();
 
 		// Decide the next mutable operation to perform.
 		if (!Private->MutableOperationQueue.IsEmpty())
diff --git a/Engine/Plugins/Experimental/Mutable/Source/CustomizableObject/Private/MuCO/CustomizableObjectSystemPrivate.h b/Engine/Plugins/Experimental/Mutable/Source/CustomizableObject/Private/MuCO/CustomizableObjectSystemPrivate.h
--- a/Engine/Plugins/Experimental/Mutable/Source/CustomizableObject/Private/MuCO/CustomizableObjectSystemPrivate.h
+++ b/Engine/Plugins/Experimental/Mutable/Source/CustomizableObject/Private/MuCO/CustomizableObjectSystemPrivate.h
@@ -527,23 +527,21 @@ public:
 	static UCustomizableObjectSystem* SSystem;
 
 	// Pointer to the lower level mutable system that actually does the work.
-	mu::SystemPtr MutableSystem;
+	mu::Ptr<mu::System> MutableSystem;
 
-	// Store the last streaming memory size in bytes, to change it when it is safe.
-	uint64 LastStreamingMemorySize = 0;
+	/** Store the last streaming memory size in bytes, to change it when it is safe. */
+	uint64 LastWorkingMemoryBytes = 0;
 
 	// This object is responsible for streaming data to the MutableSystem.
-	// Non-owned reference
-	class FUnrealMutableModelBulkStreamer* Streamer = nullptr;
+	TSharedPtr<class FUnrealMutableModelBulkStreamer> Streamer;
 
-	// Non-owned reference
-	class FUnrealExtensionDataStreamer* ExtensionDataStreamer = nullptr;
+	// 
+	TSharedPtr<class FUnrealExtensionDataStreamer> ExtensionDataStreamer;
 
 	// This object is responsible for providing custom images to mutable models (for image parameters)
 	// This object is called from the mutable thread, and it should only access data already safely submitted from
 	// the game thread and stored in FUnrealMutableImageProvider::GlobalExternalImages.
-	// Non-owned reference
-	class FUnrealMutableImageProvider* ImageProvider = nullptr;
+	TSharedPtr<class FUnrealMutableImageProvider> ImageProvider;
 
 	// Cache of weak references to generated resources to see if they can be reused.
 	TArray<FMutableResourceCache> ModelResourcesCache;
@@ -738,8 +736,8 @@ public:
 	// Important!!! Never call when there's a Begin Update thread running!
 	void ReleasePendingMutableInstances();
 
-	// Check and update the streaming memory limit. Only safe from game thread and when the mutable thread is idle.
-	void UpdateStreamingLimit();
+	/** Update the last set amount of internal memory Mutable can use to build objects. */
+	void UpdateMemoryLimit();
 
 	bool IsMutableAnimInfoDebuggingEnabled() const;
 
diff --git a/Engine/Plugins/Experimental/Mutable/Source/CustomizableObjectEditor/Private/MuCOE/Nodes/CustomizableObjectNodeSkeletalMesh.cpp b/Engine/Plugins/Experimental/Mutable/Source/CustomizableObjectEditor/Private/MuCOE/Nodes/CustomizableObjectNodeSkeletalMesh.cpp
--- a/Engine/Plugins/Experimental/Mutable/Source/CustomizableObjectEditor/Private/MuCOE/Nodes/CustomizableObjectNodeSkeletalMesh.cpp
+++ b/Engine/Plugins/Experimental/Mutable/Source/CustomizableObjectEditor/Private/MuCOE/Nodes/CustomizableObjectNodeSkeletalMesh.cpp
@@ -567,29 +567,33 @@ void UCustomizableObjectNodeSkeletalMesh::BackwardsCompatibleFixup()
 				{
 					if (const FSkeletalMaterial* SkeletalMaterial = GetSkeletalMaterialFor(LODIndex, SectionIndex))
 					{
-						for (int32 ImageIndex = 0; ImageIndex < Section.ImagePinsRef.Num(); ++ImageIndex)
+						if (SkeletalMaterial && SkeletalMaterial->MaterialInterface)
 						{
-							const UEdGraphPin* ImagePin = Section.ImagePinsRef[ImageIndex].Get();
-
-							FGuid TextureParameterId;
-
 							TArray<FGuid> ParameterIds;
 							TArray<FMaterialParameterInfo> ParameterInfo;
 							SkeletalMaterial->MaterialInterface->GetAllParameterInfoOfType(EMaterialParameterType::Texture, ParameterInfo, ParameterIds);
-							check(ParameterIds.Num() == ParameterInfo.Num())
-							for (int32 Index = 0; Index < ParameterIds.Num(); ++Index)
+							check(ParameterIds.Num() == ParameterInfo.Num());
+
+							for (int32 ImageIndex = 0; ImageIndex < Section.ImagePinsRef.Num(); ++ImageIndex)
 							{
-								if (ParameterInfo[Index].Name == ImagePin->PinFriendlyName.ToString())
+								const UEdGraphPin* ImagePin = Section.ImagePinsRef[ImageIndex].Get();
+
+								FGuid TextureParameterId;
+
+								for (int32 Index = 0; Index < ParameterIds.Num(); ++Index)
 								{
-									TextureParameterId = ParameterIds[Index];
-									break;	
+									if (ParameterInfo[Index].Name == ImagePin->PinFriendlyName.ToString())
+									{
+										TextureParameterId = ParameterIds[Index];
+										break;
+									}
 								}
+
+								UCustomizableObjectNodeSkeletalMeshPinDataImage* PinData = NewObject<UCustomizableObjectNodeSkeletalMeshPinDataImage>(this);
+								PinData->Init(LODIndex, SectionIndex, TextureParameterId);
+
+								AddPinData(*ImagePin, *PinData);
 							}
-							
-							UCustomizableObjectNodeSkeletalMeshPinDataImage* PinData = NewObject<UCustomizableObjectNodeSkeletalMeshPinDataImage>(this);
-							PinData->Init(LODIndex, SectionIndex, TextureParameterId);
-						
-							AddPinData(*ImagePin, *PinData);
 						}
 					}
 				}
diff --git a/Engine/Plugins/Experimental/Mutable/Source/CustomizableObjectEditor/Private/MuCOE/SMutableCodeViewer.cpp b/Engine/Plugins/Experimental/Mutable/Source/CustomizableObjectEditor/Private/MuCOE/SMutableCodeViewer.cpp
--- a/Engine/Plugins/Experimental/Mutable/Source/CustomizableObjectEditor/Private/MuCOE/SMutableCodeViewer.cpp
+++ b/Engine/Plugins/Experimental/Mutable/Source/CustomizableObjectEditor/Private/MuCOE/SMutableCodeViewer.cpp
@@ -2407,7 +2407,7 @@ void SMutableCodeViewer::Tick(const FGeometry& AllottedGeometry, const double In
 	const mu::Ptr<mu::Settings> Settings = new mu::Settings();
 	const mu::Ptr<mu::System> System = new mu::System(Settings);
 
-	TestImageProvider* ImageProvider = new TestImageProvider;
+	TSharedPtr<TestImageProvider> ImageProvider = MakeShared<TestImageProvider>();
 	System->SetImageParameterGenerator(ImageProvider);
 
 	System->GetPrivate()->BeginBuild(MutableModel);
@@ -2474,12 +2474,8 @@ void SMutableCodeViewer::Tick(const FGeometry& AllottedGeometry, const double In
 	case mu::DT_COLOUR:
 	{
 		check(PreviewColorViewer);
-		float RedValue;
-		float GreenValue;
-		float BlueValue;
-		float AlphaValue;
-		System->GetPrivate()->BuildColour(MutableModel, PreviewParameters.get(), SelectedOperationAddress, &RedValue, &GreenValue, &BlueValue, &AlphaValue);
-		PreviewColorViewer->SetColor(RedValue, GreenValue, BlueValue, AlphaValue);
+		FVector4f Color = System->GetPrivate()->BuildColour(MutableModel, PreviewParameters.get(), SelectedOperationAddress);
+		PreviewColorViewer->SetColor(Color);
 		break;
 	}
 
diff --git a/Engine/Plugins/Experimental/Mutable/Source/CustomizableObjectEditor/Private/MuCOE/SMutableColorViewer.cpp b/Engine/Plugins/Experimental/Mutable/Source/CustomizableObjectEditor/Private/MuCOE/SMutableColorViewer.cpp
--- a/Engine/Plugins/Experimental/Mutable/Source/CustomizableObjectEditor/Private/MuCOE/SMutableColorViewer.cpp
+++ b/Engine/Plugins/Experimental/Mutable/Source/CustomizableObjectEditor/Private/MuCOE/SMutableColorViewer.cpp
@@ -150,45 +150,37 @@ void SMutableColorViewer::Construct(const FArguments& InArgs)
 	];
 }
 
-void SMutableColorViewer::SetColor(float InRed, float InGreen, float InBlue, float InAlpha)
+void SMutableColorViewer::SetColor(FVector4f InColor)
 {
-	RedValue = InRed;
-	GreenValue = InGreen;
-	BlueValue = InBlue;
-	AlphaValue = InAlpha;
+	Color = InColor;
 }
 
 
 FLinearColor SMutableColorViewer::GetColor() const
 {
-	FLinearColor BoxColor;
-	
-	BoxColor.R = RedValue;
-	BoxColor.G = GreenValue;
-	BoxColor.B = BlueValue;
-	BoxColor.A = AlphaValue;
+	FLinearColor BoxColor(Color);
 	
 	return BoxColor;
 }
 
 FText SMutableColorViewer::GetRedValue() const
 {
-	return FText::AsNumber(RedValue);
+	return FText::AsNumber(Color[0]);
 }
 
 FText SMutableColorViewer::GetGreenValue() const
 {
-	return FText::AsNumber(GreenValue);
+	return FText::AsNumber(Color[1]);
 }
 
 FText SMutableColorViewer::GetBlueValue() const
 {
-	return FText::AsNumber(BlueValue);
+	return FText::AsNumber(Color[2]);
 }
 
 FText SMutableColorViewer::GetAlphaValue() const
 {
-	return FText::AsNumber(AlphaValue);
+	return FText::AsNumber(Color[3]);
 }
 
 #undef LOCTEXT_NAMESPACE
diff --git a/Engine/Plugins/Experimental/Mutable/Source/CustomizableObjectEditor/Private/MuCOE/SMutableColorViewer.h b/Engine/Plugins/Experimental/Mutable/Source/CustomizableObjectEditor/Private/MuCOE/SMutableColorViewer.h
--- a/Engine/Plugins/Experimental/Mutable/Source/CustomizableObjectEditor/Private/MuCOE/SMutableColorViewer.h
+++ b/Engine/Plugins/Experimental/Mutable/Source/CustomizableObjectEditor/Private/MuCOE/SMutableColorViewer.h
@@ -17,15 +17,12 @@ public:
 	void Construct(const FArguments& InArgs);
 
 	/** Set the Mutable Bool to be used for this widget */
-	void SetColor(float InRed, float InGreen, float InBlue, float InAlpha);
+	void SetColor(FVector4f Color);
 
 private:
 
 	/** Color being displayed */
-	float RedValue = 0.0f;
-	float GreenValue = 0.0f;
-	float BlueValue = 0.0f;
-	float AlphaValue = 1.0f;
+	FVector4f Color = FVector4f(0,0,0,1);
 
 	/** Color box widget designed to serve as a preview of the color reported by mutable */
 	TSharedPtr<SColorBlock> ColorPreview;
