[mutable] Implement memory budget control for image operations.

diff --git a/Engine/Plugins/Experimental/Mutable/Source/MutableRuntime/Private/MuR/CodeRunner.cpp b/Engine/Plugins/Experimental/Mutable/Source/MutableRuntime/Private/MuR/CodeRunner.cpp
--- a/Engine/Plugins/Experimental/Mutable/Source/MutableRuntime/Private/MuR/CodeRunner.cpp
+++ b/Engine/Plugins/Experimental/Mutable/Source/MutableRuntime/Private/MuR/CodeRunner.cpp
@@ -35,12 +35,10 @@
 #include "MuR/OpImageLuminance.h"
 #include "MuR/OpImageMipmap.h"
 #include "MuR/OpImageNormalCombine.h"
-#include "MuR/OpImagePixelFormat.h"
 #include "MuR/OpImageProject.h"
 #include "MuR/OpImageRasterMesh.h"
 #include "MuR/OpImageResize.h"
 #include "MuR/OpImageSaturate.h"
-#include "MuR/OpImageSelectColour.h"
 #include "MuR/OpImageTransform.h"
 #include "MuR/OpLayoutPack.h"
 #include "MuR/OpLayoutRemoveBlocks.h"
@@ -66,7 +64,6 @@
 #include "MuR/ParametersPrivate.h"
 #include "MuR/PhysicsBody.h"
 #include "MuR/Platform.h"
-#include "MuR/SettingsPrivate.h"
 #include "MuR/Skeleton.h"
 #include "MuR/SystemPrivate.h"
 #include "Templates/Tuple.h"
@@ -102,7 +99,7 @@ namespace mu
 {
 
     //---------------------------------------------------------------------------------------------
-    CodeRunner::CodeRunner(const SettingsPtrConst& InSettings,
+    CodeRunner::CodeRunner(const Ptr<const Settings>& InSettings,
 		class System::Private* InSystem,
 		EExecutionStrategy InExecutionStrategy,
 		const TSharedPtr<const Model>& InModel,
@@ -116,16 +113,15 @@ namespace mu
 		, m_pParams(InParams)
 		, m_lodMask(InLodMask)
 	{
-		ScheduledStagePerOp.resize(m_pModel->GetPrivate()->m_program.m_opAddress.Num());
+		FProgram& program = m_pModel->GetPrivate()->m_program;
+		ScheduledStagePerOp.resize(program.m_opAddress.Num());
 
 		// We will read this in the end, so make sure we keep it.
    		if (Type == FScheduledOp::EType::Full)
    		{
-			m_pSystem->m_memory->IncreaseHitCount(FCacheAddress(at, 0, executionOptions));
+			GetMemory().IncreaseHitCount(FCacheAddress(at, 0, executionOptions));
 		}
     	
-		FProgram& program = m_pModel->GetPrivate()->m_program;
-		m_romPendingOps.SetNum(program.m_roms.Num());
 
 		ImageLOD = InImageLOD;
 
@@ -141,7 +137,7 @@ namespace mu
     //---------------------------------------------------------------------------------------------
 	FProgramCache& CodeRunner::GetMemory()
     {
-		return *m_pSystem->m_memory;
+		return *m_pSystem->WorkingMemoryManager.CurrentInstanceCache;
 	}
 
 
@@ -156,9 +152,9 @@ namespace mu
 
 		check(m_pSystem);
 
-		if (m_pSystem->m_pImageParameterGenerator)
+		if (m_pSystem->ImageParameterGenerator)
 		{
-			return m_pSystem->m_pImageParameterGenerator->GetImageAsync(Id, MipmapsToSkip, ResultCallback);
+			return m_pSystem->ImageParameterGenerator->GetImageAsync(Id, MipmapsToSkip, ResultCallback);
 
 			// Don't cache for now. Need to figure out how to invalidate them.
 			// \TODO: Like constants? attached to a cache level?
@@ -190,9 +186,9 @@ namespace mu
 
 		check(m_pSystem);
 
-		if (m_pSystem->m_pImageParameterGenerator)
+		if (m_pSystem->ImageParameterGenerator)
 		{
-			return m_pSystem->m_pImageParameterGenerator->GetImageDesc(Id, MipmapsToSkip);
+			return m_pSystem->ImageParameterGenerator->GetImageDesc(Id, MipmapsToSkip);
 		}
 		else
 		{
@@ -231,7 +227,7 @@ namespace mu
 
             // If there is no expression, we'll assume true.
             bool value = true;
-            value = GetMemory().GetBool( FCacheAddress(args.condition, item.ExecutionIndex, item.ExecutionOptions) );
+            value = LoadBool( FCacheAddress(args.condition, item.ExecutionIndex, item.ExecutionOptions) );
 
             OP::ADDRESS resultAt = value ? args.yes : args.no;
 
@@ -251,16 +247,16 @@ namespace mu
             FCacheAddress rat( resultAt, item );
             switch (GetOpDataType(type))
             {
-            case DT_BOOL:       GetMemory().SetBool( cat, GetMemory().GetBool(rat) ); break;
-            case DT_INT:        GetMemory().SetInt( cat, GetMemory().GetInt(rat) ); break;
-            case DT_SCALAR:     GetMemory().SetScalar( cat, GetMemory().GetScalar(rat) ); break;
-			case DT_STRING:		GetMemory().SetString( cat, GetMemory().GetString( rat ) ); break;
-            case DT_COLOUR:		GetMemory().SetColour( cat, GetMemory().GetColour( rat ) ); break;
-            case DT_PROJECTOR:  GetMemory().SetProjector( cat, GetMemory().GetProjector(rat) ); break;
-            case DT_MESH:       GetMemory().SetMesh( cat, GetMemory().GetMesh(rat) ); break;
-            case DT_IMAGE:      GetMemory().SetImage( cat, GetMemory().GetImage(rat) ); break;
-            case DT_LAYOUT:     GetMemory().SetLayout( cat, GetMemory().GetLayout(rat) ); break;
-            case DT_INSTANCE:   GetMemory().SetInstance( cat, GetMemory().GetInstance(rat) ); break;
+            case DT_BOOL:       StoreBool( cat, LoadBool(rat) ); break;
+            case DT_INT:        StoreInt( cat, LoadInt(rat) ); break;
+            case DT_SCALAR:     StoreScalar( cat, LoadScalar(rat) ); break;
+			case DT_STRING:		StoreString( cat, LoadString( rat ) ); break;
+            case DT_COLOUR:		StoreColor( cat, LoadColor( rat ) ); break;
+            case DT_PROJECTOR:  StoreProjector( cat, LoadProjector(rat) ); break;
+            case DT_MESH:       StoreMesh( cat, LoadMesh(rat) ); break;
+            case DT_IMAGE:      StoreImage( cat, LoadImage(rat) ); break;
+            case DT_LAYOUT:     StoreLayout( cat, LoadLayout(rat) ); break;
+            case DT_INSTANCE:   StoreInstance( cat, LoadInstance(rat) ); break;
             default:
                 // Not implemented
                 check( false );
@@ -300,23 +296,22 @@ namespace mu
 		{
 			if (VarAddress)
 			{
-				AddOp(FScheduledOp(item.At, item, 1),
-					FScheduledOp(VarAddress, item));
+				AddOp(FScheduledOp(item.At, item, 1), FScheduledOp(VarAddress, item));
 			}
 			else
 			{
 				switch (GetOpDataType(type))
 				{
-				case DT_BOOL:       GetMemory().SetBool(item, false); break;
-				case DT_INT:        GetMemory().SetInt(item, 0); break;
-				case DT_SCALAR:		GetMemory().SetScalar(item, 0.0f); break;
-				case DT_STRING:		GetMemory().SetString(item, nullptr); break;
-				case DT_COLOUR:		GetMemory().SetColour(item, FVector4f()); break;
-				case DT_PROJECTOR:  GetMemory().SetProjector(item, nullptr); break;
-				case DT_MESH:       GetMemory().SetMesh(item, nullptr); break;
-				case DT_IMAGE:      GetMemory().SetImage(item, nullptr); break;
-				case DT_LAYOUT:     GetMemory().SetLayout(item, nullptr); break;
-				case DT_INSTANCE:   GetMemory().SetInstance(item, nullptr); break;
+				case DT_BOOL:       StoreBool(item, false); break;
+				case DT_INT:        StoreInt(item, 0); break;
+				case DT_SCALAR:		StoreScalar(item, 0.0f); break;
+				case DT_STRING:		StoreString(item, nullptr); break;
+				case DT_COLOUR:		StoreColor(item, FVector4f()); break;
+				case DT_PROJECTOR:  StoreProjector(item, nullptr); break;
+				case DT_MESH:       StoreMesh(item, nullptr); break;
+				case DT_IMAGE:      StoreImage(item, nullptr); break;
+				case DT_LAYOUT:     StoreLayout(item, nullptr); break;
+				case DT_INSTANCE:   StoreInstance(item, nullptr); break;
 				default:
 					// Not implemented
 					check(false);
@@ -328,7 +323,7 @@ namespace mu
 		case 1:
 		{
 			// Get the variable result
-			int var = GetMemory().GetInt(FCacheAddress(VarAddress, item));
+			int var = LoadInt(FCacheAddress(VarAddress, item));
 
 			OP::ADDRESS valueAt = DefAddress;
 			for (uint32 C = 0; C < CaseCount; ++C)
@@ -364,16 +359,16 @@ namespace mu
             FCacheAddress rat( resultAt, item );
             switch (GetOpDataType(type))
             {
-            case DT_BOOL:       GetMemory().SetBool( cat, GetMemory().GetBool(rat) ); break;
-            case DT_INT:        GetMemory().SetInt( cat, GetMemory().GetInt(rat) ); break;
-            case DT_SCALAR:     GetMemory().SetScalar( cat, GetMemory().GetScalar(rat) ); break;
-            case DT_STRING:		GetMemory().SetString( cat, GetMemory().GetString( rat ) ); break;
-            case DT_COLOUR:		GetMemory().SetColour( cat, GetMemory().GetColour( rat ) ); break;
-            case DT_PROJECTOR:  GetMemory().SetProjector( cat, GetMemory().GetProjector(rat) ); break;
-            case DT_MESH:       GetMemory().SetMesh( cat, GetMemory().GetMesh(rat) ); break;
-            case DT_IMAGE:      GetMemory().SetImage( cat, GetMemory().GetImage(rat) ); break;
-            case DT_LAYOUT:     GetMemory().SetLayout( cat, GetMemory().GetLayout(rat) ); break;
-            case DT_INSTANCE:   GetMemory().SetInstance( cat, GetMemory().GetInstance(rat) ); break;
+            case DT_BOOL:       StoreBool( cat, LoadBool(rat) ); break;
+            case DT_INT:        StoreInt( cat, LoadInt(rat) ); break;
+            case DT_SCALAR:     StoreScalar( cat, LoadScalar(rat) ); break;
+            case DT_STRING:		StoreString( cat, LoadString( rat ) ); break;
+            case DT_COLOUR:		StoreColor( cat, LoadColor( rat ) ); break;
+            case DT_PROJECTOR:  StoreProjector( cat, LoadProjector(rat) ); break;
+            case DT_MESH:       StoreMesh( cat, LoadMesh(rat) ); break;
+            case DT_IMAGE:      StoreImage( cat, LoadImage(rat) ); break;
+            case DT_LAYOUT:     StoreLayout( cat, LoadLayout(rat) ); break;
+            case DT_INSTANCE:   StoreInstance( cat, LoadInstance(rat) ); break;
             default:
                 // Not implemented
                 check( false );
@@ -411,7 +406,7 @@ namespace mu
 
             case 1:
             {
-                InstancePtrConst pBase = GetMemory().GetInstance( FCacheAddress(args.instance,item) );
+                InstancePtrConst pBase = LoadInstance( FCacheAddress(args.instance,item) );
                 InstancePtr pResult;
                 if (!pBase)
                 {
@@ -419,12 +414,12 @@ namespace mu
                 }
                 else
                 {
-					pResult = CloneOrTakeOver<Instance>(pBase.get());
+					pResult = mu::CloneOrTakeOver<Instance>(pBase.get());
                 }
 
                 if ( args.value )
                 {
-					FVector4f value = GetMemory().GetColour( FCacheAddress(args.value,item) );
+					FVector4f value = LoadColor( FCacheAddress(args.value,item) );
 
                     OP::ADDRESS nameAd = args.name;
                     check(  nameAd < (uint32)pModel->GetPrivate()->m_program.m_constantStrings.Num() );
@@ -432,7 +427,7 @@ namespace mu
 
                     pResult->GetPrivate()->AddVector( 0, 0, 0, value, strName );
                 }
-                GetMemory().SetInstance( item, pResult );
+                StoreInstance( item, pResult );
                 break;
             }
 
@@ -456,7 +451,7 @@ namespace mu
 
             case 1:
             {
-                InstancePtrConst pBase = GetMemory().GetInstance( FCacheAddress(args.instance,item) );
+                InstancePtrConst pBase = LoadInstance( FCacheAddress(args.instance,item) );
                 InstancePtr pResult;
                 if (!pBase)
                 {
@@ -464,12 +459,12 @@ namespace mu
                 }
                 else
                 {
-                    pResult = CloneOrTakeOver<Instance>(pBase.get());
+                    pResult = mu::CloneOrTakeOver<Instance>(pBase.get());
                 }
 
                 if ( args.value )
                 {
-                    float value = GetMemory().GetScalar( FCacheAddress(args.value,item) );
+                    float value = LoadScalar( FCacheAddress(args.value,item) );
 
                     OP::ADDRESS nameAd = args.name;
                     check(  nameAd < (uint32)pModel->GetPrivate()->m_program.m_constantStrings.Num() );
@@ -477,7 +472,7 @@ namespace mu
 
                     pResult->GetPrivate()->AddScalar( 0, 0, 0, value, strName );
                 }
-                GetMemory().SetInstance( item, pResult );
+                StoreInstance( item, pResult );
                 break;
             }
 
@@ -501,7 +496,7 @@ namespace mu
             case 1:
             {
                 InstancePtrConst pBase =
-                    GetMemory().GetInstance( FCacheAddress( args.instance, item ) );
+                    LoadInstance( FCacheAddress( args.instance, item ) );
                 InstancePtr pResult;
                 if ( !pBase )
                 {
@@ -509,13 +504,13 @@ namespace mu
                 }
                 else
                 {
-					pResult = CloneOrTakeOver<Instance>(pBase.get());
+					pResult = mu::CloneOrTakeOver<Instance>(pBase.get());
 				}
 
                 if ( args.value )
                 {
                     Ptr<const String> value =
-                        GetMemory().GetString( FCacheAddress( args.value, item ) );
+                        LoadString( FCacheAddress( args.value, item ) );
 
                     OP::ADDRESS nameAd = args.name;
                     check( nameAd < (uint32)pModel->GetPrivate()->m_program.m_constantStrings.Num() );
@@ -524,7 +519,7 @@ namespace mu
 
                     pResult->GetPrivate()->AddString( 0, 0, 0, value->GetValue(), strName );
                 }
-                GetMemory().SetInstance( item, pResult );
+                StoreInstance( item, pResult );
                 break;
             }
 
@@ -548,7 +543,7 @@ namespace mu
 
             case 1:
             {
-                InstancePtrConst pBase = GetMemory().GetInstance( FCacheAddress(args.instance,item) );
+                InstancePtrConst pBase = LoadInstance( FCacheAddress(args.instance,item) );
                 InstancePtr pResult;
                 if (!pBase)
                 {
@@ -556,12 +551,12 @@ namespace mu
                 }
                 else
                 {
-					pResult = CloneOrTakeOver<Instance>(pBase.get());
+					pResult = mu::CloneOrTakeOver<Instance>(pBase.get());
 				}
 
                 if ( args.value )
                 {
-                    InstancePtrConst pComp = GetMemory().GetInstance( FCacheAddress(args.value,item) );
+                    InstancePtrConst pComp = LoadInstance( FCacheAddress(args.value,item) );
 
                     int cindex = pResult->GetPrivate()->AddComponent( 0 );
 
@@ -583,7 +578,7 @@ namespace mu
                         pResult->GetPrivate()->SetComponentName( 0, cindex, strName );
                     }
                 }
-                GetMemory().SetInstance( item, pResult );
+                StoreInstance( item, pResult );
                 break;
             }
 
@@ -607,12 +602,12 @@ namespace mu
 
             case 1:
             {
-                InstancePtrConst pBase = GetMemory().GetInstance( FCacheAddress(args.instance,item) );
+                InstancePtrConst pBase = LoadInstance( FCacheAddress(args.instance,item) );
 
                 InstancePtr pResult;
 				if (pBase)
 				{
-					pResult = CloneOrTakeOver<Instance>(pBase.get());
+					pResult = mu::CloneOrTakeOver<Instance>(pBase.get());
 				}
 				else
 				{
@@ -623,7 +618,7 @@ namespace mu
                 // additional information like internal or external IDs
                 //if ( args.value )
                 {
-                    InstancePtrConst pSurf = GetMemory().GetInstance( FCacheAddress(args.value,item) );
+                    InstancePtrConst pSurf = LoadInstance( FCacheAddress(args.value,item) );
 
                     int sindex = pResult->GetPrivate()->AddSurface( 0, 0 );
 
@@ -654,7 +649,7 @@ namespace mu
                     pResult->GetPrivate()->m_lods[0].m_components[0].m_surfaces[sindex].ExternalId = args.ExternalId;
                     pResult->GetPrivate()->m_lods[0].m_components[0].m_surfaces[sindex].SharedId = args.SharedSurfaceId;
                 }
-                GetMemory().SetInstance( item, pResult );
+                StoreInstance( item, pResult );
                 break;
             }
 
@@ -704,7 +699,7 @@ namespace mu
 
                         if ( selectedLod )
                         {
-                            InstancePtrConst pLOD = GetMemory().GetInstance( FCacheAddress(args.lod[i],item) );
+                            InstancePtrConst pLOD = LoadInstance( FCacheAddress(args.lod[i],item) );
 
                             int LODIndex = pResult->GetPrivate()->AddLOD();
 
@@ -722,7 +717,7 @@ namespace mu
                     }
                 }
 
-                GetMemory().SetInstance( item, pResult );
+                StoreInstance( item, pResult );
                 break;
             }
 
@@ -756,11 +751,11 @@ namespace mu
 				case 1:
 				{
 					// Assemble result
-					InstancePtrConst InstanceOpResult = GetMemory().GetInstance(FCacheAddress(Args.Instance, item));
+					InstancePtrConst InstanceOpResult = LoadInstance(FCacheAddress(Args.Instance, item));
 
-					InstancePtr Result = CloneOrTakeOver<Instance>(InstanceOpResult.get());
+					InstancePtr Result = mu::CloneOrTakeOver<Instance>(InstanceOpResult.get());
 
-					if (ExtensionDataPtrConst ExtensionData = GetMemory().GetExtensionData(FCacheAddress(Args.ExtensionData, item)))
+					if (ExtensionDataPtrConst ExtensionData = LoadExtensionData(FCacheAddress(Args.ExtensionData, item)))
 					{
 						const OP::ADDRESS NameAddress = Args.ExtensionDataName;
 						check(NameAddress < (uint32)pModel->GetPrivate()->m_program.m_constantStrings.Num());
@@ -769,7 +764,7 @@ namespace mu
 						Result->GetPrivate()->AddExtensionData(ExtensionData, NameString);
 					}
 
-					GetMemory().SetInstance(item, Result);
+					StoreInstance(item, Result);
 					break;
 				}
 
@@ -810,7 +805,7 @@ namespace mu
 
             case 1:
             {
-                InstancePtrConst pBase = GetMemory().GetInstance( FCacheAddress(args.instance,item) );
+                InstancePtrConst pBase = LoadInstance( FCacheAddress(args.instance,item) );
                 InstancePtr pResult;
                 if (!pBase)
                 {
@@ -818,7 +813,7 @@ namespace mu
                 }
                 else
                 {
-					pResult = CloneOrTakeOver<Instance>(pBase.get());
+					pResult = mu::CloneOrTakeOver<Instance>(pBase.get());
 				}
 
                 if ( args.value )
@@ -832,7 +827,7 @@ namespace mu
                     const char* strName = pModel->GetPrivate()->m_program.m_constantStrings[ nameAd ].c_str();
                     pResult->GetPrivate()->AddMesh( 0, 0, meshId, strName );
                 }
-                GetMemory().SetInstance( item, pResult );
+                StoreInstance( item, pResult );
                 break;
             }
 
@@ -854,7 +849,7 @@ namespace mu
 
             case 1:
             {
-                InstancePtrConst pBase = GetMemory().GetInstance( FCacheAddress(args.instance,item) );
+                InstancePtrConst pBase = LoadInstance( FCacheAddress(args.instance,item) );
                 InstancePtr pResult;
                 if (!pBase)
                 {
@@ -862,7 +857,7 @@ namespace mu
                 }
                 else
                 {
-					pResult = CloneOrTakeOver<Instance>(pBase.get());
+					pResult = mu::CloneOrTakeOver<Instance>(pBase.get());
 				}
 
                 if ( args.value )
@@ -876,7 +871,7 @@ namespace mu
                     const char* strName = pModel->GetPrivate()->m_program.m_constantStrings[ nameAd ].c_str();
                     pResult->GetPrivate()->AddImage( 0, 0, 0, imageId, strName );
                 }
-                GetMemory().SetInstance( item, pResult );
+                StoreInstance( item, pResult );
                 break;
             }
 
@@ -915,24 +910,24 @@ namespace mu
 
             MeshPtrConst pSourceConst;
             program.GetConstant( cat, pSourceConst );
-            MeshPtr pSource = pSourceConst->Clone();
+            MeshPtr Source = pSourceConst->Clone();
 
             // Set the separate skeleton if necessary
             if (args.skeleton>=0)
             {
                 check( program.m_constantSkeletons.Num()>size_t(args.skeleton)  );
                 Ptr<const Skeleton> pSkeleton = program.m_constantSkeletons[args.skeleton];
-                pSource->SetSkeleton(pSkeleton);
+                Source->SetSkeleton(pSkeleton);
             }
 
 			if (args.physicsBody >= 0)
 			{
                 check( program.m_constantPhysicsBodies.Num()>size_t(args.physicsBody)  );
                 Ptr<const PhysicsBody> pPhysicsBody = program.m_constantPhysicsBodies[args.physicsBody];
-                pSource->SetPhysicsBody(pPhysicsBody);
+                Source->SetPhysicsBody(pPhysicsBody);
 			}
 
-            GetMemory().SetMesh( item, pSource );
+            StoreMesh( item, Source );
 			//UE_LOG(LogMutableCore, Log, TEXT("Set mesh constant %d."), item.At);
             break;
         }
@@ -945,13 +940,13 @@ namespace mu
             const FProgram& program = pModel->GetPrivate()->m_program;
 
 			int32 MipsToSkip = item.ExecutionOptions;
-            Ptr<const Image> pSource;
-            program.GetConstant( cat, pSource, MipsToSkip);
+            Ptr<const Image> Source;
+            program.GetConstant( cat, Source, MipsToSkip);
 
 			// Assume the ROM has been loaded previously in a task generated at IssueOp
-			check(pSource);
+			check(Source);
 
-            GetMemory().SetImage( item, pSource );
+            StoreImage( item, Source );
 			//UE_LOG(LogMutableCore, Log, TEXT("Set image constant %d."), item.At);
 			break;
         }
@@ -968,7 +963,7 @@ namespace mu
 
 			check(SourceConst);
 
-            GetMemory().SetExtensionData(item, SourceConst);
+            StoreExtensionData(item, SourceConst);
             break;
 		}
 
@@ -1006,19 +1001,19 @@ namespace mu
             {
            		MUTABLE_CPUPROFILER_SCOPE(ME_APPLYLAYOUT)
             		
-                Ptr<const Mesh> pBase = GetMemory().GetMesh( FCacheAddress(args.mesh,item) );
+                Ptr<const Mesh> pBase = LoadMesh( FCacheAddress(args.mesh,item) );
 
                 MeshPtr pApplied;
                 if (pBase)
                 {
                     pApplied = pBase->Clone();
 
-                    Ptr<const Layout> pLayout = GetMemory().GetLayout( FCacheAddress(args.layout,item) );
+                    Ptr<const Layout> pLayout = LoadLayout( FCacheAddress(args.layout,item) );
                     int texCoordsSet = args.channel;
                     MeshApplyLayout( pApplied.get(), pLayout.get(), texCoordsSet );
                 }
 
-                GetMemory().SetMesh( item, pApplied );
+                StoreMesh( item, pApplied );
                 break;
             }
 
@@ -1050,7 +1045,7 @@ namespace mu
                 }
                 else
                 {
-                    GetMemory().SetMesh( item, nullptr );
+                    StoreMesh( item, nullptr );
                 }
                 break;
 
@@ -1058,8 +1053,8 @@ namespace mu
             {
        	        MUTABLE_CPUPROFILER_SCOPE(ME_DIFFERENCE)
 
-                Ptr<const Mesh> pBase = GetMemory().GetMesh(FCacheAddress(BaseAt,item));
-                Ptr<const Mesh> pTarget = GetMemory().GetMesh(FCacheAddress(TargetAt,item));
+                Ptr<const Mesh> pBase = LoadMesh(FCacheAddress(BaseAt,item));
+                Ptr<const Mesh> pTarget = LoadMesh(FCacheAddress(TargetAt,item));
 
 				TArray<MESH_BUFFER_SEMANTIC, TInlineAllocator<8>> Semantics;
 				TArray<int32, TInlineAllocator<8>> SemanticIndices;
@@ -1085,7 +1080,7 @@ namespace mu
                                           NumChannels, Semantics.GetData(), SemanticIndices.GetData(),
                                           bIgnoreTextureCoords!=0 );
 
-                GetMemory().SetMesh( item, pResult );
+                StoreMesh( item, pResult );
                 break;
             }
 
@@ -1127,7 +1122,7 @@ namespace mu
                 }
                 else
                 {
-                    GetMemory().SetMesh( item, nullptr );
+                    StoreMesh( item, nullptr );
                 }
                 break;
 
@@ -1136,7 +1131,7 @@ namespace mu
                 MUTABLE_CPUPROFILER_SCOPE(ME_MORPH2_1)
 
                 bool baseValid = GetMemory().IsValid( FCacheAddress(BaseAt,item) );
-                float factor = GetMemory().GetScalar( FCacheAddress(FactorAt,item) );
+                float factor = LoadScalar( FCacheAddress(FactorAt,item) );
 
                 // Factor goes from -1 to 1 across all targets. [0 - 1] represents positive morphs, while [-1, 0) represent negative morphs.
 				factor = FMath::Clamp(factor, -1.0f, 1.0f); // Is the factor not in range [-1, 1], it will index a non existing morph.
@@ -1208,7 +1203,7 @@ namespace mu
             {
        		    MUTABLE_CPUPROFILER_SCOPE(ME_MORPH2_2)
 
-                Ptr<const Mesh> pBase = GetMemory().GetMesh( FCacheAddress(BaseAt,item) );
+                Ptr<const Mesh> pBase = LoadMesh( FCacheAddress(BaseAt,item) );
 
                 // Factor from 0 to 1 between the two targets
                 const FScheduledOpData& HeapData = m_heapData[ (size_t)item.CustomState ];
@@ -1224,7 +1219,7 @@ namespace mu
 					if ( bifactor < UE_SMALL_NUMBER && bifactor > -UE_SMALL_NUMBER )
                     {
                         // Base with one full morph
-                        Ptr<const Mesh> pMorph = GetMemory().GetMesh( FCacheAddress(Targets[min],item) );
+                        Ptr<const Mesh> pMorph = LoadMesh( FCacheAddress(Targets[min],item) );
                         if (pMorph)
                         {
 							pResult = MeshMorph(pBase.get(), pMorph.get());
@@ -1234,7 +1229,7 @@ namespace mu
                     else if ( bifactor > 1.0f - UE_SMALL_NUMBER )
                     {
                         check( max>0 );
-                        Ptr<const Mesh> pMorph = GetMemory().GetMesh( FCacheAddress(Targets[max],item) );
+                        Ptr<const Mesh> pMorph = LoadMesh( FCacheAddress(Targets[max],item) );
                         if (pMorph)
                         {
 							pResult = MeshMorph(pBase.get(), pMorph.get());
@@ -1244,7 +1239,7 @@ namespace mu
 					else if (bifactor <= -1.0f + UE_SMALL_NUMBER)
 					{
 						check( max > 0 );
-						Ptr<const Mesh> pMorph = GetMemory().GetMesh(FCacheAddress(Targets[max], item));
+						Ptr<const Mesh> pMorph = LoadMesh(FCacheAddress(Targets[max], item));
 						if (pMorph)
 						{
 							pResult = MeshMorph(pBase.get(), pMorph.get(), bifactor);
@@ -1253,8 +1248,8 @@ namespace mu
                     // Mix two targets on the base
                     else
                     {
-                        Ptr<const Mesh> pMin = GetMemory().GetMesh( FCacheAddress(Targets[min],item) );
-                        Ptr<const Mesh> pMax = GetMemory().GetMesh( FCacheAddress(Targets[max],item) );
+                        Ptr<const Mesh> pMin = LoadMesh( FCacheAddress(Targets[min],item) );
+                        Ptr<const Mesh> pMax = LoadMesh( FCacheAddress(Targets[max],item) );
                         if (pMin && pMax)
                         {
                             pResult = MeshMorph2( pBase.get(), pMin.get(), pMax.get(), bifactor );
@@ -1266,7 +1261,7 @@ namespace mu
 					// In other words, we have a discontinuity arround 0.
 					// With the current implementation, the only solutions to get a negative pMin morph is to apply a really small negative value (lim x -> -0) (for example: -0.00001).
                 }
-                GetMemory().SetMesh( item, pResult );
+                StoreMesh( item, pResult );
 
                 break;
             }
@@ -1293,8 +1288,8 @@ namespace mu
             {
 				MUTABLE_CPUPROFILER_SCOPE(ME_MERGE_1)
 
-                MeshPtrConst pA = GetMemory().GetMesh( FCacheAddress(args.base,item) );
-                MeshPtrConst pB = GetMemory().GetMesh( FCacheAddress(args.added,item) );
+                MeshPtrConst pA = LoadMesh( FCacheAddress(args.base,item) );
+                MeshPtrConst pB = LoadMesh( FCacheAddress(args.added,item) );
 
                 MeshPtr pResult;
 
@@ -1326,7 +1321,7 @@ namespace mu
                     pResult = new Mesh();
                 }
 
-                GetMemory().SetMesh( item, pResult );
+                StoreMesh( item, pResult );
                 break;
             }
 
@@ -1351,7 +1346,7 @@ namespace mu
                 }
                 else
                 {
-                    GetMemory().SetMesh( item, nullptr );
+                    StoreMesh( item, nullptr );
                 }
                 break;
 
@@ -1368,7 +1363,7 @@ namespace mu
                 }
 
                 // Factor from 0 to 1 across all targets
-                float factor = GetMemory().GetScalar( FCacheAddress(args.factor,item) );
+                float factor = LoadScalar( FCacheAddress(args.factor,item) );
 
                 float delta = 1.0f/(count-1);
                 int min = (int)floorf( factor/delta );
@@ -1389,13 +1384,13 @@ namespace mu
                     if( min==0 )
                     {
                         // Just the base
-                            Ptr<const Mesh> pBase = GetMemory().GetMesh( FCacheAddress(args.base,item) );
+                            Ptr<const Mesh> pBase = LoadMesh( FCacheAddress(args.base,item) );
                             MeshPtr pResult;
                             if (pBase)
                             {
                                 pResult = pBase->Clone();
                             }
-                            GetMemory().SetMesh( item, pResult );
+                            StoreMesh( item, pResult );
                         }
                     else
                     {
@@ -1455,7 +1450,7 @@ namespace mu
                 int min = data.Interpolate.Min;
                 int max = data.Interpolate.Max;
 
-                Ptr<const Mesh> pBase = GetMemory().GetMesh( FCacheAddress(args.base, item) );
+                Ptr<const Mesh> pBase = LoadMesh( FCacheAddress(args.base, item) );
 
                 MeshPtr pResult;
 
@@ -1472,7 +1467,7 @@ namespace mu
                         else
                         {
                             // Base with one full morph
-                            Ptr<const Mesh> pMorph = GetMemory().GetMesh( FCacheAddress(args.targets[min-1],item) );
+                            Ptr<const Mesh> pMorph = LoadMesh( FCacheAddress(args.targets[min-1],item) );
                             pResult = MeshMorph( pBase.get(), pMorph.get() );
                         }
                     }
@@ -1480,7 +1475,7 @@ namespace mu
                     else if ( bifactor > 1.0f-UE_SMALL_NUMBER )
                     {
                         check( max>0 );
-                        Ptr<const Mesh> pMorph = GetMemory().GetMesh( FCacheAddress(args.targets[max-1],item) );
+                        Ptr<const Mesh> pMorph = LoadMesh( FCacheAddress(args.targets[max-1],item) );
 
                         if (pMorph)
                         {
@@ -1494,7 +1489,7 @@ namespace mu
                     // Mix the first target on the base
                     else if ( min==0 )
                     {
-                        Ptr<const Mesh> pMorph = GetMemory().GetMesh( FCacheAddress( args.targets[0], item ) );
+                        Ptr<const Mesh> pMorph = LoadMesh( FCacheAddress( args.targets[0], item ) );
                         if (pMorph)
                         {
                             pResult = MeshMorph( pBase.get(), pMorph.get(), bifactor );
@@ -1508,8 +1503,8 @@ namespace mu
                     // Mix two targets on the base
                     else
                     {
-                        Ptr<const Mesh> pMin = GetMemory().GetMesh( FCacheAddress(args.targets[min-1],item) );
-                        Ptr<const Mesh> pMax = GetMemory().GetMesh( FCacheAddress(args.targets[max-1],item) );
+                        Ptr<const Mesh> pMin = LoadMesh( FCacheAddress(args.targets[min-1],item) );
+                        Ptr<const Mesh> pMax = LoadMesh( FCacheAddress(args.targets[max-1],item) );
 
                         if (pMin && pMax)
                         {
@@ -1522,7 +1517,7 @@ namespace mu
                     }
                 }
 
-                GetMemory().SetMesh( item, pResult );
+                StoreMesh( item, pResult );
                 break;
             }
 
@@ -1548,16 +1543,16 @@ namespace mu
             {
             	MUTABLE_CPUPROFILER_SCOPE(ME_MASKCLIPMESH_1)
             		
-                Ptr<const Mesh> pSource = GetMemory().GetMesh( FCacheAddress(args.source,item) );
-                Ptr<const Mesh> pClip = GetMemory().GetMesh( FCacheAddress(args.clip,item) );
+                Ptr<const Mesh> Source = LoadMesh( FCacheAddress(args.source,item) );
+                Ptr<const Mesh> pClip = LoadMesh( FCacheAddress(args.clip,item) );
 
                 // Only if both are valid.
                 MeshPtr pResult;
-                if (pSource.get() && pClip.get())
+                if (Source.get() && pClip.get())
                 {
-                    pResult = MeshMaskClipMesh(pSource.get(), pClip.get());
+                    pResult = MeshMaskClipMesh(Source.get(), pClip.get());
                 }
-                GetMemory().SetMesh( item, pResult );
+                StoreMesh( item, pResult );
 
                 break;
             }
@@ -1584,16 +1579,16 @@ namespace mu
             {
            		MUTABLE_CPUPROFILER_SCOPE(ME_MASKDIFF_1)
             		
-                Ptr<const Mesh> pSource = GetMemory().GetMesh( FCacheAddress(args.source,item) );
-                Ptr<const Mesh> pClip = GetMemory().GetMesh( FCacheAddress(args.fragment,item) );
+                Ptr<const Mesh> Source = LoadMesh( FCacheAddress(args.source,item) );
+                Ptr<const Mesh> pClip = LoadMesh( FCacheAddress(args.fragment,item) );
 
                 // Only if both are valid.
                 MeshPtr pResult;
-                if (pSource.get() && pClip.get())
+                if (Source.get() && pClip.get())
                 {
-                    pResult = MeshMaskDiff(pSource.get(), pClip.get());
+                    pResult = MeshMaskDiff(Source.get(), pClip.get());
                 }
-                GetMemory().SetMesh( item, pResult );
+                StoreMesh( item, pResult );
 
                 break;
             }
@@ -1619,7 +1614,7 @@ namespace mu
                     }
                     else
                     {
-                    GetMemory().SetMesh( item, nullptr );
+                    StoreMesh( item, nullptr );
                 }
                 break;
 
@@ -1627,12 +1622,12 @@ namespace mu
             {
             	MUTABLE_CPUPROFILER_SCOPE(ME_FORMAT_1)
             		
-                Ptr<const Mesh> pSource = GetMemory().GetMesh( FCacheAddress(args.source,item) );
-                Ptr<const Mesh> pFormat = GetMemory().GetMesh( FCacheAddress(args.format,item) );
+                Ptr<const Mesh> Source = LoadMesh( FCacheAddress(args.source,item) );
+                Ptr<const Mesh> pFormat = LoadMesh( FCacheAddress(args.format,item) );
 
                 uint8 flags = args.buffers;
                 MeshPtr pResult;
-                pResult = MeshFormat( pSource.get(),
+                pResult = MeshFormat( Source.get(),
                                       pFormat.get(),
                                       true,
                                       (flags & OP::MeshFormatArgs::BT_VERTEX) != 0,
@@ -1646,7 +1641,7 @@ namespace mu
                     pResult->ResetBufferIndices();
                 }
 
-                GetMemory().SetMesh( item, pResult );
+                StoreMesh( item, pResult );
                 break;
             }
 
@@ -1683,7 +1678,7 @@ namespace mu
                     }
                     else
                     {
-                    GetMemory().SetMesh( item, nullptr );
+                    StoreMesh( item, nullptr );
                 }
                 break;
 
@@ -1691,19 +1686,19 @@ namespace mu
             {
 				MUTABLE_CPUPROFILER_SCOPE(ME_EXTRACTLAYOUTBLOCK_1)
 
-                Ptr<const Mesh> pSource = GetMemory().GetMesh( FCacheAddress(source,item) );
+                Ptr<const Mesh> Source = LoadMesh( FCacheAddress(source,item) );
 
                 // Access with memcpy necessary for unaligned arm issues.
                 uint32 blocks[1024];
 				FMemory::Memcpy( blocks, data, sizeof(uint32)*FMath::Min(1024,int(blockCount)) );
 
                 MeshPtr pResult;
-                pResult = MeshExtractLayoutBlock( pSource.get(),
+                pResult = MeshExtractLayoutBlock( Source.get(),
                                                   layout,
                                                   blockCount,
                                                   blocks );
 
-                GetMemory().SetMesh( item, pResult );
+                StoreMesh( item, pResult );
                 break;
             }
 
@@ -1727,7 +1722,7 @@ namespace mu
                     }
                     else
                     {
-                    GetMemory().SetMesh( item, nullptr );
+                    StoreMesh( item, nullptr );
                 }
                 break;
 
@@ -1735,13 +1730,13 @@ namespace mu
             {
 				MUTABLE_CPUPROFILER_SCOPE(ME_EXTRACTFACEGROUP_1)
 
-                Ptr<const Mesh> pSource = GetMemory().GetMesh( FCacheAddress(args.source,item) );
+                Ptr<const Mesh> Source = LoadMesh( FCacheAddress(args.source,item) );
 
                 MeshPtr pResult;
-                pResult = MeshExtractFaceGroup( pSource.get(),
+                pResult = MeshExtractFaceGroup( Source.get(),
                                                 args.group );
 
-                GetMemory().SetMesh( item, pResult );
+                StoreMesh( item, pResult );
                 break;
             }
 
@@ -1765,7 +1760,7 @@ namespace mu
                     }
                     else
                     {
-                    GetMemory().SetMesh( item, nullptr );
+                    StoreMesh( item, nullptr );
                 }
                 break;
 
@@ -1773,14 +1768,14 @@ namespace mu
             {
            		MUTABLE_CPUPROFILER_SCOPE(ME_TRANSFORM_1)
             		
-                Ptr<const Mesh> pSource = GetMemory().GetMesh( FCacheAddress(args.source,item) );
+                Ptr<const Mesh> Source = LoadMesh( FCacheAddress(args.source,item) );
 
                 const mat4f& mat = pModel->GetPrivate()->m_program.
                     m_constantMatrices[args.matrix];
 
-                MeshPtr pResult = MeshTransform(pSource.get(), ToUnreal(mat) );
+                MeshPtr pResult = MeshTransform(Source.get(), ToUnreal(mat) );
 
-                GetMemory().SetMesh( item, pResult );
+                StoreMesh( item, pResult );
                 break;
             }
 
@@ -1804,7 +1799,7 @@ namespace mu
                     }
                     else
                     {
-                    GetMemory().SetMesh( item, nullptr );
+                    StoreMesh( item, nullptr );
                 }
                 break;
 
@@ -1812,7 +1807,7 @@ namespace mu
             {
            		MUTABLE_CPUPROFILER_SCOPE(ME_CLIPMORPHPLANE_1)
             		
-                Ptr<const Mesh> pSource = GetMemory().GetMesh( FCacheAddress(args.source,item) );
+                Ptr<const Mesh> Source = LoadMesh( FCacheAddress(args.source,item) );
 
                 MeshPtr pResult;
 
@@ -1831,7 +1826,7 @@ namespace mu
 
                     // Should be None or an axis aligned box
                     const FShape& selectionShape = pModel->GetPrivate()->m_program.m_constantShapes[args.vertexSelectionShapeOrBone];
-                    pResult = MeshClipMorphPlane(pSource.get(), origin, normal, args.dist, args.factor, morphShape.size[0], morphShape.size[1], morphShape.size[2], selectionShape);
+                    pResult = MeshClipMorphPlane(Source.get(), origin, normal, args.dist, args.factor, morphShape.size[0], morphShape.size[1], morphShape.size[2], selectionShape);
                 }
 
                 else if (args.vertexSelectionType == OP::MeshClipMorphPlaneArgs::VS_BONE_HIERARCHY)
@@ -1841,17 +1836,17 @@ namespace mu
                     FShape selectionShape;
                     selectionShape.type = (uint8)FShape::Type::None;
                     const string& selectionBone = pModel->GetPrivate()->m_program.m_constantStrings[args.vertexSelectionShapeOrBone];
-					pResult = MeshClipMorphPlane(pSource.get(), origin, normal, args.dist, args.factor, morphShape.size[0], morphShape.size[1], morphShape.size[2], selectionShape, selectionBone, args.maxBoneRadius);
+					pResult = MeshClipMorphPlane(Source.get(), origin, normal, args.dist, args.factor, morphShape.size[0], morphShape.size[1], morphShape.size[2], selectionShape, selectionBone, args.maxBoneRadius);
                 }
                 else
                 {
                     // No vertex selection
                     FShape selectionShape;
                     selectionShape.type = (uint8)FShape::Type::None;
-                    pResult = MeshClipMorphPlane(pSource.get(), origin, normal, args.dist, args.factor, morphShape.size[0], morphShape.size[1], morphShape.size[2], selectionShape);
+                    pResult = MeshClipMorphPlane(Source.get(), origin, normal, args.dist, args.factor, morphShape.size[0], morphShape.size[1], morphShape.size[2], selectionShape);
                 }
 
-                GetMemory().SetMesh( item, pResult );
+                StoreMesh( item, pResult );
                 break;
             }
 
@@ -1877,7 +1872,7 @@ namespace mu
                     }
                     else
                     {
-                    GetMemory().SetMesh( item, nullptr );
+                    StoreMesh( item, nullptr );
                 }
                 break;
 
@@ -1885,21 +1880,21 @@ namespace mu
             {
 				MUTABLE_CPUPROFILER_SCOPE(ME_CLIPWITHMESH_1)
 
-                Ptr<const Mesh> pSource = GetMemory().GetMesh( FCacheAddress(args.source,item) );
-                Ptr<const Mesh> pClip = GetMemory().GetMesh( FCacheAddress(args.clipMesh,item) );
+                Ptr<const Mesh> Source = LoadMesh( FCacheAddress(args.source,item) );
+                Ptr<const Mesh> pClip = LoadMesh( FCacheAddress(args.clipMesh,item) );
 
                 // Only if both are valid.
                 MeshPtr pResult;
-                if (pSource && pClip)
+                if (Source && pClip)
                 {
-                    pResult = MeshClipWithMesh(pSource.get(), pClip.get());
+                    pResult = MeshClipWithMesh(Source.get(), pClip.get());
                 }
-                else if (pSource)
+                else if (Source)
                 {
-                    pResult = pSource->Clone();
+                    pResult = Source->Clone();
                 }
 
-                GetMemory().SetMesh( item, pResult );
+                StoreMesh( item, pResult );
                 break;
             }
 
@@ -1923,7 +1918,7 @@ namespace mu
 					}
 					else
 					{
-					GetMemory().SetMesh(item, nullptr);
+					StoreMesh(item, nullptr);
 				}
 				break;
 			
@@ -1931,8 +1926,8 @@ namespace mu
 			{
 				MUTABLE_CPUPROFILER_SCOPE(ME_CLIPDEFORM_1)
 
-				Ptr<const Mesh> BaseMesh = GetMemory().GetMesh(FCacheAddress(args.mesh, item));
-				Ptr<const Mesh> ClipShape = GetMemory().GetMesh(FCacheAddress(args.clipShape, item));
+				Ptr<const Mesh> BaseMesh = LoadMesh(FCacheAddress(args.mesh, item));
+				Ptr<const Mesh> ClipShape = LoadMesh(FCacheAddress(args.clipShape, item));
 				Ptr<Mesh> pResult;
 
 				if ( BaseMesh && ClipShape )
@@ -1944,7 +1939,7 @@ namespace mu
 					pResult = BaseMesh->Clone();
 				}
 
-				GetMemory().SetMesh(item, pResult);
+				StoreMesh(item, pResult);
 				break;
 			}
 
@@ -1969,7 +1964,7 @@ namespace mu
                     }
                     else
                     {
-                    GetMemory().SetMesh( item, nullptr );
+                    StoreMesh( item, nullptr );
                 }
                 break;
 
@@ -1977,21 +1972,21 @@ namespace mu
             {
     			MUTABLE_CPUPROFILER_SCOPE(ME_REMAPINDICES_1)
 
-                Ptr<const Mesh> pSource = GetMemory().GetMesh( FCacheAddress(args.source,item) );
-                Ptr<const Mesh> pReference = GetMemory().GetMesh( FCacheAddress(args.reference,item) );
+                Ptr<const Mesh> Source = LoadMesh( FCacheAddress(args.source,item) );
+                Ptr<const Mesh> pReference = LoadMesh( FCacheAddress(args.reference,item) );
 
                 // Only if both are valid.
                 MeshPtr pResult;
-                if (pSource && pReference)
+                if (Source && pReference)
                 {
-                    pResult = MeshRemapIndices(pSource.get(), pReference.get());
+                    pResult = MeshRemapIndices(Source.get(), pReference.get());
                 }
-                else if (pSource)
+                else if (Source)
                 {
-                    pResult = pSource->Clone();
+                    pResult = Source->Clone();
                 }
 
-                GetMemory().SetMesh( item, pResult );
+                StoreMesh( item, pResult );
                 break;
             }
 
@@ -2017,7 +2012,7 @@ namespace mu
                     }
                     else
                     {
-                    GetMemory().SetMesh( item, nullptr );
+                    StoreMesh( item, nullptr );
                 }
                 break;
 
@@ -2025,8 +2020,8 @@ namespace mu
             {
           		MUTABLE_CPUPROFILER_SCOPE(ME_APPLYPOSE_1)
 
-                Ptr<const Mesh> pBase = GetMemory().GetMesh( FCacheAddress(args.base,item) );
-                Ptr<const Mesh> pPose = GetMemory().GetMesh( FCacheAddress(args.pose,item) );
+                Ptr<const Mesh> pBase = LoadMesh( FCacheAddress(args.base,item) );
+                Ptr<const Mesh> pPose = LoadMesh( FCacheAddress(args.pose,item) );
 
                 // Only if both are valid.
                 MeshPtr pResult;
@@ -2039,7 +2034,7 @@ namespace mu
                     pResult = pBase->Clone();
                 }
 
-                GetMemory().SetMesh( item, pResult );
+                StoreMesh( item, pResult );
                 break;
             }
 
@@ -2067,7 +2062,7 @@ namespace mu
 					}
 					else
 					{
-					GetMemory().SetMesh(item, nullptr);
+					StoreMesh(item, nullptr);
 				}
 				break;
 
@@ -2075,14 +2070,14 @@ namespace mu
 			{
 				MUTABLE_CPUPROFILER_SCOPE(ME_GEOMETRYOPERATION_1)
 				
-				Ptr<const Mesh> MeshA = GetMemory().GetMesh(FCacheAddress(args.meshA, item));
-				Ptr<const Mesh> MeshB = GetMemory().GetMesh(FCacheAddress(args.meshB, item));
-				float ScalarA = GetMemory().GetScalar(FCacheAddress(args.scalarA, item));
-				float ScalarB = GetMemory().GetScalar(FCacheAddress(args.scalarB, item));
+				Ptr<const Mesh> MeshA = LoadMesh(FCacheAddress(args.meshA, item));
+				Ptr<const Mesh> MeshB = LoadMesh(FCacheAddress(args.meshB, item));
+				float ScalarA = LoadScalar(FCacheAddress(args.scalarA, item));
+				float ScalarB = LoadScalar(FCacheAddress(args.scalarB, item));
 
 				MeshPtr pResult = MeshGeometryOperation(MeshA.get(),MeshB.get(),ScalarA,ScalarB);
 
-				GetMemory().SetMesh(item, pResult);
+				StoreMesh(item, pResult);
 				break;
 			}
 
@@ -2110,15 +2105,15 @@ namespace mu
 				}
 				else
 				{
-					GetMemory().SetMesh(item, nullptr);
+					StoreMesh(item, nullptr);
 				}
 				break;
 
 			case 1:
 			{
 				MUTABLE_CPUPROFILER_SCOPE(ME_BINDSHAPE_1)
-				Ptr<const Mesh> BaseMesh = GetMemory().GetMesh(FCacheAddress(Args.mesh, item));
-				Ptr<const Mesh> Shape = GetMemory().GetMesh(FCacheAddress(Args.shape, item));
+				Ptr<const Mesh> BaseMesh = LoadMesh(FCacheAddress(Args.mesh, item));
+				Ptr<const Mesh> Shape = LoadMesh(FCacheAddress(Args.shape, item));
 				
 				
 				EShapeBindingMethod BindingMethod = static_cast<EShapeBindingMethod>(Args.bindingMethod); 
@@ -2170,22 +2165,22 @@ namespace mu
 						{
 							Ptr<Mesh> BindMeshNoVerts = BaseMesh->Clone();
 							BindMeshNoVerts->m_AdditionalBuffers = MoveTemp(BindMesh->m_AdditionalBuffers);
-							GetMemory().SetMesh(item, BindMeshNoVerts);
+							StoreMesh(item, BindMeshNoVerts);
 						}
 						else
 						{
-							GetMemory().SetMesh(item, BindMesh);
+							StoreMesh(item, BindMesh);
 						}
 					}
 					else
 					{
-						GetMemory().SetMesh(item, BaseMesh);
+						StoreMesh(item, BaseMesh);
 					}	
 				}	
 				else
 				{
 					Ptr<Mesh> ResultPtr = MeshBindShapeClipDeform( BaseMesh.get(), Shape.get(), BindingMethod );
-					GetMemory().SetMesh(item, ResultPtr);
+					StoreMesh(item, ResultPtr);
 				}
 
 				break;
@@ -2213,7 +2208,7 @@ namespace mu
 				}
 				else
 				{
-					GetMemory().SetMesh(item, nullptr);
+					StoreMesh(item, nullptr);
 				}
 				break;
 
@@ -2221,8 +2216,8 @@ namespace mu
 			{
 				MUTABLE_CPUPROFILER_SCOPE(ME_APPLYSHAPE_1)
 					
-				Ptr<const Mesh> BaseMesh = GetMemory().GetMesh(FCacheAddress(args.mesh, item));
-				Ptr<const Mesh> Shape = GetMemory().GetMesh(FCacheAddress(args.shape, item));
+				Ptr<const Mesh> BaseMesh = LoadMesh(FCacheAddress(args.mesh, item));
+				Ptr<const Mesh> Shape = LoadMesh(FCacheAddress(args.shape, item));
 
 				const EMeshBindShapeFlags ReshapeFlags = static_cast<EMeshBindShapeFlags>(args.flags);
 				Ptr<Mesh> ReshapedMesh = MeshApplyShape(BaseMesh.get(), Shape.get(), ReshapeFlags);
@@ -2231,7 +2226,7 @@ namespace mu
 
 				if (!ReshapedMesh)
 				{
-					GetMemory().SetMesh(item, BaseMesh);
+					StoreMesh(item, BaseMesh);
 				}
 				else if (!bReshapeVertices)
 				{
@@ -2248,11 +2243,11 @@ namespace mu
 					Result->AdditionalPhysicsBodies = ReshapedMesh->AdditionalPhysicsBodies;
 					Result->BonePoses = ReshapedMesh->BonePoses;
 
-					GetMemory().SetMesh(item, Result);
+					StoreMesh(item, Result);
 				}
 				else
 				{
-					GetMemory().SetMesh(item, ReshapedMesh);
+					StoreMesh(item, ReshapedMesh);
 				}
 
 				break;
@@ -2280,7 +2275,7 @@ namespace mu
 				}
 				else 
 				{
-					GetMemory().SetMesh(item, nullptr);
+					StoreMesh(item, nullptr);
 				}
 				break;
 			}
@@ -2288,8 +2283,8 @@ namespace mu
 			{
 				MUTABLE_CPUPROFILER_SCOPE(ME_MORPHRESHAPE_1)
 					
-				Ptr<const Mesh> MorphedMesh = GetMemory().GetMesh(FCacheAddress(Args.Morph, item));
-				Ptr<const Mesh> ReshapeMesh = GetMemory().GetMesh(FCacheAddress(Args.Reshape, item));
+				Ptr<const Mesh> MorphedMesh = LoadMesh(FCacheAddress(Args.Morph, item));
+				Ptr<const Mesh> ReshapeMesh = LoadMesh(FCacheAddress(Args.Reshape, item));
 
 				if (ReshapeMesh)
 				{
@@ -2304,12 +2299,12 @@ namespace mu
 					ResultPtr->SetPhysicsBody(ReshapeMesh->GetPhysicsBody().get());
 					ResultPtr->BonePoses = ReshapeMesh->BonePoses;
 
-					GetMemory().SetMesh(item, ResultPtr);
+					StoreMesh(item, ResultPtr);
 				}
 				else
 				{
 					Ptr<Mesh> ResultPtr = MorphedMesh->Clone();
-					GetMemory().SetMesh(item, ResultPtr);
+					StoreMesh(item, ResultPtr);
 				}
 
 				break;
@@ -2336,7 +2331,7 @@ namespace mu
                     }
                     else
                     {
-                    GetMemory().SetMesh( item, nullptr );
+                    StoreMesh( item, nullptr );
                 }
                 break;
 
@@ -2344,41 +2339,41 @@ namespace mu
             {
             	MUTABLE_CPUPROFILER_SCOPE(ME_SETSKELETON_1)
             		
-                Ptr<const Mesh> pSource = GetMemory().GetMesh( FCacheAddress(args.source,item) );
-                Ptr<const Mesh> pSkeleton = GetMemory().GetMesh( FCacheAddress(args.skeleton,item) );
+                Ptr<const Mesh> Source = LoadMesh( FCacheAddress(args.source,item) );
+                Ptr<const Mesh> pSkeleton = LoadMesh( FCacheAddress(args.skeleton,item) );
 
                 // Only if both are valid.
                 MeshPtr pResult;
-                if (pSource && pSkeleton)
+                if (Source && pSkeleton)
                 {
-                    if ( pSource->GetSkeleton()
+                    if ( Source->GetSkeleton()
                          &&
-                         !pSource->GetSkeleton()->m_bones.IsEmpty() )
+                         !Source->GetSkeleton()->m_bones.IsEmpty() )
                     {
                         // For some reason we already have bone data, so we can't just overwrite it
                         // or the skinning may break. This may happen because of a problem in the
                         // optimiser that needs investigation.
                         // \TODO Be defensive, for now.
                         UE_LOG(LogMutableCore,Warning, TEXT("Performing a MeshRemapSkeleton, instead of MeshSetSkeletonData because source mesh already has some skeleton."));
-                        pResult = MeshRemapSkeleton( pSource.get(), pSkeleton->GetSkeleton().get() );
+                        pResult = MeshRemapSkeleton( Source.get(), pSkeleton->GetSkeleton().get() );
                         if (!pResult)
                         {
-                            pResult = pSource->Clone();
+                            pResult = Source->Clone();
                         }
                     }
                     else
                     {
-                        pResult = pSource->Clone();
+                        pResult = Source->Clone();
                         pResult->SetSkeleton(pSkeleton->GetSkeleton().get());
                     }
                     //pResult->GetPrivate()->CheckIntegrity();
                 }
-                else if (pSource)
+                else if (Source)
                 {
-                    pResult = pSource->Clone();
+                    pResult = Source->Clone();
                 }
 
-                GetMemory().SetMesh( item, pResult );
+                StoreMesh( item, pResult );
                 break;
             }
 
@@ -2429,7 +2424,7 @@ namespace mu
                 }
                 else
                 {
-                    GetMemory().SetMesh( item, nullptr );
+                    StoreMesh( item, nullptr );
                 }
                 break;
 
@@ -2447,7 +2442,7 @@ namespace mu
                     bool value = true;
                     if (conditions[r].At)
                     {
-                        value = GetMemory().GetBool( FCacheAddress(conditions[r].At, item) );
+                        value = LoadBool( FCacheAddress(conditions[r].At, item) );
                     }
 
                     if (value)
@@ -2468,9 +2463,9 @@ namespace mu
             	MUTABLE_CPUPROFILER_SCOPE(ME_REMOVEMASK_2)
             	
                 // \todo: single remove operation with all masks?
-                Ptr<const Mesh> pSource = GetMemory().GetMesh( FCacheAddress(source,item) );
+                Ptr<const Mesh> Source = LoadMesh( FCacheAddress(source,item) );
 
-                MeshPtrConst pResult = pSource;
+                MeshPtrConst pResult = Source;
 
                 for( size_t r=0; pResult && r<conditions.Num(); ++r )
                 {
@@ -2478,20 +2473,20 @@ namespace mu
                     bool value = true;
                     if (conditions[r].At)
                     {
-                        value = GetMemory().GetBool( FCacheAddress(conditions[r].At, item) );
+                        value = LoadBool( FCacheAddress(conditions[r].At, item) );
                     }
 
                     if (value)
                     {
-                        Ptr<const Mesh> pMask = GetMemory().GetMesh( FCacheAddress(masks[r],item) );
-                        if (pMask)
+                        Ptr<const Mesh> Mask = LoadMesh( FCacheAddress(masks[r],item) );
+                        if (Mask)
                         {
-                            pResult = MeshRemoveMask(pResult.get(), pMask.get());
+                            pResult = MeshRemoveMask(pResult.get(), Mask.get());
                         }
                     }
                 }
 
-                GetMemory().SetMesh( item, pResult );
+                StoreMesh( item, pResult );
                 break;
             }
 
@@ -2516,7 +2511,7 @@ namespace mu
                 }
                 else
                 {
-                    GetMemory().SetMesh( item, nullptr );
+                    StoreMesh( item, nullptr );
                 }
                 break;
 
@@ -2524,8 +2519,8 @@ namespace mu
             {
 				MUTABLE_CPUPROFILER_SCOPE(ME_PROJECT_1)
 
-                Ptr<const Mesh> pMesh = GetMemory().GetMesh( FCacheAddress(args.mesh,item) );
-                Ptr<const Projector> pProjector = GetMemory().GetProjector( FCacheAddress(args.projector,item) );
+                Ptr<const Mesh> pMesh = LoadMesh( FCacheAddress(args.mesh,item) );
+                Ptr<const Projector> pProjector = LoadProjector( FCacheAddress(args.projector,item) );
 
                 // Only if both are valid.
                 MeshPtr pResult;
@@ -2546,7 +2541,7 @@ namespace mu
                         pResult = pMesh->Clone();
                     }
                 }
-                GetMemory().SetMesh( item, pResult );
+                StoreMesh( item, pResult );
                 break;
             }
 
@@ -2569,7 +2564,7 @@ namespace mu
 				}
 				else
 				{
-					GetMemory().SetMesh(item, nullptr);
+					StoreMesh(item, nullptr);
 				}
 				break;
 
@@ -2577,11 +2572,11 @@ namespace mu
 			{
 				MUTABLE_CPUPROFILER_SCOPE(ME_OPTIMIZESKINNING_1)
 
-				Ptr<const Mesh> pSource = GetMemory().GetMesh(FCacheAddress(args.source, item));
+				Ptr<const Mesh> Source = LoadMesh(FCacheAddress(args.source, item));
 
-				MeshPtrConst pResult = MeshOptimizeSkinning(pSource.get());
+				MeshPtrConst pResult = MeshOptimizeSkinning(Source.get());
 
-				GetMemory().SetMesh(item, pResult ? pResult : pSource);
+				StoreMesh(item, pResult ? pResult : Source);
 				break;
 			}
 
@@ -2714,50 +2709,58 @@ namespace mu
                     DATATYPE RangeSizeType = GetOpDataType( pModel->GetPrivate()->m_program.GetOpType(args.rangeSize) );
                     if (RangeSizeType == DT_INT)
                     {
-						Iterations = GetMemory().GetInt(RangeAddress);
+						Iterations = LoadInt(RangeAddress);
                     }
                     else if (RangeSizeType == DT_SCALAR)
                     {
-						Iterations = int32( GetMemory().GetScalar(RangeAddress) );
+						Iterations = int32( LoadScalar(RangeAddress) );
                     }
                 }
 
-				Ptr<const Image> Base = GetMemory().GetImage(FCacheAddress(args.base, item));
+				Ptr<const Image> Base = LoadImage(FCacheAddress(args.base, item));
 
 				if (Iterations <= 0)
 				{
 					// There are no layers: return the base
-					GetMemory().SetImage(item, Base);
+					StoreImage(item, Base);
 				}
 				else
 				{
 					// Store the base
-					Ptr<Image> New = mu::CloneOrTakeOver(Base.get());
+					Ptr<Image> New = CloneOrTakeOver(Base);
+					EImageFormat InitialBaseFormat = New->GetFormat();
 
 					// Reset relevancy map.
 					New->m_flags &= ~Image::EImageFlags::IF_HAS_RELEVANCY_MAP;
 
 					// This shouldn't happen in optimised models, but it could happen in editors, etc.
 					// \todo: raise a performance warning?
-					EImageFormat BaseFormat = GetUncompressedFormat(Base->GetFormat());
-					if (Base->GetFormat() != BaseFormat)
+					EImageFormat BaseFormat = GetUncompressedFormat(New->GetFormat());
+					if (New->GetFormat() != BaseFormat)
 					{
-						Base = ImagePixelFormat(m_pSettings->GetPrivate()->m_imageCompressionQuality, Base.get(), BaseFormat);
+						Ptr<Image> Formatted = CreateImage( New->GetSizeX(), New->GetSizeY(), New->GetLODCount(), BaseFormat, EInitializationType::NotInitialized );
+
+						bool bSuccess = false;
+						ImagePixelFormat(bSuccess, m_pSettings->ImageCompressionQuality, Formatted.get(), New.get());
+						check(bSuccess); // Decompression cannot fail
+
+						Release(New);
+						New = Formatted;
 					}
 
 					FScheduledOpData Data;
 					Data.Resource = New;
 					Data.MultiLayer.Iterations = Iterations;
-					Data.MultiLayer.OriginalBaseFormat = Base->GetFormat();
+					Data.MultiLayer.OriginalBaseFormat = InitialBaseFormat;
 					Data.MultiLayer.bBlendOnlyOneMip = false;
 					int32 DataPos = m_heapData.Add(Data);
 
 					// Request the first layer
 					int32 CurrentIteration = 0;
 					FScheduledOp ItemCopy = item;
-					ExecutionIndex Index = GetMemory().GetRageIndex(item.ExecutionIndex);
+					ExecutionIndex Index = GetMemory().GetRangeIndex(item.ExecutionIndex);
 					Index.SetFromModelRangeIndex(args.rangeId, CurrentIteration);
-					ItemCopy.ExecutionIndex = GetMemory().GetRageIndexIndex(Index);
+					ItemCopy.ExecutionIndex = GetMemory().GetRangeIndexIndex(Index);
 					AddOp(FScheduledOp(item.At, item, 2, DataPos), FScheduledOp(args.base, item), FScheduledOp(args.blended, ItemCopy), FScheduledOp(args.mask, ItemCopy));
 				}
 
@@ -2781,32 +2784,44 @@ namespace mu
 				Ptr<Image> Base = static_cast<Image*>(Data.Resource.get());
  
                 FScheduledOp itemCopy = item;
-                ExecutionIndex index = GetMemory().GetRageIndex( item.ExecutionIndex );
+                ExecutionIndex index = GetMemory().GetRangeIndex( item.ExecutionIndex );
 				
                 {
                     index.SetFromModelRangeIndex( args.rangeId, CurrentIteration);
-                    itemCopy.ExecutionIndex = GetMemory().GetRageIndexIndex(index);
+                    itemCopy.ExecutionIndex = GetMemory().GetRangeIndexIndex(index);
 					itemCopy.CustomState = 0;
 
-                    Ptr<const Image> pBlended = GetMemory().GetImage( FCacheAddress(args.blended,itemCopy) );
+                    Ptr<const Image> Blended = LoadImage( FCacheAddress(args.blended,itemCopy) );
 
                     // This shouldn't happen in optimised models, but it could happen in editors, etc.
                     // \todo: raise a performance warning?
-                    if (pBlended && pBlended->GetFormat()!=Base->GetFormat() )
+                    if (Blended && Blended->GetFormat()!=Base->GetFormat() )
                     {
 						MUTABLE_CPUPROFILER_SCOPE(ImageResize_BlendedReformat);
-						pBlended = ImagePixelFormat( m_pSettings->GetPrivate()->m_imageCompressionQuality, pBlended.get(), Base->GetFormat());
+
+						Ptr<Image> Formatted = CreateImage(Blended->GetSizeX(), Blended->GetSizeY(), Blended->GetLODCount(), Base->GetFormat(), EInitializationType::NotInitialized);
+
+						bool bSuccess = false;
+						ImagePixelFormat(bSuccess, m_pSettings->ImageCompressionQuality, Formatted.get(), Blended.get());
+						check(bSuccess);
+
+						Release(Blended);
+						Blended = Formatted;
                     }
 
 					// TODO: This shouldn't happen, but be defensive.
 					FImageSize ResultSize = Base->GetSize();
-					if (pBlended && pBlended->GetSize() != ResultSize)
+					if (Blended && Blended->GetSize() != ResultSize)
 					{
 						MUTABLE_CPUPROFILER_SCOPE(ImageResize_BlendedFixForMultilayer);
-						pBlended = ImageResizeLinear(0, pBlended.get(), ResultSize);
+
+						Ptr<Image> Resized = CreateImage(ResultSize[0], ResultSize[1], Blended->GetLODCount(), Blended->GetFormat(), EInitializationType::NotInitialized);
+						ImageResizeLinear(Resized.get(), 0, Blended.get());
+						Release(Blended);
+						Blended = Resized;
 					}
 
-					if (pBlended->GetLODCount() < Base->GetLODCount())
+					if (Blended->GetLODCount() < Base->GetLODCount())
 					{
 						Data.MultiLayer.bBlendOnlyOneMip = true;
 					}
@@ -2832,23 +2847,27 @@ namespace mu
 						constexpr bool bUseVectorImpl = false;
 						if constexpr (bUseVectorImpl)
 						{
-							BufferLayerCompositeVector<VectorBlendChannelMasked, VectorLightenChannel, false>(Base.get(), pBlended.get(), bBlendOnlyOneMip, args.BlendAlphaSourceChannel);
+							BufferLayerCompositeVector<VectorBlendChannelMasked, VectorLightenChannel, false>(Base.get(), Blended.get(), bBlendOnlyOneMip, args.BlendAlphaSourceChannel);
 						}
 						else
 						{
-							BufferLayerComposite<BlendChannelMasked, LightenChannel, false>(Base.get(), pBlended.get(), bBlendOnlyOneMip, args.BlendAlphaSourceChannel);
+							BufferLayerComposite<BlendChannelMasked, LightenChannel, false>(Base.get(), Blended.get(), bBlendOnlyOneMip, args.BlendAlphaSourceChannel);
 						}
 					}
 
                     if (!bDone && args.mask)
                     {
-                        Ptr<const Image> pMask = GetMemory().GetImage( FCacheAddress(args.mask,itemCopy) );
+                        Ptr<const Image> Mask = LoadImage( FCacheAddress(args.mask,itemCopy) );
 
 						// TODO: This shouldn't happen, but be defensive.
-						if (pMask && pMask->GetSize() != ResultSize)
+						if (Mask && Mask->GetSize() != ResultSize)
 						{
 							MUTABLE_CPUPROFILER_SCOPE(ImageResize_MaskFixForMultilayer);
-							pMask = ImageResizeLinear(0, pMask.get(), ResultSize);
+
+							Ptr<Image> Resized = CreateImage(ResultSize[0], ResultSize[1], Mask->GetLODCount(), Mask->GetFormat(), EInitializationType::NotInitialized);
+							ImageResizeLinear(Resized.get(), 0, Mask.get());
+							Release(Mask);
+							Mask = Resized;
 						}
 
 						// Not implemented yet
@@ -2857,18 +2876,19 @@ namespace mu
                         switch (EBlendType(args.blendType))
                         {
 						case EBlendType::BT_NORMAL_COMBINE: check(false); break;
-                        case EBlendType::BT_SOFTLIGHT: BufferLayer<SoftLightChannelMasked, SoftLightChannel, false>( Base->GetData(), Base.get(), pMask.get(), pBlended.get(), bApplyColorBlendToAlpha, bBlendOnlyOneMip); break;
-                        case EBlendType::BT_HARDLIGHT: BufferLayer<HardLightChannelMasked, HardLightChannel, false>(Base->GetData(), Base.get(), pMask.get(), pBlended.get(), bApplyColorBlendToAlpha, bBlendOnlyOneMip); break;
-                        case EBlendType::BT_BURN: BufferLayer<BurnChannelMasked, BurnChannel, false>(Base->GetData(), Base.get(), pMask.get(), pBlended.get(), bApplyColorBlendToAlpha, bBlendOnlyOneMip); break;
-                        case EBlendType::BT_DODGE: BufferLayer<DodgeChannelMasked, DodgeChannel, false>(Base->GetData(), Base.get(), pMask.get(), pBlended.get(), bApplyColorBlendToAlpha, bBlendOnlyOneMip); break;
-                        case EBlendType::BT_SCREEN: BufferLayer<ScreenChannelMasked, ScreenChannel, false>(Base->GetData(), Base.get(), pMask.get(), pBlended.get(), bApplyColorBlendToAlpha, bBlendOnlyOneMip); break;
-                        case EBlendType::BT_OVERLAY: BufferLayer<OverlayChannelMasked, OverlayChannel, false>(Base->GetData(), Base.get(), pMask.get(), pBlended.get(), bApplyColorBlendToAlpha, bBlendOnlyOneMip); break;
-                        case EBlendType::BT_LIGHTEN: BufferLayer<LightenChannelMasked, LightenChannel, false>(Base->GetData(), Base.get(), pMask.get(), pBlended.get(), bApplyColorBlendToAlpha, bBlendOnlyOneMip); break;
-                        case EBlendType::BT_MULTIPLY: BufferLayer<MultiplyChannelMasked, MultiplyChannel, false>(Base->GetData(), Base.get(), pMask.get(), pBlended.get(), bApplyColorBlendToAlpha, bBlendOnlyOneMip); break;
-                        case EBlendType::BT_BLEND: BufferLayer<BlendChannelMasked, BlendChannel, false>(Base->GetData(), Base.get(), pMask.get(), pBlended.get(), bApplyColorBlendToAlpha, bBlendOnlyOneMip); break;
+                        case EBlendType::BT_SOFTLIGHT: BufferLayer<SoftLightChannelMasked, SoftLightChannel, false>( Base->GetData(), Base.get(), Mask.get(), Blended.get(), bApplyColorBlendToAlpha, bBlendOnlyOneMip); break;
+                        case EBlendType::BT_HARDLIGHT: BufferLayer<HardLightChannelMasked, HardLightChannel, false>(Base->GetData(), Base.get(), Mask.get(), Blended.get(), bApplyColorBlendToAlpha, bBlendOnlyOneMip); break;
+                        case EBlendType::BT_BURN: BufferLayer<BurnChannelMasked, BurnChannel, false>(Base->GetData(), Base.get(), Mask.get(), Blended.get(), bApplyColorBlendToAlpha, bBlendOnlyOneMip); break;
+                        case EBlendType::BT_DODGE: BufferLayer<DodgeChannelMasked, DodgeChannel, false>(Base->GetData(), Base.get(), Mask.get(), Blended.get(), bApplyColorBlendToAlpha, bBlendOnlyOneMip); break;
+                        case EBlendType::BT_SCREEN: BufferLayer<ScreenChannelMasked, ScreenChannel, false>(Base->GetData(), Base.get(), Mask.get(), Blended.get(), bApplyColorBlendToAlpha, bBlendOnlyOneMip); break;
+                        case EBlendType::BT_OVERLAY: BufferLayer<OverlayChannelMasked, OverlayChannel, false>(Base->GetData(), Base.get(), Mask.get(), Blended.get(), bApplyColorBlendToAlpha, bBlendOnlyOneMip); break;
+                        case EBlendType::BT_LIGHTEN: BufferLayer<LightenChannelMasked, LightenChannel, false>(Base->GetData(), Base.get(), Mask.get(), Blended.get(), bApplyColorBlendToAlpha, bBlendOnlyOneMip); break;
+                        case EBlendType::BT_MULTIPLY: BufferLayer<MultiplyChannelMasked, MultiplyChannel, false>(Base->GetData(), Base.get(), Mask.get(), Blended.get(), bApplyColorBlendToAlpha, bBlendOnlyOneMip); break;
+                        case EBlendType::BT_BLEND: BufferLayer<BlendChannelMasked, BlendChannel, false>(Base->GetData(), Base.get(), Mask.get(), Blended.get(), bApplyColorBlendToAlpha, bBlendOnlyOneMip); break;
                         default: check(false);
                         }
 
+						Release(Mask);
                     }
 					else if (!bDone && args.bUseMaskFromBlended)
 					{
@@ -2878,15 +2898,15 @@ namespace mu
 						switch (EBlendType(args.blendType))
 						{
 						case EBlendType::BT_NORMAL_COMBINE: check(false); break;
-						case EBlendType::BT_SOFTLIGHT: BufferLayerEmbeddedMask<SoftLightChannelMasked, SoftLightChannel, false>(Base->GetData(), Base.get(), pBlended.get(), bApplyColorBlendToAlpha, bBlendOnlyOneMip); break;
-						case EBlendType::BT_HARDLIGHT: BufferLayerEmbeddedMask<HardLightChannelMasked, HardLightChannel, false>(Base->GetData(), Base.get(), pBlended.get(), bApplyColorBlendToAlpha, bBlendOnlyOneMip); break;
-						case EBlendType::BT_BURN: BufferLayerEmbeddedMask<BurnChannelMasked, BurnChannel, false>(Base->GetData(), Base.get(), pBlended.get(), bApplyColorBlendToAlpha, bBlendOnlyOneMip); break;
-						case EBlendType::BT_DODGE: BufferLayerEmbeddedMask<DodgeChannelMasked, DodgeChannel, false>(Base->GetData(), Base.get(), pBlended.get(), bApplyColorBlendToAlpha, bBlendOnlyOneMip); break;
-						case EBlendType::BT_SCREEN: BufferLayerEmbeddedMask<ScreenChannelMasked, ScreenChannel, false>(Base->GetData(), Base.get(), pBlended.get(), bApplyColorBlendToAlpha, bBlendOnlyOneMip); break;
-						case EBlendType::BT_OVERLAY: BufferLayerEmbeddedMask<OverlayChannelMasked, OverlayChannel, false>(Base->GetData(), Base.get(), pBlended.get(), bApplyColorBlendToAlpha, bBlendOnlyOneMip); break;
-						case EBlendType::BT_LIGHTEN: BufferLayerEmbeddedMask<LightenChannelMasked, LightenChannel, false>(Base->GetData(), Base.get(), pBlended.get(), bApplyColorBlendToAlpha, bBlendOnlyOneMip); break;
-						case EBlendType::BT_MULTIPLY: BufferLayerEmbeddedMask<MultiplyChannelMasked, MultiplyChannel, false>(Base->GetData(), Base.get(), pBlended.get(), bApplyColorBlendToAlpha, bBlendOnlyOneMip); break;
-						case EBlendType::BT_BLEND: BufferLayerEmbeddedMask<BlendChannelMasked, BlendChannel, false>(Base->GetData(), Base.get(), pBlended.get(), bApplyColorBlendToAlpha, bBlendOnlyOneMip); break;
+						case EBlendType::BT_SOFTLIGHT: BufferLayerEmbeddedMask<SoftLightChannelMasked, SoftLightChannel, false>(Base->GetData(), Base.get(), Blended.get(), bApplyColorBlendToAlpha, bBlendOnlyOneMip); break;
+						case EBlendType::BT_HARDLIGHT: BufferLayerEmbeddedMask<HardLightChannelMasked, HardLightChannel, false>(Base->GetData(), Base.get(), Blended.get(), bApplyColorBlendToAlpha, bBlendOnlyOneMip); break;
+						case EBlendType::BT_BURN: BufferLayerEmbeddedMask<BurnChannelMasked, BurnChannel, false>(Base->GetData(), Base.get(), Blended.get(), bApplyColorBlendToAlpha, bBlendOnlyOneMip); break;
+						case EBlendType::BT_DODGE: BufferLayerEmbeddedMask<DodgeChannelMasked, DodgeChannel, false>(Base->GetData(), Base.get(), Blended.get(), bApplyColorBlendToAlpha, bBlendOnlyOneMip); break;
+						case EBlendType::BT_SCREEN: BufferLayerEmbeddedMask<ScreenChannelMasked, ScreenChannel, false>(Base->GetData(), Base.get(), Blended.get(), bApplyColorBlendToAlpha, bBlendOnlyOneMip); break;
+						case EBlendType::BT_OVERLAY: BufferLayerEmbeddedMask<OverlayChannelMasked, OverlayChannel, false>(Base->GetData(), Base.get(), Blended.get(), bApplyColorBlendToAlpha, bBlendOnlyOneMip); break;
+						case EBlendType::BT_LIGHTEN: BufferLayerEmbeddedMask<LightenChannelMasked, LightenChannel, false>(Base->GetData(), Base.get(), Blended.get(), bApplyColorBlendToAlpha, bBlendOnlyOneMip); break;
+						case EBlendType::BT_MULTIPLY: BufferLayerEmbeddedMask<MultiplyChannelMasked, MultiplyChannel, false>(Base->GetData(), Base.get(), Blended.get(), bApplyColorBlendToAlpha, bBlendOnlyOneMip); break;
+						case EBlendType::BT_BLEND: BufferLayerEmbeddedMask<BlendChannelMasked, BlendChannel, false>(Base->GetData(), Base.get(), Blended.get(), bApplyColorBlendToAlpha, bBlendOnlyOneMip); break;
 						default: check(false);
 						}
 					}
@@ -2895,15 +2915,15 @@ namespace mu
                         switch (EBlendType(args.blendType))
                         {
 						case EBlendType::BT_NORMAL_COMBINE: check(false); break;
-                        case EBlendType::BT_SOFTLIGHT: BufferLayer<SoftLightChannel, false>(Base.get(), Base.get(), pBlended.get(), bApplyColorBlendToAlpha, bBlendOnlyOneMip, bUseBlendSourceFromBlendAlpha); break;
-                        case EBlendType::BT_HARDLIGHT: BufferLayer<HardLightChannel, false>(Base.get(), Base.get(), pBlended.get(), bApplyColorBlendToAlpha, bBlendOnlyOneMip, bUseBlendSourceFromBlendAlpha); break;
-                        case EBlendType::BT_BURN: BufferLayer<BurnChannel, false>(Base.get(), Base.get(), pBlended.get(), bApplyColorBlendToAlpha, bBlendOnlyOneMip, bUseBlendSourceFromBlendAlpha); break;
-                        case EBlendType::BT_DODGE: BufferLayer<DodgeChannel, false>(Base.get(), Base.get(), pBlended.get(), bApplyColorBlendToAlpha, bBlendOnlyOneMip, bUseBlendSourceFromBlendAlpha); break;
-                        case EBlendType::BT_SCREEN: BufferLayer<ScreenChannel, false>(Base.get(), Base.get(),  pBlended.get(), bApplyColorBlendToAlpha, bBlendOnlyOneMip, bUseBlendSourceFromBlendAlpha); break;
-                        case EBlendType::BT_OVERLAY: BufferLayer<OverlayChannel, false>(Base.get(), Base.get(), pBlended.get(), bApplyColorBlendToAlpha, bBlendOnlyOneMip, bUseBlendSourceFromBlendAlpha); break;
-                        case EBlendType::BT_LIGHTEN: BufferLayer<LightenChannel, false>(Base.get(), Base.get(), pBlended.get(), bApplyColorBlendToAlpha, bBlendOnlyOneMip, bUseBlendSourceFromBlendAlpha); break;
-                        case EBlendType::BT_MULTIPLY: BufferLayer<MultiplyChannel, false>(Base.get(), Base.get(), pBlended.get(), bApplyColorBlendToAlpha, bBlendOnlyOneMip, bUseBlendSourceFromBlendAlpha); break;
-                        case EBlendType::BT_BLEND: BufferLayer<BlendChannel, false>(Base.get(), Base.get(), pBlended.get(), bApplyColorBlendToAlpha, bBlendOnlyOneMip, bUseBlendSourceFromBlendAlpha); break;
+                        case EBlendType::BT_SOFTLIGHT: BufferLayer<SoftLightChannel, false>(Base.get(), Base.get(), Blended.get(), bApplyColorBlendToAlpha, bBlendOnlyOneMip, bUseBlendSourceFromBlendAlpha); break;
+                        case EBlendType::BT_HARDLIGHT: BufferLayer<HardLightChannel, false>(Base.get(), Base.get(), Blended.get(), bApplyColorBlendToAlpha, bBlendOnlyOneMip, bUseBlendSourceFromBlendAlpha); break;
+                        case EBlendType::BT_BURN: BufferLayer<BurnChannel, false>(Base.get(), Base.get(), Blended.get(), bApplyColorBlendToAlpha, bBlendOnlyOneMip, bUseBlendSourceFromBlendAlpha); break;
+                        case EBlendType::BT_DODGE: BufferLayer<DodgeChannel, false>(Base.get(), Base.get(), Blended.get(), bApplyColorBlendToAlpha, bBlendOnlyOneMip, bUseBlendSourceFromBlendAlpha); break;
+                        case EBlendType::BT_SCREEN: BufferLayer<ScreenChannel, false>(Base.get(), Base.get(),  Blended.get(), bApplyColorBlendToAlpha, bBlendOnlyOneMip, bUseBlendSourceFromBlendAlpha); break;
+                        case EBlendType::BT_OVERLAY: BufferLayer<OverlayChannel, false>(Base.get(), Base.get(), Blended.get(), bApplyColorBlendToAlpha, bBlendOnlyOneMip, bUseBlendSourceFromBlendAlpha); break;
+                        case EBlendType::BT_LIGHTEN: BufferLayer<LightenChannel, false>(Base.get(), Base.get(), Blended.get(), bApplyColorBlendToAlpha, bBlendOnlyOneMip, bUseBlendSourceFromBlendAlpha); break;
+                        case EBlendType::BT_MULTIPLY: BufferLayer<MultiplyChannel, false>(Base.get(), Base.get(), Blended.get(), bApplyColorBlendToAlpha, bBlendOnlyOneMip, bUseBlendSourceFromBlendAlpha); break;
+                        case EBlendType::BT_BLEND: BufferLayer<BlendChannel, false>(Base.get(), Base.get(), Blended.get(), bApplyColorBlendToAlpha, bBlendOnlyOneMip, bUseBlendSourceFromBlendAlpha); break;
                         default: check(false);
                         }
                     }
@@ -2914,18 +2934,20 @@ namespace mu
 						// Separate alpha operation ignores the mask.
 						switch (EBlendType(args.blendTypeAlpha))
 						{
-						case EBlendType::BT_SOFTLIGHT: BufferLayerInPlace<SoftLightChannel, false, 1>(Base.get(), pBlended.get(), bBlendOnlyOneMip, 3, args.BlendAlphaSourceChannel); break;
-						case EBlendType::BT_HARDLIGHT: BufferLayerInPlace<HardLightChannel, false, 1>(Base.get(), pBlended.get(), bBlendOnlyOneMip, 3, args.BlendAlphaSourceChannel); break;
-						case EBlendType::BT_BURN: BufferLayerInPlace<BurnChannel, false, 1>(Base.get(), pBlended.get(), bBlendOnlyOneMip, 3, args.BlendAlphaSourceChannel); break;
-						case EBlendType::BT_DODGE: BufferLayerInPlace<DodgeChannel, false, 1>(Base.get(), pBlended.get(), bBlendOnlyOneMip, 3, args.BlendAlphaSourceChannel); break;
-						case EBlendType::BT_SCREEN: BufferLayerInPlace<ScreenChannel, false, 1>(Base.get(), pBlended.get(), bBlendOnlyOneMip, 3, args.BlendAlphaSourceChannel); break;
-						case EBlendType::BT_OVERLAY: BufferLayerInPlace<OverlayChannel, false, 1>(Base.get(), pBlended.get(), bBlendOnlyOneMip, 3, args.BlendAlphaSourceChannel); break;
-						case EBlendType::BT_LIGHTEN: BufferLayerInPlace<LightenChannel, false, 1>(Base.get(), pBlended.get(), bBlendOnlyOneMip, 3, args.BlendAlphaSourceChannel); break;
-						case EBlendType::BT_MULTIPLY: BufferLayerInPlace<MultiplyChannel, false, 1>(Base.get(), pBlended.get(), bBlendOnlyOneMip, 3, args.BlendAlphaSourceChannel); break;
-						case EBlendType::BT_BLEND: BufferLayerInPlace<BlendChannel, false, 1>(Base.get(), pBlended.get(), bBlendOnlyOneMip, 3, args.BlendAlphaSourceChannel); break;
+						case EBlendType::BT_SOFTLIGHT: BufferLayerInPlace<SoftLightChannel, false, 1>(Base.get(), Blended.get(), bBlendOnlyOneMip, 3, args.BlendAlphaSourceChannel); break;
+						case EBlendType::BT_HARDLIGHT: BufferLayerInPlace<HardLightChannel, false, 1>(Base.get(), Blended.get(), bBlendOnlyOneMip, 3, args.BlendAlphaSourceChannel); break;
+						case EBlendType::BT_BURN: BufferLayerInPlace<BurnChannel, false, 1>(Base.get(), Blended.get(), bBlendOnlyOneMip, 3, args.BlendAlphaSourceChannel); break;
+						case EBlendType::BT_DODGE: BufferLayerInPlace<DodgeChannel, false, 1>(Base.get(), Blended.get(), bBlendOnlyOneMip, 3, args.BlendAlphaSourceChannel); break;
+						case EBlendType::BT_SCREEN: BufferLayerInPlace<ScreenChannel, false, 1>(Base.get(), Blended.get(), bBlendOnlyOneMip, 3, args.BlendAlphaSourceChannel); break;
+						case EBlendType::BT_OVERLAY: BufferLayerInPlace<OverlayChannel, false, 1>(Base.get(), Blended.get(), bBlendOnlyOneMip, 3, args.BlendAlphaSourceChannel); break;
+						case EBlendType::BT_LIGHTEN: BufferLayerInPlace<LightenChannel, false, 1>(Base.get(), Blended.get(), bBlendOnlyOneMip, 3, args.BlendAlphaSourceChannel); break;
+						case EBlendType::BT_MULTIPLY: BufferLayerInPlace<MultiplyChannel, false, 1>(Base.get(), Blended.get(), bBlendOnlyOneMip, 3, args.BlendAlphaSourceChannel); break;
+						case EBlendType::BT_BLEND: BufferLayerInPlace<BlendChannel, false, 1>(Base.get(), Blended.get(), bBlendOnlyOneMip, 3, args.BlendAlphaSourceChannel); break;
 						default: check(false);
 						}
 					}
+
+					Release(Blended);
 				}
 
 				// Are we done?
@@ -2935,13 +2957,13 @@ namespace mu
 					{
 						MUTABLE_CPUPROFILER_SCOPE(ImageLayer_MipFix);
 						FMipmapGenerationSettings DummyMipSettings{};
-						ImageMipmapInPlace(m_pSettings->GetPrivate()->m_imageCompressionQuality, Base.get(), DummyMipSettings);
+						ImageMipmapInPlace(m_pSettings->ImageCompressionQuality, Base.get(), DummyMipSettings);
 					}
 
 					// TODO: Reconvert to OriginalBaseFormat if necessary?
 
-					GetMemory().SetImage(item, Base);
 					Data.Resource = nullptr;
+					StoreImage(item, Base);
 					break;
 				}
 				else
@@ -2949,9 +2971,9 @@ namespace mu
 					// Request a new layer
 					++CurrentIteration;
 					FScheduledOp ItemCopy = item;
-					ExecutionIndex Index = GetMemory().GetRageIndex(item.ExecutionIndex);
+					ExecutionIndex Index = GetMemory().GetRangeIndex(item.ExecutionIndex);
 					Index.SetFromModelRangeIndex(args.rangeId, CurrentIteration);
-					ItemCopy.ExecutionIndex = GetMemory().GetRageIndexIndex(Index);
+					ItemCopy.ExecutionIndex = GetMemory().GetRangeIndexIndex(Index);
 					AddOp(FScheduledOp(item.At, item, 2+CurrentIteration, item.CustomState), FScheduledOp(args.blended, ItemCopy), FScheduledOp(args.mask, ItemCopy));
 
 				}
@@ -2964,38 +2986,6 @@ namespace mu
             break;
         }
 
-        case OP_TYPE::IM_DIFFERENCE:
-        {
-			OP::ImageDifferenceArgs args = pModel->GetPrivate()->m_program.GetOpArgs<OP::ImageDifferenceArgs>(item.At);
-            switch (item.Stage)
-            {
-            case 0:
-                    AddOp( FScheduledOp( item.At, item, 1),
-                           FScheduledOp( args.a, item),
-                           FScheduledOp( args.b, item) );
-                break;
-
-            case 1:
-            {
-            	MUTABLE_CPUPROFILER_SCOPE(IM_DIFFERENCE_1)
-            		
-                // TODO: Reuse base if possible
-                Ptr<const Image> pA = GetMemory().GetImage( FCacheAddress(args.a,item) );
-                Ptr<const Image> pB = GetMemory().GetImage( FCacheAddress(args.b,item) );
-
-                ImagePtr pResult = ImageDifference( pA.get(), pB.get() );
-
-                GetMemory().SetImage( item, pResult );
-                break;
-            }
-
-            default:
-                check(false);
-            }
-
-			break;
-		}
-
 		case OP_TYPE::IM_NORMALCOMPOSITE:
 		{
 			OP::ImageNormalCompositeArgs args = pModel->GetPrivate()->m_program.GetOpArgs<OP::ImageNormalCompositeArgs>(item.At);
@@ -3004,13 +2994,13 @@ namespace mu
 			case 0:
 				if (args.base && args.normal)
 				{
-						AddOp(FScheduledOp(item.At, item, 1),
-							  FScheduledOp(args.base, item),
-							  FScheduledOp(args.normal, item));
-					}
-					else
-					{
-					GetMemory().SetImage(item, nullptr);
+					AddOp(FScheduledOp(item.At, item, 1),
+							FScheduledOp(args.base, item),
+							FScheduledOp(args.normal, item));
+			}
+				else
+				{
+					StoreImage(item, nullptr);
 				}
 				break;
 
@@ -3018,31 +3008,30 @@ namespace mu
 			{
 				MUTABLE_CPUPROFILER_SCOPE(IM_NORMALCOMPOSITE_1)
 
-				Ptr<const Image> Base = GetMemory().GetImage(FCacheAddress(args.base, item));
-				Ptr<const Image> Normal = GetMemory().GetImage(FCacheAddress(args.normal, item));
+				Ptr<const Image> Base = LoadImage(FCacheAddress(args.base, item));
+				Ptr<const Image> Normal = LoadImage(FCacheAddress(args.normal, item));
 
 				if (Normal->GetLODCount() < Base->GetLODCount())
 				{
 					MUTABLE_CPUPROFILER_SCOPE(ImageNormalComposite_EmergencyFix);
 
 					int levelCount = Base->GetLODCount();
-					ImagePtr pDest = new Image(Normal->GetSizeX(), Normal->GetSizeY(),
-						levelCount,
-						Normal->GetFormat());
+					ImagePtr pDest = CreateImage(Normal->GetSizeX(), Normal->GetSizeY(), levelCount, Normal->GetFormat());
 
-					SCRATCH_IMAGE_MIPMAP scratch;
 					FMipmapGenerationSettings mipSettings{};
-					ImageMipmap_PrepareScratch(pDest.get(), Normal.get(), levelCount, &scratch);
-					ImageMipmap(m_pSettings->GetPrivate()->m_imageCompressionQuality,
-						pDest.get(), Normal.get(), levelCount, &scratch, mipSettings);
+					ImageMipmap(m_pSystem->WorkingMemoryManager,m_pSettings->ImageCompressionQuality, pDest.get(), Normal.get(), levelCount, mipSettings);
+
+					Release(Normal);
 					Normal = pDest;
 				}
 
 
-                ImagePtr pResult = new Image( Base->GetSizeX(), Base->GetSizeY(), Base->GetLODCount(), Base->GetFormat() );
-				ImageNormalComposite(pResult.get(), Base.get(), Normal.get(), args.mode, args.power);
+                ImagePtr Result = CreateImage( Base->GetSizeX(), Base->GetSizeY(), Base->GetLODCount(), Base->GetFormat() );
+				ImageNormalComposite(Result.get(), Base.get(), Normal.get(), args.mode, args.power);
 
-				GetMemory().SetImage(item, pResult);
+				Release(Base);
+				Release(Normal);
+				StoreImage(item, Result);
 				break;
 			}
 
@@ -3059,8 +3048,7 @@ namespace mu
             switch (item.Stage)
             {
             case 0:
-                    AddOp( FScheduledOp( item.At, item, 1),
-                           FScheduledOp( args.source, item) );
+                AddOp( FScheduledOp( item.At, item, 1), FScheduledOp( args.source, item) );
                 break;
 
             case 1:
@@ -3126,55 +3114,52 @@ namespace mu
             switch (item.Stage)
             {
             case 0:
-                    AddOp( FScheduledOp( item.At, item, 1),
-                           FScheduledOp( args.source, item),
-                           FScheduledOp( args.sizeSource, item) );
+                AddOp( FScheduledOp( item.At, item, 1),
+                        FScheduledOp( args.source, item),
+                        FScheduledOp( args.sizeSource, item) );
                 break;
 
             case 1:
             {
             	MUTABLE_CPUPROFILER_SCOPE(IM_RESIZELIKE_1)
             	
-                Ptr<const Image> pBase = GetMemory().GetImage( FCacheAddress(args.source,item) );
-                Ptr<const Image> pSizeBase = GetMemory().GetImage( FCacheAddress(args.sizeSource,item) );
-                ImagePtr pResult;
+                Ptr<const Image> Base = LoadImage( FCacheAddress(args.source,item) );
+                Ptr<const Image> SizeBase = LoadImage( FCacheAddress(args.sizeSource,item) );
+				FImageSize DestSize = SizeBase->GetSize();
+				Release(SizeBase);
 
-                if ( pBase->GetSizeX()!=pSizeBase->GetSizeX()
-                     ||
-                     pBase->GetSizeY()!=pSizeBase->GetSizeY() )
+                if ( Base->GetSize()!=DestSize )
                 {
-                    FImageSize destSize( pSizeBase->GetSizeX(), pSizeBase->GetSizeY() );
-
-                    pResult = ImageResize( pBase.get(), destSize );
+					Ptr<Image> Result = CreateImage(DestSize[0], DestSize[1], Base->GetLODCount(), Base->GetFormat(), EInitializationType::NotInitialized);
+					ImageResizeLinear( Result.get(), m_pSettings->ImageCompressionQuality, Base.get() );
+					Release(Base);
 
                     // If the source image had mips, generate them as well for the resized image.
                     // This shouldn't happen often since "ResizeLike" should be usually optimised out
                     // during model compilation. The mipmap generation below is not very precise with
                     // the number of mips that are needed and will probably generate too many
-                    bool sourceHasMips = pBase->GetLODCount()>1;
+                    bool sourceHasMips = Base->GetLODCount()>1;
                     if (sourceHasMips)
                     {
-                        int levelCount = Image::GetMipmapCount( pResult->GetSizeX(), pResult->GetSizeY() );
-                        ImagePtr pMipmapped = new Image( pResult->GetSizeX(), pResult->GetSizeY(),
-                                                         levelCount,
-                                                         pResult->GetFormat() );
+                        int levelCount = Image::GetMipmapCount( Result->GetSizeX(), Result->GetSizeY() );
+                        Ptr<Image> Mipmapped = CreateImage( Result->GetSizeX(), Result->GetSizeY(), levelCount, Result->GetFormat() );
 
-                        SCRATCH_IMAGE_MIPMAP scratch;
 						FMipmapGenerationSettings mipSettings{};
 
-                        ImageMipmap_PrepareScratch( pMipmapped.get(), pResult.get(), levelCount, &scratch );
-                        ImageMipmap( m_pSettings->GetPrivate()->m_imageCompressionQuality,
-                                     pMipmapped.get(), pResult.get(), levelCount, &scratch, mipSettings );
+                        ImageMipmap( m_pSystem->WorkingMemoryManager, m_pSettings->ImageCompressionQuality,
+                                     Mipmapped.get(), Result.get(), levelCount, mipSettings );
 
-                        pResult = pMipmapped;
-                    }
-                }
+						Release(Result);
+						Result = Mipmapped;
+                    }				
+
+					StoreImage(item, Result);
+				}
                 else
                 {
-                    pResult = pBase->Clone();
-                }
-
-                GetMemory().SetImage( item, pResult );
+					StoreImage(item, Base);
+				}
+				
                 break;
             }
 
@@ -3191,13 +3176,11 @@ namespace mu
             switch (item.Stage)
             {
             case 0:
-                    AddOp( FScheduledOp( item.At, item, 1),
-                           FScheduledOp( args.source, item) );
+                AddOp( FScheduledOp( item.At, item, 1), FScheduledOp( args.source, item) );
                 break;
 
             case 1:
             {
-
 				// This has been moved to a task. It should have been intercepted in IssueOp.
 				check(false);
             }
@@ -3216,15 +3199,14 @@ namespace mu
             switch (item.Stage)
             {
             case 0:
-                    AddOp( FScheduledOp( item.At, item, 1),
-                           FScheduledOp::FromOpAndOptions( args.layout, item, 0) );
+                AddOp( FScheduledOp( item.At, item, 1), FScheduledOp::FromOpAndOptions( args.layout, item, 0) );
                 break;
 
             case 1:
             {
             	MUTABLE_CPUPROFILER_SCOPE(IM_BLANKLAYOUT_1)
             		
-                Ptr<const Layout> pLayout = GetMemory().GetLayout(FScheduledOp::FromOpAndOptions(args.layout, item, 0));
+                Ptr<const Layout> pLayout = LoadLayout(FScheduledOp::FromOpAndOptions(args.layout, item, 0));
 
                 FIntPoint SizeInBlocks = pLayout->GetGridSize();
 
@@ -3273,11 +3255,8 @@ namespace mu
                     }
                 }
 
-                ImagePtr pNew = new Image(ImageSizeInPixels.X, ImageSizeInPixels.Y, MipsToGenerate, EImageFormat(args.format) );
-
-                FMemory::Memzero( pNew->GetData(), pNew->GetDataSize() );
-
-                GetMemory().SetImage( item, pNew );
+                ImagePtr pNew = CreateImage(ImageSizeInPixels.X, ImageSizeInPixels.Y, MipsToGenerate, EImageFormat(args.format) );
+                StoreImage( item, pNew );
                 break;
             }
 
@@ -3295,15 +3274,14 @@ namespace mu
             switch ( item.Stage )
             {
             case 0:
-                AddOp( FScheduledOp( item.At, item, 1 ), 
-					FScheduledOp::FromOpAndOptions( args.layout, item, 0 ) );
+                AddOp( FScheduledOp( item.At, item, 1 ), FScheduledOp::FromOpAndOptions( args.layout, item, 0 ) );
                 break;
 
             case 1:
             {
             	MUTABLE_CPUPROFILER_SCOPE(IM_COMPOSE_1)
             		
-                Ptr<const Layout> ComposeLayout = GetMemory().GetLayout( FCacheAddress( args.layout, FScheduledOp::FromOpAndOptions(args.layout, item, 0)) );
+                Ptr<const Layout> ComposeLayout = LoadLayout( FCacheAddress( args.layout, FScheduledOp::FromOpAndOptions(args.layout, item, 0)) );
 
                 FScheduledOpData data;
                 data.Resource = const_cast<Layout*>(ComposeLayout.get());
@@ -3360,7 +3338,7 @@ namespace mu
                     count++;
                 }
 
-                float factor = GetMemory().GetScalar( FCacheAddress(args.factor,item) );
+                float factor = LoadScalar( FCacheAddress(args.factor,item) );
 
                 float delta = 1.0f/(count-1);
                 int min = (int)floorf( factor/delta );
@@ -3376,20 +3354,20 @@ namespace mu
 
                 if ( bifactor < UE_SMALL_NUMBER )
                 {
-                        AddOp( FScheduledOp( item.At, item, 2, dataPos),
-                               FScheduledOp( args.targets[min], item) );
-                    }
+                    AddOp( FScheduledOp( item.At, item, 2, dataPos),
+                            FScheduledOp( args.targets[min], item) );
+                }
                 else if ( bifactor > 1.0f-UE_SMALL_NUMBER )
                 {
-                        AddOp( FScheduledOp( item.At, item, 2, dataPos),
-                               FScheduledOp( args.targets[max], item) );
-                    }
-                    else
-                    {
-                        AddOp( FScheduledOp( item.At, item, 2, dataPos),
-                               FScheduledOp( args.targets[min], item),
-                               FScheduledOp( args.targets[max], item) );
-                    }
+                    AddOp( FScheduledOp( item.At, item, 2, dataPos),
+                            FScheduledOp( args.targets[max], item) );
+                }
+                else
+                {
+                    AddOp( FScheduledOp( item.At, item, 2, dataPos),
+                            FScheduledOp( args.targets[min], item),
+                            FScheduledOp( args.targets[max], item) );
+                }
                 break;
             }
 
@@ -3412,48 +3390,47 @@ namespace mu
                 int min = data.Interpolate.Min;
                 int max = data.Interpolate.Max;
 
-                ImagePtr pResult;
                 if ( bifactor < UE_SMALL_NUMBER )
                 {
-                    Ptr<const Image> pSource = GetMemory().GetImage( FCacheAddress(args.targets[min],item) );
-                    pResult = pSource->Clone();
-                }
+                    Ptr<const Image> Source = LoadImage( FCacheAddress(args.targets[min],item) );
+					StoreImage(item, Source);
+				}
                 else if ( bifactor > 1.0f-UE_SMALL_NUMBER )
                 {
-                    Ptr<const Image> pSource = GetMemory().GetImage( FCacheAddress(args.targets[max],item) );
-                    pResult = pSource->Clone();
-                }
+                    Ptr<const Image> Source = LoadImage( FCacheAddress(args.targets[max],item) );
+					StoreImage(item, Source);
+				}
                 else
                 {
-                    Ptr<const Image> pMin = GetMemory().GetImage( FCacheAddress(args.targets[min],item) );
-                    Ptr<const Image> pMax = GetMemory().GetImage( FCacheAddress(args.targets[max],item) );
+					Ptr<const Image> pMin = LoadImage( FCacheAddress(args.targets[min],item) );
+                    Ptr<const Image> pMax = LoadImage( FCacheAddress(args.targets[max],item) );
 
                     if (pMin && pMax)
                     {
 						int32 LevelCount = FMath::Max(pMin->GetLODCount(), pMax->GetLODCount());
 						
-						ImagePtr pNew = new Image( pMin->GetSizeX(), pMin->GetSizeY(), LevelCount, pMin->GetFormat() );
+						Ptr<Image> pNew = CreateImage( pMin->GetSizeX(), pMin->GetSizeY(), LevelCount, pMin->GetFormat() );
 
 						// Be defensive: ensure image sizes match.
 						if (pMin->GetSize() != pMax->GetSize())
 						{
 							MUTABLE_CPUPROFILER_SCOPE(ImageResize_ForInterpolate);
-							pMax = ImageResizeLinear(0, pMax.get(), pMin->GetSize());
+							Ptr<Image> Resized = CreateImage(pMin->GetSizeX(), pMin->GetSizeY(), pMax->GetLODCount(), pMax->GetFormat(), EInitializationType::NotInitialized);
+							ImageResizeLinear(Resized.get(), 0, pMax.get());
+							Release(pMax);
+							pMax = Resized;
 						}
 
 						if (pMin->GetLODCount() != LevelCount)
 						{
 							MUTABLE_CPUPROFILER_SCOPE(Mipmap_ForInterpolate);
 
-							ImagePtr pDest = new Image(pMin->GetSizeX(), pMin->GetSizeY(), LevelCount, pMin->GetFormat());
+							ImagePtr pDest = CreateImage(pMin->GetSizeX(), pMin->GetSizeY(), LevelCount, pMin->GetFormat());
 
-							SCRATCH_IMAGE_MIPMAP scratch;
 							FMipmapGenerationSettings settings{};
+							ImageMipmap(m_pSystem->WorkingMemoryManager, m_pSettings->ImageCompressionQuality, pDest.get(), pMin.get(), LevelCount, settings);
 
-							ImageMipmap_PrepareScratch(pDest.get(), pMin.get(), LevelCount, &scratch);
-							ImageMipmap(m_pSettings->GetPrivate()->m_imageCompressionQuality, pDest.get(), pMin.get(), LevelCount,
-								&scratch, settings);
-
+							Release(pMin);
 							pMin = pDest;
 						}
 
@@ -3461,34 +3438,32 @@ namespace mu
 						{
 							MUTABLE_CPUPROFILER_SCOPE(Mipmap_ForInterpolate);
 
-							ImagePtr pDest = new Image(pMax->GetSizeX(), pMax->GetSizeY(), LevelCount, pMax->GetFormat());
+							ImagePtr pDest = CreateImage(pMax->GetSizeX(), pMax->GetSizeY(), LevelCount, pMax->GetFormat());
 
-							SCRATCH_IMAGE_MIPMAP scratch;
 							FMipmapGenerationSettings settings{};
+							ImageMipmap(m_pSystem->WorkingMemoryManager, m_pSettings->ImageCompressionQuality, pDest.get(), pMax.get(), LevelCount, settings);
 
-							ImageMipmap_PrepareScratch(pDest.get(), pMax.get(), LevelCount, &scratch);
-							ImageMipmap(m_pSettings->GetPrivate()->m_imageCompressionQuality, pDest.get(), pMax.get(), LevelCount,
-								&scratch, settings);
-
+							Release(pMax);
 							pMax = pDest;
 						}
 
-
                         ImageInterpolate( pNew.get(), pMin.get(), pMax.get(), bifactor );
 
-                        pResult = pNew;
-                    }
+						Release(pMin);
+						Release(pMax);
+						StoreImage(item, pNew);
+					}
                     else if (pMin)
                     {
-                        pResult = pMin->Clone();
-                    }
+						StoreImage(item, pMin);
+					}
                     else if (pMax)
                     {
-                        pResult = pMax->Clone();
-                    }
-                }
+						StoreImage(item, pMax);
+					}
+
+				}
 
-                GetMemory().SetImage( item, pResult );
                 break;
             }
 
@@ -3505,12 +3480,12 @@ namespace mu
             switch (item.Stage)
             {
             case 0:
-                    AddOp( FScheduledOp( item.At, item, 1),
-                           FScheduledOp::FromOpAndOptions( args.factor1, item, 0),
-                           FScheduledOp::FromOpAndOptions( args.factor2, item, 0),
-                           FScheduledOp( args.target0, item),
-                           FScheduledOp( args.target1, item),
-                           FScheduledOp( args.target2, item) );
+                AddOp( FScheduledOp( item.At, item, 1),
+                        FScheduledOp::FromOpAndOptions( args.factor1, item, 0),
+                        FScheduledOp::FromOpAndOptions( args.factor2, item, 0),
+                        FScheduledOp( args.target0, item),
+                        FScheduledOp( args.target1, item),
+                        FScheduledOp( args.target2, item) );
                 break;
 
             case 1:
@@ -3518,19 +3493,27 @@ namespace mu
            		MUTABLE_CPUPROFILER_SCOPE(IM_INTERPOLATE3_1)
 
                 // \TODO Optimise limit cases
+                float factor1 = LoadScalar(FScheduledOp::FromOpAndOptions(args.factor1, item, 0));
+                float factor2 = LoadScalar(FScheduledOp::FromOpAndOptions(args.factor2, item, 0));
 
-                float factor1 = GetMemory().GetScalar(FScheduledOp::FromOpAndOptions(args.factor1, item, 0));
-                float factor2 = GetMemory().GetScalar(FScheduledOp::FromOpAndOptions(args.factor2, item, 0));
-
-                Ptr<const Image> pTarget0 = GetMemory().GetImage( FCacheAddress(args.target0,item) );
-                Ptr<const Image> pTarget1 = GetMemory().GetImage( FCacheAddress(args.target1,item) );
-                Ptr<const Image> pTarget2 = GetMemory().GetImage( FCacheAddress(args.target2,item) );
+				Ptr<const Image> Target0 = LoadImage( FCacheAddress(args.target0,item) );
+				Ptr<const Image> Target1 = LoadImage( FCacheAddress(args.target1,item) );
+				Ptr<const Image> Target2 = LoadImage( FCacheAddress(args.target2,item) );
 
-                ImagePtr pResult = ImageInterpolate( pTarget0.get(), pTarget1.get(), pTarget2.get(),
-                                                     factor1, factor2 );
+				Ptr<Image> Result;
+				
+				if (Target0)
+				{
+					// \TODO Optimise by doing in-place if possible
+					Result = CreateImageLike(Target0.get(), EInitializationType::NotInitialized);
+					ImageInterpolate(Result.get(), Target0.get(), Target1.get(), Target2.get(), factor1, factor2);
+				}
 
-                GetMemory().SetImage( item, pResult );
-                break;
+				Release(Target0);
+				Release(Target1);
+				Release(Target2);
+				StoreImage(item, Result);
+				break;
             }
 
             default:
@@ -3546,9 +3529,9 @@ namespace mu
             switch (item.Stage)
             {
             case 0:
-                    AddOp( FScheduledOp( item.At, item, 1),
-                           FScheduledOp( args.base, item ),
-                           FScheduledOp::FromOpAndOptions( args.factor, item, 0 ));
+                AddOp( FScheduledOp( item.At, item, 1),
+                        FScheduledOp( args.base, item ),
+                        FScheduledOp::FromOpAndOptions( args.factor, item, 0 ));
                 break;
 
             case 1:
@@ -3578,10 +3561,13 @@ namespace mu
             {
            		MUTABLE_CPUPROFILER_SCOPE(IM_LUMINANCE_1)
             		
-                Ptr<const Image> pBase = GetMemory().GetImage( FCacheAddress(args.base,item) );
+                Ptr<const Image> Base = LoadImage( FCacheAddress(args.base,item) );
+
+				Ptr<Image> Result = CreateImage(Base->GetSizeX(), Base->GetSizeY(), Base->GetLODCount(), EImageFormat::IF_L_UBYTE, EInitializationType::NotInitialized);
+                ImageLuminance( Result.get(),Base.get() );
 
-                ImagePtr pResult = ImageLuminance( pBase.get() );
-                GetMemory().SetImage( item, pResult );
+				Release(Base);
+				StoreImage( item, Result );
                 break;
             }
 
@@ -3617,37 +3603,6 @@ namespace mu
             break;
         }
 
-        case OP_TYPE::IM_SELECTCOLOUR:
-        {
-			OP::ImageSelectColourArgs args = pModel->GetPrivate()->m_program.GetOpArgs<OP::ImageSelectColourArgs>(item.At);
-            switch (item.Stage)
-            {
-            case 0:
-                AddOp( FScheduledOp( item.At, item, 1),
-                        FScheduledOp( args.base, item ),
-                        FScheduledOp::FromOpAndOptions( args.colour, item, 0 ) );
-                break;
-
-            case 1:
-            {
-            	MUTABLE_CPUPROFILER_SCOPE(IM_SELECTCOLOUR_1)
-            		
-                Ptr<const Image> pBase = GetMemory().GetImage( FCacheAddress(args.base,item) );
-				FVector4f colour = GetMemory().GetColour(FScheduledOp::FromOpAndOptions(args.colour, item, 0));
-
-				Ptr<Image> pResult = ImageSelectColour( pBase.get(), vec3f(colour) );
-
-                GetMemory().SetImage( item, pResult );
-                break;
-            }
-
-            default:
-                check(false);
-            }
-
-            break;
-        }
-
         case OP_TYPE::IM_COLOURMAP:
         {
 			OP::ImageColourMapArgs args = pModel->GetPrivate()->m_program.GetOpArgs<OP::ImageColourMapArgs>(item.At);
@@ -3664,30 +3619,36 @@ namespace mu
             {
             	MUTABLE_CPUPROFILER_SCOPE(IM_COLOURMAP_1)
             		
-                Ptr<const Image> pSource = GetMemory().GetImage( FCacheAddress(args.base,item) );
-                Ptr<const Image> pMask = GetMemory().GetImage( FCacheAddress(args.mask,item) );
-                Ptr<const Image> pMap = GetMemory().GetImage( FCacheAddress(args.map,item) );
+                Ptr<const Image> Source = LoadImage( FCacheAddress(args.base,item) );
+                Ptr<const Image> Mask = LoadImage( FCacheAddress(args.mask,item) );
+                Ptr<const Image> Map = LoadImage( FCacheAddress(args.map,item) );
 
-				bool bOnlyOneMip = (pMask->GetLODCount() < pSource->GetLODCount());
+				bool bOnlyOneMip = (Mask->GetLODCount() < Source->GetLODCount());
 
 				// Be defensive: ensure image sizes match.
-				if (pMask->GetSize() != pSource->GetSize())
+				if (Mask->GetSize() != Source->GetSize())
 				{
 					MUTABLE_CPUPROFILER_SCOPE(ImageResize_ForColourmap);
-					pMask = ImageResizeLinear(0, pMask.get(), pSource->GetSize());
+					Ptr<Image> Resized = CreateImage(Source->GetSizeX(), Source->GetSizeY(), 1, Mask->GetFormat(), EInitializationType::NotInitialized);
+					ImageResizeLinear(Resized.get(), 0, Mask.get());
+					Release(Mask);
+					Mask = Resized;
 				}
 
-
-                Ptr<Image> pResult = ImageColourMap( pSource.get(), pMask.get(), pMap.get(), bOnlyOneMip);
+				Ptr<Image> Result = CreateImage(Source->GetSizeX(), Source->GetSizeY(), Source->GetLODCount(), Source->GetFormat(), EInitializationType::NotInitialized);
+				ImageColourMap( Result.get(), Source.get(), Mask.get(), Map.get(), bOnlyOneMip);
 
 				if (bOnlyOneMip)
 				{
 					MUTABLE_CPUPROFILER_SCOPE(ImageColourMap_MipFix);
 					FMipmapGenerationSettings DummyMipSettings{};
-					ImageMipmapInPlace(m_pSettings->GetPrivate()->m_imageCompressionQuality, pResult.get(), DummyMipSettings);
+					ImageMipmapInPlace(m_pSettings->ImageCompressionQuality, Result.get(), DummyMipSettings);
 				}
 
-                GetMemory().SetImage( item, pResult );
+				Release(Source);
+				Release(Mask);
+				Release(Map);
+				StoreImage( item, Result );
                 break;
             }
 
@@ -3704,23 +3665,22 @@ namespace mu
             switch (item.Stage)
             {
             case 0:
-                    AddOp( FScheduledOp( item.At, item, 1),
-                           FScheduledOp::FromOpAndOptions( args.colour0, item, 0 ),
-                           FScheduledOp::FromOpAndOptions( args.colour1, item, 0 ) );
+                AddOp( FScheduledOp( item.At, item, 1),
+                        FScheduledOp::FromOpAndOptions( args.colour0, item, 0 ),
+                        FScheduledOp::FromOpAndOptions( args.colour1, item, 0 ) );
                 break;
 
             case 1:
             {
 				MUTABLE_CPUPROFILER_SCOPE(IM_GRADIENT_1)
             		
-				FVector4f colour0 = GetMemory().GetColour(FScheduledOp::FromOpAndOptions(args.colour0, item, 0));
-				FVector4f colour1 = GetMemory().GetColour(FScheduledOp::FromOpAndOptions(args.colour1, item, 0));
+				FVector4f colour0 = LoadColor(FScheduledOp::FromOpAndOptions(args.colour0, item, 0));
+				FVector4f colour1 = LoadColor(FScheduledOp::FromOpAndOptions(args.colour1, item, 0));
 
-                ImagePtr pResult = ImageGradient( colour0, colour1,
-                                                  args.size[0],
-                                                  args.size[1] );
+				ImagePtr pResult = CreateImage(args.size[0], args.size[1], 1, EImageFormat::IF_RGB_UBYTE);
+                ImageGradient( pResult.get(), colour0, colour1 );
 
-                GetMemory().SetImage( item, pResult );
+				StoreImage( item, pResult );
                 break;
             }
 
@@ -3737,22 +3697,24 @@ namespace mu
             switch (item.Stage)
             {
             case 0:
-                    AddOp( FScheduledOp( item.At, item, 1),
-                           FScheduledOp( args.base, item ),
-                           FScheduledOp::FromOpAndOptions( args.threshold, item, 0 ) );
+                AddOp( FScheduledOp( item.At, item, 1),
+                        FScheduledOp( args.base, item ),
+                        FScheduledOp::FromOpAndOptions( args.threshold, item, 0 ) );
                 break;
 
             case 1:
             {
             	MUTABLE_CPUPROFILER_SCOPE(IM_BINARISE_1)
             		
-                Ptr<const Image> pA = GetMemory().GetImage( FCacheAddress(args.base,item) );
+                Ptr<const Image> pA = LoadImage( FCacheAddress(args.base,item) );
 
-                float c = GetMemory().GetScalar(FScheduledOp::FromOpAndOptions(args.threshold, item, 0));
+                float c = LoadScalar(FScheduledOp::FromOpAndOptions(args.threshold, item, 0));
 
-                ImagePtr pResult = ImageBinarise( pA.get(), c );
+                Ptr<Image> Result = CreateImage(pA->GetSizeX(), pA->GetSizeY(), pA->GetLODCount(), EImageFormat::IF_L_UBYTE);
+				ImageBinarise( Result.get(), pA.get(), c );
 
-                GetMemory().SetImage( item, pResult );
+				Release(pA);
+				StoreImage( item, Result );
                 break;
             }
 
@@ -3769,8 +3731,7 @@ namespace mu
 			switch (item.Stage)
 			{
 			case 0:
-					AddOp(FScheduledOp(item.At, item, 1),
-						FScheduledOp(args.base, item));
+				AddOp(FScheduledOp(item.At, item, 1), FScheduledOp(args.base, item));
 				break;
 
 			case 1:
@@ -3792,15 +3753,14 @@ namespace mu
             switch (item.Stage)
             {
             case 0:
-				AddOp( FScheduledOp( item.At, item, 1),
-					FScheduledOp::FromOpAndOptions( args.colour, item, 0 ) );
+				AddOp( FScheduledOp( item.At, item, 1), FScheduledOp::FromOpAndOptions( args.colour, item, 0 ) );
                 break;
 
             case 1:
             {
             	MUTABLE_CPUPROFILER_SCOPE(IM_PLAINCOLOUR_1)
             		
-				FVector4f c = GetMemory().GetColour(FScheduledOp::FromOpAndOptions(args.colour, item, 0));
+				FVector4f c = LoadColor(FScheduledOp::FromOpAndOptions(args.colour, item, 0));
 
 				uint16 SizeX = args.size[0];
 				uint16 SizeY = args.size[1];
@@ -3819,14 +3779,11 @@ namespace mu
 					--LODs;
 				}
 
-                ImagePtr pA = new Image( SizeX, SizeY,
-                                         FMath::Max(LODs,1),
-                                         EImageFormat(args.format),
-										 EInitializationType::NotInitialized );
+                ImagePtr pA = CreateImage( SizeX, SizeY, FMath::Max(LODs,1), EImageFormat(args.format), EInitializationType::NotInitialized );
 
                 pA->FillColour(c);
 
-                GetMemory().SetImage( item, pA );
+				StoreImage( item, pA );
                 break;
             }
 
@@ -3845,7 +3802,7 @@ namespace mu
 			case 0:
 			{
 				Ptr<Image> Result = Image::CreateAsReference(Args.ID);
-				GetMemory().SetImage(item, Result);
+				StoreImage(item, Result);
 				break;
 			}
 
@@ -3862,15 +3819,14 @@ namespace mu
             switch (item.Stage)
             {
             case 0:
-                AddOp( FScheduledOp( item.At, item, 1),
-					FScheduledOp( args.source, item ) );
+                AddOp( FScheduledOp( item.At, item, 1), FScheduledOp( args.source, item ) );
                 break;
 
             case 1:
             {
             	MUTABLE_CPUPROFILER_SCOPE(IM_CROP_1)
             		
-                Ptr<const Image> pA = GetMemory().GetImage( FCacheAddress(args.source,item) );
+                Ptr<const Image> pA = LoadImage( FCacheAddress(args.source,item) );
 
                 box< vec2<int> > rect;
                 rect.min[0] = args.minX;
@@ -3889,10 +3845,12 @@ namespace mu
 				ImagePtr pResult;
 				if (!rect.IsEmpty())
 				{
-					pResult = ImageCrop(m_pSettings->GetPrivate()->m_imageCompressionQuality, pA.get(), rect);
+					pResult = CreateImage( rect.size[0], rect.size[1], 1, pA->GetFormat() );
+					ImageCrop(pResult.get(), m_pSettings->ImageCompressionQuality, pA.get(), rect);
 				}
 
-                GetMemory().SetImage( item, pResult );
+				Release(pA);
+				StoreImage( item, pResult );
                 break;
             }
 
@@ -3921,13 +3879,14 @@ namespace mu
             {
                 MUTABLE_CPUPROFILER_SCOPE(IM_PATCH_1)
 
-                Ptr<const Image> pA = GetMemory().GetImage( FCacheAddress(args.base,item) );
-                Ptr<const Image> pB = GetMemory().GetImage( FCacheAddress(args.patch,item) );
+                Ptr<const Image> pA = LoadImage( FCacheAddress(args.base,item) );
+                Ptr<const Image> pB = LoadImage( FCacheAddress(args.patch,item) );
 
 				// Failsafe
 				if (!pA || !pB)
 				{
-					GetMemory().SetImage(item, pA);
+					Release(pB);
+					StoreImage(item, pA);
 					break;
 				}
 
@@ -3939,7 +3898,7 @@ namespace mu
 				rect.size[0] = pB->GetSizeX();
 				rect.size[1] = pB->GetSizeY();
 
-                ImagePtr pResult = mu::CloneOrTakeOver(pA.get());
+                ImagePtr pResult = CloneOrTakeOver(pA);
 
 				bool bApplyPatch = !rect.IsEmpty();
 				if (bApplyPatch)
@@ -3960,11 +3919,20 @@ namespace mu
 
 						if (pResult->GetFormat() != format)
 						{
-							pResult = ImagePixelFormat(m_pSettings->GetPrivate()->m_imageCompressionQuality, pResult.get(), format);
+							Ptr<Image> Formatted = CreateImage(pResult->GetSizeX(), pResult->GetSizeY(), pResult->GetLODCount(), format, EInitializationType::NotInitialized);
+							bool bSuccess = false;
+							ImagePixelFormat(bSuccess, m_pSettings->ImageCompressionQuality, Formatted.get(), pResult.get());
+							check(bSuccess);
+							Release(pResult);
+							pResult = Formatted;
 						}
 						if (pB->GetFormat() != format)
 						{
-							pB = ImagePixelFormat(m_pSettings->GetPrivate()->m_imageCompressionQuality, pB.get(), format);
+							Ptr<Image> Formatted = CreateImage(pB->GetSizeX(), pB->GetSizeY(), pB->GetLODCount(), format, EInitializationType::NotInitialized);
+							bool bSuccess = false;
+							ImagePixelFormat(bSuccess, m_pSettings->ImageCompressionQuality, Formatted.get(), pB.get());
+							check(bSuccess);
+							Release(pB);
 						}
 					}
 
@@ -3992,7 +3960,8 @@ namespace mu
 					//	rect.min[0], rect.min[1], rect.size[0], rect.size[1], pResult->GetSizeX(), pResult->GetSizeY());
 				}
 
-                GetMemory().SetImage( item, pResult );
+				Release(pB);
+				StoreImage( item, pResult );
                 break;
             }
 
@@ -4026,7 +3995,7 @@ namespace mu
 			{
 				MUTABLE_CPUPROFILER_SCOPE(IM_RASTERMESH_1)
 
-				Ptr<const Mesh> pMesh = GetMemory().GetMesh(FScheduledOp::FromOpAndOptions(args.mesh, item, 0));
+				Ptr<const Mesh> pMesh = LoadMesh(FScheduledOp::FromOpAndOptions(args.mesh, item, 0));
 
 				// If no image, we are generating a flat mesh UV raster. This is the final stage in this case.
 				if (!args.image)
@@ -4054,14 +4023,14 @@ namespace mu
 					}
 
                     // Flat mesh UV raster
-					Ptr<Image> ResultImage = new Image(SizeX, SizeY, 1, EImageFormat::IF_L_UBYTE);
+					Ptr<Image> ResultImage = CreateImage(SizeX, SizeY, 1, EImageFormat::IF_L_UBYTE);
 					if (pMesh)
 					{
 						ImageRasterMesh(pMesh.get(), ResultImage.get(), args.blockIndex, CropMin, UncroppedSize);
 					}
 
 					// Stop execution.
-					GetMemory().SetImage(item, ResultImage);
+					StoreImage(item, ResultImage);
 					break;
 				}
 
@@ -4071,7 +4040,7 @@ namespace mu
 				FScheduledOpData Data;
 				Data.RasterMesh.Mip = ProjectionMip;
 				Data.RasterMesh.MipValue = static_cast<float>(ProjectionMip);
-				Ptr<const Projector> pProjector = GetMemory().GetProjector(FScheduledOp::FromOpAndOptions(args.projector, item, 0));
+				Ptr<const Projector> pProjector = LoadProjector(FScheduledOp::FromOpAndOptions(args.projector, item, 0));
 				if (pProjector)
 				{
 					EMinFilterMethod MinFilterMethod = Invoke([&]() -> EMinFilterMethod
@@ -4108,8 +4077,8 @@ namespace mu
 				// Set mesh and projector back to the cache to prevent possible code re-execution.
 				// this needs to be done before the operation is re-scheduled otherwise the execution may deadlock.	
 				// TODO: Review if this is a good way of reusing data between stages.
-				GetMemory().SetMesh(FScheduledOp::FromOpAndOptions(args.mesh, item, 0), pMesh);
-				GetMemory().SetProjector(FScheduledOp::FromOpAndOptions(args.projector, item, 0), pProjector);
+				StoreMesh(FScheduledOp::FromOpAndOptions(args.mesh, item, 0), pMesh);
+				StoreProjector(FScheduledOp::FromOpAndOptions(args.projector, item, 0), pProjector);
 				
 
 				int32 DataHeapAddress = m_heapData.Add(Data);
@@ -4132,16 +4101,16 @@ namespace mu
 				{
 					// This case is treated at the previous stage.
 					check(false);
-					GetMemory().SetImage(item, nullptr);
+					StoreImage(item, nullptr);
 					break;
 				}
 
-                Ptr<const Mesh> pMesh = GetMemory().GetMesh(FScheduledOp::FromOpAndOptions(args.mesh, item, 0));
+                Ptr<const Mesh> pMesh = LoadMesh(FScheduledOp::FromOpAndOptions(args.mesh, item, 0));
 
 				if (!pMesh)
 				{
 					check(false);
-					GetMemory().SetImage(item, nullptr);
+					StoreImage(item, nullptr);
 					break;
 				}
 
@@ -4170,20 +4139,23 @@ namespace mu
 				FScheduledOpData Data = m_heapData[item.CustomState];
 
 				// Raster with projection
-				Ptr<const Image> pSource = GetMemory().GetImage(
-					FCacheAddress(args.image, item.ExecutionIndex, Data.RasterMesh.Mip));
+				Ptr<const Image> Source = LoadImage(FCacheAddress(args.image, item.ExecutionIndex, Data.RasterMesh.Mip));
 
-				Ptr<const Image> pMask = nullptr;
+				Ptr<const Image> Mask = nullptr;
 				if (args.mask)
 				{
-					pMask = GetMemory().GetImage(FCacheAddress(args.mask, item));
+					Mask = LoadImage(FCacheAddress(args.mask, item));
 
 					// TODO: This shouldn't happen, but be defensive.
 					FImageSize ResultSize(SizeX, SizeY);
-					if (pMask && pMask->GetSize()!= ResultSize)
+					if (Mask && Mask->GetSize()!= ResultSize)
 					{
 						MUTABLE_CPUPROFILER_SCOPE(ImageResize_MaskFixForProjection);
-						pMask = ImageResizeLinear(0, pMask.get(), ResultSize);
+
+						Ptr<Image> Resized = CreateImage(SizeX, SizeY, Mask->GetLODCount(), Mask->GetFormat(), EInitializationType::NotInitialized);
+						ImageResizeLinear(Resized.get(), 0, Mask.get());
+						Release(Mask);
+						Mask = Resized;
 					}
 				}
 
@@ -4191,19 +4163,24 @@ namespace mu
 				float fadeEnd = 180.0f;
 				if ( args.angleFadeProperties )
 				{
-					FVector4f fadeProperties = GetMemory().GetColour(FScheduledOp::FromOpAndOptions(args.angleFadeProperties, item, 0));
+					FVector4f fadeProperties = LoadColor(FScheduledOp::FromOpAndOptions(args.angleFadeProperties, item, 0));
 					fadeStart = fadeProperties[0];
 					fadeEnd = fadeProperties[1];
 				}
 				const float FadeStartRad = FMath::DegreesToRadians(fadeStart);
 				const float FadeEndRad = FMath::DegreesToRadians(fadeEnd);
 
-				EImageFormat Format = pSource ? GetUncompressedFormat(pSource->GetFormat()) : EImageFormat::IF_L_UBYTE;
+				EImageFormat Format = Source ? GetUncompressedFormat(Source->GetFormat()) : EImageFormat::IF_L_UBYTE;
 
-				if (pSource && pSource->GetFormat()!=Format)
+				if (Source && Source->GetFormat()!=Format)
 				{
 					MUTABLE_CPUPROFILER_SCOPE(RunCode_RasterMesh_ReformatSource);
-					pSource = ImagePixelFormat( m_pSettings->GetPrivate()->m_imageCompressionQuality, pSource.get(), Format );
+					Ptr<Image> Formatted = CreateImage(Source->GetSizeX(), Source->GetSizeY(), Source->GetLODCount(), Format, EInitializationType::NotInitialized);
+					bool bSuccess = false;
+					ImagePixelFormat(bSuccess, m_pSettings->ImageCompressionQuality, Formatted.get(), Source.get());
+					check(bSuccess); 
+					Release(Source);
+					Source = Formatted;
 				}
 
 				// Allocate memory for the temporary buffers
@@ -4231,34 +4208,35 @@ namespace mu
 
 				if (SamplingMethod == ESamplingMethod::BiLinear)
 				{
-					if (pSource->GetLODCount() < 2 && pSource->GetSizeX() > 1 && pSource->GetSizeY() > 1)
+					if (Source->GetLODCount() < 2 && Source->GetSizeX() > 1 && Source->GetSizeY() > 1)
 					{
 						MUTABLE_CPUPROFILER_SCOPE(RunCode_RasterMesh_BilinearMipGen);
 
-						Ptr<Image> NewImage = new Image(pSource->GetSizeX(), pSource->GetSizeY(), 2, pSource->GetFormat());
+						Ptr<Image> NewImage = CreateImage(Source->GetSizeX(), Source->GetSizeY(), 2, Source->GetFormat());
 
-						check(NewImage->GetDataSize() >= pSource->GetDataSize());
-						FMemory::Memcpy(NewImage->GetData(), pSource->GetData(), pSource->GetDataSize());
+						check(NewImage->GetDataSize() >= Source->GetDataSize());
+						FMemory::Memcpy(NewImage->GetData(), Source->GetData(), Source->GetDataSize());
 
 						ImageMipmapInPlace(0, NewImage.get(), FMipmapGenerationSettings{});
 
-						pSource = NewImage;
+						Release(Source);
+						Source = NewImage;
 					}
 				}
 
 				// Allocate new image after bilinear mip generation to reduce operation memory peak.
-                Ptr<Image> pNew = new Image(SizeX, SizeY, 1, Format);
+				Ptr<Image> pNew = CreateImage(SizeX, SizeY, 1, Format);
 
-				if (args.projector && pSource && pSource->GetSizeX() > 0 && pSource->GetSizeY() > 0)
+				if (args.projector && Source && Source->GetSizeX() > 0 && Source->GetSizeY() > 0)
 				{
-					Ptr<const Projector> pProjector = GetMemory().GetProjector(FScheduledOp::FromOpAndOptions(args.projector, item, 0));
+					Ptr<const Projector> pProjector = LoadProjector(FScheduledOp::FromOpAndOptions(args.projector, item, 0));
 					if (pProjector)
 					{
 						switch (pProjector->m_value.type)
 						{
 						case PROJECTOR_TYPE::PLANAR:
 							ImageRasterProjectedPlanar(pMesh.get(), pNew.get(),
-								pSource.get(), pMask.get(),
+								Source.get(), Mask.get(),
 								args.bIsRGBFadingEnabled, args.bIsAlphaFadingEnabled,
 								SamplingMethod,
 								FadeStartRad, FadeEndRad, FMath::Frac(Data.RasterMesh.MipValue),
@@ -4269,7 +4247,7 @@ namespace mu
 
 						case PROJECTOR_TYPE::WRAPPING:
 							ImageRasterProjectedWrapping(pMesh.get(), pNew.get(),
-								pSource.get(), pMask.get(),
+								Source.get(), Mask.get(),
 								args.bIsRGBFadingEnabled, args.bIsAlphaFadingEnabled,
 								SamplingMethod,
 								FadeStartRad, FadeEndRad, FMath::Frac(Data.RasterMesh.MipValue),
@@ -4280,7 +4258,7 @@ namespace mu
 
 						case PROJECTOR_TYPE::CYLINDRICAL:
 							ImageRasterProjectedCylindrical(pMesh.get(), pNew.get(),
-								pSource.get(), pMask.get(),
+								Source.get(), Mask.get(),
 								args.bIsRGBFadingEnabled, args.bIsAlphaFadingEnabled,
 								FadeStartRad, FadeEndRad,
 								layout,
@@ -4296,7 +4274,9 @@ namespace mu
 					}
 				}
 
-				GetMemory().SetImage(item, pNew);
+				Release(Source);
+				Release(Mask);
+				StoreImage(item, pNew);
 
                 break;
             }
@@ -4314,22 +4294,22 @@ namespace mu
             switch (item.Stage)
             {
             case 0:
-                    AddOp( FScheduledOp( item.At, item, 1),
-                           FScheduledOp( args.mask, item) );
+                AddOp( FScheduledOp( item.At, item, 1), FScheduledOp( args.mask, item) );
                 break;
 
             case 1:
             {
                 MUTABLE_CPUPROFILER_SCOPE(IM_MAKEGROWMAP_1)
 
-                Ptr<const Image> pMask = GetMemory().GetImage( FCacheAddress(args.mask,item) );
+                Ptr<const Image> Mask = LoadImage( FCacheAddress(args.mask,item) );
 
-                ImagePtr pNew = new Image( pMask->GetSizeX(), pMask->GetSizeY(), pMask->GetLODCount(), EImageFormat::IF_L_UBYTE);
+                Ptr<Image> Result = CreateImage( Mask->GetSizeX(), Mask->GetSizeY(), Mask->GetLODCount(), EImageFormat::IF_L_UBYTE);
 
-                ImageMakeGrowMap( pNew.get(), pMask.get(), args.border );
-				pNew->m_flags |= Image::IF_CANNOT_BE_SCALED;
+                ImageMakeGrowMap(Result.get(), Mask.get(), args.border );
+				Result->m_flags |= Image::IF_CANNOT_BE_SCALED;
 
-                GetMemory().SetImage( item, pNew );
+				Release(Mask);
+                StoreImage( item, Result);
                 break;
             }
 
@@ -4346,51 +4326,59 @@ namespace mu
             switch (item.Stage)
             {
             case 0:
-                    AddOp( FScheduledOp( item.At, item, 1),
-                           FScheduledOp( args.source, item ),
-                           FScheduledOp( args.displacementMap, item ) );
-                break;
+                AddOp( FScheduledOp( item.At, item, 1),
+                        FScheduledOp( args.source, item ),
+                        FScheduledOp( args.displacementMap, item ) );
+				break;
 
             case 1:
             {
                 MUTABLE_CPUPROFILER_SCOPE(IM_DISPLACE_1)
 
-                Ptr<const Image> pSource = GetMemory().GetImage( FCacheAddress(args.source,item) );
-                Ptr<const Image> pMap = GetMemory().GetImage( FCacheAddress(args.displacementMap,item) );
+                Ptr<const Image> Source = LoadImage( FCacheAddress(args.source,item) );
+                Ptr<const Image> pMap = LoadImage( FCacheAddress(args.displacementMap,item) );
 
-				if (!pSource)
+				if (!Source)
 				{
-					GetMemory().SetImage(item, nullptr);
+					Release(pMap);
+					StoreImage(item, nullptr);
 					break;
 				}
 
 				// TODO: This shouldn't happen: displacement maps cannot be scaled because their information
 				// is resolution sensitive (pixel offsets). If the size doesn't match, scale the source, apply 
 				// displacement and then unscale it.
-				FImageSize OriginalSourceScale = pSource->GetSize();
+				FImageSize OriginalSourceScale = Source->GetSize();
 				if (OriginalSourceScale[0]>0 && OriginalSourceScale[1]>0 && OriginalSourceScale != pMap->GetSize())
 				{
 					MUTABLE_CPUPROFILER_SCOPE(ImageResize_EmergencyHackForDisplacementStep1);
-					pSource = ImageResizeLinear(0, pSource.get(), pMap->GetSize());
+
+					Ptr<Image> Resized = CreateImage(pMap->GetSizeX(), pMap->GetSizeY(), Source->GetLODCount(), Source->GetFormat(), EInitializationType::NotInitialized);
+					ImageResizeLinear(Resized.get(), 0, Source.get());
+					Release(Source);
+					Source = Resized;
 				}
 
 				// This works based on the assumption that displacement maps never read from a position they actually write to.
 				// Since they are used for UV border expansion, this should always be the case.
-				//ImagePtr pNew = new Image(pSource->GetSizeX(), pSource->GetSizeY(), 1, pSource->GetFormat());
-				Ptr<Image> pNew = mu::CloneOrTakeOver(pSource.get());
+				Ptr<Image> Result = CloneOrTakeOver(Source);
 
 				if (OriginalSourceScale[0] > 0 && OriginalSourceScale[1] > 0)
 				{
-					ImageDisplace(pNew.get(), pSource.get(), pMap.get());
+					ImageDisplace(Result.get(), Result.get(), pMap.get());
 
-					if (OriginalSourceScale != pNew->GetSize())
+					if (OriginalSourceScale != Result->GetSize())
 					{
 						MUTABLE_CPUPROFILER_SCOPE(ImageResize_EmergencyHackForDisplacementStep2);
-						pNew = ImageResizeLinear(0, pNew.get(), OriginalSourceScale);
+						Ptr<Image> Resized = CreateImage(OriginalSourceScale[0], OriginalSourceScale[1], Result->GetLODCount(), Result->GetFormat(), EInitializationType::NotInitialized);
+						ImageResizeLinear(Resized.get(), 0, Result.get());
+						Release(Result);
+						Result = Resized;
 					}
 				}
 
-                GetMemory().SetImage( item, pNew );
+				Release(pMap);
+                StoreImage( item, Result);
                 break;
             }
 
@@ -4425,30 +4413,35 @@ namespace mu
             {
             	MUTABLE_CPUPROFILER_SCOPE(IM_TRANSFORM_1)
             		
-                Ptr<const Image> pBaseImage = GetMemory().GetImage( FCacheAddress(Args.base, item) );
+                Ptr<const Image> pBaseImage = LoadImage( FCacheAddress(Args.base, item) );
                 
                 const FVector2f Offset = FVector2f(
-                        Args.offsetX ? GetMemory().GetScalar( FCacheAddress(Args.offsetX, item) ) : 0.0f,
-                        Args.offsetY ? GetMemory().GetScalar( FCacheAddress(Args.offsetY, item) ) : 0.0f );
+                        Args.offsetX ? LoadScalar( FCacheAddress(Args.offsetX, item) ) : 0.0f,
+                        Args.offsetY ? LoadScalar( FCacheAddress(Args.offsetY, item) ) : 0.0f );
 
                 const FVector2f Scale = FVector2f(
-                        Args.scaleX ? GetMemory().GetScalar( FCacheAddress(Args.scaleX, item) ) : 1.0f,
-                        Args.scaleY ? GetMemory().GetScalar( FCacheAddress(Args.scaleY, item) ) : 1.0f );
+                        Args.scaleX ? LoadScalar( FCacheAddress(Args.scaleX, item) ) : 1.0f,
+                        Args.scaleY ? LoadScalar( FCacheAddress(Args.scaleY, item) ) : 1.0f );
 
 				// Map Range 0-1 to a full rotation
-                const float Rotation = GetMemory().GetScalar( FCacheAddress(Args.rotation, item) ) * UE_TWO_PI;
-
-
-				Ptr<Image> pResult = new Image( pBaseImage->GetSizeX(), pBaseImage->GetSizeY(), 1, pBaseImage->GetFormat());
+                const float Rotation = LoadScalar( FCacheAddress(Args.rotation, item) ) * UE_TWO_PI;
 			
-				FImageSize SampleImageSize =  FImageSize( 
-					static_cast<uint16>( FMath::Clamp( FMath::FloorToInt( float(pBaseImage->GetSizeX()) * FMath::Abs(Scale.X)), 2, pBaseImage->GetSizeX() ) ),
-					static_cast<uint16>( FMath::Clamp( FMath::FloorToInt( float(pBaseImage->GetSizeY()) * FMath::Abs(Scale.Y)), 2, pBaseImage->GetSizeY() ) ) );
+				EImageFormat BaseFormat = pBaseImage->GetFormat();
+				int32 BaseLODs = pBaseImage->GetLODCount();
+				FImageSize BaseSize = pBaseImage->GetSize();
+				FImageSize SampleImageSize =  FImageSize(
+					static_cast<uint16>( FMath::Clamp( FMath::FloorToInt( float(BaseSize.X) * FMath::Abs(Scale.X)), 2, BaseSize.X) ),
+					static_cast<uint16>( FMath::Clamp( FMath::FloorToInt( float(BaseSize.Y) * FMath::Abs(Scale.Y)), 2, BaseSize.Y) ) );
 	
-				Ptr<Image> pSampleImage = ImageResizeLinear( 0, pBaseImage.get(), SampleImageSize );
-                ImageTransform( pResult.get(), pSampleImage.get(), Offset, Scale, Rotation );
+				Ptr<Image> pSampleImage = CreateImage(SampleImageSize[0], SampleImageSize[1], BaseLODs, BaseFormat, EInitializationType::NotInitialized);
+				ImageResizeLinear( pSampleImage.get(), 0, pBaseImage.get() );
+				Release(pBaseImage);
+
+				Ptr<Image> Result = CreateImage(BaseSize.X, BaseSize.Y, 1, BaseFormat);
+				ImageTransform( Result.get(), pSampleImage.get(), Offset, Scale, Rotation );
 
-                GetMemory().SetImage( item, pResult );
+				Release(pSampleImage);
+				StoreImage( item, Result );
 
                 break;
             }
@@ -4493,7 +4486,7 @@ namespace mu
              ++rangeIndexInParam )
         {
             uint32 rangeIndexInModel = paramDesc.m_ranges[rangeIndexInParam];
-            const ExecutionIndex& currentIndex = GetMemory().GetRageIndex( item.ExecutionIndex );
+            const ExecutionIndex& currentIndex = GetMemory().GetRangeIndex( item.ExecutionIndex );
             int position = currentIndex.GetFromModelRangeIndex(rangeIndexInModel);
             index->GetPrivate()->m_values[rangeIndexInParam] = position;
         }
@@ -4516,7 +4509,7 @@ namespace mu
         {
 			OP::BoolConstantArgs args = program.GetOpArgs<OP::BoolConstantArgs>(item.At);
             bool result = args.value;
-            GetMemory().SetBool( item, result );
+            StoreBool( item, result );
             break;
         }
 
@@ -4526,7 +4519,7 @@ namespace mu
             bool result = false;
 			Ptr<RangeIndex> index = BuildCurrentOpRangeIndex( item, pParams, pModel, args.variable );
             result = pParams->GetBoolValue( args.variable, index );
-            GetMemory().SetBool( item, result );
+            StoreBool( item, result );
             break;
         }
 
@@ -4536,17 +4529,17 @@ namespace mu
             switch (item.Stage)
             {
             case 0:
-                    AddOp( FScheduledOp( item.At, item, 1),
-                           FScheduledOp( args.a, item),
-                           FScheduledOp( args.b, item) );
+                AddOp( FScheduledOp( item.At, item, 1),
+                        FScheduledOp( args.a, item),
+                        FScheduledOp( args.b, item) );
                 break;
 
             case 1:
             {
-                float a = GetMemory().GetScalar( FCacheAddress(args.a,item) );
-                float b = GetMemory().GetScalar( FCacheAddress(args.b,item) );
+                float a = LoadScalar( FCacheAddress(args.a,item) );
+                float b = LoadScalar( FCacheAddress(args.b,item) );
                 bool result = a<b;
-                GetMemory().SetBool( item, result );
+                StoreBool( item, result );
                 break;
             }
 
@@ -4568,20 +4561,20 @@ namespace mu
                     bool skip = false;
                     if ( args.a && GetMemory().IsValid( FCacheAddress(args.a,item) ) )
                     {
-                         bool a = GetMemory().GetBool( FCacheAddress(args.a,item) );
+                         bool a = LoadBool( FCacheAddress(args.a,item) );
                          if (!a)
                          {
-                            GetMemory().SetBool( item, false );
+                            StoreBool( item, false );
                             skip=true;
                          }
                     }
 
                     if ( !skip && args.b && GetMemory().IsValid( FCacheAddress(args.b,item) ) )
                     {
-                         bool b = GetMemory().GetBool( FCacheAddress(args.b,item) );
+                         bool b = LoadBool( FCacheAddress(args.b,item) );
                          if (!b)
                          {
-                            GetMemory().SetBool( item, false );
+                            StoreBool( item, false );
                             skip=true;
                          }
                     }
@@ -4596,10 +4589,10 @@ namespace mu
 
             case 1:
             {
-                bool a = args.a ? GetMemory().GetBool( FCacheAddress(args.a,item) ) : true;
+                bool a = args.a ? LoadBool( FCacheAddress(args.a,item) ) : true;
                 if (!a)
                 {
-                    GetMemory().SetBool( item, false );
+                    StoreBool( item, false );
                 }
                 else
                 {
@@ -4612,8 +4605,8 @@ namespace mu
             case 2:
             {
                 // We arrived here because a is true
-                bool b = args.b ? GetMemory().GetBool( FCacheAddress(args.b,item) ) : true;
-                GetMemory().SetBool( item, b );
+                bool b = args.b ? LoadBool( FCacheAddress(args.b,item) ) : true;
+                StoreBool( item, b );
                 break;
             }
 
@@ -4634,20 +4627,20 @@ namespace mu
                     bool skip = false;
                     if ( args.a && GetMemory().IsValid( FCacheAddress(args.a,item) ) )
                     {
-                         bool a = GetMemory().GetBool( FCacheAddress(args.a,item) );
+                         bool a = LoadBool( FCacheAddress(args.a,item) );
                          if (a)
                          {
-                            GetMemory().SetBool( item, true );
+                            StoreBool( item, true );
                             skip=true;
                          }
                     }
 
                     if ( !skip && args.b && GetMemory().IsValid( FCacheAddress(args.b,item) ) )
                     {
-                         bool b = GetMemory().GetBool( FCacheAddress(args.b,item) );
+                         bool b = LoadBool( FCacheAddress(args.b,item) );
                          if (b)
                          {
-                            GetMemory().SetBool( item, true );
+                            StoreBool( item, true );
                             skip=true;
                          }
                     }
@@ -4662,10 +4655,10 @@ namespace mu
 
             case 1:
             {
-                bool a = args.a ? GetMemory().GetBool( FCacheAddress(args.a,item) ) : false;
+                bool a = args.a ? LoadBool( FCacheAddress(args.a,item) ) : false;
                 if (a)
                 {
-                    GetMemory().SetBool( item, true );
+                    StoreBool( item, true );
                 }
                 else
                 {
@@ -4678,8 +4671,8 @@ namespace mu
             case 2:
             {
                 // We arrived here because a is false
-                bool b = args.b ? GetMemory().GetBool( FCacheAddress(args.b,item) ) : false;
-                GetMemory().SetBool( item, b );
+                bool b = args.b ? LoadBool( FCacheAddress(args.b,item) ) : false;
+                StoreBool( item, b );
                 break;
             }
 
@@ -4701,8 +4694,8 @@ namespace mu
 
             case 1:
             {
-                bool result = !GetMemory().GetBool( FCacheAddress(args.source,item) );
-                GetMemory().SetBool( item, result );
+                bool result = !LoadBool( FCacheAddress(args.source,item) );
+                StoreBool( item, result );
                 break;
             }
 
@@ -4724,9 +4717,9 @@ namespace mu
 
             case 1:
             {
-                int a = GetMemory().GetInt( FCacheAddress(args.value,item) );
+                int a = LoadInt( FCacheAddress(args.value,item) );
                 bool result = a == args.constant;
-                GetMemory().SetBool( item, result );
+                StoreBool( item, result );
                 break;
             }
 
@@ -4755,7 +4748,7 @@ namespace mu
         {
 			OP::IntConstantArgs args = pModel->GetPrivate()->m_program.GetOpArgs<OP::IntConstantArgs>(item.At);
             int result = args.value;
-            GetMemory().SetInt( item, result );
+            StoreInt( item, result );
             break;
         }
 
@@ -4785,7 +4778,7 @@ namespace mu
                 }
             }
 
-            GetMemory().SetInt( item, result );
+            StoreInt( item, result );
             break;
         }
 
@@ -4809,7 +4802,7 @@ namespace mu
         {
 			OP::ScalarConstantArgs args = pModel->GetPrivate()->m_program.GetOpArgs<OP::ScalarConstantArgs>(item.At);
             float result = args.value;
-            GetMemory().SetScalar( item, result );
+            StoreScalar( item, result );
             break;
         }
 
@@ -4818,7 +4811,7 @@ namespace mu
 			OP::ParameterArgs args = pModel->GetPrivate()->m_program.GetOpArgs<OP::ParameterArgs>(item.At);
 			Ptr<RangeIndex> index = BuildCurrentOpRangeIndex( item, pParams, pModel, args.variable );
             float result = pParams->GetFloatValue( args.variable, index );
-            GetMemory().SetScalar( item, result );
+            StoreScalar( item, result );
             break;
         }
 
@@ -4834,12 +4827,12 @@ namespace mu
 
             case 1:
             {
-                float time = GetMemory().GetScalar( FCacheAddress(args.time,item) );
+                float time = LoadScalar( FCacheAddress(args.time,item) );
 
                 const Curve& curve = pModel->GetPrivate()->m_program.m_constantCurves[args.curve];
                 float result = EvalCurve(curve, time);
 
-                GetMemory().SetScalar( item, result );
+                StoreScalar( item, result );
                 break;
             }
 
@@ -4867,8 +4860,8 @@ namespace mu
 
             case 1:
             {
-                float a = GetMemory().GetScalar( FCacheAddress(args.a,item) );
-                float b = GetMemory().GetScalar( FCacheAddress(args.b,item) );
+                float a = LoadScalar( FCacheAddress(args.a,item) );
+                float b = LoadScalar( FCacheAddress(args.b,item) );
 
                 float result = 1.0f;
                 switch (args.operation)
@@ -4894,7 +4887,7 @@ namespace mu
                     break;
                 }
 
-                GetMemory().SetScalar( item, result );
+                StoreScalar( item, result );
                 break;
             }
 
@@ -4926,7 +4919,7 @@ namespace mu
             check( args.value < (uint32)pModel->GetPrivate()->m_program.m_constantStrings.Num() );
 
             const std::string& result = pModel->GetPrivate()->m_program.m_constantStrings[args.value];
-            GetMemory().SetString( item, new String(result.c_str()) );
+            StoreString( item, new String(result.c_str()) );
 
             break;
         }
@@ -4936,7 +4929,7 @@ namespace mu
 			OP::ParameterArgs args = pModel->GetPrivate()->m_program.GetOpArgs<OP::ParameterArgs>( item.At );
 			Ptr<RangeIndex> index = BuildCurrentOpRangeIndex( item, pParams, pModel, args.variable );
             string result = pParams->GetStringValue( args.variable, index );
-            GetMemory().SetString( item, new String( result.c_str() ) );
+            StoreString( item, new String( result.c_str() ) );
             break;
         }
 
@@ -4967,7 +4960,7 @@ namespace mu
             result[1] = args.value[1];
             result[2] = args.value[2];
             result[3] = args.value[3];
-            GetMemory().SetColour( item, result );
+            StoreColor( item, result );
             break;
         }
 
@@ -4979,7 +4972,7 @@ namespace mu
             float g=0.0f;
             float b=0.0f;            
             pParams->GetColourValue( args.variable, &r, &g, &b, index );
-            GetMemory().SetColour( item, FVector4f(r,g,b,1.0f) );
+            StoreColor( item, FVector4f(r,g,b,1.0f) );
             break;
         }
 
@@ -4998,10 +4991,10 @@ namespace mu
 
             case 1:
             {
-                float x = args.x ? GetMemory().GetScalar( FCacheAddress(args.x,item) ) : 0.5f;
-                float y = args.y ? GetMemory().GetScalar( FCacheAddress(args.y,item) ) : 0.5f;
+                float x = args.x ? LoadScalar( FCacheAddress(args.x,item) ) : 0.5f;
+                float y = args.y ? LoadScalar( FCacheAddress(args.y,item) ) : 0.5f;
 
-                Ptr<const Image> pImage = GetMemory().GetImage(FScheduledOp::FromOpAndOptions(args.image, item, 0));
+                Ptr<const Image> pImage = LoadImage(FScheduledOp::FromOpAndOptions(args.image, item, 0));
 
 				FVector4f result;
                 if (pImage)
@@ -5021,7 +5014,8 @@ namespace mu
                     result = FVector4f();
                 }
 
-                GetMemory().SetColour( item, result );
+				Release(pImage);
+                StoreColor( item, result );
                 break;
             }
 
@@ -5053,12 +5047,12 @@ namespace mu
                 {
                     if ( args.sources[t] )
                     {
-                        FVector4f p = GetMemory().GetColour( FCacheAddress(args.sources[t],item) );
+                        FVector4f p = LoadColor( FCacheAddress(args.sources[t],item) );
                         result[t] = p[ args.sourceChannels[t] ];
                     }
                 }
 
-                GetMemory().SetColour( item, result );
+                StoreColor( item, result );
                 break;
             }
 
@@ -5081,11 +5075,12 @@ namespace mu
 
             case 1:
             {
-                Ptr<const Image> pImage = GetMemory().GetImage( FCacheAddress(args.image,item) );
+                Ptr<const Image> pImage = LoadImage( FCacheAddress(args.image,item) );
 
 				FVector4f result = FVector4f( (float)pImage->GetSizeX(), (float)pImage->GetSizeY(), 0.0f, 0.0f );
 
-                GetMemory().SetColour( item, result );
+				Release(pImage);
+                StoreColor( item, result );
                 break;
             }
 
@@ -5108,7 +5103,7 @@ namespace mu
 
             case 1:
             {
-                Ptr<const Layout> pLayout = GetMemory().GetLayout( FCacheAddress(args.layout,item) );
+                Ptr<const Layout> pLayout = LoadLayout( FCacheAddress(args.layout,item) );
 
 				FVector4f result = FVector4f(0,0,0,0);
                 if ( pLayout )
@@ -5135,7 +5130,7 @@ namespace mu
                     }
                 }
 
-                GetMemory().SetColour( item, result );
+                StoreColor( item, result );
                 break;
             }
 
@@ -5165,25 +5160,25 @@ namespace mu
 
                 if (args.x)
                 {
-                    result[0] = GetMemory().GetScalar( FCacheAddress(args.x,item) );
+                    result[0] = LoadScalar( FCacheAddress(args.x,item) );
                 }
 
                 if (args.y)
                 {
-                    result[1] = GetMemory().GetScalar( FCacheAddress(args.y,item) );
+                    result[1] = LoadScalar( FCacheAddress(args.y,item) );
                 }
 
                 if (args.z)
                 {
-                    result[2] = GetMemory().GetScalar( FCacheAddress(args.z,item) );
+                    result[2] = LoadScalar( FCacheAddress(args.z,item) );
                 }
 
                 if (args.w)
                 {
-                    result[3] = GetMemory().GetScalar( FCacheAddress(args.w,item) );
+                    result[3] = LoadScalar( FCacheAddress(args.w,item) );
                 }
 
-                GetMemory().SetColour( item, result );
+                StoreColor( item, result );
                 break;
             }
 
@@ -5213,9 +5208,9 @@ namespace mu
                 otype = program.GetOpType( args.b );
                 dtype = GetOpDataType( otype );
                 check( dtype == DT_COLOUR );
-				FVector4f a = args.a ? GetMemory().GetColour( FCacheAddress( args.a, item ) )
+				FVector4f a = args.a ? LoadColor( FCacheAddress( args.a, item ) )
                                  : FVector4f( 0, 0, 0, 0 );
-				FVector4f b = args.b ? GetMemory().GetColour( FCacheAddress( args.b, item ) )
+				FVector4f b = args.b ? LoadColor( FCacheAddress( args.b, item ) )
                                  : FVector4f( 0, 0, 0, 0 );
 
 				FVector4f result = FVector4f(0,0,0,0);
@@ -5242,7 +5237,7 @@ namespace mu
                     break;
                 }
 
-                GetMemory().SetColour( item, result );
+                StoreColor( item, result );
                 break;
             }
 
@@ -5275,7 +5270,7 @@ namespace mu
 			OP::ResourceConstantArgs args = pModel->GetPrivate()->m_program.GetOpArgs<OP::ResourceConstantArgs>(item.At);
             ProjectorPtr pResult = new Projector();
             pResult->m_value = program.m_constantProjectors[args.value];
-            GetMemory().SetProjector( item, pResult );
+            StoreProjector( item, pResult );
             break;
         }
 
@@ -5290,7 +5285,7 @@ namespace mu
             const FProjector& def = program.m_parameters[args.variable].m_defaultValue.Get<ParamProjectorType>();
             pResult->m_value.type = def.type;
 
-            GetMemory().SetProjector( item, pResult );
+            StoreProjector( item, pResult );
             break;
         }
 
@@ -5317,7 +5312,7 @@ namespace mu
 
             LayoutPtrConst pResult = pModel->GetPrivate()->m_program.m_constantLayouts
                     [ args.value ];
-            GetMemory().SetLayout( item, pResult );
+            StoreLayout( item, pResult );
             break;
         }
 
@@ -5334,8 +5329,8 @@ namespace mu
 
             case 1:
             {
-                Ptr<const Layout> pA = GetMemory().GetLayout( FCacheAddress(args.Base,item) );
-                Ptr<const Layout> pB = GetMemory().GetLayout( FCacheAddress(args.Added,item) );
+                Ptr<const Layout> pA = LoadLayout( FCacheAddress(args.Base,item) );
+                Ptr<const Layout> pB = LoadLayout( FCacheAddress(args.Added,item) );
 
                 LayoutPtrConst pResult;
 
@@ -5352,7 +5347,7 @@ namespace mu
                     pResult = pB->Clone();
                 }
 
-                GetMemory().SetLayout( item, pResult );
+                StoreLayout( item, pResult );
                 break;
             }
 
@@ -5375,16 +5370,16 @@ namespace mu
 
             case 1:
             {
-                Ptr<const Layout> pSource = GetMemory().GetLayout( FCacheAddress(args.Source,item) );
+                Ptr<const Layout> Source = LoadLayout( FCacheAddress(args.Source,item) );
 
 				LayoutPtr pResult;
 
-				if (pSource)
+				if (Source)
 				{
-					pResult = pSource->Clone();
+					pResult = Source->Clone();
 
 					SCRATCH_LAYOUT_PACK scratch;
-					int32 BlockCount = pSource->GetBlockCount();
+					int32 BlockCount = Source->GetBlockCount();
 					scratch.blocks.SetNum(BlockCount);
 					scratch.sorted.SetNum(BlockCount);
 					scratch.positions.SetNum(BlockCount);
@@ -5392,10 +5387,10 @@ namespace mu
 					scratch.reductions.SetNum(BlockCount);
 					scratch.useSymmetry.SetNum(BlockCount);
 
-					LayoutPack3(pResult.get(), pSource.get(), &scratch);
+					LayoutPack3(pResult.get(), Source.get(), &scratch);
 				}
 
-                GetMemory().SetLayout( item, pResult );
+                StoreLayout( item, pResult );
                 break;
             }
 
@@ -5418,11 +5413,11 @@ namespace mu
 
 			case 1:
 			{
-				Ptr<const Mesh> Mesh = GetMemory().GetMesh(FCacheAddress(args.Mesh, item));
+				Ptr<const Mesh> Mesh = LoadMesh(FCacheAddress(args.Mesh, item));
 
 				Ptr<const Layout> Result = LayoutFromMesh_RemoveBlocks(Mesh.get(), args.LayoutIndex);
 
-				GetMemory().SetLayout(item, Result);
+				StoreLayout(item, Result);
 				break;
 			}
 
@@ -5446,8 +5441,8 @@ namespace mu
 
 			case 1:
 			{
-				Ptr<const Layout> Source = GetMemory().GetLayout(FCacheAddress(args.Source, item));
-				Ptr<const Layout> ReferenceLayout = GetMemory().GetLayout(FCacheAddress(args.ReferenceLayout, item));
+				Ptr<const Layout> Source = LoadLayout(FCacheAddress(args.Source, item));
+				Ptr<const Layout> ReferenceLayout = LoadLayout(FCacheAddress(args.ReferenceLayout, item));
 
 				Ptr<const Layout> pResult;
 
@@ -5460,7 +5455,7 @@ namespace mu
 					pResult = Source;
 				}
 
-				GetMemory().SetLayout(item, pResult);
+				StoreLayout(item, pResult);
 				break;
 			}
 
@@ -5488,6 +5483,12 @@ namespace mu
 		OP_TYPE type = pModel->GetPrivate()->m_program.GetOpType(item.At);
 		//UE_LOG(LogMutableCore, Log, TEXT("Running :%5d , %d, of type %d "), item.At, item.Stage, type);
 
+		// Very spammy, for debugging purposes.
+		//if (m_pSystem)
+		//{
+		//	m_pSystem->WorkingMemoryManager.LogWorkingMemory( this );
+		//}
+
 		switch ( type )
         {
         case OP_TYPE::NONE:
@@ -5610,7 +5611,7 @@ namespace mu
 			Result.m_size[0] = program.m_constantImages[ImageIndex].ImageSizeX;
 			Result.m_size[1] = program.m_constantImages[ImageIndex].ImageSizeY;
 			Result.m_lods = program.m_constantImages[ImageIndex].LODCount;
-			GetMemory().SetValidDesc(item);
+			StoreValidDesc(item);
 			break;
 		}
 
@@ -5621,7 +5622,7 @@ namespace mu
 			EXTERNAL_IMAGE_ID id = pParams->GetImageValue(args.variable);
 			uint8 MipsToSkip = item.ExecutionOptions;
 			m_heapImageDesc[item.CustomState] = GetExternalImageDesc(id, MipsToSkip);
-			GetMemory().SetValidDesc(item);
+			StoreValidDesc(item);
 			break;
 		}
 
@@ -5633,7 +5634,7 @@ namespace mu
 			Result.m_size[0] = 0;
 			Result.m_size[1] = 0;
 			Result.m_lods = 0;
-			GetMemory().SetValidDesc(item);
+			StoreValidDesc(item);
 			break;
 		}
 
@@ -5653,13 +5654,13 @@ namespace mu
 
 			case 1:
 			{
-				bool value = GetMemory().GetBool(FCacheAddress(args.condition, item.ExecutionIndex, item.ExecutionOptions));
+				bool value = LoadBool(FCacheAddress(args.condition, item.ExecutionIndex, item.ExecutionOptions));
 				OP::ADDRESS resultAt = value ? args.yes : args.no;
 				AddOp(FScheduledOp(item.At, item, 2), FScheduledOp(resultAt, item));
 				break;
 			}
 
-			case 2: GetMemory().SetValidDesc(item); break;
+			case 2: StoreValidDesc(item); break;
 			default: check(false);
 			}
 			break;
@@ -5694,7 +5695,7 @@ namespace mu
 				}
 				else
 				{
-					GetMemory().SetValidDesc(item);
+					StoreValidDesc(item);
 				}
 				break;
 			}
@@ -5702,7 +5703,7 @@ namespace mu
 			case 1:
 			{
 				// Get the variable result
-				int var = GetMemory().GetInt(FCacheAddress(VarAddress, item));
+				int var = LoadInt(FCacheAddress(VarAddress, item));
 
 				OP::ADDRESS valueAt = DefAddress;
 				for (uint32 C = 0; C < CaseCount; ++C)
@@ -5728,7 +5729,7 @@ namespace mu
 				break;
 			}
 
-			case 2: GetMemory().SetValidDesc(item); break;
+			case 2: StoreValidDesc(item); break;
 			default: check(false); break;
 			}
 			break;
@@ -5740,7 +5741,7 @@ namespace mu
 			switch (item.Stage)
 			{
 			case 0: AddOp(FScheduledOp(item.At, item, 1), FScheduledOp(args.base, item)); break;
-			case 1: GetMemory().SetValidDesc(item); break;
+			case 1: StoreValidDesc(item); break;
 			default: check(false);
 			}
 			break;
@@ -5752,7 +5753,7 @@ namespace mu
 			switch (item.Stage)
 			{
 			case 0: AddOp(FScheduledOp(item.At, item, 1), FScheduledOp(args.base, item)); break;
-			case 1: GetMemory().SetValidDesc(item); break;
+			case 1: StoreValidDesc(item); break;
 			default: check(false);
 			}
 			break;
@@ -5764,19 +5765,7 @@ namespace mu
 			switch (item.Stage)
 			{
 			case 0: AddOp(FScheduledOp(item.At, item, 1), FScheduledOp(args.base, item)); break;
-			case 1: GetMemory().SetValidDesc(item); break;
-			default: check(false);
-			}
-			break;
-		}
-
-		case OP_TYPE::IM_DIFFERENCE:
-		{
-			OP::ImageDifferenceArgs args = program.GetOpArgs<OP::ImageDifferenceArgs>(item.At);
-			switch (item.Stage)
-			{
-			case 0: AddOp(FScheduledOp(item.At, item, 1), FScheduledOp(args.a, item)); break;
-			case 1: GetMemory().SetValidDesc(item); break;
+			case 1: StoreValidDesc(item); break;
 			default: check(false);
 			}
 			break;
@@ -5788,7 +5777,7 @@ namespace mu
 			switch (item.Stage)
 			{
 			case 0: AddOp(FScheduledOp(item.At, item, 1), FScheduledOp(args.base, item)); break;
-			case 1: GetMemory().SetValidDesc(item); break;
+			case 1: StoreValidDesc(item); break;
 			default: check(false);
 			}
 			break;
@@ -5815,7 +5804,7 @@ namespace mu
 					NewFormat = args.formatIfAlpha;
 				}
 				m_heapImageDesc[item.CustomState].m_format = NewFormat;				
-				GetMemory().SetValidDesc(item);
+				StoreValidDesc(item);
 				break;
 			}
 
@@ -5858,7 +5847,7 @@ namespace mu
 
 				// Update result.
 				m_heapImageDesc[item.CustomState].m_lods = levelCount;
-				GetMemory().SetValidDesc(item);
+				StoreValidDesc(item);
 				break;
 			}
 
@@ -5880,7 +5869,7 @@ namespace mu
 			case 1:
 				m_heapImageDesc[item.CustomState].m_size[0] = args.size[0];
 				m_heapImageDesc[item.CustomState].m_size[1] = args.size[1];
-				GetMemory().SetValidDesc(item);
+				StoreValidDesc(item);
 				break;
 
 			default:
@@ -5914,7 +5903,7 @@ namespace mu
 				FImageDesc& ResultAndBaseDesc = m_heapImageDesc[SecondStageData.ResizeLike.ResultDescAt];
 				const FImageDesc& SourceDesc = m_heapImageDesc[SecondStageData.ResizeLike.SourceDescAt];
 				ResultAndBaseDesc.m_size = SourceDesc.m_size;
-				GetMemory().SetValidDesc(item);
+				StoreValidDesc(item);
 				break;
 			}
 
@@ -5941,7 +5930,7 @@ namespace mu
 					uint16(ResultAndBaseDesc.m_size[0] * args.factor[0] + 0.5f),
 					uint16(ResultAndBaseDesc.m_size[1] * args.factor[1] + 0.5f));
 				ResultAndBaseDesc.m_size = destSize;
-				GetMemory().SetValidDesc(item);
+				StoreValidDesc(item);
 				break;
 			}
 
@@ -5968,7 +5957,7 @@ namespace mu
 
 			case 1:
 			{
-				Ptr<const Layout> pLayout = GetMemory().GetLayout(FCacheAddress(args.layout, item));
+				Ptr<const Layout> pLayout = LoadLayout(FCacheAddress(args.layout, item));
 
 				FIntPoint SizeInBlocks = pLayout->GetGridSize();
 				FIntPoint BlockSizeInPixels(args.blockSize[0], args.blockSize[1]);
@@ -5990,7 +5979,7 @@ namespace mu
 						ResultAndBaseDesc.m_lods = args.mipmapCount;
 					}
 				}
-				GetMemory().SetValidDesc(item);
+				StoreValidDesc(item);
 				break;
 			}
 
@@ -6007,7 +5996,7 @@ namespace mu
 			switch (item.Stage)
 			{
 			case 0: AddOp(FScheduledOp(item.At, item, 1), FScheduledOp(args.base, item)); break;
-			case 1: GetMemory().SetValidDesc(item); break;
+			case 1: StoreValidDesc(item); break;
 			default: check(false);
 			}
 			break;
@@ -6019,7 +6008,7 @@ namespace mu
 			switch (item.Stage)
 			{
 			case 0: AddOp(FScheduledOp(item.At, item, 1), FScheduledOp(args.targets[0], item)); break;
-			case 1: GetMemory().SetValidDesc(item); break;
+			case 1: StoreValidDesc(item); break;
 			default: check(false);
 			}
 			break;
@@ -6031,7 +6020,7 @@ namespace mu
 			switch (item.Stage)
 			{
 			case 0: AddOp(FScheduledOp(item.At, item, 1), FScheduledOp(args.target0, item)); break;
-			case 1: GetMemory().SetValidDesc(item); break;
+			case 1: StoreValidDesc(item); break;
 			default: check(false);
 			}
 			break;
@@ -6043,7 +6032,7 @@ namespace mu
 			switch (item.Stage)
 			{
 			case 0: AddOp(FScheduledOp(item.At, item, 1), FScheduledOp(args.base, item)); break;
-			case 1: GetMemory().SetValidDesc(item); break;
+			case 1: StoreValidDesc(item); break;
 			default: check(false);
 			}
 			break;
@@ -6060,7 +6049,7 @@ namespace mu
 
 			case 1:
 				m_heapImageDesc[item.CustomState].m_format = EImageFormat::IF_L_UBYTE;
-				GetMemory().SetValidDesc(item);
+				StoreValidDesc(item);
 				break;
 
 			default:
@@ -6081,28 +6070,7 @@ namespace mu
 
 			case 1:
 				m_heapImageDesc[item.CustomState].m_format = args.format;
-				GetMemory().SetValidDesc(item);
-				break;
-
-			default:
-				check(false);
-			}
-
-			break;
-		}
-
-		case OP_TYPE::IM_SELECTCOLOUR:
-		{
-			OP::ImageSelectColourArgs args = program.GetOpArgs<OP::ImageSelectColourArgs>(item.At);
-			switch (item.Stage)
-			{
-			case 0:
-				AddOp(FScheduledOp(item.At, item, 1), FScheduledOp(args.base, item));
-				break;
-
-			case 1:
-				m_heapImageDesc[item.CustomState].m_format = EImageFormat::IF_L_UBYTE;
-				GetMemory().SetValidDesc(item);
+				StoreValidDesc(item);
 				break;
 
 			default:
@@ -6118,7 +6086,7 @@ namespace mu
 			switch (item.Stage)
 			{
 			case 0: AddOp(FScheduledOp(item.At, item, 1), FScheduledOp(args.base, item)); break;
-			case 1: GetMemory().SetValidDesc(item); break;
+			case 1: StoreValidDesc(item); break;
 			default: check(false);
 			}
 			break;
@@ -6131,7 +6099,7 @@ namespace mu
 			m_heapImageDesc[item.CustomState].m_size[1] = args.size[1];
 			m_heapImageDesc[item.CustomState].m_lods = 1;
 			m_heapImageDesc[item.CustomState].m_format = EImageFormat::IF_RGB_UBYTE;
-			GetMemory().SetValidDesc(item);
+			StoreValidDesc(item);
 			break;
 		}
 
@@ -6146,7 +6114,7 @@ namespace mu
 
 			case 1:
 				m_heapImageDesc[item.CustomState].m_format = EImageFormat::IF_L_UBYTE;
-				GetMemory().SetValidDesc(item);
+				StoreValidDesc(item);
 				break;
 
 			default:
@@ -6161,7 +6129,7 @@ namespace mu
 			switch (item.Stage)
 			{
 			case 0: AddOp(FScheduledOp(item.At, item, 1), FScheduledOp(args.base, item)); break;
-			case 1: GetMemory().SetValidDesc(item); break;
+			case 1: StoreValidDesc(item); break;
 			default: check(false);
 			}
 			break;
@@ -6174,7 +6142,7 @@ namespace mu
 			m_heapImageDesc[item.CustomState].m_size[1] = args.size[1];
 			m_heapImageDesc[item.CustomState].m_lods = args.LODs;
 			m_heapImageDesc[item.CustomState].m_format = args.format;
-			GetMemory().SetValidDesc(item);
+			StoreValidDesc(item);
 			break;
 		}
 
@@ -6191,7 +6159,7 @@ namespace mu
 				m_heapImageDesc[item.CustomState].m_size[0] = args.sizeX;
 				m_heapImageDesc[item.CustomState].m_size[1] = args.sizeY;
 				m_heapImageDesc[item.CustomState].m_lods = 1;
-				GetMemory().SetValidDesc(item);
+				StoreValidDesc(item);
 				break;
 
 			default:
@@ -6206,7 +6174,7 @@ namespace mu
 			switch (item.Stage)
 			{
 			case 0: AddOp(FScheduledOp(item.At, item, 1), FScheduledOp(args.base, item)); break;
-			case 1: GetMemory().SetValidDesc(item); break;
+			case 1: StoreValidDesc(item); break;
 			default: check(false);
 			}
 			break;
@@ -6219,7 +6187,7 @@ namespace mu
 			m_heapImageDesc[item.CustomState].m_size[1] = args.sizeY;
 			m_heapImageDesc[item.CustomState].m_lods = 1;
 			m_heapImageDesc[item.CustomState].m_format = EImageFormat::IF_L_UBYTE;
-			GetMemory().SetValidDesc(item);
+			StoreValidDesc(item);
 			break;
 		}
 
@@ -6235,7 +6203,7 @@ namespace mu
 			case 1:
 				m_heapImageDesc[item.CustomState].m_format = EImageFormat::IF_L_UBYTE;
 				m_heapImageDesc[item.CustomState].m_lods = 1;
-				GetMemory().SetValidDesc(item);
+				StoreValidDesc(item);
 				break;
 
 			default:
@@ -6251,7 +6219,7 @@ namespace mu
 			switch (item.Stage)
 			{
 			case 0: AddOp(FScheduledOp(item.At, item, 1), FScheduledOp(args.source, item)); break;
-			case 1: GetMemory().SetValidDesc(item); break;
+			case 1: StoreValidDesc(item); break;
 			default: check(false);
 			}
 			break;
@@ -6271,7 +6239,7 @@ namespace mu
 			}
             case 1:
             {
-				GetMemory().SetValidDesc(item);
+				StoreValidDesc(item);
                 break;
             }
 
