[mutable] Implement memory budget control for image operations.

diff --git a/Engine/Plugins/Experimental/Mutable/Source/MutableRuntime/Private/MuR/CodeRunner.h b/Engine/Plugins/Experimental/Mutable/Source/MutableRuntime/Private/MuR/CodeRunner.h
--- a/Engine/Plugins/Experimental/Mutable/Source/MutableRuntime/Private/MuR/CodeRunner.h
+++ b/Engine/Plugins/Experimental/Mutable/Source/MutableRuntime/Private/MuR/CodeRunner.h
@@ -40,20 +40,16 @@ namespace  mu
     //---------------------------------------------------------------------------------------------
     //! Code execution of the mutable virtual machine.
     //---------------------------------------------------------------------------------------------
-    class CodeRunner : public Base
+    class CodeRunner
     {
     public:
-		CodeRunner(const SettingsPtrConst&, class System::Private*, 
+		CodeRunner(const Ptr<const Settings>&, class System::Private*, 
 			EExecutionStrategy,
 			const TSharedPtr<const Model>&, const Parameters* pParams,
 			OP::ADDRESS at, uint32 lodMask, uint8 executionOptions, int32 InImageLOD, FScheduledOp::EType );
 
     protected:
 
-        //! Settings that may affect the execution of some operations, like image conversion
-        //! quality.
-        SettingsPtrConst m_pSettings;
-
         //! Type of data sometimes stored in the code runner heap to pass info between
         //! operation stages.
         struct FScheduledOpData
@@ -121,6 +117,9 @@ namespace  mu
 #endif
  	    mu::FImageDesc GetExternalImageDesc(EXTERNAL_IMAGE_ID Id, uint8 MipmapsToSkip);
 
+		/** Settings that may affect the execution of some operations, like image conversion quality. */
+		Ptr<const Settings> m_pSettings;
+
     protected:
         //! Heap of intermediate data pushed by some instructions and referred by others.
         //! It is not released until no operations are pending.
@@ -145,11 +144,36 @@ namespace  mu
 		{
 			FTask() {}
 			FTask(const FScheduledOp& InOp) : Op(InOp) {}
-			FTask(const FScheduledOp& InOp, const FScheduledOp& InDep0) : Op(InOp) { Deps.Add(InDep0); }
-			FTask(const FScheduledOp& InOp, const FScheduledOp& InDep0, const FScheduledOp& InDep1) : Op(InOp) { Deps.Add(InDep0); Deps.Add(InDep1); }
-			FTask(const FScheduledOp& InOp, const FScheduledOp& InDep0, const FScheduledOp& InDep1, const FScheduledOp& InDep2) : Op(InOp) { Deps.Add(InDep0); Deps.Add(InDep1); Deps.Add(InDep2); }
-			FTask(const FScheduledOp& InOp, const FScheduledOp& InDep0, const FScheduledOp& InDep1, const FScheduledOp& InDep2, const FScheduledOp& InDep3) : Op(InOp) { Deps.Add(InDep0); Deps.Add(InDep1); Deps.Add(InDep2); Deps.Add(InDep3); }
-			FTask(const FScheduledOp& InOp, const FScheduledOp& InDep0, const FScheduledOp& InDep1, const FScheduledOp& InDep2, const FScheduledOp& InDep3, const FScheduledOp& InDep4) : Op(InOp) { Deps.Add(InDep0); Deps.Add(InDep1); Deps.Add(InDep2); Deps.Add(InDep3); Deps.Add(InDep4); }
+			FTask(const FScheduledOp& InOp, const FScheduledOp& InDep0) : Op(InOp) 
+			{ 
+				if (InDep0.At) Deps.Add(InDep0); 
+			}
+			FTask(const FScheduledOp& InOp, const FScheduledOp& InDep0, const FScheduledOp& InDep1) : Op(InOp) 
+			{ 
+				if (InDep0.At) Deps.Add(InDep0); 
+				if (InDep1.At) Deps.Add(InDep1); 
+			}
+			FTask(const FScheduledOp& InOp, const FScheduledOp& InDep0, const FScheduledOp& InDep1, const FScheduledOp& InDep2) : Op(InOp) 
+			{ 
+				if (InDep0.At) Deps.Add(InDep0); 
+				if (InDep1.At) Deps.Add(InDep1); 
+				if (InDep2.At) Deps.Add(InDep2); 
+			}
+			FTask(const FScheduledOp& InOp, const FScheduledOp& InDep0, const FScheduledOp& InDep1, const FScheduledOp& InDep2, const FScheduledOp& InDep3) : Op(InOp) 
+			{ 
+				if (InDep0.At) Deps.Add(InDep0); 
+				if (InDep1.At) Deps.Add(InDep1); 
+				if (InDep2.At) Deps.Add(InDep2); 
+				if (InDep3.At) Deps.Add(InDep3); 
+			}
+			FTask(const FScheduledOp& InOp, const FScheduledOp& InDep0, const FScheduledOp& InDep1, const FScheduledOp& InDep2, const FScheduledOp& InDep3, const FScheduledOp& InDep4) : Op(InOp) 
+			{ 
+				if (InDep0.At) Deps.Add(InDep0); 
+				if (InDep1.At) Deps.Add(InDep1); 
+				if (InDep2.At) Deps.Add(InDep2); 
+				if (InDep3.At) Deps.Add(InDep3); 
+				if (InDep4.At) Deps.Add(InDep4); 
+			}
 
 			FScheduledOp Op;
 			TArray<FCacheAddress, TInlineAllocator<3>> Deps;
@@ -168,7 +192,9 @@ namespace  mu
 
 			FIssuedTask(const FScheduledOp& InOp) : Op(InOp) {}
 			virtual ~FIssuedTask() {}
-			virtual bool Prepare(CodeRunner*, const TSharedPtr<const Model>&, bool& bOutFailed) { bOutFailed = false; return true; }
+
+			/** */
+			virtual bool Prepare(CodeRunner*, bool& bOutFailed) { bOutFailed = false; return true; }
 			virtual void DoWork() {}
 			virtual void Complete(CodeRunner*) = 0;
 			virtual bool IsComplete(CodeRunner*)
@@ -192,7 +218,7 @@ namespace  mu
 			}
 
 			// FIssuedTask interface
-			bool Prepare(CodeRunner*, const TSharedPtr<const Model>&, bool& bOutFailed) override;
+			bool Prepare(CodeRunner*, bool& bOutFailed) override;
 			void Complete(CodeRunner*) override;
 			bool IsComplete(CodeRunner*) override;
 
@@ -213,7 +239,7 @@ namespace  mu
 			}
 
 			// FIssuedTask interface
-			bool Prepare(CodeRunner*, const TSharedPtr<const Model>&, bool& bOutFailed) override;
+			bool Prepare(CodeRunner*, bool& bOutFailed) override;
 			void Complete(CodeRunner*) override;
 			bool IsComplete(CodeRunner*) override;
 
@@ -236,7 +262,7 @@ namespace  mu
 			}
 
 			// FIssuedTask interface
-			bool Prepare(CodeRunner*, const TSharedPtr<const Model>&, bool& bOutFailed) override;
+			bool Prepare(CodeRunner*, bool& bOutFailed) override;
 			void Complete(CodeRunner*) override;
 			bool IsComplete(CodeRunner*) override;
 
@@ -332,6 +358,15 @@ namespace  mu
 			}
 		}
 
+		/** Calculate an approximation of memory used by manging structures in this class. */
+		int32 GetInternalMemoryBytes() const
+		{
+			return sizeof(CodeRunner) 
+				+ m_heapData.GetAllocatedSize() + m_heapImageDesc.GetAllocatedSize()
+				+ ClosedTasks.GetAllocatedSize() + OpenTasks.GetAllocatedSize() + ScheduledStagePerOp.GetAllocatedSize()
+				// this contains smart pointers, approximate size like this:
+				+ IssuedTasks.Max() * ( sizeof(FIssuedTask) + 16);
+		}
 
 	protected:
 
@@ -350,6 +385,153 @@ namespace  mu
 		/** List of tasks that have been set to run concurrently and their completion is unknown. */
 		TArray< TSharedPtr<FIssuedTask> > IssuedTasks;
 
+	public:
+
+		bool LoadBool(const FCacheAddress& From)
+		{
+			return m_pSystem->WorkingMemoryManager.CurrentInstanceCache->GetBool(From);
+		}
+
+		float LoadInt(const FCacheAddress& From)
+		{
+			return m_pSystem->WorkingMemoryManager.CurrentInstanceCache->GetInt(From);
+		}
+
+		float LoadScalar(const FCacheAddress& From)
+		{
+			return m_pSystem->WorkingMemoryManager.CurrentInstanceCache->GetScalar(From);
+		}
+
+		FVector4f LoadColor(const FCacheAddress& From)
+		{
+			return m_pSystem->WorkingMemoryManager.CurrentInstanceCache->GetColour(From);
+		}
+
+		Ptr<const String> LoadString(const FCacheAddress& From)
+		{
+			return m_pSystem->WorkingMemoryManager.CurrentInstanceCache->GetString(From);
+		}
+
+		Ptr<const Projector> LoadProjector(const FCacheAddress& From)
+		{
+			return m_pSystem->WorkingMemoryManager.CurrentInstanceCache->GetProjector(From);
+		}
+
+		Ptr<const Mesh> LoadMesh(const FCacheAddress& From)
+		{
+			return m_pSystem->WorkingMemoryManager.CurrentInstanceCache->GetMesh(From);
+		}
+
+		Ptr<const Image> LoadImage(const FCacheAddress& From)
+		{
+			return m_pSystem->WorkingMemoryManager.LoadImage(From);
+		}
+
+		Ptr<const Layout> LoadLayout(const FCacheAddress& From)
+		{
+			return m_pSystem->WorkingMemoryManager.CurrentInstanceCache->GetLayout(From);
+		}
+
+		Ptr<const Instance> LoadInstance(const FCacheAddress& From)
+		{
+			return m_pSystem->WorkingMemoryManager.CurrentInstanceCache->GetInstance(From);
+		}
+
+		Ptr<const ExtensionData> LoadExtensionData(const FCacheAddress& From)
+		{
+			return m_pSystem->WorkingMemoryManager.CurrentInstanceCache->GetExtensionData(From);
+		}
+
+		void StoreValidDesc(const FCacheAddress& To)
+		{
+			m_pSystem->WorkingMemoryManager.CurrentInstanceCache->SetValidDesc(To);
+		}
+
+		void StoreBool(const FCacheAddress& To, bool Value)
+		{
+			m_pSystem->WorkingMemoryManager.CurrentInstanceCache->SetBool(To, Value);
+		}
+
+		void StoreInt(const FCacheAddress& To, int32 Value)
+		{
+			m_pSystem->WorkingMemoryManager.CurrentInstanceCache->SetInt(To, Value);
+		}
+
+		void StoreScalar(const FCacheAddress& To, float Value)
+		{
+			m_pSystem->WorkingMemoryManager.CurrentInstanceCache->SetScalar(To, Value);
+		}
+
+		void StoreString(const FCacheAddress& To, Ptr<const String> Value)
+		{
+			m_pSystem->WorkingMemoryManager.CurrentInstanceCache->SetString(To, Value);
+		}
+
+		void StoreColor(const FCacheAddress& To, const FVector4f& Value)
+		{
+			m_pSystem->WorkingMemoryManager.CurrentInstanceCache->SetColour(To, Value);
+		}
+
+		void StoreProjector(const FCacheAddress& To, Ptr<const Projector> Value)
+		{
+			m_pSystem->WorkingMemoryManager.CurrentInstanceCache->SetProjector(To, Value);
+		}
+
+		void StoreMesh(const FCacheAddress& To, Ptr<const Mesh> Resource)
+		{
+			m_pSystem->WorkingMemoryManager.CurrentInstanceCache->SetMesh(To, Resource);
+		}
+
+		void StoreImage(const FCacheAddress& To, Ptr<const Image> Resource)
+		{
+			m_pSystem->WorkingMemoryManager.StoreImage(To, Resource);
+		}
+
+		void StoreLayout(const FCacheAddress& To, Ptr<const Layout> Resource)
+		{
+			m_pSystem->WorkingMemoryManager.CurrentInstanceCache->SetLayout(To, Resource);
+		}
+
+		void StoreInstance(const FCacheAddress& To, Ptr<const Instance> Resource)
+		{
+			m_pSystem->WorkingMemoryManager.CurrentInstanceCache->SetInstance(To, Resource);
+		}
+
+		void StoreExtensionData(const FCacheAddress& To, Ptr<const ExtensionData> Resource)
+		{
+			m_pSystem->WorkingMemoryManager.CurrentInstanceCache->SetExtensionData(To, Resource);
+		}
+
+		Ptr<Image> CreateImage(uint32 SizeX, uint32 SizeY, uint32 Lods, EImageFormat Format, EInitializationType Init = EInitializationType::Black)
+		{
+			Ptr<Image> Result = m_pSystem->WorkingMemoryManager.CreateImage(SizeX, SizeY, Lods, Format, Init);
+			return MoveTemp(Result);
+		}
+
+		Ptr<Image> CreateImageLike(const Image* Ref, EInitializationType Init = EInitializationType::Black)
+		{
+			Ptr<Image> Result = m_pSystem->WorkingMemoryManager.CreateImage(Ref->GetSizeX(), Ref->GetSizeY(), Ref->GetLODCount(), Ref->GetFormat(), Init);
+			return MoveTemp(Result);
+		}
+
+		/** Ref will be nulled and relesed in any case. */
+		Ptr<Image> CloneOrTakeOver(Ptr<const Image>& Ref)
+		{
+			Ptr<Image> Result = m_pSystem->WorkingMemoryManager.CloneOrTakeOver(Ref);
+			return MoveTemp(Result);
+		}
+
+		void Release(Ptr<const Image>& Resource)
+		{
+			return m_pSystem->WorkingMemoryManager.Release(Resource);
+		}
+
+		void Release(Ptr<Image>& Resource)
+		{
+			return m_pSystem->WorkingMemoryManager.Release(Resource);
+		}
+
+
 	// TODO: protect this.
 	public:
 		
@@ -373,9 +555,6 @@ namespace  mu
 		};
 		TArray<FRomLoadOp> m_romLoadOps;
 
-		//! Count of pending operations for every rom index
-		TArray<uint16> m_romPendingOps;
-
 	private:
 
 		inline void AddChildren(const FScheduledOp& dep)
@@ -392,7 +571,7 @@ namespace  mu
 
 			if (dep.Type == FScheduledOp::EType::Full)
 			{
-				m_pSystem->m_memory->IncreaseHitCount(at);
+				m_pSystem->WorkingMemoryManager.CurrentInstanceCache->IncreaseHitCount(at);
 			}
 		}
 
