[mutable] Implement memory budget control for image operations.

diff --git a/Engine/Plugins/Experimental/Mutable/Source/MutableRuntime/Private/MuR/CodeRunnerTasked.cpp b/Engine/Plugins/Experimental/Mutable/Source/MutableRuntime/Private/MuR/CodeRunnerTasked.cpp
--- a/Engine/Plugins/Experimental/Mutable/Source/MutableRuntime/Private/MuR/CodeRunnerTasked.cpp
+++ b/Engine/Plugins/Experimental/Mutable/Source/MutableRuntime/Private/MuR/CodeRunnerTasked.cpp
@@ -26,9 +26,7 @@
 #include "MuR/OpImageCompose.h"
 #include "MuR/OpImageMipmap.h"
 #include "MuR/OpImageNormalCombine.h"
-#include "MuR/OpImagePixelFormat.h"
 #include "MuR/OpImageResize.h"
-#include "MuR/OpImageSwizzle.h"
 #include "MuR/OpImageSaturate.h"
 #include "MuR/OpImageInvert.h"
 #include "MuR/Operations.h"
@@ -38,7 +36,6 @@
 #include "MuR/Serialisation.h"
 #include "MuR/SerialisationPrivate.h"
 #include "MuR/Settings.h"
-#include "MuR/SettingsPrivate.h"
 #include "MuR/SystemPrivate.h"
 #include "Stats/Stats2.h"
 #include "Templates/RefCounting.h"
@@ -68,7 +65,7 @@ namespace mu
 		if (DebugRom && (DebugRomAll || o.m_romIndex == DebugRomIndex))
 			UE_LOG(LogMutableCore, Log, TEXT("CodeRunner::CompleteRomLoadOp for rom %d."), o.m_romIndex);
 
-		m_pSystem->m_pStreamInterface->EndRead(o.m_streamID);
+		m_pSystem->StreamInterface->EndRead(o.m_streamID);
 
 		FProgram& program = m_pModel->GetPrivate()->m_program;
 		{
@@ -117,7 +114,7 @@ namespace mu
 	{
 #if UE_MUTABLE_ENABLE_SLOW_TRACES && !UE_BUILD_SHIPPING && !UE_BUILD_TEST
 
-		//// Code runner status
+		//// Code Runner status
 		//TRACE_COUNTER_SET(MutableRuntime_OpenTask, OpenTasks.Num());
 		//TRACE_COUNTER_SET(MutableRuntime_ClosedTasks, ClosedTasks.Num());
 		//TRACE_COUNTER_SET(MutableRuntime_IssuedTasks, IssuedTasks.Num());
@@ -263,13 +260,11 @@ namespace mu
 		{   0,   0,   0,   0 },	// IM_RESIZEREL		(special case)
 		{   0,  20,   0,   0 },	// IM_BLANKLAYOUT
 		{   0,   0, -20,   0 },	// IM_COMPOSE
-		{   0,   0,   0,   0 },	// IM_DIFFERENCE
 		{   0,   0, -20,   0 },	// IM_INTERPOLATE
 		{   0,   0, -30,   0 },	// IM_INTERPOLATE3
 		{   0,   0,   0,   0 },	// IM_SATURATE
 		{   0,   0,   0,   0 },	// IM_LUMINANCE
 		{   0,   0,   0,   0 },	// IM_SWIZZLE
-		{   0,   0,   0,   0 },	// IM_SELECTCOLOUR
 		{   0,   0,   0,   0 },	// IM_COLOURMAP
 		{   0,   5,   0,   0 },	// IM_GRADIENT
 		{   0,   0,   0,   0 },	// IM_BINARISE
@@ -482,7 +477,7 @@ namespace mu
 				if (IssuedTask)
 				{
 					bool bFailed = false;
-					bool hasWork = IssuedTask->Prepare(this, m_pModel, bFailed);
+					bool hasWork = IssuedTask->Prepare(this, bFailed);
 					if (bFailed)
 					{
 						bUnrecoverableError = true;
@@ -541,7 +536,7 @@ namespace mu
 #endif
 						}
 					}
-
+					
 					// Remember it for later processing.
 					IssuedTasks.Add(IssuedTask);
 				}
@@ -571,7 +566,7 @@ namespace mu
 			// Look for completed streaming ops and complete the rom loading
 			for (FRomLoadOp& o : m_romLoadOps)
 			{
-				if (o.m_romIndex>=0 && m_pSystem->m_pStreamInterface->IsReadCompleted(o.m_streamID))
+				if (o.m_romIndex>=0 && m_pSystem->StreamInterface->IsReadCompleted(o.m_streamID))
 				{
 					CompleteRomLoadOp(o);
 					UpdateTraces();
@@ -585,13 +580,13 @@ namespace mu
 				bool Ready = true;
 				for ( const FCacheAddress& Dep: ClosedTasks[Index].Deps )
 				{
-					bool bDependencySatisfied = false;					
-					bDependencySatisfied = Dep.At && !GetMemory().IsValid(Dep);
+					bool bDependencyFailed = false;
+					bDependencyFailed = Dep.At && !GetMemory().IsValid(Dep);
 					
-					if (bDependencySatisfied)
+					if (bDependencyFailed)
 					{
 						Ready = false;
-						continue;
+						break;
 					}
 				}
 
@@ -717,259 +712,294 @@ namespace mu
 	class FImageLayerTask : public CodeRunner::FIssuedTask
 	{
 	public:
-		FImageLayerTask(const FScheduledOp&, FProgramCache&, const OP::ImageLayerArgs&, int imageCompressionQuality );
+		FImageLayerTask(const FScheduledOp& InOp, const OP::ImageLayerArgs& InArgs)
+			: FIssuedTask(InOp), Args(InArgs)
+		{}
 
 		// FIssuedTask interface
+		bool Prepare(CodeRunner*, bool& bOutFailed) override;
 		void DoWork() override;
-		void Complete(CodeRunner* staged) override;
+		void Complete(CodeRunner*) override;
 
 	private:
-		int m_imageCompressionQuality = 0;
+		int32 ImageCompressionQuality = 0;
 		OP::ImageLayerArgs Args;
-		Ptr<const Image> m_base;
-		Ptr<const Image> m_blended;
-		Ptr<const Image> m_mask;
-		Ptr<const Image> Result;
+		Ptr<const Image> Blended;
+		Ptr<const Image> Mask;
+		Ptr<Image> Result;
+		EImageFormat InitialFormat = EImageFormat::IF_NONE;
 	};
 
 
-	//---------------------------------------------------------------------------------------------
-	FImageLayerTask::FImageLayerTask(const FScheduledOp& InOp, FProgramCache& Memory, const OP::ImageLayerArgs& InArgs, int imageCompressionQuality)
-		: FIssuedTask(InOp)
+	bool FImageLayerTask::Prepare(CodeRunner* Runner, bool& bOutFailed)
 	{
-		m_imageCompressionQuality = imageCompressionQuality;
-		Args = InArgs;
-		m_base = Memory.GetImage({ Args.base, InOp.ExecutionIndex, InOp.ExecutionOptions });
-		check(!m_base || m_base->GetFormat() < EImageFormat::IF_COUNT);
+		// This runs in the mutable Runner thread
+		MUTABLE_CPUPROFILER_SCOPE(FImageLayerTask_Prepare);
+		bOutFailed = false;
 
-		m_blended = Memory.GetImage({ Args.blended, InOp.ExecutionIndex, InOp.ExecutionOptions });
+		ImageCompressionQuality = Runner->m_pSettings->ImageCompressionQuality;
+
+		Ptr<const Image> Base = Runner->LoadImage({ Args.base, Op.ExecutionIndex, Op.ExecutionOptions });
+		check(!Base || Base->GetFormat() < EImageFormat::IF_COUNT);
+
+		Blended = Runner->LoadImage({ Args.blended, Op.ExecutionIndex, Op.ExecutionOptions });
 		if (Args.mask)
 		{
-			m_mask = Memory.GetImage({ Args.mask, InOp.ExecutionIndex, InOp.ExecutionOptions });
-			check(!m_mask || m_mask->GetFormat() < EImageFormat::IF_COUNT);
+			Mask = Runner->LoadImage({ Args.mask, Op.ExecutionIndex, Op.ExecutionOptions });
+			check(!Mask || Mask->GetFormat() < EImageFormat::IF_COUNT);
 		}
-	}
-
-	//---------------------------------------------------------------------------------------------
-	void FImageLayerTask::DoWork()
-	{
-		MUTABLE_CPUPROFILER_SCOPE(FImageLayerTask);
 
-		// Warning: This runs in a worker thread.
+		// Shortcuts
+		if (!Base)
+		{
+			Runner->Release(Blended);
+			Runner->Release(Mask);
+			Runner->StoreImage(Op, Base);
+			return false;
+		}
 
-		if (!m_base)
+		bool bValid = Base->GetSizeX() > 0 && Base->GetSizeY() > 0;
+		if (!bValid || !Blended)
 		{
-			// This shouldn't really happen if data is correct.
-			return;
+			Runner->Release(Blended);
+			Runner->Release(Mask);
+			Runner->StoreImage(Op, Base);
+			return false;
 		}
 
-		EImageFormat InitialFormat = m_base->GetFormat();
-		check(InitialFormat < EImageFormat::IF_COUNT);
+		// Input data fixes
+		InitialFormat = Base->GetFormat();
 
 		if (IsCompressedFormat(InitialFormat))
 		{
-			UE_LOG(LogMutableCore, Error, TEXT("Image layer on compressed image is not supported. Ignored."));
 			EImageFormat UncompressedFormat = GetUncompressedFormat(InitialFormat);
-			m_base = ImagePixelFormat(m_imageCompressionQuality, m_base.get(), UncompressedFormat);
+			Ptr<Image> Formatted = Runner->CreateImage( Base->GetSizeX(), Base->GetSizeY(), Base->GetLODCount(), UncompressedFormat, EInitializationType::NotInitialized );
+			bool bSuccess = false;
+			ImagePixelFormat(bSuccess, ImageCompressionQuality, Formatted.get(), Base.get());
+			check(bSuccess); // Decompression cannot fail
+			Runner->Release(Base);
+			Base = Formatted;
 		}
 
-		if (m_blended && InitialFormat != m_blended->GetFormat() && !Args.flags)
+		if (Blended && InitialFormat != Blended->GetFormat() && !Args.flags)
 		{
-			m_blended = ImagePixelFormat(m_imageCompressionQuality, m_blended.get(), m_base->GetFormat());
+			Ptr<Image> Formatted = Runner->CreateImage(Blended->GetSizeX(), Blended->GetSizeY(), Blended->GetLODCount(), Base->GetFormat(), EInitializationType::NotInitialized);
+			bool bSuccess = false;
+			ImagePixelFormat(bSuccess, ImageCompressionQuality, Formatted.get(), Blended.get());
+			check(bSuccess); // Decompression cannot fail
+			Runner->Release(Blended);
+			Blended = Formatted;
 		}
 
-		ImagePtr pNew = mu::CloneOrTakeOver<Image>(m_base.get());
+		if (Base->GetSize() != Blended->GetSize())
+		{
+			MUTABLE_CPUPROFILER_SCOPE(ImageResize_EmergencyFix);
+			Ptr<Image> Resized = Runner->CreateImage(Base->GetSizeX(), Base->GetSizeY(), 1, Blended->GetFormat(), EInitializationType::NotInitialized);
+			ImageResizeLinear(Resized.get(), ImageCompressionQuality, Blended.get());
+			Runner->Release(Blended);
+			Blended = Resized;
 
-		check(pNew->GetDataSize() == m_base->GetDataSize());
+		}
 
-		bool bValid = pNew->GetSizeX() > 0 && pNew->GetSizeY() > 0;
-		if (bValid && m_blended)
+		if (Args.mask)
 		{
-			// TODO: This shouldn't happen, but check it to be safe
-			if (m_base->GetSize() != m_blended->GetSize())
+			if (Base->GetSize() != Mask->GetSize())
 			{
 				MUTABLE_CPUPROFILER_SCOPE(ImageResize_EmergencyFix);
-				m_blended = ImageResize(m_blended.get(), m_base->GetSize());
-			}
-
-			check(pNew->GetDataSize() == m_base->GetDataSize());
+				Ptr<Image> Resized = Runner->CreateImage(Base->GetSizeX(), Base->GetSizeY(), 1, Mask->GetFormat(), EInitializationType::NotInitialized);
+				ImageResizeLinear(Resized.get(), ImageCompressionQuality, Mask.get());
+				Runner->Release(Mask);
+				Mask = Resized;
 
-			bool bOnlyOneMip = false;
-			if (m_blended->GetLODCount() < m_base->GetLODCount())
-			{
-				bOnlyOneMip = true;
 			}
 
-			bool bDone = false;
-
-			if (!m_mask 
-				&& 
-				// Flags have to match exactly for this optimize case. Other flags are not supported.
-				Args.flags== OP::ImageLayerArgs::F_USE_MASK_FROM_BLENDED
-				&&
-				Args.blendType == uint8(EBlendType::BT_BLEND)
-				&&
-				Args.blendTypeAlpha == uint8(EBlendType::BT_LIGHTEN))
+			// emergy fix c36adf47-e40d-490f-b709-41142bafad78
+			if (Mask->GetLODCount() < Base->GetLODCount())
 			{
-				MUTABLE_CPUPROFILER_SCOPE(ImageLayerTask_Optimized);
+				MUTABLE_CPUPROFILER_SCOPE(ImageLayer_EmergencyFix);
 
-				// This is a frequent critical-path case because of multilayer projectors.
-				bDone = true;
+				int32 levelCount = Base->GetLODCount();
+				ImagePtr pDest = Runner->CreateImage(Mask->GetSizeX(), Mask->GetSizeY(),
+					levelCount,
+					Mask->GetFormat());
 
-				constexpr bool bUseVectorImplementation = false;	
-				if constexpr (bUseVectorImplementation)
-				{
-					BufferLayerCompositeVector<VectorBlendChannelMasked, VectorLightenChannel, false>(pNew.get(), m_blended.get(), bOnlyOneMip, Args.BlendAlphaSourceChannel);
-				}
-				else
-				{
-					BufferLayerComposite<BlendChannelMasked, LightenChannel, false>(pNew.get(), m_blended.get(), bOnlyOneMip, Args.BlendAlphaSourceChannel);
-				}
-			}
+				FMipmapGenerationSettings settings{};
+				ImageMipmap(Runner->m_pSystem->WorkingMemoryManager, ImageCompressionQuality, pDest.get(), Mask.get(), levelCount, settings);
 
-			bool bApplyColorBlendToAlpha = (Args.flags & OP::ImageLayerArgs::F_APPLY_TO_ALPHA) != 0;
-			bool bUseBlendSourceFromBlendAlpha = (Args.flags & OP::ImageLayerArgs::F_BLENDED_RGB_FROM_ALPHA) != 0;
-			bool bUseMaskFromBlendAlpha = (Args.flags & OP::ImageLayerArgs::F_USE_MASK_FROM_BLENDED);
+				Runner->Release(Mask);
+				Mask = pDest;
+			}
+		}
 
-			if (!bDone && Args.mask)
-			{
-				if (m_base->GetSize() != m_mask->GetSize())
-				{
-					MUTABLE_CPUPROFILER_SCOPE(ImageResize_EmergencyFix);
-					m_mask = ImageResizeLinear(m_imageCompressionQuality, m_mask.get(), m_base->GetSize());
-				}
+		// Create destination data
+		Result = Runner->CloneOrTakeOver(Base);
+		return true;
+	}
 
-				// emergy fix c36adf47-e40d-490f-b709-41142bafad78
-				if (m_mask->GetLODCount() < m_base->GetLODCount())
-				{
-					MUTABLE_CPUPROFILER_SCOPE(ImageLayer_EmergencyFix);
 
-					int32 levelCount = m_base->GetLODCount();
-					ImagePtr pDest = new Image(m_mask->GetSizeX(), m_mask->GetSizeY(),
-						levelCount,
-						m_mask->GetFormat());
+	void FImageLayerTask::DoWork()
+	{
+		MUTABLE_CPUPROFILER_SCOPE(FImageLayerTask);
 
-					SCRATCH_IMAGE_MIPMAP scratch;
-					FMipmapGenerationSettings settings{};
+		// This runs in a worker thread.
 
-					ImageMipmap_PrepareScratch(pDest.get(), m_mask.get(), levelCount, &scratch);
-					ImageMipmap(m_imageCompressionQuality, pDest.get(), m_mask.get(), levelCount,
-						&scratch, settings);
+		bool bOnlyOneMip = false;
+		if (Blended->GetLODCount() < Result->GetLODCount())
+		{
+			bOnlyOneMip = true;
+		}
 
-					m_mask = pDest;
-				}
+		bool bDone = false;
 
-				// Not implemented yet
-				check(!bUseBlendSourceFromBlendAlpha);
+		if (!Mask 
+			&& 
+			// Flags have to match exactly for this optimize case. Other flags are not supported.
+			Args.flags== OP::ImageLayerArgs::F_USE_MASK_FROM_BLENDED
+			&&
+			Args.blendType == uint8(EBlendType::BT_BLEND)
+			&&
+			Args.blendTypeAlpha == uint8(EBlendType::BT_LIGHTEN))
+		{
+			MUTABLE_CPUPROFILER_SCOPE(ImageLayerTask_Optimized);
 
-				switch (EBlendType(Args.blendType))
-				{
-				case EBlendType::BT_NORMAL_COMBINE: ImageNormalCombine(pNew.get(), m_base.get(), m_mask.get(), m_blended.get(), bOnlyOneMip); break;
-				case EBlendType::BT_SOFTLIGHT: BufferLayer<SoftLightChannelMasked, SoftLightChannel, true>(pNew->GetData(), m_base.get(), m_mask.get(), m_blended.get(), bApplyColorBlendToAlpha, bOnlyOneMip); break;
-				case EBlendType::BT_HARDLIGHT: BufferLayer<HardLightChannelMasked, HardLightChannel, true>(pNew->GetData(), m_base.get(), m_mask.get(), m_blended.get(), bApplyColorBlendToAlpha, bOnlyOneMip); break;
-				case EBlendType::BT_BURN: BufferLayer<BurnChannelMasked, BurnChannel, true>(pNew->GetData(), m_base.get(), m_mask.get(), m_blended.get(), bApplyColorBlendToAlpha, bOnlyOneMip); break;
-				case EBlendType::BT_DODGE: BufferLayer<DodgeChannelMasked, DodgeChannel, true>(pNew->GetData(), m_base.get(), m_mask.get(), m_blended.get(), bApplyColorBlendToAlpha, bOnlyOneMip); break;
-				case EBlendType::BT_SCREEN: BufferLayer<ScreenChannelMasked, ScreenChannel, true>(pNew->GetData(), m_base.get(), m_mask.get(), m_blended.get(), bApplyColorBlendToAlpha, bOnlyOneMip); break;
-				case EBlendType::BT_OVERLAY: BufferLayer<OverlayChannelMasked, OverlayChannel, true>(pNew->GetData(), m_base.get(), m_mask.get(), m_blended.get(), bApplyColorBlendToAlpha, bOnlyOneMip); break;
-				case EBlendType::BT_LIGHTEN: BufferLayer<LightenChannelMasked, LightenChannel, true>(pNew->GetData(), m_base.get(), m_mask.get(), m_blended.get(), bApplyColorBlendToAlpha, bOnlyOneMip); break;
-				case EBlendType::BT_MULTIPLY: BufferLayer<MultiplyChannelMasked, MultiplyChannel, true>(pNew->GetData(), m_base.get(), m_mask.get(), m_blended.get(), bApplyColorBlendToAlpha, bOnlyOneMip); break;
-				case EBlendType::BT_BLEND: BufferLayer<BlendChannelMasked, BlendChannel, true>(pNew->GetData(), m_base.get(), m_mask.get(), m_blended.get(), bApplyColorBlendToAlpha, bOnlyOneMip); break;
-				case EBlendType::BT_NONE: break;
-				default: check(false);
-				}
+			// This is a frequent critical-path case because of multilayer projectors.
+			bDone = true;
 
-			}
-			else if (!bDone && bUseMaskFromBlendAlpha)
+			constexpr bool bUseVectorImplementation = false;	
+			if constexpr (bUseVectorImplementation)
 			{
-				// Not implemented yet
-				check(!bUseBlendSourceFromBlendAlpha);
-
-				// Apply blend without to RGB using mask in blended alpha
-				switch (EBlendType(Args.blendType))
-				{
-				case EBlendType::BT_NORMAL_COMBINE: check(false); break;
-				case EBlendType::BT_SOFTLIGHT: BufferLayerEmbeddedMask<SoftLightChannelMasked, SoftLightChannel, false>(pNew->GetData(), pNew.get(), m_blended.get(), bApplyColorBlendToAlpha, bOnlyOneMip); break;
-				case EBlendType::BT_HARDLIGHT: BufferLayerEmbeddedMask<HardLightChannelMasked, HardLightChannel, false>(pNew->GetData(), pNew.get(), m_blended.get(), bApplyColorBlendToAlpha, bOnlyOneMip); break;
-				case EBlendType::BT_BURN: BufferLayerEmbeddedMask<BurnChannelMasked, BurnChannel, false>(pNew->GetData(), pNew.get(), m_blended.get(), bApplyColorBlendToAlpha, bOnlyOneMip); break;
-				case EBlendType::BT_DODGE: BufferLayerEmbeddedMask<DodgeChannelMasked, DodgeChannel, false>(pNew->GetData(), pNew.get(), m_blended.get(), bApplyColorBlendToAlpha, bOnlyOneMip); break;
-				case EBlendType::BT_SCREEN: BufferLayerEmbeddedMask<ScreenChannelMasked, ScreenChannel, false>(pNew->GetData(), pNew.get(), m_blended.get(), bApplyColorBlendToAlpha, bOnlyOneMip); break;
-				case EBlendType::BT_OVERLAY: BufferLayerEmbeddedMask<OverlayChannelMasked, OverlayChannel, false>(pNew->GetData(), pNew.get(), m_blended.get(), bApplyColorBlendToAlpha, bOnlyOneMip); break;
-				case EBlendType::BT_LIGHTEN: BufferLayerEmbeddedMask<LightenChannelMasked, LightenChannel, false>(pNew->GetData(), pNew.get(), m_blended.get(), bApplyColorBlendToAlpha, bOnlyOneMip); break;
-				case EBlendType::BT_MULTIPLY: BufferLayerEmbeddedMask<MultiplyChannelMasked, MultiplyChannel, false>(pNew->GetData(), pNew.get(), m_blended.get(), bApplyColorBlendToAlpha, bOnlyOneMip); break;
-				case EBlendType::BT_BLEND: BufferLayerEmbeddedMask<BlendChannelMasked, BlendChannel, false>(pNew->GetData(), pNew.get(), m_blended.get(), bApplyColorBlendToAlpha, bOnlyOneMip); break;
-				case EBlendType::BT_NONE: break;
-				default: check(false);
-				}
+				BufferLayerCompositeVector<VectorBlendChannelMasked, VectorLightenChannel, false>(Result.get(), Blended.get(), bOnlyOneMip, Args.BlendAlphaSourceChannel);
 			}
-			else if (!bDone)
+			else
 			{
-				// Apply blend without mask to RGB
-				switch (EBlendType(Args.blendType))
-				{
-				case EBlendType::BT_NORMAL_COMBINE: ImageNormalCombine(pNew.get(), m_base.get(), m_blended.get(), bOnlyOneMip); check(!bUseBlendSourceFromBlendAlpha);  break;
-				case EBlendType::BT_SOFTLIGHT: BufferLayer<SoftLightChannel, true>(pNew.get(), m_base.get(), m_blended.get(), bApplyColorBlendToAlpha, bOnlyOneMip, bUseBlendSourceFromBlendAlpha); break;
-				case EBlendType::BT_HARDLIGHT: BufferLayer<HardLightChannel, true>(pNew.get(), m_base.get(), m_blended.get(), bApplyColorBlendToAlpha, bOnlyOneMip, bUseBlendSourceFromBlendAlpha); break;
-				case EBlendType::BT_BURN: BufferLayer<BurnChannel, true>(pNew.get(), m_base.get(), m_blended.get(), bApplyColorBlendToAlpha, bOnlyOneMip, bUseBlendSourceFromBlendAlpha); break;
-				case EBlendType::BT_DODGE: BufferLayer<DodgeChannel, true>(pNew.get(), m_base.get(), m_blended.get(), bApplyColorBlendToAlpha, bOnlyOneMip, bUseBlendSourceFromBlendAlpha); break;
-				case EBlendType::BT_SCREEN: BufferLayer<ScreenChannel, true>(pNew.get(), m_base.get(), m_blended.get(), bApplyColorBlendToAlpha, bOnlyOneMip, bUseBlendSourceFromBlendAlpha); break;
-				case EBlendType::BT_OVERLAY: BufferLayer<OverlayChannel, true>(pNew.get(), m_base.get(), m_blended.get(), bApplyColorBlendToAlpha, bOnlyOneMip, bUseBlendSourceFromBlendAlpha); break;
-				case EBlendType::BT_LIGHTEN: BufferLayer<LightenChannel, true>(pNew.get(), m_base.get(), m_blended.get(), bApplyColorBlendToAlpha, bOnlyOneMip, bUseBlendSourceFromBlendAlpha); break;
-				case EBlendType::BT_MULTIPLY: BufferLayer<MultiplyChannel, true>(pNew.get(), m_base.get(), m_blended.get(), bApplyColorBlendToAlpha, bOnlyOneMip, bUseBlendSourceFromBlendAlpha); break;
-				case EBlendType::BT_BLEND: BufferLayer<BlendChannel, true>(pNew.get(), m_base.get(), m_blended.get(), bApplyColorBlendToAlpha, bOnlyOneMip, bUseBlendSourceFromBlendAlpha); break;
-				case EBlendType::BT_NONE: break;
-				default: check(false);
-				}
+				BufferLayerComposite<BlendChannelMasked, LightenChannel, false>(Result.get(), Blended.get(), bOnlyOneMip, Args.BlendAlphaSourceChannel);
 			}
+		}
 
-			// Apply the separate blend operation for alpha
-			if (!bDone && !bApplyColorBlendToAlpha)
+		bool bApplyColorBlendToAlpha = (Args.flags & OP::ImageLayerArgs::F_APPLY_TO_ALPHA) != 0;
+		bool bUseBlendSourceFromBlendAlpha = (Args.flags & OP::ImageLayerArgs::F_BLENDED_RGB_FROM_ALPHA) != 0;
+		bool bUseMaskFromBlendAlpha = (Args.flags & OP::ImageLayerArgs::F_USE_MASK_FROM_BLENDED);
+
+		if (!bDone && Args.mask)
+		{
+			// Not implemented yet
+			check(!bUseBlendSourceFromBlendAlpha);
+
+			// TODO: in-place
+
+			switch (EBlendType(Args.blendType))
 			{
-				// Separate alpha operation ignores the mask.
-				switch (EBlendType(Args.blendTypeAlpha))
-				{
-				case EBlendType::BT_SOFTLIGHT: BufferLayerInPlace<SoftLightChannel, false, 1>(pNew.get(), m_blended.get(), bOnlyOneMip, 3, Args.BlendAlphaSourceChannel); break;
-				case EBlendType::BT_HARDLIGHT: BufferLayerInPlace<HardLightChannel, false, 1>(pNew.get(), m_blended.get(), bOnlyOneMip, 3, Args.BlendAlphaSourceChannel); break;
-				case EBlendType::BT_BURN: BufferLayerInPlace<BurnChannel, false, 1>(pNew.get(), m_blended.get(), bOnlyOneMip, 3, Args.BlendAlphaSourceChannel); break;
-				case EBlendType::BT_DODGE: BufferLayerInPlace<DodgeChannel, false, 1>(pNew.get(), m_blended.get(), bOnlyOneMip, 3, Args.BlendAlphaSourceChannel); break;
-				case EBlendType::BT_SCREEN: BufferLayerInPlace<ScreenChannel, false, 1>(pNew.get(), m_blended.get(), bOnlyOneMip, 3, Args.BlendAlphaSourceChannel); break;
-				case EBlendType::BT_OVERLAY: BufferLayerInPlace<OverlayChannel, false, 1>(pNew.get(), m_blended.get(), bOnlyOneMip, 3, Args.BlendAlphaSourceChannel); break;
-				case EBlendType::BT_LIGHTEN: BufferLayerInPlace<LightenChannel, false, 1>(pNew.get(), m_blended.get(), bOnlyOneMip, 3, Args.BlendAlphaSourceChannel); break;
-				case EBlendType::BT_MULTIPLY: BufferLayerInPlace<MultiplyChannel, false, 1>(pNew.get(), m_blended.get(), bOnlyOneMip, 3, Args.BlendAlphaSourceChannel); break;
-				case EBlendType::BT_BLEND: BufferLayerInPlace<BlendChannel, false, 1>(pNew.get(), m_blended.get(), bOnlyOneMip, 3, Args.BlendAlphaSourceChannel); break;
-				case EBlendType::BT_NONE: break;
-				default: check(false);
-				}
+			case EBlendType::BT_NORMAL_COMBINE: ImageNormalCombine(Result.get(), Result.get(), Mask.get(), Blended.get(), bOnlyOneMip); break;
+			case EBlendType::BT_SOFTLIGHT: BufferLayer<SoftLightChannelMasked, SoftLightChannel, true>(Result->GetData(), Result.get(), Mask.get(), Blended.get(), bApplyColorBlendToAlpha, bOnlyOneMip); break;
+			case EBlendType::BT_HARDLIGHT: BufferLayer<HardLightChannelMasked, HardLightChannel, true>(Result->GetData(), Result.get(), Mask.get(), Blended.get(), bApplyColorBlendToAlpha, bOnlyOneMip); break;
+			case EBlendType::BT_BURN: BufferLayer<BurnChannelMasked, BurnChannel, true>(Result->GetData(), Result.get(), Mask.get(), Blended.get(), bApplyColorBlendToAlpha, bOnlyOneMip); break;
+			case EBlendType::BT_DODGE: BufferLayer<DodgeChannelMasked, DodgeChannel, true>(Result->GetData(), Result.get(), Mask.get(), Blended.get(), bApplyColorBlendToAlpha, bOnlyOneMip); break;
+			case EBlendType::BT_SCREEN: BufferLayer<ScreenChannelMasked, ScreenChannel, true>(Result->GetData(), Result.get(), Mask.get(), Blended.get(), bApplyColorBlendToAlpha, bOnlyOneMip); break;
+			case EBlendType::BT_OVERLAY: BufferLayer<OverlayChannelMasked, OverlayChannel, true>(Result->GetData(), Result.get(), Mask.get(), Blended.get(), bApplyColorBlendToAlpha, bOnlyOneMip); break;
+			case EBlendType::BT_LIGHTEN: BufferLayer<LightenChannelMasked, LightenChannel, true>(Result->GetData(), Result.get(), Mask.get(), Blended.get(), bApplyColorBlendToAlpha, bOnlyOneMip); break;
+			case EBlendType::BT_MULTIPLY: BufferLayer<MultiplyChannelMasked, MultiplyChannel, true>(Result->GetData(), Result.get(), Mask.get(), Blended.get(), bApplyColorBlendToAlpha, bOnlyOneMip); break;
+			case EBlendType::BT_BLEND: BufferLayer<BlendChannelMasked, BlendChannel, true>(Result->GetData(), Result.get(), Mask.get(), Blended.get(), bApplyColorBlendToAlpha, bOnlyOneMip); break;
+			case EBlendType::BT_NONE: break;
+			default: check(false);
 			}
 
-			if (bOnlyOneMip)
-			{
-				MUTABLE_CPUPROFILER_SCOPE(ImageLayer_MipFix);
-				FMipmapGenerationSettings DummyMipSettings{};
-				ImageMipmapInPlace(m_imageCompressionQuality, pNew.get(), DummyMipSettings);
+		}
+		else if (!bDone && bUseMaskFromBlendAlpha)
+		{
+			// Not implemented yet
+			check(!bUseBlendSourceFromBlendAlpha);
+
+			// Apply blend without to RGB using mask in blended alpha
+			switch (EBlendType(Args.blendType))
+			{
+			case EBlendType::BT_NORMAL_COMBINE: check(false); break;
+			case EBlendType::BT_SOFTLIGHT: BufferLayerEmbeddedMask<SoftLightChannelMasked, SoftLightChannel, false>(Result->GetData(), Result.get(), Blended.get(), bApplyColorBlendToAlpha, bOnlyOneMip); break;
+			case EBlendType::BT_HARDLIGHT: BufferLayerEmbeddedMask<HardLightChannelMasked, HardLightChannel, false>(Result->GetData(), Result.get(), Blended.get(), bApplyColorBlendToAlpha, bOnlyOneMip); break;
+			case EBlendType::BT_BURN: BufferLayerEmbeddedMask<BurnChannelMasked, BurnChannel, false>(Result->GetData(), Result.get(), Blended.get(), bApplyColorBlendToAlpha, bOnlyOneMip); break;
+			case EBlendType::BT_DODGE: BufferLayerEmbeddedMask<DodgeChannelMasked, DodgeChannel, false>(Result->GetData(), Result.get(), Blended.get(), bApplyColorBlendToAlpha, bOnlyOneMip); break;
+			case EBlendType::BT_SCREEN: BufferLayerEmbeddedMask<ScreenChannelMasked, ScreenChannel, false>(Result->GetData(), Result.get(), Blended.get(), bApplyColorBlendToAlpha, bOnlyOneMip); break;
+			case EBlendType::BT_OVERLAY: BufferLayerEmbeddedMask<OverlayChannelMasked, OverlayChannel, false>(Result->GetData(), Result.get(), Blended.get(), bApplyColorBlendToAlpha, bOnlyOneMip); break;
+			case EBlendType::BT_LIGHTEN: BufferLayerEmbeddedMask<LightenChannelMasked, LightenChannel, false>(Result->GetData(), Result.get(), Blended.get(), bApplyColorBlendToAlpha, bOnlyOneMip); break;
+			case EBlendType::BT_MULTIPLY: BufferLayerEmbeddedMask<MultiplyChannelMasked, MultiplyChannel, false>(Result->GetData(), Result.get(), Blended.get(), bApplyColorBlendToAlpha, bOnlyOneMip); break;
+			case EBlendType::BT_BLEND: BufferLayerEmbeddedMask<BlendChannelMasked, BlendChannel, false>(Result->GetData(), Result.get(), Blended.get(), bApplyColorBlendToAlpha, bOnlyOneMip); break;
+			case EBlendType::BT_NONE: break;
+			default: check(false);
+			}
+		}
+		else if (!bDone)
+		{
+			// Apply blend without mask to RGB
+			switch (EBlendType(Args.blendType))
+			{
+			case EBlendType::BT_NORMAL_COMBINE: ImageNormalCombine(Result.get(), Result.get(), Blended.get(), bOnlyOneMip); check(!bUseBlendSourceFromBlendAlpha);  break;
+			case EBlendType::BT_SOFTLIGHT: BufferLayer<SoftLightChannel, true>(Result.get(), Result.get(), Blended.get(), bApplyColorBlendToAlpha, bOnlyOneMip, bUseBlendSourceFromBlendAlpha); break;
+			case EBlendType::BT_HARDLIGHT: BufferLayer<HardLightChannel, true>(Result.get(), Result.get(), Blended.get(), bApplyColorBlendToAlpha, bOnlyOneMip, bUseBlendSourceFromBlendAlpha); break;
+			case EBlendType::BT_BURN: BufferLayer<BurnChannel, true>(Result.get(), Result.get(), Blended.get(), bApplyColorBlendToAlpha, bOnlyOneMip, bUseBlendSourceFromBlendAlpha); break;
+			case EBlendType::BT_DODGE: BufferLayer<DodgeChannel, true>(Result.get(), Result.get(), Blended.get(), bApplyColorBlendToAlpha, bOnlyOneMip, bUseBlendSourceFromBlendAlpha); break;
+			case EBlendType::BT_SCREEN: BufferLayer<ScreenChannel, true>(Result.get(), Result.get(), Blended.get(), bApplyColorBlendToAlpha, bOnlyOneMip, bUseBlendSourceFromBlendAlpha); break;
+			case EBlendType::BT_OVERLAY: BufferLayer<OverlayChannel, true>(Result.get(), Result.get(), Blended.get(), bApplyColorBlendToAlpha, bOnlyOneMip, bUseBlendSourceFromBlendAlpha); break;
+			case EBlendType::BT_LIGHTEN: BufferLayer<LightenChannel, true>(Result.get(), Result.get(), Blended.get(), bApplyColorBlendToAlpha, bOnlyOneMip, bUseBlendSourceFromBlendAlpha); break;
+			case EBlendType::BT_MULTIPLY: BufferLayer<MultiplyChannel, true>(Result.get(), Result.get(), Blended.get(), bApplyColorBlendToAlpha, bOnlyOneMip, bUseBlendSourceFromBlendAlpha); break;
+			case EBlendType::BT_BLEND: BufferLayer<BlendChannel, true>(Result.get(), Result.get(), Blended.get(), bApplyColorBlendToAlpha, bOnlyOneMip, bUseBlendSourceFromBlendAlpha); break;
+			case EBlendType::BT_NONE: break;
+			default: check(false);
 			}
 		}
 
-		if (InitialFormat != pNew->GetFormat())
+		// Apply the separate blend operation for alpha
+		if (!bDone && !bApplyColorBlendToAlpha)
 		{
-			pNew = ImagePixelFormat(m_imageCompressionQuality, pNew.get(), InitialFormat);
+			// Separate alpha operation ignores the mask.
+			switch (EBlendType(Args.blendTypeAlpha))
+			{
+			case EBlendType::BT_SOFTLIGHT: BufferLayerInPlace<SoftLightChannel, false, 1>(Result.get(), Blended.get(), bOnlyOneMip, 3, Args.BlendAlphaSourceChannel); break;
+			case EBlendType::BT_HARDLIGHT: BufferLayerInPlace<HardLightChannel, false, 1>(Result.get(), Blended.get(), bOnlyOneMip, 3, Args.BlendAlphaSourceChannel); break;
+			case EBlendType::BT_BURN: BufferLayerInPlace<BurnChannel, false, 1>(Result.get(), Blended.get(), bOnlyOneMip, 3, Args.BlendAlphaSourceChannel); break;
+			case EBlendType::BT_DODGE: BufferLayerInPlace<DodgeChannel, false, 1>(Result.get(), Blended.get(), bOnlyOneMip, 3, Args.BlendAlphaSourceChannel); break;
+			case EBlendType::BT_SCREEN: BufferLayerInPlace<ScreenChannel, false, 1>(Result.get(), Blended.get(), bOnlyOneMip, 3, Args.BlendAlphaSourceChannel); break;
+			case EBlendType::BT_OVERLAY: BufferLayerInPlace<OverlayChannel, false, 1>(Result.get(), Blended.get(), bOnlyOneMip, 3, Args.BlendAlphaSourceChannel); break;
+			case EBlendType::BT_LIGHTEN: BufferLayerInPlace<LightenChannel, false, 1>(Result.get(), Blended.get(), bOnlyOneMip, 3, Args.BlendAlphaSourceChannel); break;
+			case EBlendType::BT_MULTIPLY: BufferLayerInPlace<MultiplyChannel, false, 1>(Result.get(), Blended.get(), bOnlyOneMip, 3, Args.BlendAlphaSourceChannel); break;
+			case EBlendType::BT_BLEND: BufferLayerInPlace<BlendChannel, false, 1>(Result.get(), Blended.get(), bOnlyOneMip, 3, Args.BlendAlphaSourceChannel); break;
+			case EBlendType::BT_NONE: break;
+			default: check(false);
+			}
 		}
 
-		// Reset relevancy map.
-		pNew->m_flags &= ~Image::EImageFlags::IF_HAS_RELEVANCY_MAP;
+		if (bOnlyOneMip)
+		{
+			MUTABLE_CPUPROFILER_SCOPE(ImageLayer_MipFix);
+			FMipmapGenerationSettings DummyMipSettings{};
+			ImageMipmapInPlace(ImageCompressionQuality, Result.get(), DummyMipSettings);
+		}
 
-		// Free resources
-		m_base = nullptr;
-		m_blended = nullptr;
-		m_mask = nullptr;
-		Result = pNew;
+		// Reset relevancy map.
+		Result->m_flags &= ~Image::EImageFlags::IF_HAS_RELEVANCY_MAP;
 	}
 
 
-	//---------------------------------------------------------------------------------------------
-	void FImageLayerTask::Complete( CodeRunner* runner )
+	void FImageLayerTask::Complete( CodeRunner* Runner )
 	{
-		// This runs in the runner thread
-		runner->GetMemory().SetImage(Op, Result);
+		// This runs in the Runner thread
+		Runner->Release(Blended);
+		Runner->Release(Mask);
+
+		// If no shortcut was taken
+		if (Result)
+		{
+			if (InitialFormat != Result->GetFormat())
+			{
+				Ptr<Image> Formatted = Runner->CreateImage(Result->GetSizeX(), Result->GetSizeY(), Result->GetLODCount(), InitialFormat, EInitializationType::NotInitialized);
+				bool bSuccess = false;
+				ImagePixelFormat(bSuccess, ImageCompressionQuality, Formatted.get(), Result.get());
+				check(bSuccess);
+
+				Runner->Release(Result);
+				Result = Formatted;
+			}
+
+			Runner->StoreImage(Op, Result);
+		}
 	}
 
 
@@ -979,88 +1009,107 @@ namespace mu
 	class FImageLayerColourTask : public CodeRunner::FIssuedTask
 	{
 	public:
-		FImageLayerColourTask(const FScheduledOp&, FProgramCache&, const OP::ImageLayerColourArgs&, int imageCompressionQuality);
+		FImageLayerColourTask(const FScheduledOp& InOp, const OP::ImageLayerColourArgs& InArgs)
+			: FIssuedTask(InOp), Args(InArgs)
+		{}
 
 		// FIssuedTask interface
+		bool Prepare(CodeRunner*, bool& bOutFailed) override;
 		void DoWork() override;
-		void Complete(CodeRunner* staged) override;
+		void Complete(CodeRunner*) override;
 
 	private:
-		int m_imageCompressionQuality = 0;
+		int32 ImageCompressionQuality = 0;
 		OP::ImageLayerColourArgs Args;
-		FVector4f m_col;
-		Ptr<const Image> m_base;
-		Ptr<const Image> m_mask;
-		Ptr<const Image> Result;
+		FVector4f Color;
+		Ptr<const Image> Mask;
+		Ptr<Image> Result;
+		EImageFormat InitialFormat = EImageFormat::IF_NONE;
 	};
 
 
-	//---------------------------------------------------------------------------------------------
-	FImageLayerColourTask::FImageLayerColourTask(const FScheduledOp& InOp, FProgramCache& Memory, const OP::ImageLayerColourArgs& InArgs, int imageCompressionQuality)
-		: FIssuedTask(InOp)
+	bool FImageLayerColourTask::Prepare(CodeRunner* Runner, bool& bOutFailed)
 	{
-		m_imageCompressionQuality = imageCompressionQuality;
-		Args = InArgs;
-		m_base = Memory.GetImage({ Args.base, InOp.ExecutionIndex, InOp.ExecutionOptions });
-		check(!m_base || m_base->GetFormat() < EImageFormat::IF_COUNT);
+		// This runs in the mutable Runner thread
+		MUTABLE_CPUPROFILER_SCOPE(FImageLayerColourTask_Prepare);
+		bOutFailed = false;
+
+		ImageCompressionQuality = Runner->m_pSettings->ImageCompressionQuality;
+
+		Ptr<const Image> Base = Runner->LoadImage({ Args.base, Op.ExecutionIndex, Op.ExecutionOptions });
+		check(!Base || Base->GetFormat() < EImageFormat::IF_COUNT);
 
-		m_col = Memory.GetColour({ Args.colour, InOp.ExecutionIndex, 0 });
+		Color = Runner->LoadColor({ Args.colour, Op.ExecutionIndex, 0 });
 		if (Args.mask)
 		{
-			m_mask = Memory.GetImage({ Args.mask, InOp.ExecutionIndex, InOp.ExecutionOptions });
-			check(!m_mask || m_mask->GetFormat() < EImageFormat::IF_COUNT);
+			Mask = Runner->LoadImage({ Args.mask, Op.ExecutionIndex, Op.ExecutionOptions });
+			check(!Mask || Mask->GetFormat() < EImageFormat::IF_COUNT);
 		}
-	}
-
 
-	//---------------------------------------------------------------------------------------------
-	void FImageLayerColourTask::DoWork()
-	{
-		MUTABLE_CPUPROFILER_SCOPE(FImageLayerColourTask);
+		// Shortcuts
+		if (!Base)
+		{
+			Runner->Release(Mask);
+			Runner->StoreImage(Op, Base);
+			return false;
+		}
 
-		// Warning: This runs in a worker thread.
-		Ptr<Image> pNew;
-		if (!m_base)
+		bool bValid = Base->GetSizeX() > 0 && Base->GetSizeY() > 0;
+		if (!bValid)
 		{
-			return;
+			Runner->Release(Mask);
+			Runner->StoreImage(Op, Base);
+			return false;
 		}
 
-		EImageFormat InitialFormat = m_base->GetFormat();
+		// Fix input data
+		InitialFormat = Base->GetFormat();
 		check(InitialFormat < EImageFormat::IF_COUNT);
 
-		if (Args.mask && m_mask)
+		if (Args.mask && Mask)
 		{
-			if (m_base->GetSize() != m_mask->GetSize())
+			if (Base->GetSize() != Mask->GetSize())
 			{
 				MUTABLE_CPUPROFILER_SCOPE(ImageResize_EmergencyFixSize);
-				m_mask = ImageResizeLinear(m_imageCompressionQuality, m_mask.get(), m_base->GetSize());
+				Ptr<Image> Resized = Runner->CreateImage(Base->GetSizeX(), Base->GetSizeY(), 1, Mask->GetFormat(), EInitializationType::NotInitialized);
+				ImageResizeLinear(Resized.get(), ImageCompressionQuality, Mask.get());
+				Runner->Release(Mask);
+				Mask = Resized;
 			}
 
 			// emergy fix c36adf47-e40d-490f-b709-41142bafad78
-			if (m_mask->GetLODCount() < m_base->GetLODCount())
+			if (Mask->GetLODCount() < Base->GetLODCount())
 			{
 				MUTABLE_CPUPROFILER_SCOPE(ImageResize_EmergencyFixMips);
-				int32 levelCount = m_base->GetLODCount();
-				Ptr<Image> pDest = new Image(m_mask->GetSizeX(), m_mask->GetSizeY(), levelCount, m_mask->GetFormat());
+				int32 levelCount = Base->GetLODCount();
+				Ptr<Image> pDest = Runner->CreateImage(Mask->GetSizeX(), Mask->GetSizeY(), levelCount, Mask->GetFormat());
 
-				SCRATCH_IMAGE_MIPMAP scratch;
 				FMipmapGenerationSettings settings{};
+				ImageMipmap(Runner->m_pSystem->WorkingMemoryManager, ImageCompressionQuality, pDest.get(), Mask.get(), levelCount, settings);
 
-				ImageMipmap_PrepareScratch(pDest.get(), m_mask.get(), levelCount, &scratch);
-				ImageMipmap(m_imageCompressionQuality, pDest.get(), m_mask.get(), levelCount, &scratch, settings);
-				m_mask = pDest;
+				Runner->Release(Mask);
+				Mask = pDest;
 			}
 		}
 
+		// Create destination data
+		Result = Runner->CloneOrTakeOver(Base);
+		return true;
+	}
+
+
+	void FImageLayerColourTask::DoWork()
+	{
+		// This runs in a worker thread
+		MUTABLE_CPUPROFILER_SCOPE(FImageLayerColourTask);
+
 		bool bOnlyOneMip = false;
 
 		// Does it apply to colour?
 		if (EBlendType(Args.blendType) != EBlendType::BT_NONE)
 		{
 			// TODO: It could be done "in-place"
-			pNew = mu::CloneOrTakeOver<mu::Image>(m_base.get());
-
-			if (Args.mask && m_mask)
+			if (Args.mask && Mask)
 			{			
 				// Not implemented yet
 				check(Args.flags==0);
@@ -1068,16 +1117,16 @@ namespace mu
 				// \todo: precalculated tables for softlight
 				switch (EBlendType(Args.blendType))
 				{
-				case EBlendType::BT_NORMAL_COMBINE: ImageNormalCombine(pNew.get(), m_base.get(), m_mask.get(), m_col); break;
-				case EBlendType::BT_SOFTLIGHT: BufferLayerColour<SoftLightChannelMasked, SoftLightChannel>(pNew->GetData(), m_base.get(), m_mask.get(), m_col); break;
-				case EBlendType::BT_HARDLIGHT: BufferLayerColour<HardLightChannelMasked, HardLightChannel>(pNew->GetData(), m_base.get(), m_mask.get(), m_col); break;
-				case EBlendType::BT_BURN: BufferLayerColour<BurnChannelMasked, BurnChannel>(pNew->GetData(), m_base.get(), m_mask.get(), m_col); break;
-				case EBlendType::BT_DODGE: BufferLayerColour<DodgeChannelMasked, DodgeChannel>(pNew->GetData(), m_base.get(), m_mask.get(), m_col); break;
-				case EBlendType::BT_SCREEN: BufferLayerColour<ScreenChannelMasked, ScreenChannel>(pNew->GetData(), m_base.get(), m_mask.get(), m_col); break;
-				case EBlendType::BT_OVERLAY: BufferLayerColour<OverlayChannelMasked, OverlayChannel>(pNew->GetData(), m_base.get(), m_mask.get(), m_col); break;
-				case EBlendType::BT_LIGHTEN: BufferLayerColour<LightenChannelMasked, LightenChannel>(pNew->GetData(), m_base.get(), m_mask.get(), m_col); break;
-				case EBlendType::BT_MULTIPLY: BufferLayerColour<MultiplyChannelMasked, MultiplyChannel>(pNew->GetData(), m_base.get(), m_mask.get(), m_col); break;
-				case EBlendType::BT_BLEND: BufferLayerColour<BlendChannelMasked, BlendChannel>(pNew->GetData(), m_base.get(), m_mask.get(), m_col); break;
+				case EBlendType::BT_NORMAL_COMBINE: ImageNormalCombine(Result.get(), Result.get(), Mask.get(), Color); break;
+				case EBlendType::BT_SOFTLIGHT: BufferLayerColour<SoftLightChannelMasked, SoftLightChannel>(Result->GetData(), Result.get(), Mask.get(), Color); break;
+				case EBlendType::BT_HARDLIGHT: BufferLayerColour<HardLightChannelMasked, HardLightChannel>(Result->GetData(), Result.get(), Mask.get(), Color); break;
+				case EBlendType::BT_BURN: BufferLayerColour<BurnChannelMasked, BurnChannel>(Result->GetData(), Result.get(), Mask.get(), Color); break;
+				case EBlendType::BT_DODGE: BufferLayerColour<DodgeChannelMasked, DodgeChannel>(Result->GetData(), Result.get(), Mask.get(), Color); break;
+				case EBlendType::BT_SCREEN: BufferLayerColour<ScreenChannelMasked, ScreenChannel>(Result->GetData(), Result.get(), Mask.get(), Color); break;
+				case EBlendType::BT_OVERLAY: BufferLayerColour<OverlayChannelMasked, OverlayChannel>(Result->GetData(), Result.get(), Mask.get(), Color); break;
+				case EBlendType::BT_LIGHTEN: BufferLayerColour<LightenChannelMasked, LightenChannel>(Result->GetData(), Result.get(), Mask.get(), Color); break;
+				case EBlendType::BT_MULTIPLY: BufferLayerColour<MultiplyChannelMasked, MultiplyChannel>(Result->GetData(), Result.get(), Mask.get(), Color); break;
+				case EBlendType::BT_BLEND: BufferLayerColour<BlendChannelMasked, BlendChannel>(Result->GetData(), Result.get(), Mask.get(), Color); break;
 				default: check(false);
 				}
 
@@ -1089,14 +1138,14 @@ namespace mu
 					switch (EBlendType(Args.blendType))
 					{
 					case EBlendType::BT_NORMAL_COMBINE: check(false); break;
-					case EBlendType::BT_SOFTLIGHT: BufferLayerColourFromAlpha<SoftLightChannel>(pNew.get(), m_base.get(), m_col); break;
-					case EBlendType::BT_HARDLIGHT: BufferLayerColourFromAlpha<HardLightChannel>(pNew.get(), m_base.get(), m_col); break;
-					case EBlendType::BT_BURN: BufferLayerColourFromAlpha<BurnChannel>(pNew.get(), m_base.get(), m_col); break;
-					case EBlendType::BT_DODGE: BufferLayerColourFromAlpha<DodgeChannel>(pNew.get(), m_base.get(), m_col); break;
-					case EBlendType::BT_SCREEN: BufferLayerColourFromAlpha<ScreenChannel>(pNew.get(), m_base.get(), m_col); break;
-					case EBlendType::BT_OVERLAY: BufferLayerColourFromAlpha<OverlayChannel>(pNew.get(), m_base.get(), m_col); break;
-					case EBlendType::BT_LIGHTEN: BufferLayerColourFromAlpha<LightenChannel>(pNew.get(), m_base.get(), m_col); break;
-					case EBlendType::BT_MULTIPLY: BufferLayerColourFromAlpha<MultiplyChannel>(pNew.get(), m_base.get(), m_col); break;
+					case EBlendType::BT_SOFTLIGHT: BufferLayerColourFromAlpha<SoftLightChannel>(Result.get(), Result.get(), Color); break;
+					case EBlendType::BT_HARDLIGHT: BufferLayerColourFromAlpha<HardLightChannel>(Result.get(), Result.get(), Color); break;
+					case EBlendType::BT_BURN: BufferLayerColourFromAlpha<BurnChannel>(Result.get(), Result.get(), Color); break;
+					case EBlendType::BT_DODGE: BufferLayerColourFromAlpha<DodgeChannel>(Result.get(), Result.get(), Color); break;
+					case EBlendType::BT_SCREEN: BufferLayerColourFromAlpha<ScreenChannel>(Result.get(), Result.get(), Color); break;
+					case EBlendType::BT_OVERLAY: BufferLayerColourFromAlpha<OverlayChannel>(Result.get(), Result.get(), Color); break;
+					case EBlendType::BT_LIGHTEN: BufferLayerColourFromAlpha<LightenChannel>(Result.get(), Result.get(), Color); break;
+					case EBlendType::BT_MULTIPLY: BufferLayerColourFromAlpha<MultiplyChannel>(Result.get(), Result.get(), Color); break;
 					case EBlendType::BT_BLEND: check(false); break;
 					default: check(false);
 					}
@@ -1105,16 +1154,16 @@ namespace mu
 				{
 					switch (EBlendType(Args.blendType))
 					{
-					case EBlendType::BT_NORMAL_COMBINE: ImageNormalCombine(pNew.get(), m_base.get(), m_col); break;
-					case EBlendType::BT_SOFTLIGHT: BufferLayerColour<SoftLightChannel>(pNew.get(), m_base.get(), m_col); break;
-					case EBlendType::BT_HARDLIGHT: BufferLayerColour<HardLightChannel>(pNew.get(), m_base.get(), m_col); break;
-					case EBlendType::BT_BURN: BufferLayerColour<BurnChannel>(pNew.get(), m_base.get(), m_col); break;
-					case EBlendType::BT_DODGE: BufferLayerColour<DodgeChannel>(pNew.get(), m_base.get(), m_col); break;
-					case EBlendType::BT_SCREEN: BufferLayerColour<ScreenChannel>(pNew.get(), m_base.get(), m_col); break;
-					case EBlendType::BT_OVERLAY: BufferLayerColour<OverlayChannel>(pNew.get(), m_base.get(), m_col); break;
-					case EBlendType::BT_LIGHTEN: BufferLayerColour<LightenChannel>(pNew.get(), m_base.get(), m_col); break;
-					case EBlendType::BT_MULTIPLY: BufferLayerColour<MultiplyChannel>(pNew.get(), m_base.get(), m_col); break;
-					case EBlendType::BT_BLEND: pNew->FillColour(m_col); break;
+					case EBlendType::BT_NORMAL_COMBINE: ImageNormalCombine(Result.get(), Result.get(), Color); break;
+					case EBlendType::BT_SOFTLIGHT: BufferLayerColour<SoftLightChannel>(Result.get(), Result.get(), Color); break;
+					case EBlendType::BT_HARDLIGHT: BufferLayerColour<HardLightChannel>(Result.get(), Result.get(), Color); break;
+					case EBlendType::BT_BURN: BufferLayerColour<BurnChannel>(Result.get(), Result.get(), Color); break;
+					case EBlendType::BT_DODGE: BufferLayerColour<DodgeChannel>(Result.get(), Result.get(), Color); break;
+					case EBlendType::BT_SCREEN: BufferLayerColour<ScreenChannel>(Result.get(), Result.get(), Color); break;
+					case EBlendType::BT_OVERLAY: BufferLayerColour<OverlayChannel>(Result.get(), Result.get(), Color); break;
+					case EBlendType::BT_LIGHTEN: BufferLayerColour<LightenChannel>(Result.get(), Result.get(), Color); break;
+					case EBlendType::BT_MULTIPLY: BufferLayerColour<MultiplyChannel>(Result.get(), Result.get(), Color); break;
+					case EBlendType::BT_BLEND: Result->FillColour(Color); break;
 					default: check(false);
 					}
 				}
@@ -1124,26 +1173,21 @@ namespace mu
 		// Does it apply to alpha?
 		if (EBlendType(Args.blendTypeAlpha) != EBlendType::BT_NONE)
 		{
-			if (!pNew)
-			{
-				pNew = mu::CloneOrTakeOver<mu::Image>(m_base.get());
-			}
-
-			if (Args.mask && m_mask)
+			if (Args.mask && Mask)
 			{
 				// \todo: precalculated tables for softlight
 				switch (EBlendType(Args.blendTypeAlpha))
 				{
 				case EBlendType::BT_NORMAL_COMBINE: check(false); break;
-				case EBlendType::BT_SOFTLIGHT: BufferLayerColourInPlace<SoftLightChannelMasked, SoftLightChannel, 1>(pNew.get(), m_mask.get(), m_col, bOnlyOneMip, 3, Args.BlendAlphaSourceChannel); break;
-				case EBlendType::BT_HARDLIGHT: BufferLayerColourInPlace<HardLightChannelMasked, HardLightChannel, 1>(pNew.get(), m_mask.get(), m_col, bOnlyOneMip, 3, Args.BlendAlphaSourceChannel); break;
-				case EBlendType::BT_BURN: BufferLayerColourInPlace<BurnChannelMasked, BurnChannel, 1>(pNew.get(), m_mask.get(), m_col, bOnlyOneMip, 3, Args.BlendAlphaSourceChannel); break;
-				case EBlendType::BT_DODGE: BufferLayerColourInPlace<DodgeChannelMasked, DodgeChannel, 1>(pNew.get(), m_mask.get(), m_col, bOnlyOneMip, 3, Args.BlendAlphaSourceChannel); break;
-				case EBlendType::BT_SCREEN: BufferLayerColourInPlace<ScreenChannelMasked, ScreenChannel, 1>(pNew.get(), m_mask.get(), m_col, bOnlyOneMip, 3, Args.BlendAlphaSourceChannel); break;
-				case EBlendType::BT_OVERLAY: BufferLayerColourInPlace<OverlayChannelMasked, OverlayChannel, 1>(pNew.get(), m_mask.get(), m_col, bOnlyOneMip, 3, Args.BlendAlphaSourceChannel); break;
-				case EBlendType::BT_LIGHTEN: BufferLayerColourInPlace<LightenChannelMasked, LightenChannel, 1>(pNew.get(), m_mask.get(), m_col, bOnlyOneMip, 3, Args.BlendAlphaSourceChannel); break;
-				case EBlendType::BT_MULTIPLY: BufferLayerColourInPlace<MultiplyChannelMasked, MultiplyChannel, 1>(pNew.get(), m_mask.get(), m_col, bOnlyOneMip, 3, Args.BlendAlphaSourceChannel); break;
-				case EBlendType::BT_BLEND: BufferLayerColourInPlace<BlendChannelMasked, BlendChannel, 1>(pNew.get(), m_mask.get(), m_col, bOnlyOneMip, 3, Args.BlendAlphaSourceChannel); break;
+				case EBlendType::BT_SOFTLIGHT: BufferLayerColourInPlace<SoftLightChannelMasked, SoftLightChannel, 1>(Result.get(), Mask.get(), Color, bOnlyOneMip, 3, Args.BlendAlphaSourceChannel); break;
+				case EBlendType::BT_HARDLIGHT: BufferLayerColourInPlace<HardLightChannelMasked, HardLightChannel, 1>(Result.get(), Mask.get(), Color, bOnlyOneMip, 3, Args.BlendAlphaSourceChannel); break;
+				case EBlendType::BT_BURN: BufferLayerColourInPlace<BurnChannelMasked, BurnChannel, 1>(Result.get(), Mask.get(), Color, bOnlyOneMip, 3, Args.BlendAlphaSourceChannel); break;
+				case EBlendType::BT_DODGE: BufferLayerColourInPlace<DodgeChannelMasked, DodgeChannel, 1>(Result.get(), Mask.get(), Color, bOnlyOneMip, 3, Args.BlendAlphaSourceChannel); break;
+				case EBlendType::BT_SCREEN: BufferLayerColourInPlace<ScreenChannelMasked, ScreenChannel, 1>(Result.get(), Mask.get(), Color, bOnlyOneMip, 3, Args.BlendAlphaSourceChannel); break;
+				case EBlendType::BT_OVERLAY: BufferLayerColourInPlace<OverlayChannelMasked, OverlayChannel, 1>(Result.get(), Mask.get(), Color, bOnlyOneMip, 3, Args.BlendAlphaSourceChannel); break;
+				case EBlendType::BT_LIGHTEN: BufferLayerColourInPlace<LightenChannelMasked, LightenChannel, 1>(Result.get(), Mask.get(), Color, bOnlyOneMip, 3, Args.BlendAlphaSourceChannel); break;
+				case EBlendType::BT_MULTIPLY: BufferLayerColourInPlace<MultiplyChannelMasked, MultiplyChannel, 1>(Result.get(), Mask.get(), Color, bOnlyOneMip, 3, Args.BlendAlphaSourceChannel); break;
+				case EBlendType::BT_BLEND: BufferLayerColourInPlace<BlendChannelMasked, BlendChannel, 1>(Result.get(), Mask.get(), Color, bOnlyOneMip, 3, Args.BlendAlphaSourceChannel); break;
 				default: check(false);
 				}
 
@@ -1153,35 +1197,35 @@ namespace mu
 				switch (EBlendType(Args.blendTypeAlpha))
 				{
 				case EBlendType::BT_NORMAL_COMBINE: check(false); break;
-				case EBlendType::BT_SOFTLIGHT: BufferLayerColourInPlace<SoftLightChannel, 1>(pNew.get(), m_col, bOnlyOneMip, 3, Args.BlendAlphaSourceChannel); break;
-				case EBlendType::BT_HARDLIGHT: BufferLayerColourInPlace<HardLightChannel, 1>(pNew.get(), m_col, bOnlyOneMip, 3, Args.BlendAlphaSourceChannel); break;
-				case EBlendType::BT_BURN: BufferLayerColourInPlace<BurnChannel, 1>(pNew.get(), m_col, bOnlyOneMip, 3, Args.BlendAlphaSourceChannel); break;
-				case EBlendType::BT_DODGE: BufferLayerColourInPlace<DodgeChannel, 1>(pNew.get(), m_col, bOnlyOneMip, 3, Args.BlendAlphaSourceChannel); break;
-				case EBlendType::BT_SCREEN: BufferLayerColourInPlace<ScreenChannel, 1>(pNew.get(), m_col, bOnlyOneMip, 3, Args.BlendAlphaSourceChannel); break;
-				case EBlendType::BT_OVERLAY: BufferLayerColourInPlace<OverlayChannel, 1>(pNew.get(), m_col, bOnlyOneMip, 3, Args.BlendAlphaSourceChannel); break;
-				case EBlendType::BT_LIGHTEN: BufferLayerColourInPlace<LightenChannel, 1>(pNew.get(), m_col, bOnlyOneMip, 3, Args.BlendAlphaSourceChannel); break;
-				case EBlendType::BT_MULTIPLY: BufferLayerColourInPlace<MultiplyChannel, 1>(pNew.get(), m_col, bOnlyOneMip, 3, Args.BlendAlphaSourceChannel); break;
-				case EBlendType::BT_BLEND: BufferLayerColourInPlace<BlendChannel, 1>(pNew.get(), m_col, bOnlyOneMip, 3, Args.BlendAlphaSourceChannel); break;
+				case EBlendType::BT_SOFTLIGHT: BufferLayerColourInPlace<SoftLightChannel, 1>(Result.get(), Color, bOnlyOneMip, 3, Args.BlendAlphaSourceChannel); break;
+				case EBlendType::BT_HARDLIGHT: BufferLayerColourInPlace<HardLightChannel, 1>(Result.get(), Color, bOnlyOneMip, 3, Args.BlendAlphaSourceChannel); break;
+				case EBlendType::BT_BURN: BufferLayerColourInPlace<BurnChannel, 1>(Result.get(), Color, bOnlyOneMip, 3, Args.BlendAlphaSourceChannel); break;
+				case EBlendType::BT_DODGE: BufferLayerColourInPlace<DodgeChannel, 1>(Result.get(), Color, bOnlyOneMip, 3, Args.BlendAlphaSourceChannel); break;
+				case EBlendType::BT_SCREEN: BufferLayerColourInPlace<ScreenChannel, 1>(Result.get(), Color, bOnlyOneMip, 3, Args.BlendAlphaSourceChannel); break;
+				case EBlendType::BT_OVERLAY: BufferLayerColourInPlace<OverlayChannel, 1>(Result.get(), Color, bOnlyOneMip, 3, Args.BlendAlphaSourceChannel); break;
+				case EBlendType::BT_LIGHTEN: BufferLayerColourInPlace<LightenChannel, 1>(Result.get(), Color, bOnlyOneMip, 3, Args.BlendAlphaSourceChannel); break;
+				case EBlendType::BT_MULTIPLY: BufferLayerColourInPlace<MultiplyChannel, 1>(Result.get(), Color, bOnlyOneMip, 3, Args.BlendAlphaSourceChannel); break;
+				case EBlendType::BT_BLEND: BufferLayerColourInPlace<BlendChannel, 1>(Result.get(), Color, bOnlyOneMip, 3, Args.BlendAlphaSourceChannel); break;
 				default: check(false);
 				}
 			}
-
 		}
 
 		// Reset relevancy map.
-		pNew->m_flags &= ~Image::EImageFlags::IF_HAS_RELEVANCY_MAP;
-
-		m_base = nullptr;
-		m_mask = nullptr;
-		Result = pNew;
+		Result->m_flags &= ~Image::EImageFlags::IF_HAS_RELEVANCY_MAP;
 	}
 
 
-	//---------------------------------------------------------------------------------------------
-	void FImageLayerColourTask::Complete(CodeRunner* runner)
+	void FImageLayerColourTask::Complete(CodeRunner* Runner)
 	{
-		// This runs in the runner thread
-		runner->GetMemory().SetImage(Op, Result);
+		// This runs in the Runner thread
+		Runner->Release(Mask);
+
+		// If no shortcut was taken
+		if (Result)
+		{
+			Runner->StoreImage(Op, Result);
+		}
 	}
 
 
@@ -1191,75 +1235,104 @@ namespace mu
 	class FImagePixelFormatTask : public CodeRunner::FIssuedTask
 	{
 	public:
-		FImagePixelFormatTask(const FScheduledOp&, FProgramCache&, const OP::ImagePixelFormatArgs&, int imageCompressionQuality);
+		FImagePixelFormatTask(const FScheduledOp& InOp, const OP::ImagePixelFormatArgs& InArgs)
+			: FIssuedTask(InOp), Args(InArgs)
+		{}
 
 		// FIssuedTask interface
+		bool Prepare(CodeRunner*, bool& bOutFailed) override;
 		void DoWork() override;
-		void Complete(CodeRunner* staged) override;
+		void Complete(CodeRunner*) override;
 
 	private:
-		int m_imageCompressionQuality = 0;
+		int ImageCompressionQuality = 0;
 		OP::ImagePixelFormatArgs Args;
-		Ptr<const Image> m_base;
-		Ptr<const Image> Result;
+		EImageFormat TargetFormat = EImageFormat::IF_NONE;
+		Ptr<const Image> Base;
+		Ptr<Image> Result;
 	};
 
 
-	//---------------------------------------------------------------------------------------------
-	FImagePixelFormatTask::FImagePixelFormatTask(const FScheduledOp& InOp, FProgramCache& Memory, const OP::ImagePixelFormatArgs& InArgs, int imageCompressionQuality)
-		: FIssuedTask(InOp)
+	bool FImagePixelFormatTask::Prepare(CodeRunner* Runner, bool& bOutFailed)
 	{
-		m_imageCompressionQuality = imageCompressionQuality;
-		Args = InArgs;
-		m_base = Memory.GetImage({ Args.source, InOp.ExecutionIndex, InOp.ExecutionOptions });
-	}
+		// This runs in the mutable Runner thread
+		MUTABLE_CPUPROFILER_SCOPE(FImageLayerColourTask_Prepare);
+		bOutFailed = false;
 
+		ImageCompressionQuality = Runner->m_pSettings->ImageCompressionQuality;
 
-	//---------------------------------------------------------------------------------------------
-	void FImagePixelFormatTask::DoWork()
-	{
-		MUTABLE_CPUPROFILER_SCOPE(FImagePixelFormatTask);
+		Base = Runner->LoadImage({ Args.source, Op.ExecutionIndex, Op.ExecutionOptions });
+		check(!Base || Base->GetFormat() < EImageFormat::IF_COUNT);		
 
-		// Warning: This runs in a worker thread.
+		// Shortcuts
+		if (!Base)
+		{
+			Runner->StoreImage(Op, Base);
+			return false;
+		}
 
-		if (!m_base)
-			return;
+		bool bValid = Base->GetSizeX() > 0 && Base->GetSizeY() > 0;
+		if (!bValid)
+		{
+			Runner->StoreImage(Op, Base);
+			return false;
+		}
 
-		EImageFormat targetFormat = Args.format;
+		TargetFormat = Args.format;
 		if (Args.formatIfAlpha != EImageFormat::IF_NONE
 			&&
-			GetImageFormatData(m_base->GetFormat()).m_channels > 3)
+			GetImageFormatData(Base->GetFormat()).m_channels > 3)
 		{
-			targetFormat = Args.formatIfAlpha;
+			TargetFormat = Args.formatIfAlpha;
 		}
 
-		if (targetFormat != EImageFormat::IF_NONE)
+		if (TargetFormat == EImageFormat::IF_NONE || TargetFormat == Base->GetFormat())
 		{
-			if (targetFormat == m_base->GetFormat())
-			{
-				Result = CloneOrTakeOver<>(m_base.get());
-			}
-			else
-			{
-				Ptr<Image> TempResult = ImagePixelFormat(m_imageCompressionQuality, m_base.get(), targetFormat, -1);
-				TempResult->m_flags = m_base->m_flags;
-				Result = TempResult;
-			}
+			Runner->StoreImage(Op, Base);
+			return false;
 		}
-		else
+
+		// Create destination data
+		Result = Runner->CreateImage(Base->GetSizeX(), Base->GetSizeY(), Base->GetLODCount(), TargetFormat, EInitializationType::NotInitialized);
+		return true;
+	}
+
+
+	void FImagePixelFormatTask::DoWork()
+	{
+		// This runs in a worker thread
+		MUTABLE_CPUPROFILER_SCOPE(FImagePixelFormatTask);
+		
+		bool bSuccess = false;
+		ImagePixelFormat(bSuccess, ImageCompressionQuality, Result.get(), Base.get(), -1);
+
+		int32 OriginalDataSize = FMath::Max(Result->m_data.Num(), Base->m_data.Num());
+		int32 ExcessDataSize = OriginalDataSize * 4 + 4;
+		while (!bSuccess)
 		{
-			Result = m_base;
+			MUTABLE_CPUPROFILER_SCOPE(Recompression_OutOfSpace);
+
+			// Bad case, where the RLE compressed data requires more memory than the uncompressed data.
+			// We need to support it anyway for small mips or scaled images.
+			Result->m_data.SetNum(ExcessDataSize);
+			bSuccess = false;
+			ImagePixelFormat(bSuccess,ImageCompressionQuality, Result.get(), Base.get(), -1);
+			ExcessDataSize *= 4;
 		}
 
-		m_base = nullptr;
 	}
 
 
-	//---------------------------------------------------------------------------------------------
-	void FImagePixelFormatTask::Complete(CodeRunner* runner)
+	void FImagePixelFormatTask::Complete(CodeRunner* Runner)
 	{
-		// This runs in the runner thread
-		runner->GetMemory().SetImage(Op, Result);
+		// This runs in the Runner thread
+		Runner->Release(Base);
+
+		// If no shortcut was taken
+		if (Result)
+		{
+			Runner->StoreImage(Op, Result);
+		}
 	}
 
 
@@ -1269,78 +1342,102 @@ namespace mu
 	class FImageMipmapTask : public CodeRunner::FIssuedTask
 	{
 	public:
-		FImageMipmapTask(const FScheduledOp&, FProgramCache&, const OP::ImageMipmapArgs&, int imageCompressionQuality);
+		FImageMipmapTask(const FScheduledOp& InOp, const OP::ImageMipmapArgs& InArgs)
+			:FIssuedTask(InOp), Args(InArgs)
+		{}
 
 		// FIssuedTask interface
+		bool Prepare(CodeRunner*, bool& bOutFailed) override;
 		void DoWork() override;
-		void Complete(CodeRunner* staged) override;
+		void Complete(CodeRunner*) override;
 
 	private:
-		int m_imageCompressionQuality = 0;
+		int ImageCompressionQuality = 0;
 		OP::ImageMipmapArgs Args;
-		Ptr<const Image> m_base;
-		Ptr<const Image> Result;
+		Ptr<const Image> Base;
+		Ptr<Image> Result;
+		FScratchImageMipmap Scratch;
 	};
 
 
-	//---------------------------------------------------------------------------------------------
-	FImageMipmapTask::FImageMipmapTask(const FScheduledOp& InOp, FProgramCache& Memory, const OP::ImageMipmapArgs& InArgs, int imageCompressionQuality)
-		: FIssuedTask(InOp)
+	bool FImageMipmapTask::Prepare(CodeRunner* Runner, bool& bOutFailed)
 	{
-		m_imageCompressionQuality = imageCompressionQuality;
-		Args = InArgs;
-		m_base = Memory.GetImage({ Args.source, InOp.ExecutionIndex, InOp.ExecutionOptions });
-	}
+		// This runs in the mutable Runner thread
+		MUTABLE_CPUPROFILER_SCOPE(FImageLayerColourTask_Prepare);
+		bOutFailed = false;
 
+		ImageCompressionQuality = Runner->m_pSettings->ImageCompressionQuality;
 
-	//---------------------------------------------------------------------------------------------
-	void FImageMipmapTask::DoWork()
-	{
-		MUTABLE_CPUPROFILER_SCOPE(FImageMipmapTask);
+		Base = Runner->LoadImage({ Args.source, Op.ExecutionIndex, Op.ExecutionOptions });
+		check(!Base || Base->GetFormat() < EImageFormat::IF_COUNT);
 
-		// Warning: This runs in a worker thread.
+		// Shortcuts
+		if (!Base)
+		{
+			Runner->StoreImage(Op, Base);
+			return false;
+		}
 
-		if (!m_base)
+		bool bValid = Base->GetSizeX() > 0 && Base->GetSizeY() > 0;
+		if (!bValid)
 		{
-			return;
+			Runner->StoreImage(Op, Base);
+			return false;
 		}
 
-		int32 levelCount = Args.levels;
-		int32 maxLevelCount = Image::GetMipmapCount(m_base->GetSizeX(), m_base->GetSizeY());
-		if (levelCount == 0)
+		int32 LevelCount = Args.levels;
+		int32 MaxLevelCount = Image::GetMipmapCount(Base->GetSizeX(), Base->GetSizeY());
+		if (LevelCount == 0)
 		{
-			levelCount = maxLevelCount;
+			LevelCount = MaxLevelCount;
 		}
-		else if (levelCount > maxLevelCount)
+		else if (LevelCount > MaxLevelCount)
 		{
 			// If code generation is smart enough, this should never happen.
 			// \todo But apparently it does, sometimes.
-			levelCount = maxLevelCount;
+			LevelCount = MaxLevelCount;
 		}
 
 		// At least keep the levels we already have.
-		int startLevel = m_base->GetLODCount();
-		levelCount = FMath::Max(startLevel, levelCount);
+		int StartLevel = Base->GetLODCount();
+		LevelCount = FMath::Max(StartLevel, LevelCount);
 
-		ImagePtr pDest = new Image(m_base->GetSizeX(), m_base->GetSizeY(), levelCount, m_base->GetFormat(),EInitializationType::NotInitialized);
-		pDest->m_flags = m_base->m_flags;
+		// Create destination data
+		Result = Runner->CreateImage(Base->GetSizeX(), Base->GetSizeY(), LevelCount, Base->GetFormat(), EInitializationType::NotInitialized);
+		Result->m_flags = Base->m_flags;
 
-		SCRATCH_IMAGE_MIPMAP scratch;
-		FMipmapGenerationSettings settings{};
+		ImageMipmap_PrepareScratch( Runner->m_pSystem->WorkingMemoryManager, Result.get(), Base.get(), LevelCount, Scratch);
+
+		return true;
+	}
 
-		ImageMipmap_PrepareScratch(pDest.get(), m_base.get(), levelCount, &scratch);
-		ImageMipmap(m_imageCompressionQuality, pDest.get(), m_base.get(), levelCount, &scratch, settings);
 
-		m_base = nullptr;
-		Result = pDest;
+	//---------------------------------------------------------------------------------------------
+	void FImageMipmapTask::DoWork()
+	{
+		// This runs in a worker thread
+		MUTABLE_CPUPROFILER_SCOPE(FImageMipmapTask);
+
+		FMipmapGenerationSettings settings{};
+		ImageMipmap(Scratch, ImageCompressionQuality, Result.get(), Base.get(), Result->GetLODCount(), settings);
 	}
 
 
 	//---------------------------------------------------------------------------------------------
-	void FImageMipmapTask::Complete(CodeRunner* runner)
+	void FImageMipmapTask::Complete(CodeRunner* Runner)
 	{
-		// This runs in the runner thread
-		runner->GetMemory().SetImage(Op, Result);
+		ImageMipmap_ReleaseScratch(Runner->m_pSystem->WorkingMemoryManager, Scratch);
+
+		// This runs in the Runner thread
+		Runner->Release(Base);
+
+		check(Result);
+
+		// If no shortcut was taken
+		if (Result)
+		{
+			Runner->StoreImage(Op, Result);
+		}
 	}
 
 
@@ -1350,40 +1447,49 @@ namespace mu
 	class FImageSwizzleTask : public CodeRunner::FIssuedTask
 	{
 	public:
-		FImageSwizzleTask(const FScheduledOp&, FProgramCache&, const OP::ImageSwizzleArgs&);
+		FImageSwizzleTask(const FScheduledOp& InOp, const OP::ImageSwizzleArgs& InArgs)
+			: FIssuedTask(InOp), Args(InArgs)
+		{}
 
 		// FIssuedTask interface
+		bool Prepare(CodeRunner*, bool& bOutFailed);
 		void DoWork() override;
-		void Complete(CodeRunner* staged) override;
+		void Complete(CodeRunner*) override;
 
 	private:
 		OP::ImageSwizzleArgs Args;
 		Ptr<const Image> Sources[MUTABLE_OP_MAX_SWIZZLE_CHANNELS];
-		Ptr<const Image> Result;
+		Ptr<Image> Result;
 	};
 
 
-	//---------------------------------------------------------------------------------------------
-	FImageSwizzleTask::FImageSwizzleTask(const FScheduledOp& InOp, FProgramCache& Memory, const OP::ImageSwizzleArgs& InArgs)
-		: FIssuedTask(InOp)
+	bool FImageSwizzleTask::Prepare(CodeRunner* Runner, bool& bOutFailed)
 	{
-		Args = InArgs;
+		// This runs in the mutable Runner thread
+		MUTABLE_CPUPROFILER_SCOPE(FImageLayerColourTask_Prepare);
+		bOutFailed = false;
+
 		for (int i = 0; i < MUTABLE_OP_MAX_SWIZZLE_CHANNELS; ++i)
 		{
 			if (Args.sources[i])
 			{
-				Sources[i] = Memory.GetImage({ Args.sources[i], InOp.ExecutionIndex, InOp.ExecutionOptions });
+				Sources[i] = Runner->LoadImage({ Args.sources[i], Op.ExecutionIndex, Op.ExecutionOptions });
 			}
 		}
-	}
 
+		// Shortcuts
+		if (!Sources[0])
+		{
+			for (int i = 0; i < MUTABLE_OP_MAX_SWIZZLE_CHANNELS; ++i)
+			{
+				Runner->Release(Sources[i]);
+			}
+			Runner->StoreImage(Op, nullptr);
+			return false;
+		}
 
-	//---------------------------------------------------------------------------------------------
-	void FImageSwizzleTask::DoWork()
-	{
-		MUTABLE_CPUPROFILER_SCOPE(FImageSwizzleTask);
-
-		// Warning: This runs in a worker thread.
+		// Create destination data
+		EImageFormat format = (EImageFormat)Args.format;
 
 		// Be defensive: ensure image sizes match.
 		for (int i = 1; i < MUTABLE_OP_MAX_SWIZZLE_CHANNELS; ++i)
@@ -1391,25 +1497,40 @@ namespace mu
 			if (Sources[i] && Sources[i]->GetSize() != Sources[0]->GetSize())
 			{
 				MUTABLE_CPUPROFILER_SCOPE(ImageResize_ForSwizzle);
-				Sources[i] = ImageResizeLinear(0, Sources[i].get(), Sources[0]->GetSize());
+				Ptr<Image> Resized = Runner->CreateImage(Sources[0]->GetSizeX(), Sources[0]->GetSizeY(), 1, Sources[i]->GetFormat(), EInitializationType::NotInitialized);
+				ImageResizeLinear(Resized.get(), 0, Sources[i].get());
+				Runner->Release(Sources[i]);
+				Sources[i] = Resized;
 			}
 		}
 
+		Result = Runner->CreateImage(Sources[0]->GetSizeX(), Sources[0]->GetSizeY(), Sources[0]->GetLODCount(), Args.format, EInitializationType::Black);
+		return true;
+	}
 
-		EImageFormat format = (EImageFormat)Args.format;
-		Result = ImageSwizzle(format, Sources, Args.sourceChannels);
-		for (int i = 0; i < MUTABLE_OP_MAX_SWIZZLE_CHANNELS; ++i)
-		{
-			Sources[i] = nullptr;
-		}
+
+	void FImageSwizzleTask::DoWork()
+	{
+		// This runs in a worker thread
+		MUTABLE_CPUPROFILER_SCOPE(FImageSwizzleTask);
+
+		ImageSwizzle(Result.get(), Sources, Args.sourceChannels);
 	}
 
 
-	//---------------------------------------------------------------------------------------------
-	void FImageSwizzleTask::Complete(CodeRunner* runner)
+	void FImageSwizzleTask::Complete(CodeRunner* Runner)
 	{
-		// This runs in the runner thread
-		runner->GetMemory().SetImage(Op, Result);
+		// This runs in the Runner thread
+		for (int i = 0; i < MUTABLE_OP_MAX_SWIZZLE_CHANNELS; ++i)
+		{
+			Runner->Release(Sources[i]);
+		}
+
+		// If no shortcut was taken
+		if (Result)
+		{
+			Runner->StoreImage(Op, Result);
+		}
 	}
 
 
@@ -1419,70 +1540,74 @@ namespace mu
 	class FImageSaturateTask : public CodeRunner::FIssuedTask
 	{
 	public:
-		FImageSaturateTask(const FScheduledOp&, FProgramCache&, const OP::ImageSaturateArgs&);
+		FImageSaturateTask(const FScheduledOp&, const OP::ImageSaturateArgs&);
 
 		// FIssuedTask interface
+		bool Prepare(CodeRunner*, bool& bOutFailed) override;
 		void DoWork() override;
-		void Complete(CodeRunner* staged) override;
+		void Complete(CodeRunner*) override;
 
 	private:
-		Ptr<const Image> Source;
-		Ptr<const Image> Result;
+		const OP::ImageSaturateArgs Args;
+		Ptr<Image> Result;
 		float Factor;
 	};
 
 
 	//---------------------------------------------------------------------------------------------
-	FImageSaturateTask::FImageSaturateTask(const FScheduledOp& InOp, FProgramCache& Memory, const OP::ImageSaturateArgs& InArgs)
+	FImageSaturateTask::FImageSaturateTask(const FScheduledOp& InOp, const OP::ImageSaturateArgs& InArgs)
 		: FIssuedTask(InOp)
+		, Args(InArgs)
 	{
-		Source = Memory.GetImage(FCacheAddress(InArgs.base, InOp));
-		Factor = Memory.GetScalar(FScheduledOp::FromOpAndOptions(InArgs.factor, InOp, 0));
 	}
 
 	//---------------------------------------------------------------------------------------------
-	void FImageSaturateTask::DoWork()
+	bool FImageSaturateTask::Prepare(CodeRunner* Runner, bool& bOutFailed)
 	{
-		MUTABLE_CPUPROFILER_SCOPE(FImageSaturateTask);
+		// This runs in the mutable Runner thread
+		MUTABLE_CPUPROFILER_SCOPE(FImageSaturateTask_Prepare);
+		bOutFailed = false;
+
+		Ptr<const Image> Source = Runner->LoadImage(FCacheAddress(Args.base, Op));
+		Factor = Runner->LoadScalar(FScheduledOp::FromOpAndOptions(Args.factor, Op, 0));
 
 		if (!Source)
 		{
-			check(false);
-			return;
+			Runner->StoreImage(Op, Source);
+			return false;
 		}
 
 		const bool bOptimizeUnchanged = FMath::IsNearlyEqual(Factor, 1.0f);
-		if (Source->IsUnique())
+		if (bOptimizeUnchanged)
 		{
-			Ptr<Image> MutableImage = mu::CloneOrTakeOver<>(Source.get());
-			if (!bOptimizeUnchanged)
-			{
-				constexpr bool bUseVectorIntrinsics = false;
-				ImageSaturateInPlace<bUseVectorIntrinsics>(MutableImage.get(), Factor);
-			}
-
-			Result = MutableImage;
-		}
-		else
-		{
-			if (bOptimizeUnchanged)
-			{
-				Result = Source->Clone();
-			}
-			else
-			{
-				Result = ImageSaturate(Source.get(), Factor);
-			}
+			Runner->StoreImage(Op, Source);
+			return false;
 		}
 
-		Source = nullptr;
+		Result = Runner->CloneOrTakeOver(Source);
+		return true;
+	}
+
+	//---------------------------------------------------------------------------------------------
+	void FImageSaturateTask::DoWork()
+	{
+		// This runs on a random worker thread
+		MUTABLE_CPUPROFILER_SCOPE(FImageSaturateTask);
+
+		constexpr bool bUseVectorIntrinsics = false;
+		ImageSaturateInPlace<bUseVectorIntrinsics>(Result.get(), Factor);
 	}
 
 	//---------------------------------------------------------------------------------------------
 	void FImageSaturateTask::Complete(CodeRunner* Runner)
 	{
-		// This runs in the runner thread
-		Runner->GetMemory().SetImage(Op, Result);
+		// This runs in the mutable Runner thread
+
+		// If we didn't take a shortcut
+		if (Result)
+		{
+			Runner->StoreImage(Op, Result);
+		}
 	}
 
 
@@ -1492,92 +1617,105 @@ namespace mu
 	class FImageResizeTask : public CodeRunner::FIssuedTask
 	{
 	public:
-		FImageResizeTask(const FScheduledOp& InOp, FProgramCache& Memory, const Ptr<const Image>& InBase, uint16 InSizeX, uint16 InSizeY, int32 InImageCompressionQuality);
+		FImageResizeTask(const FScheduledOp& InOp, const OP::ImageResizeArgs& InArgs)
+			: FIssuedTask(InOp), Args(InArgs)
+		{}
 
 		// FIssuedTask interface
+		bool Prepare(CodeRunner* Runner, bool& bFailed) override;
 		void DoWork() override;
-		void Complete(CodeRunner* staged) override;
+		void Complete(CodeRunner*) override;
 
 	private:
-		Ptr<const Image> pBase;
-		Ptr<const Image> Result;
-		int32 ImageCompressionQuality;
-		uint16 SizeX, SizeY;
+		int32 ImageCompressionQuality = 0;
+		OP::ImageResizeArgs Args;
+		Ptr<const Image> Base;
+		Ptr<Image> Result;
 	};
 
 
-	//---------------------------------------------------------------------------------------------
-	FImageResizeTask::FImageResizeTask(const FScheduledOp& InOp, FProgramCache& Memory, const Ptr<const Image>& InBase, uint16 InSizeX, uint16 InSizeY, int32 InImageCompressionQuality)
-		: FIssuedTask(InOp)
-	{
-		SizeX = InSizeX;
-		SizeY = InSizeY;
-		pBase = InBase;
-		ImageCompressionQuality = InImageCompressionQuality;
-	}
-
-	//---------------------------------------------------------------------------------------------
-	void FImageResizeTask::DoWork()
+	bool FImageResizeTask::Prepare(CodeRunner* Runner, bool& bOutFailed)
 	{
-		MUTABLE_CPUPROFILER_SCOPE(FImageResizeTask);
+		// This runs in the mutable Runner thread
+		MUTABLE_CPUPROFILER_SCOPE(FImageResizeTask_Prepare);
+		bOutFailed = false;
 
-		if (!pBase)
-		{
-			Result = nullptr;
-			return;
-		}
+		ImageCompressionQuality = Runner->m_pSettings->ImageCompressionQuality;
 
-		FImageSize destSize = FImageSize( SizeX, SizeY );
+		FImageSize destSize = FImageSize(Args.size[0], Args.size[1]);
 
 		// Apply the mips-to-skip to the dest size
 		int32 MipsToSkip = Op.ExecutionOptions;
 		destSize[0] = FMath::Max(destSize[0] >> MipsToSkip, 1);
 		destSize[1] = FMath::Max(destSize[1] >> MipsToSkip, 1);
 
-		if ( destSize[0]!=pBase->GetSizeX()
-			 ||
-			 destSize[1]!=pBase->GetSizeY() )
+		Base = Runner->LoadImage(FCacheAddress(Args.source, Op));
+		if (!Base 
+			|| 
+			( Base->GetSizeX()==destSize[0] && Base->GetSizeY()==destSize[1] )
+			)
 		{
-			ImagePtr Resized = ImageResizeLinear( ImageCompressionQuality, pBase.get(), destSize );
-			Result = Resized;
+			Runner->StoreImage(Op, Base);
+			return false;
+		}
 
-			// If the source image had mips, generate them as well for the resized image.
-			// This shouldn't happen often since "ResizeLike" should be usually optimised out
-			// during model compilation. The mipmap generation below is not very precise with
-			// the number of mips that are needed and will probably generate too many
-			bool sourceHasMips = pBase->GetLODCount()>1;
+		int32 Lods = 1;
 
-			// No longer needed.
-			pBase = nullptr;
+		// If the source image had mips, generate them as well for the resized image.
+		// This shouldn't happen often since it should be usually optimised  during model compilation. 
+		// The mipmap generation below is not very precise with the number of mips that are needed and 
+		// will probably generate too many
+		bool bSourceHasMips = Base->GetLODCount() > 1;
+		if (bSourceHasMips)
+		{
+			Lods = Image::GetMipmapCount(destSize[0], destSize[1]);
+		}
 
-			if (sourceHasMips)
-			{
-				int levelCount = Image::GetMipmapCount(Resized->GetSizeX(), Resized->GetSizeY() );
-				ImagePtr Mipmapped = new Image(Resized->GetSizeX(), Resized->GetSizeY(), levelCount, Resized->GetFormat() );
+		Result = Runner->CreateImage( destSize[0], destSize[1], Lods, Base->GetFormat(), EInitializationType::NotInitialized );
+
+		return true;
+	}
 
-				SCRATCH_IMAGE_MIPMAP scratch;
-				FMipmapGenerationSettings mipSettings{};
 
-				ImageMipmap_PrepareScratch(Mipmapped.get(), Resized.get(), levelCount, &scratch );
-				ImageMipmap( ImageCompressionQuality,
-							 Mipmapped.get(), Resized.get(), levelCount, &scratch, mipSettings );
+	//---------------------------------------------------------------------------------------------
+	void FImageResizeTask::DoWork()
+	{
+		// This runs on a random worker thread
+		MUTABLE_CPUPROFILER_SCOPE(FImageResizeTask);
 
-				Result = Mipmapped;
-			}			
-		}
-		else
+		FImageSize destSize = FImageSize( Args.size[0], Args.size[1]);
+
+		// Apply the mips-to-skip to the dest size
+		int32 MipsToSkip = Op.ExecutionOptions;
+		destSize[0] = FMath::Max(destSize[0] >> MipsToSkip, 1);
+		destSize[1] = FMath::Max(destSize[1] >> MipsToSkip, 1);
+
+		// Warning: This will actually allocate temp memory that may exceed the budget.
+		// \TODO: Fix it.
+		// Warning: This may allocate memory if the image is compressed.
+		// \TODO: treat this case in the Prepare.
+		ImageResizeLinear( Result.get(), ImageCompressionQuality, Base.get() );
+
+		int32 LodCount = Result->GetLODCount();
+		if (LodCount>1)
 		{
-			Result = pBase;
+			FMipmapGenerationSettings mipSettings{};
+			ImageMipmapInPlace( ImageCompressionQuality, Result.get(), mipSettings );
 		}
-
-		pBase = nullptr;
 	}
 
+
 	//---------------------------------------------------------------------------------------------
 	void FImageResizeTask::Complete(CodeRunner* Runner)
 	{
-		// This runs in the runner thread
-		Runner->GetMemory().SetImage(Op, Result);
+		// This runs in the Runner thread
+		Runner->Release(Base);
+
+		// If didn't take a shortcut and set it already
+		if (Result)
+		{
+			Runner->StoreImage(Op, Result);
+		}
 	}
 
 
@@ -1587,133 +1725,152 @@ namespace mu
 	class FImageResizeRelTask : public CodeRunner::FIssuedTask
 	{
 	public:
-		FImageResizeRelTask(const FScheduledOp&, FProgramCache&, const OP::ImageResizeRelArgs&, int32 ImageCompressionQuality);
+		FImageResizeRelTask(const FScheduledOp& InOp, const OP::ImageResizeRelArgs& InArgs)
+			: FIssuedTask(InOp), Args(InArgs)
+		{}
 
 		// FIssuedTask interface
+		bool Prepare(CodeRunner* Runner, bool& bFailed) override;
 		void DoWork() override;
-		void Complete(CodeRunner* staged) override;
+		void Complete(CodeRunner*) override;
 
 	private:
-		Ptr<const Image> pBase;
-		Ptr<const Image> Result;
 		OP::ImageResizeRelArgs Args;
-		int32 ImageCompressionQuality;
+		int32 ImageCompressionQuality=0;
+		Ptr<const Image> Base;
+		Ptr<Image> Result;
+		FImageSize DestSize;
 	};
 
 
-	//---------------------------------------------------------------------------------------------
-	FImageResizeRelTask::FImageResizeRelTask(const FScheduledOp& InOp, FProgramCache& Memory, const OP::ImageResizeRelArgs& InArgs, int32 InImageCompressionQuality)
-		: FIssuedTask(InOp)
+	bool FImageResizeRelTask::Prepare(CodeRunner* Runner, bool& bOutFailed)
 	{
-		Args = InArgs;
-		pBase = Memory.GetImage(FCacheAddress(InArgs.source, InOp));
-		ImageCompressionQuality = InImageCompressionQuality;
-	}
+		// This runs in the mutable Runner thread
+		MUTABLE_CPUPROFILER_SCOPE(FImageResizeRelTask_Prepare);
+		bOutFailed = false;
 
-	//---------------------------------------------------------------------------------------------
-	void FImageResizeRelTask::DoWork()
-	{
-		MUTABLE_CPUPROFILER_SCOPE(FImageResizeRelTask);
+		ImageCompressionQuality = Runner->m_pSettings->ImageCompressionQuality;
 
-		if (!pBase)
+		Base = Runner->LoadImage(FCacheAddress(Args.source, Op));
+		if (!Base)
 		{
-			Result = nullptr;
-			return;
+			Runner->StoreImage(Op, Base);
+			return false;
 		}
 
-		FImageSize destSize(
-					uint16( FMath::Max(1.0, pBase->GetSizeX()*Args.factor[0] + 0.5f ) ),
-					uint16( FMath::Max(1.0, pBase->GetSizeY()*Args.factor[1] + 0.5f ) ) );
+		DestSize = FImageSize(
+			uint16(FMath::Max(1.0, Base->GetSizeX() * Args.factor[0] + 0.5f)),
+			uint16(FMath::Max(1.0, Base->GetSizeY() * Args.factor[1] + 0.5f)));
 
-		//pResult = ImageResize( pBase.get(), destSize );
-		ImagePtr pResult = ImageResizeLinear(
-			ImageCompressionQuality, pBase.get(), destSize );
+		// Apply the mips-to-skip to the dest size
+		int32 MipsToSkip = Op.ExecutionOptions;
+		DestSize[0] = FMath::Max(DestSize[0] >> MipsToSkip, 1);
+		DestSize[1] = FMath::Max(DestSize[1] >> MipsToSkip, 1);
 
-		// If the source image had mips, generate them as well for the resized image.
-		// This shouldn't happen often since "ResizeLike" should be usually optimised out
-		// during model compilation. The mipmap generation below is not very precise with
-		// the number of mips that are needed and will probably generate too many
-		bool sourceHasMips = pBase->GetLODCount()>1;
+		if (Base->GetSizeX() == DestSize[0] && Base->GetSizeY() == DestSize[1])
+		{
+			Runner->StoreImage(Op, Base);
+			return false;
+		}
 
-		// We don't need this anymore
-		pBase = nullptr;
+		int32 Lods = 1;
 
-		if (sourceHasMips)
+		// If the source image had mips, generate them as well for the resized image.
+		// This shouldn't happen often since it should be usually optimised  during model compilation. 
+		// The mipmap generation below is not very precise with the number of mips that are needed and 
+		// will probably generate too many
+		bool bSourceHasMips = Base->GetLODCount() > 1;
+		if (bSourceHasMips)
 		{
-			int levelCount = Image::GetMipmapCount(pResult->GetSizeX(), pResult->GetSizeY());
-			ImagePtr pMipmapped = new Image(pResult->GetSizeX(), pResult->GetSizeY(),
-				levelCount,
-				pResult->GetFormat());
+			Lods = Image::GetMipmapCount(DestSize[0], DestSize[1]);
+		}
 
-			SCRATCH_IMAGE_MIPMAP scratch;
-			FMipmapGenerationSettings mipSettings{};
+		Result = Runner->CreateImage(DestSize[0], DestSize[1], Lods, Base->GetFormat(), EInitializationType::NotInitialized);
+
+		return true;
+	}
 
-			ImageMipmap_PrepareScratch(pMipmapped.get(), pResult.get(), levelCount, &scratch);
-			ImageMipmap(ImageCompressionQuality,
-				pMipmapped.get(), pResult.get(), levelCount, &scratch, mipSettings);
 
-			pResult = pMipmapped;
-		}
+	void FImageResizeRelTask::DoWork()
+	{
+		// This runs on a random worker thread
+		MUTABLE_CPUPROFILER_SCOPE(FImageResizeRelTask);
+
+		ImageResizeLinear(Result.get(), ImageCompressionQuality, Base.get());
 
-		Result = pResult;
+		int32 LodCount = Result->GetLODCount();
+		if (LodCount > 1)
+		{
+			FMipmapGenerationSettings mipSettings{};
+			ImageMipmapInPlace(ImageCompressionQuality, Result.get(), mipSettings);
+		}
 	}
 
-	//---------------------------------------------------------------------------------------------
+
 	void FImageResizeRelTask::Complete(CodeRunner* Runner)
 	{
-		// This runs in the runner thread
-		Runner->GetMemory().SetImage(Op, Result);
+		// This runs in the Runner thread
+		Runner->Release(Base);
+
+		// If didn't take a shortcut and set it already
+		if (Result)
+		{
+			Runner->StoreImage(Op, Result);
+		}
 	}
 
+
 	//---------------------------------------------------------------------------------------------
 	//---------------------------------------------------------------------------------------------
 	//---------------------------------------------------------------------------------------------
 	class FImageInvertTask : public CodeRunner::FIssuedTask
 	{
 	public:
-		FImageInvertTask(const FScheduledOp&, FProgramCache&, const OP::ImageInvertArgs&);
+		FImageInvertTask(const FScheduledOp& InOp, const OP::ImageInvertArgs& InArgs)
+			: FIssuedTask(InOp), Args(InArgs)
+		{}
 
 		// FIssuedTask interface
+		bool Prepare(CodeRunner*, bool& bOutFailed) override;
 		void DoWork() override;
 		void Complete(CodeRunner* Runner) override;
 
 	private:
-		Ptr<const Image> Source;
-		Ptr<const Image> Result;
+		Ptr<Image> Result;
+		OP::ImageInvertArgs Args;
 	};
 
 
-	//---------------------------------------------------------------------------------------------
-	FImageInvertTask::FImageInvertTask(const FScheduledOp& InOp, FProgramCache& Memory, const OP::ImageInvertArgs& InArgs)
-		: FIssuedTask(InOp)
+	bool FImageInvertTask::Prepare(CodeRunner* Runner, bool& bOutFailed)
 	{
-		Source = Memory.GetImage(FCacheAddress(InArgs.base, InOp));
+		// This runs in the mutable Runner thread
+		MUTABLE_CPUPROFILER_SCOPE(FImageLayerTask_Prepare);
+		bOutFailed = false;
+
+		Ptr<const Image> Source = Runner->LoadImage({ Args.base, Op.ExecutionIndex, Op.ExecutionOptions });
+
+		// Create destination data
+		Result = Runner->CloneOrTakeOver(Source);
+		return true;
 	}
 
-	//---------------------------------------------------------------------------------------------
+
 	void FImageInvertTask::DoWork()
 	{
+		// This runs on a random worker thread
 		MUTABLE_CPUPROFILER_SCOPE(FImageInvertTask);
 
-		if (Source->IsUnique())
-		{
-			Ptr<Image> MutableImage = mu::CloneOrTakeOver<>(Source.get());
-			ImageInvertInPlace(MutableImage.get());
-			Result = MutableImage;
-		}
-		else
-		{
-			Result = ImageInvert(Source.get());
-		}
-
-		Source = nullptr;
+		ImageInvertInPlace(Result.get());
 	}
 
-	//---------------------------------------------------------------------------------------------
+
 	void FImageInvertTask::Complete(CodeRunner* Runner)
 	{
-		// This runs in the runner thread
-		Runner->GetMemory().SetImage(Op, Result);
+		// If we didn't take a shortcut
+		if (Result)
+		{
+			Runner->StoreImage(Op, Result);
+		}
 	}
 
 
@@ -1723,187 +1880,205 @@ namespace mu
 	class FImageComposeTask : public CodeRunner::FIssuedTask
 	{
 	public:
-		FImageComposeTask(const FScheduledOp&, FProgramCache&, const OP::ImageComposeArgs&, int imageCompressionQuality, const Ptr<const Layout>& );
+		FImageComposeTask(const FScheduledOp& InOp, const OP::ImageComposeArgs& InArgs, const Ptr<const Layout>& InLayout)
+			: FIssuedTask(InOp), Args(InArgs), Layout(InLayout)
+		{}
 
 		// FIssuedTask interface
+		bool Prepare(CodeRunner*, bool& bOutFailed) override;
 		void DoWork() override;
-		void Complete(CodeRunner* staged) override;
+		void Complete(CodeRunner*) override;
 
 	private:
-		int m_imageCompressionQuality = 0;
+		int32 ImageCompressionQuality = 0;
 		OP::ImageComposeArgs Args;
-		Ptr<const Image> m_base;
-		Ptr<const Layout> m_layout;
-		Ptr<const Image> m_block;
-		Ptr<const Image> m_mask;
-		Ptr<const Image> Result;
+		Ptr<const Layout> Layout;
+		Ptr<const Image> Block;
+		Ptr<const Image> Mask;
+		Ptr<Image> Result;
+		box< UE::Math::TIntVector2<uint16> > Rect;
 	};
 
 
-	//---------------------------------------------------------------------------------------------
-	FImageComposeTask::FImageComposeTask(const FScheduledOp& InOp, FProgramCache& Memory, const OP::ImageComposeArgs& InArgs, int imageCompressionQuality, const Ptr<const Layout>& InLayout)
-		: FIssuedTask(InOp)
+	bool FImageComposeTask::Prepare(CodeRunner* Runner, bool& bOutFailed)
 	{
-		m_imageCompressionQuality = imageCompressionQuality;
-		Args = InArgs;
-		m_layout = InLayout;
-
-		m_base = Memory.GetImage({ Args.base, Op.ExecutionIndex, InOp.ExecutionOptions });
-		
-		int relBlockIndex = m_layout->FindBlock(Args.blockIndex);
+		// This runs in the mutable Runner thread
+		MUTABLE_CPUPROFILER_SCOPE(FImageLayerTask_Prepare);
+		bOutFailed = false;
 
-		if (relBlockIndex >= 0)
-		{
-			m_block = Memory.GetImage({ Args.blockImage, Op.ExecutionIndex, InOp.ExecutionOptions });
-			if (Args.mask)
-			{
-				m_mask = Memory.GetImage({ Args.mask, Op.ExecutionIndex, InOp.ExecutionOptions });
-			}
-		}
-	}
+		ImageCompressionQuality = Runner->m_pSettings->ImageCompressionQuality;
 
+		Ptr<const Image> Base = Runner->LoadImage({ Args.base, Op.ExecutionIndex, Op.ExecutionOptions });
 
-	//---------------------------------------------------------------------------------------------
-	void FImageComposeTask::DoWork()
-	{
-		MUTABLE_CPUPROFILER_SCOPE(FImageComposeTask);
+		int32 RelBlockIndex = Layout->FindBlock(Args.blockIndex);
 
-		// Warning: This runs in a worker thread.
-		int relBlockIndex = m_layout->FindBlock(Args.blockIndex);
+		// Shortcuts
+		if (RelBlockIndex < 0)
+		{
+			Runner->StoreImage(Op, Base);
+			return false;
+		}
 
-		if (relBlockIndex < 0)
+		// Only load the image is RelBlockIndex is valid, otherwise, we won't have requested it.
+		Block = Runner->LoadImage({ Args.blockImage, Op.ExecutionIndex, Op.ExecutionOptions });
+		if (Args.mask)
 		{
-			Result = m_base;
+			Mask = Runner->LoadImage({ Args.mask, Op.ExecutionIndex, Op.ExecutionOptions });
 		}
-		else
+
+		box< UE::Math::TIntVector2<uint16> > RectInblocks;
+		Layout->GetBlock
+		(
+			RelBlockIndex,
+			&RectInblocks.min[0], &RectInblocks.min[1],
+			&RectInblocks.size[0], &RectInblocks.size[1]
+		);
+
+		// Convert the rect from blocks to pixels
+		FIntPoint Grid = Layout->GetGridSize();
+		int32 BlockSizeX = Base->GetSizeX() / Grid[0];
+		int32 BlockSizeY = Base->GetSizeY() / Grid[1];
+		Rect = RectInblocks;
+		Rect.min[0] *= BlockSizeX;
+		Rect.min[1] *= BlockSizeY;
+		Rect.size[0] *= BlockSizeX;
+		Rect.size[1] *= BlockSizeY;
+
+		if (!(Block && Rect.size[0] && Rect.size[1] && Block->GetSizeX() && Block->GetSizeY()))
 		{
-			// Try to avoid the clone here
-			ImagePtr pComposed = CloneOrTakeOver<>(m_base.get());
-			pComposed->m_flags = 0;
-
-			box< UE::Math::TIntVector2<uint16> > rectInblocks;
-			m_layout->GetBlock
-			(
-				relBlockIndex,
-				&rectInblocks.min[0], &rectInblocks.min[1],
-				&rectInblocks.size[0], &rectInblocks.size[1]
-			);
-
-			// Convert the rect from blocks to pixels
-			FIntPoint grid = m_layout->GetGridSize();
-
-			int blockSizeX = m_base->GetSizeX() / grid[0];
-			int blockSizeY = m_base->GetSizeY() / grid[1];
-			box< UE::Math::TIntVector2<uint16> > rect = rectInblocks;
-			rect.min[0] *= blockSizeX;
-			rect.min[1] *= blockSizeY;
-			rect.size[0] *= blockSizeX;
-			rect.size[1] *= blockSizeY;
+			Runner->Release(Block);
+			Runner->Release(Mask);
+			Runner->StoreImage(Op, Base);
+			return false;
+		}
 
-			//
-			if (m_block && rect.size[0] && rect.size[1] && m_block->GetSizeX() && m_block->GetSizeY())
-			{
-				bool useMask = Args.mask != 0;
+		// Create destination data
+		Result = Runner->CloneOrTakeOver(Base);
+		Result->m_flags = 0;
 
-				if (!useMask)
-				{
-					MUTABLE_CPUPROFILER_SCOPE(ImageComposeWithoutMask);
+		bool useMask = Args.mask != 0;
+		if (!useMask)
+		{
+			MUTABLE_CPUPROFILER_SCOPE(ImageComposeWithoutMask);
 
-					// Resize image if it doesn't fit in the new block size
-					if (m_block->GetSizeX() != rect.size[0] ||
-						m_block->GetSizeY() != rect.size[1])
-					{
-						// This now happens more often since the generation of specific mips on request. For this reason
-						// this warning is usually acceptable.
-						//UE_LOG(LogMutableCore, Log, TEXT("Required image resize for ImageCompose: performance warning."));
-						FImageSize blockSize((uint16)rect.size[0], (uint16)rect.size[1]);
-						m_block = ImageResizeLinear(m_imageCompressionQuality, m_block.get(), blockSize);
-					}
+			EImageFormat Format = GetMostGenericFormat(Result->GetFormat(), Block->GetFormat());
 
-					// Change the block image format if it doesn't match the composed image
-					// This is usually enforced at object compilation time.
-					if (pComposed->GetFormat() != m_block->GetFormat())
-					{
-						MUTABLE_CPUPROFILER_SCOPE(ImageComposeReformat);
+			// Resize image if it doesn't fit in the new block size
+			if (Block->GetSize() != Rect.size)
+			{
+				MUTABLE_CPUPROFILER_SCOPE(ImageComposeWithoutMask_BlockResize);
 
-						EImageFormat format = GetMostGenericFormat(pComposed->GetFormat(), m_block->GetFormat());
-						if (pComposed->GetFormat() != format)
-						{
-							pComposed = ImagePixelFormat(m_imageCompressionQuality, pComposed.get(), format);
-						}
-						if (m_block->GetFormat() != format)
-						{
-							m_block = ImagePixelFormat(m_imageCompressionQuality, m_block.get(), format);
-						}
-					}
+				// This now happens more often since the generation of specific mips on request. For this reason
+				// this warning is usually acceptable.
+				Ptr<Image> Resized = Runner->CreateImage(Rect.size[0], Rect.size[1], 1, Block->GetFormat(), EInitializationType::NotInitialized );
+				ImageResizeLinear(Resized.get(), ImageCompressionQuality, Block.get() );
+				Runner->Release(Block);
+				Block = Resized;
+			}
+
+			// Change the block image format if it doesn't match the composed image
+			// This is usually enforced at object compilation time.
+			if (Result->GetFormat() != Block->GetFormat())
+			{
+				MUTABLE_CPUPROFILER_SCOPE(ImageComposeReformat);
 
-					// Compose without a mask
-					ImageCompose(pComposed.get(), m_block.get(), rect);
+				if (Result->GetFormat() != Format)
+				{
+					Ptr<Image> Formatted = Runner->CreateImage(Result->GetSizeX(), Result->GetSizeY(), Result->GetLODCount(), Format, EInitializationType::NotInitialized);
+					bool bSuccess = false;
+					ImagePixelFormat(bSuccess, ImageCompressionQuality, Formatted.get(), Result.get());
+					check(bSuccess); // Decompression cannot fail
+					Runner->Release(Result);
+					Result = Formatted;
 				}
-				else
+				if (Block->GetFormat() != Format)
 				{
-					MUTABLE_CPUPROFILER_SCOPE(ImageComposeWithMask);
-
-					// Compose with a mask
-					ImageBlendOnBaseNoAlpha(pComposed.get(), m_mask.get(), m_block.get(), rect);
+					Ptr<Image> Formatted = Runner->CreateImage(Block->GetSizeX(), Block->GetSizeY(), Block->GetLODCount(), Format, EInitializationType::NotInitialized);
+					bool bSuccess = false;
+					ImagePixelFormat(bSuccess, ImageCompressionQuality, Formatted.get(), Block.get());
+					check(bSuccess); // Decompression cannot fail
+					Runner->Release(Block);
+					Block = Formatted;
 				}
-
 			}
+		}
+
+		return true;
+	}
+
+
+	void FImageComposeTask::DoWork()
+	{
+		// This runs on a random worker thread
+		MUTABLE_CPUPROFILER_SCOPE(FImageComposeTask);
+
+		bool useMask = Args.mask != 0;
+		if (!useMask)
+		{
+			MUTABLE_CPUPROFILER_SCOPE(ImageComposeWithoutMask);
+
+			// Compose without a mask
+			ImageCompose(Result.get(), Block.get(), Rect);
+		}
+		else
+		{
+			MUTABLE_CPUPROFILER_SCOPE(ImageComposeWithMask);
 
-			Result = pComposed;
+			// Compose with a mask
+			ImageBlendOnBaseNoAlpha(Result.get(), Mask.get(), Block.get(), Rect);
 		}
 
-		m_base = nullptr;
-		m_block = nullptr;
-		m_layout = nullptr;
-		m_mask = nullptr;
+		Layout = nullptr;
 	}
 
 
-	//---------------------------------------------------------------------------------------------
-	void FImageComposeTask::Complete(CodeRunner* runner)
+	void FImageComposeTask::Complete(CodeRunner* Runner)
 	{
-		// This runs in the runner thread
-		runner->GetMemory().SetImage(Op, Result);
+		// This runs in the mutable Runner thread
+		Runner->Release(Block);
+		Runner->Release(Mask);
+
+		// If we didn't take a shortcut
+		if (Result)
+		{
+			Runner->StoreImage(Op, Result);
+		}
 	}
 
 
 	//---------------------------------------------------------------------------------------------
 	//---------------------------------------------------------------------------------------------
 	//---------------------------------------------------------------------------------------------
-	bool CodeRunner::FLoadMeshRomTask::Prepare(CodeRunner* runner, const TSharedPtr<const Model>& InModel, bool& bOutFailed )
+	bool CodeRunner::FLoadMeshRomTask::Prepare(CodeRunner* Runner, bool& bOutFailed )
 	{
+		// This runs in the mutable Runner thread
 		MUTABLE_CPUPROFILER_SCOPE(FLoadMeshRomTask_Prepare);
-
-		// This runs in th CodeRunner thread
 		bOutFailed = false;
 
-		FProgram& program = InModel->GetPrivate()->m_program;
+		FProgram& program = Runner->m_pModel->GetPrivate()->m_program;
 
 		check(RomIndex < program.m_roms.Num());
 		bool bRomIsLoaded = program.IsRomLoaded(RomIndex);
-		bool bRomHasPendingOps = runner->m_romPendingOps[RomIndex]!=0;
+
+		FWorkingMemoryManager::FModelCacheEntry& ModelCache = Runner->m_pSystem->WorkingMemoryManager.GetModelCache(Runner->m_pModel);
+		bool bRomHasPendingOps = ModelCache.PendingOpsPerRom[RomIndex]!=0;
 		if (!bRomIsLoaded && !bRomHasPendingOps)
 		{
-			check(runner->m_pSystem->m_pStreamInterface);
+			check(Runner->m_pSystem->StreamInterface);
+
+			uint32 RomSize = program.m_roms[RomIndex].Size;
+			check(RomSize > 0);
 
 			// Free roms if necessary
 			{
 				MUTABLE_CPUPROFILER_SCOPE(FreeingRoms);
-				runner->m_pSystem->m_modelCache.UpdateForLoad(RomIndex, InModel,
-					[&](const Model* pModel, int romInd)
-					{
-						return (pModel == InModel.Get())
-							&&
-							runner->m_romPendingOps[romInd] > 0;
-					});
-			}
 
-			uint32 RomSize = program.m_roms[RomIndex].Size;
-			check(RomSize>0);
+				Runner->m_pSystem->WorkingMemoryManager.MarkRomUsed(RomIndex, Runner->m_pModel);
+				Runner->m_pSystem->WorkingMemoryManager.EnsureBudgetBelow(RomSize);
+			}
 
 			CodeRunner::FRomLoadOp* op = nullptr;
-			for (FRomLoadOp& o : runner->m_romLoadOps)
+			for (FRomLoadOp& o : Runner->m_romLoadOps)
 			{
 				if (o.m_romIndex < 0)
 				{
@@ -1913,8 +2088,8 @@ namespace mu
 
 			if (!op)
 			{
-				runner->m_romLoadOps.Add(CodeRunner::FRomLoadOp());
-				op = &runner->m_romLoadOps.Last();
+				Runner->m_romLoadOps.Add(CodeRunner::FRomLoadOp());
+				op = &Runner->m_romLoadOps.Last();
 			}
 
 			op->m_romIndex = RomIndex;
@@ -1928,16 +2103,16 @@ namespace mu
 			}
 
 			uint32 RomId = program.m_roms[RomIndex].Id;
-			op->m_streamID = runner->m_pSystem->m_pStreamInterface->BeginReadBlock(InModel.Get(), RomId, op->m_streamBuffer.GetData(), RomSize);
+			op->m_streamID = Runner->m_pSystem->StreamInterface->BeginReadBlock(Runner->m_pModel.Get(), RomId, op->m_streamBuffer.GetData(), RomSize);
 			if (op->m_streamID < 0)
 			{
 				bOutFailed = true;
 				return false;
 			}
 		}
-		++runner->m_romPendingOps[RomIndex];
+		++ModelCache.PendingOpsPerRom[RomIndex];
 
-		//UE_LOG(LogMutableCore, Log, TEXT("FLoadMeshRomTask::Prepare romindex %d pending %d."), RomAt.m_rom, runner->m_romPendingOps[RomAt.m_rom]);
+		//UE_LOG(LogMutableCore, Log, TEXT("FLoadMeshRomTask::Prepare romindex %d pending %d."), RomAt.m_rom, Runner->m_romPendingOps[RomAt.m_rom]);
 
 		// No worker thread work
 		return false;
@@ -1945,24 +2120,26 @@ namespace mu
 
 
 	//---------------------------------------------------------------------------------------------
-	void CodeRunner::FLoadMeshRomTask::Complete(CodeRunner* runner)
+	void CodeRunner::FLoadMeshRomTask::Complete(CodeRunner* Runner)
 	{
 		MUTABLE_CPUPROFILER_SCOPE(FLoadMeshRomTask_Complete);
 
-		// This runs in the runner thread
+		// This runs in the Runner thread
 
 		// Process the constant op normally, now that the rom is loaded.
-		runner->RunCode(Op, runner->m_pParams, runner->m_pModel.Get(), runner->m_lodMask);
+		Runner->RunCode(Op, Runner->m_pParams, Runner->m_pModel.Get(), Runner->m_lodMask);
+
+		FWorkingMemoryManager::FModelCacheEntry& ModelCache = Runner->m_pSystem->WorkingMemoryManager.GetModelCache(Runner->m_pModel);
 
-		runner->m_pSystem->m_modelCache.MarkRomUsed(RomIndex, runner->m_pModel);
-		--runner->m_romPendingOps[RomIndex];
+		Runner->m_pSystem->WorkingMemoryManager.MarkRomUsed(RomIndex, Runner->m_pModel);
+		--ModelCache.PendingOpsPerRom[RomIndex];
 	}
 
 
 	//---------------------------------------------------------------------------------------------
-	bool CodeRunner::FLoadMeshRomTask::IsComplete(CodeRunner* runner)
+	bool CodeRunner::FLoadMeshRomTask::IsComplete(CodeRunner* Runner)
 	{ 
-		FProgram& program = runner->m_pModel->GetPrivate()->m_program;
+		FProgram& program = Runner->m_pModel->GetPrivate()->m_program;
 		return program.IsRomLoaded(RomIndex);
 	}
 
@@ -1970,9 +2147,11 @@ namespace mu
 	//---------------------------------------------------------------------------------------------
 	//---------------------------------------------------------------------------------------------
 	//---------------------------------------------------------------------------------------------
-	bool CodeRunner::FLoadExtensionDataTask::Prepare(CodeRunner* Runner, const TSharedPtr<const Model>& InModel, bool& bOutFailed)
+	bool CodeRunner::FLoadExtensionDataTask::Prepare(CodeRunner* Runner, bool& bOutFailed)
 	{
+		// This runs in the mutable Runner thread
 		MUTABLE_CPUPROFILER_SCOPE(FLoadExtensionDataTask_Prepare);
+		bOutFailed = false;
 
 		FProgram& Program = Runner->m_pModel->GetPrivate()->m_program;
 		check(Program.m_constantExtensionData.IsValidIndex(ModelConstantIndex));
@@ -2050,14 +2229,15 @@ namespace mu
 	//---------------------------------------------------------------------------------------------
 	//---------------------------------------------------------------------------------------------
 	//---------------------------------------------------------------------------------------------
-	bool CodeRunner::FLoadImageRomsTask::Prepare(CodeRunner* runner, const TSharedPtr<const Model>& InModel, bool& bOutFailed )
+	bool CodeRunner::FLoadImageRomsTask::Prepare(CodeRunner* Runner, bool& bOutFailed )
 	{
+		// This runs in the mutable Runner thread
 		MUTABLE_CPUPROFILER_SCOPE(FLoadImageRomsTask_Prepare);
-
-		// This runs in th CodeRunner thread
 		bOutFailed = false;
 
-		FProgram& program = InModel->GetPrivate()->m_program;
+		FProgram& program = Runner->m_pModel->GetPrivate()->m_program;
+
+		FWorkingMemoryManager::FModelCacheEntry& ModelCache = Runner->m_pSystem->WorkingMemoryManager.GetModelCache(Runner->m_pModel);
 
 		for (int32 i = 0; i<LODIndexCount; ++i )
 		{
@@ -2074,36 +2254,32 @@ namespace mu
 			check(RomIndex < program.m_roms.Num());
 
 			bool bRomIsLoaded = program.IsRomLoaded(RomIndex);
-			bool bRomHasAlreadyBeenRequested = runner->m_romPendingOps[RomIndex] != 0;
-			++runner->m_romPendingOps[RomIndex];
+			bool bRomHasAlreadyBeenRequested = ModelCache.PendingOpsPerRom[RomIndex] != 0;
+			++ModelCache.PendingOpsPerRom[RomIndex];
 
 			if (DebugRom && (DebugRomAll || RomIndex == DebugRomIndex))
-				UE_LOG(LogMutableCore, Log, TEXT("Preparing rom %d, now peding ops is %d."), RomIndex, runner->m_romPendingOps[RomIndex]);
+				UE_LOG(LogMutableCore, Log, TEXT("Preparing rom %d, now peding ops is %d."), RomIndex, ModelCache.PendingOpsPerRom[RomIndex]);
 
 			if (bRomIsLoaded || bRomHasAlreadyBeenRequested)
 			{
 				continue;
 			}
 
-			check(runner->m_pSystem->m_pStreamInterface);
+			check(Runner->m_pSystem->StreamInterface);
+
+			uint32 RomSize = program.m_roms[RomIndex].Size;
+			check(RomSize > 0);
 
 			// Free roms if necessary
 			{
 				MUTABLE_CPUPROFILER_SCOPE(FreeingRoms);
-				runner->m_pSystem->m_modelCache.UpdateForLoad(RomIndex, InModel,
-					[&](const Model* pModel, int romInd)
-					{
-						return (pModel == InModel.Get())
-							&&
-							runner->m_romPendingOps[romInd] > 0;
-					});
-			}
 
-			uint32 RomSize = program.m_roms[RomIndex].Size;
-			check(RomSize > 0);
+				Runner->m_pSystem->WorkingMemoryManager.MarkRomUsed(RomIndex, Runner->m_pModel);
+				Runner->m_pSystem->WorkingMemoryManager.EnsureBudgetBelow(RomSize);
+			}
 
 			CodeRunner::FRomLoadOp* op = nullptr;
-			for (FRomLoadOp& o : runner->m_romLoadOps)
+			for (FRomLoadOp& o : Runner->m_romLoadOps)
 			{
 				if (o.m_romIndex < 0)
 				{
@@ -2113,8 +2289,8 @@ namespace mu
 
 			if (!op)
 			{
-				runner->m_romLoadOps.Add(CodeRunner::FRomLoadOp());
-				op = &runner->m_romLoadOps.Last();
+				Runner->m_romLoadOps.Add(CodeRunner::FRomLoadOp());
+				op = &Runner->m_romLoadOps.Last();
 			}
 
 			op->m_romIndex = RomIndex;
@@ -2128,7 +2304,7 @@ namespace mu
 			}
 
 			uint32 RomId = program.m_roms[RomIndex].Id;
-			op->m_streamID = runner->m_pSystem->m_pStreamInterface->BeginReadBlock(InModel.Get(), RomId, op->m_streamBuffer.GetData(), RomSize);
+			op->m_streamID = Runner->m_pSystem->StreamInterface->BeginReadBlock(Runner->m_pModel.Get(), RomId, op->m_streamBuffer.GetData(), RomSize);
 			if (op->m_streamID < 0)
 			{
 				bOutFailed = true;
@@ -2136,7 +2312,7 @@ namespace mu
 			}
 		}
 
-		//UE_LOG(LogMutableCore, Log, TEXT("FLoadImageRomsTask::Prepare romindex %d pending %d."), RomAt.m_rom, runner->m_romPendingOps[RomAt.m_rom]);
+		//UE_LOG(LogMutableCore, Log, TEXT("FLoadImageRomsTask::Prepare romindex %d pending %d."), RomAt.m_rom, Runner->m_romPendingOps[RomAt.m_rom]);
 
 		// No worker thread work
 		return false;
@@ -2144,16 +2320,18 @@ namespace mu
 
 
 	//---------------------------------------------------------------------------------------------
-	void CodeRunner::FLoadImageRomsTask::Complete(CodeRunner* runner)
+	void CodeRunner::FLoadImageRomsTask::Complete(CodeRunner* Runner)
 	{
 		MUTABLE_CPUPROFILER_SCOPE(FLoadImageRomsTask_Complete);
 
-		// This runs in the runner thread
+		// This runs in the Runner thread
+
+		FWorkingMemoryManager::FModelCacheEntry& ModelCache = Runner->m_pSystem->WorkingMemoryManager.GetModelCache(Runner->m_pModel);
 
 		// Process the constant op normally, now that the rom is loaded.
-		runner->RunCode(Op, runner->m_pParams, runner->m_pModel.Get(), runner->m_lodMask);
+		Runner->RunCode(Op, Runner->m_pParams, Runner->m_pModel.Get(), Runner->m_lodMask);
 
-		FProgram& program = runner->m_pModel->GetPrivate()->m_program;
+		FProgram& program = Runner->m_pModel->GetPrivate()->m_program;
 		for (int32 i = 0; i < LODIndexCount; ++i)
 		{
 			int32 CurrentIndexIndex = LODIndexIndex + i;
@@ -2168,19 +2346,19 @@ namespace mu
 			int32 RomIndex = program.m_constantImageLODs[CurrentIndex].Key;
 			check(RomIndex < program.m_roms.Num());
 
-			runner->m_pSystem->m_modelCache.MarkRomUsed(RomIndex, runner->m_pModel);
-			--runner->m_romPendingOps[RomIndex];
+			Runner->m_pSystem->WorkingMemoryManager.MarkRomUsed(RomIndex, Runner->m_pModel);
+			--ModelCache.PendingOpsPerRom[RomIndex];
 
 			if (DebugRom && (DebugRomAll || RomIndex==DebugRomIndex))
-				UE_LOG(LogMutableCore, Log, TEXT("FLoadImageRomsTask::Complete rom %d, now peding ops is %d."), RomIndex, runner->m_romPendingOps[RomIndex]);
+				UE_LOG(LogMutableCore, Log, TEXT("FLoadImageRomsTask::Complete rom %d, now peding ops is %d."), RomIndex, ModelCache.PendingOpsPerRom[RomIndex]);
 		}
 	}
 
 
 	//---------------------------------------------------------------------------------------------
-	bool CodeRunner::FLoadImageRomsTask::IsComplete(CodeRunner* runner)
+	bool CodeRunner::FLoadImageRomsTask::IsComplete(CodeRunner* Runner)
 	{
-		FProgram& program = runner->m_pModel->GetPrivate()->m_program;
+		FProgram& program = Runner->m_pModel->GetPrivate()->m_program;
 		for (int32 i = 0; i < LODIndexCount; ++i)
 		{
 			int32 CurrentIndexIndex = LODIndexIndex + i;
@@ -2213,7 +2391,7 @@ namespace mu
 		FImageExternalLoadTask(const FScheduledOp& InItem, uint8 InMipmapsToSkip, EXTERNAL_IMAGE_ID InId);
 
 		// FIssuedTask interface
-		virtual bool Prepare(CodeRunner*, const TSharedPtr<const Model>&, bool& bOutFailed) override;
+		virtual bool Prepare(CodeRunner*, bool& bOutFailed) override;
 		virtual void Complete(CodeRunner* Runner) override;
 		
 	private:
@@ -2235,8 +2413,11 @@ namespace mu
 	}
 
 	//---------------------------------------------------------------------------------------------
-	bool FImageExternalLoadTask::Prepare(CodeRunner* Runner, const TSharedPtr<const Model>&, bool& bOutFailed)
+	bool FImageExternalLoadTask::Prepare(CodeRunner* Runner, bool& bOutFailed)
 	{
+		// This runs in the mutable Runner thread
+		MUTABLE_CPUPROFILER_SCOPE(FLoadImageRomsTask_Prepare);
+
 		// LoadExternalImageAsync will always generate some image even if it is a dummy one.
 		bOutFailed = false;
 
@@ -2263,7 +2444,7 @@ namespace mu
 			Invoke(ExternalCleanUpFunc);
 		}
 
-		Runner->GetMemory().SetImage(Op, Result);
+		Runner->StoreImage(Op, Result);
 	}	
 
 	
@@ -2373,8 +2554,8 @@ namespace mu
 		{
 			if (item.Stage == 1)
 			{
-				OP::ImagePixelFormatArgs args = program.GetOpArgs<OP::ImagePixelFormatArgs>(item.At);
-				Issued = MakeShared<FImagePixelFormatTask>(item, GetMemory(), args, m_pSettings->GetPrivate()->m_imageCompressionQuality);
+				OP::ImagePixelFormatArgs Args = program.GetOpArgs<OP::ImagePixelFormatArgs>(item.At);
+				Issued = MakeShared<FImagePixelFormatTask>(item, Args);
 			}
 			break;
 		}
@@ -2383,8 +2564,8 @@ namespace mu
 		{
 			if (item.Stage == 1)
 			{
-				OP::ImageLayerColourArgs args = program.GetOpArgs<OP::ImageLayerColourArgs>(item.At);
-				Issued = MakeShared<FImageLayerColourTask>(item, GetMemory(), args, m_pSettings->GetPrivate()->m_imageCompressionQuality);
+				OP::ImageLayerColourArgs Args = program.GetOpArgs<OP::ImageLayerColourArgs>(item.At);
+				Issued = MakeShared<FImageLayerColourTask>(item, Args);
 			}
 			break;
 		}
@@ -2396,8 +2577,8 @@ namespace mu
 				(ExecutionStrategy != EExecutionStrategy::MinimizeMemory && item.Stage == 1)
 				)
 			{
-				OP::ImageLayerArgs args = program.GetOpArgs<OP::ImageLayerArgs>(item.At);
-				Issued = MakeShared<FImageLayerTask>(item, GetMemory(), args, m_pSettings->GetPrivate()->m_imageCompressionQuality);
+				OP::ImageLayerArgs Args = program.GetOpArgs<OP::ImageLayerArgs>(item.At);
+				Issued = MakeShared<FImageLayerTask>(item, Args);
 			}
 			break;
 		}
@@ -2406,8 +2587,8 @@ namespace mu
 		{
 			if (item.Stage == 1)
 			{
-				OP::ImageMipmapArgs args = program.GetOpArgs<OP::ImageMipmapArgs>(item.At);
-				Issued = MakeShared<FImageMipmapTask>(item, GetMemory(), args, m_pSettings->GetPrivate()->m_imageCompressionQuality);
+				OP::ImageMipmapArgs Args = program.GetOpArgs<OP::ImageMipmapArgs>(item.At);
+				Issued = MakeShared<FImageMipmapTask>(item, Args);
 			}
 			break;
 		}
@@ -2416,8 +2597,8 @@ namespace mu
 		{
 			if (item.Stage == 1)
 			{
-				OP::ImageSwizzleArgs args = program.GetOpArgs<OP::ImageSwizzleArgs>(item.At);
-				Issued = MakeShared<FImageSwizzleTask>(item, GetMemory(), args);
+				OP::ImageSwizzleArgs Args = program.GetOpArgs<OP::ImageSwizzleArgs>(item.At);
+				Issued = MakeShared<FImageSwizzleTask>(item, Args);
 			}
 			break;
 		}
@@ -2427,7 +2608,7 @@ namespace mu
 			if (item.Stage == 1)
 			{
 				OP::ImageSaturateArgs Args = program.GetOpArgs<OP::ImageSaturateArgs>(item.At);
-				Issued = MakeShared<FImageSaturateTask>(item, GetMemory(), Args);
+				Issued = MakeShared<FImageSaturateTask>(item, Args);
 			}
 			break;
 		}
@@ -2437,7 +2618,7 @@ namespace mu
 			if (item.Stage == 1)
 			{
 				OP::ImageInvertArgs Args = program.GetOpArgs<OP::ImageInvertArgs>(item.At);
-				Issued = MakeShared<FImageInvertTask>(item, GetMemory(), Args);
+				Issued = MakeShared<FImageInvertTask>(item, Args);
 			}
 			break;
 		}
@@ -2447,13 +2628,7 @@ namespace mu
 			if (item.Stage == 1)
 			{
 				OP::ImageResizeArgs Args = program.GetOpArgs<OP::ImageResizeArgs>(item.At);
-				Ptr<const Image> pBase = GetMemory().GetImage(FCacheAddress(Args.source, item));
-
-				// \TODO: Shortcut without creating the task and resolving it if there is no need to resize.
-
-				// \TODO: This operation is ignoring the mips-to-skip operation data
-
-				Issued = MakeShared<FImageResizeTask>(item, GetMemory(), pBase, Args.size[0], Args.size[1], m_pSettings->GetPrivate()->m_imageCompressionQuality);
+				Issued = MakeShared<FImageResizeTask>(item, Args);
 			}
 			break;
 		}
@@ -2463,7 +2638,7 @@ namespace mu
 			if (item.Stage == 1)
 			{
 				OP::ImageResizeRelArgs Args = program.GetOpArgs<OP::ImageResizeRelArgs>(item.At);
-				Issued = MakeShared<FImageResizeRelTask>(item, GetMemory(), Args, m_pSettings->GetPrivate()->m_imageCompressionQuality);
+				Issued = MakeShared<FImageResizeRelTask>(item, Args);
 			}
 			break;
 		}
@@ -2472,9 +2647,9 @@ namespace mu
 		{
 			if (item.Stage == 2)
 			{
-				OP::ImageComposeArgs args = program.GetOpArgs<OP::ImageComposeArgs>(item.At);
+				OP::ImageComposeArgs Args = program.GetOpArgs<OP::ImageComposeArgs>(item.At);
 				Ptr<const Layout> ComposeLayout = static_cast<const Layout*>( m_heapData[item.CustomState].Resource.get());
-				Issued = MakeShared<FImageComposeTask>(item, GetMemory(), args, m_pSettings->GetPrivate()->m_imageCompressionQuality, ComposeLayout);
+				Issued = MakeShared<FImageComposeTask>(item, Args, ComposeLayout);
 			}
 			break;
 		}
