[mutable] Implement memory budget control for image operations.

diff --git a/Engine/Plugins/Experimental/Mutable/Source/MutableRuntime/Private/MuR/ExtensionData.cpp b/Engine/Plugins/Experimental/Mutable/Source/MutableRuntime/Private/MuR/ExtensionData.cpp
--- a/Engine/Plugins/Experimental/Mutable/Source/MutableRuntime/Private/MuR/ExtensionData.cpp
+++ b/Engine/Plugins/Experimental/Mutable/Source/MutableRuntime/Private/MuR/ExtensionData.cpp
@@ -59,4 +59,12 @@ void ExtensionData::Unserialise(InputArchive& Archive)
 	Origin = (EOrigin)OriginByte;
 }
 
+
+//---------------------------------------------------------------------------------------------
+int32 ExtensionData::GetDataSize() const
+{
+	return sizeof(ExtensionData);
+}
+
+
 }
diff --git a/Engine/Plugins/Experimental/Mutable/Source/MutableRuntime/Private/MuR/Image.cpp b/Engine/Plugins/Experimental/Mutable/Source/MutableRuntime/Private/MuR/Image.cpp
--- a/Engine/Plugins/Experimental/Mutable/Source/MutableRuntime/Private/MuR/Image.cpp
+++ b/Engine/Plugins/Experimental/Mutable/Source/MutableRuntime/Private/MuR/Image.cpp
@@ -8,7 +8,6 @@
 #include "MuR/ImagePrivate.h"
 #include "MuR/MutableMath.h"
 #include "MuR/MutableTrace.h"
-#include "MuR/OpImagePixelFormat.h"
 #include "MuR/OpImageResize.h"
 
 #include <initializer_list>
@@ -89,7 +88,7 @@ namespace mu
 
 
     //---------------------------------------------------------------------------------------------
-    Image::Image( uint32_t sizeX, uint32_t sizeY, uint32_t lods, EImageFormat format, EInitializationType Init )
+    Image::Image( uint32 sizeX, uint32 sizeY, uint32 lods, EImageFormat format, EInitializationType Init )
     {
 		MUTABLE_CPUPROFILER_SCOPE(NewImage)
 		LLM_SCOPE_BYNAME(TEXT("MutableRuntime"));
@@ -109,45 +108,57 @@ namespace mu
         int32 PixelsPerBlock = fdata.m_pixelsPerBlockX*fdata.m_pixelsPerBlockY;
         if (PixelsPerBlock)
         {
-			int32 DataSize = CalculateDataSize();
 			if (Init == EInitializationType::Black)
 			{
-				if (bDisableCompressedImageBlackBlockInit)
-				{
-					m_data.SetNumZeroed(DataSize);
-				}
-				else
-				{
-					check(fdata.m_bytesPerBlock <= FImageFormatData::MAX_BYTES_PER_BLOCK);
-					
-					constexpr uint8 ZeroedBlock[FImageFormatData::MAX_BYTES_PER_BLOCK] = { 0 };
+				InitToBlack();
+			}
+			else
+			{
+				int32 DataSize = CalculateDataSize();
+				m_data.SetNumUninitialized(DataSize);
+			}
+        }
+    }
 
-					const SIZE_T BlockSizeSanitized = FMath::Min<SIZE_T>(FImageFormatData::MAX_BYTES_PER_BLOCK, fdata.m_bytesPerBlock);
-					const bool bIsFormatBlackBlockZeroed = FMemory::Memcmp(fdata.BlackBlock, ZeroedBlock, BlockSizeSanitized) == 0;
 
-					if (bIsFormatBlackBlockZeroed)
-					{
-						m_data.SetNumZeroed(DataSize);
-					}
-					else
-					{
-						m_data.SetNumUninitialized(DataSize);
+	//---------------------------------------------------------------------------------------------
+	void Image::InitToBlack()
+	{
+		int32 DataSize = CalculateDataSize();
+		if (bDisableCompressedImageBlackBlockInit)
+		{
+			m_data.SetNumZeroed(DataSize);
+		}
+		else
+		{
+			const FImageFormatData& fdata = GetImageFormatData(m_format);
+			check(fdata.m_bytesPerBlock <= FImageFormatData::MAX_BYTES_PER_BLOCK);
 
-						check(fdata.m_bytesPerBlock > 0);
-						for (int32 BlockDataOffset = 0; BlockDataOffset < DataSize; BlockDataOffset += fdata.m_bytesPerBlock)
-						{
-							FMemory::Memcpy(m_data.GetData() + BlockDataOffset, fdata.BlackBlock, fdata.m_bytesPerBlock);
-						}
-					}
-				}
+			constexpr uint8 ZeroedBlock[FImageFormatData::MAX_BYTES_PER_BLOCK] = { 0 };
+
+			const SIZE_T BlockSizeSanitized = FMath::Min<SIZE_T>(FImageFormatData::MAX_BYTES_PER_BLOCK, fdata.m_bytesPerBlock);
+			const bool bIsFormatBlackBlockZeroed = FMemory::Memcmp(fdata.BlackBlock, ZeroedBlock, BlockSizeSanitized) == 0;
+
+			if (bIsFormatBlackBlockZeroed)
+			{
+				m_data.SetNumZeroed(DataSize);
 			}
 			else
 			{
 				m_data.SetNumUninitialized(DataSize);
+
+				check(fdata.m_bytesPerBlock > 0);
+				for (int32 BlockDataOffset = 0; BlockDataOffset < DataSize; BlockDataOffset += fdata.m_bytesPerBlock)
+				{
+					FMemory::Memcpy(m_data.GetData() + BlockDataOffset, fdata.BlackBlock, fdata.m_bytesPerBlock);
+				}
 			}
-        }
-    }
+		}
 
+		m_flags = 0;
+		RelevancyMinY = 0;
+		RelevancyMaxY = 0;
+	}
 
 	//---------------------------------------------------------------------------------------------
 	Ptr<Image> Image::CreateAsReference(uint32 ID)
@@ -236,6 +247,7 @@ namespace mu
 			else
 			{
 				EImageFormat uncompressedFormat = GetUncompressedFormat(m_format);
+				// TODO: OnlyLOD=Mip and then ExtractMip(0)?
 				Ptr<Image> pTemp = ImagePixelFormat(Quality, this, uncompressedFormat);
 				pTemp = pTemp->ExtractMip(Mip);
 				Ptr<Image> pResult = ImagePixelFormat(Quality, pTemp.get(), m_format);
@@ -328,33 +340,38 @@ namespace mu
 
     //---------------------------------------------------------------------------------------------
     //---------------------------------------------------------------------------------------------
-    //---------------------------------------------------------------------------------------------
-	int32 Image::CalculateDataSize() const
-    {
+	//---------------------------------------------------------------------------------------------
+	int32 Image::CalculateDataSize( int32 SizeX, int32 SizeY, int32 LodCount, EImageFormat Format )
+	{
 		int32 res = 0;
 
-        const FImageFormatData& fdata = GetImageFormatData( m_format );
-        if (fdata.m_bytesPerBlock)
-        {
-            FImageSize s = m_size;
+		const FImageFormatData& fdata = GetImageFormatData(Format);
+		if (fdata.m_bytesPerBlock)
+		{
+			for (int32 LodIndex = 0; LodIndex < FMath::Max(1, LodCount); ++LodIndex)
+			{
+				int32 blocksX = FMath::DivideAndRoundUp(SizeX, int32(fdata.m_pixelsPerBlockX));
+				int32 blocksY = FMath::DivideAndRoundUp(SizeY, int32(fdata.m_pixelsPerBlockY));
 
-            for ( int l=0; l<FMath::Max(1,(int)m_lods); ++l )
-            {
-                int blocksX = FMath::DivideAndRoundUp( s[0], (uint16)fdata.m_pixelsPerBlockX );
-                int blocksY = FMath::DivideAndRoundUp( s[1], (uint16)fdata.m_pixelsPerBlockY );
+				res += (blocksX * blocksY) * fdata.m_bytesPerBlock;
 
-                res += (blocksX*blocksY) * fdata.m_bytesPerBlock;
+				SizeX = FMath::DivideAndRoundUp(SizeX, 2);
+				SizeY = FMath::DivideAndRoundUp(SizeY, 2);
+			}
+		}
 
-                s[0] = FMath::DivideAndRoundUp( s[0], (uint16)2 );
-                s[1] = FMath::DivideAndRoundUp( s[1], (uint16)2 );
-            }
-        }
+		return res;
+	}
 
-        return res;
-    }
+
+	//---------------------------------------------------------------------------------------------
+	int32 Image::CalculateDataSize() const
+	{
+		return CalculateDataSize( m_size[0], m_size[1], m_lods, m_format );
+	}
 
 
-    //---------------------------------------------------------------------------------------------
+	//---------------------------------------------------------------------------------------------
 	int32 Image::CalculateDataSize( int mip ) const
     {
 		int32 res = 0;
diff --git a/Engine/Plugins/Experimental/Mutable/Source/MutableRuntime/Private/MuR/ImageRLE.cpp b/Engine/Plugins/Experimental/Mutable/Source/MutableRuntime/Private/MuR/ImageRLE.cpp
--- a/Engine/Plugins/Experimental/Mutable/Source/MutableRuntime/Private/MuR/ImageRLE.cpp
+++ b/Engine/Plugins/Experimental/Mutable/Source/MutableRuntime/Private/MuR/ImageRLE.cpp
@@ -10,13 +10,13 @@ namespace mu
 {
 
     //---------------------------------------------------------------------------------------------
-    uint32 CompressRLE_L( int32 width, int32 rows, const uint8* pBaseData, uint8* destData, uint32 destDataSize )
+    void CompressRLE_L( uint32& OutCompressedSize, int32 width, int32 rows, const uint8* pBaseData, uint8* destData, uint32 destDataSize )
     {
 		const uint8* InitialBaseData = pBaseData;
 
         uint8* rle = destData;
 
-        size_t maxSize = destDataSize;
+        uint32 maxSize = destDataSize;
 
         // The first uint32 will be the total mip size.
         // Then there is an offset from the initial data pointer for each line.
@@ -25,10 +25,11 @@ namespace mu
         // Could happen in an image of 1x100, for example.
         if ( offset >= maxSize )
         {
-            return 0;
+			OutCompressedSize = 0;
+            return;
         }
 
-        for ( int r = 0; r < rows; ++r )
+        for ( int32 r = 0; r < rows; ++r )
         {
             uint32* pOffset = (uint32*)&rle[sizeof( uint32 ) * ( r + 1 )];
             FMemory::Memmove( pOffset, &offset, sizeof( uint32 ) );
@@ -60,7 +61,8 @@ namespace mu
                 // Copy header
                 if ( maxSize < offset + 4 )
                 {
-                    return 0;
+					OutCompressedSize = 0;
+                    return;
                 }
                 FMemory::Memmove( &rle[offset], &equal, sizeof( uint16 ) );
                 offset += 2;
@@ -76,7 +78,8 @@ namespace mu
                 {
                     if ( maxSize < offset + different )
                     {
-                        return 0;
+						OutCompressedSize = 0;
+						return;
                     }
                     FMemory::Memmove( &rle[offset], pDifferentPixels, different );
                     offset += different;
@@ -92,7 +95,7 @@ namespace mu
 			TArray<uint8> Temp;
 			Temp.SetNum(width*rows);
 			UncompressRLE_L(width,rows,destData,Temp.GetData());
-			int Difference = FMemory::Memcmp(Temp.GetData(), InitialBaseData, width * rows);
+			int32 Difference = FMemory::Memcmp(Temp.GetData(), InitialBaseData, width * rows);
 			if (Difference)
 			{
 				// Different pos.
@@ -107,12 +110,13 @@ namespace mu
 
 				UncompressRLE_L(width, rows, destData, Temp.GetData());
 				CompressRLE_L( width, rows, InitialBaseData, destData, destDataSize);
+				check(false);
 			}
 		}
 #endif
 
         // succeded
-        return offset;
+		OutCompressedSize = offset;
     }
 
 
@@ -124,13 +128,15 @@ namespace mu
 		pBaseData += sizeof(uint32); // Total mip size
         pBaseData += rows*sizeof(uint32); // Size of each row.
 
-        for ( int r=0; r<rows; ++r )
+        for ( int32 r=0; r<rows; ++r )
         {
             const uint8* pDestRowEnd = pDestData + width;
             while ( pDestData!=pDestRowEnd )
             {
                 // Decode header
-                uint16 equal = 0;
+				//check(pBaseData + 4 <= pStartBaseData + BaseDataSize);
+				
+				uint16 equal = 0;
                 FMemory::Memmove(&equal, pBaseData, sizeof(uint16));
                 pBaseData += 2;
 
@@ -150,6 +156,7 @@ namespace mu
                 if (different)
                 {
 					check(pDestData + different <= pStartDestData + width * rows);
+					//check(pBaseData + different <= pStartBaseData + BaseDataSize);
 					FMemory::Memmove( pDestData, pBaseData, different );
                     pDestData += different;
                     pBaseData += different;
@@ -157,14 +164,14 @@ namespace mu
             }
         }
 
-        size_t totalSize = pBaseData-pStartBaseData;
-        check( totalSize==*(uint32*)pStartBaseData );
+		uint32 TotalSize = uint32(pBaseData-pStartBaseData);
+        check( TotalSize==*(uint32*)pStartBaseData );
 
-        return (uint32)totalSize;
+        return TotalSize;
     }
 
     //---------------------------------------------------------------------------------------------
-    uint32 CompressRLE_L1( int32 width, int32 rows,
+    void CompressRLE_L1( uint32& OutCompressedSize, int32 width, int32 rows,
                              const uint8* pBaseData,
                              uint8* destData,
                              uint32 destDataSize )
@@ -175,7 +182,7 @@ namespace mu
         uint32 offset = sizeof(uint32)*(rows+1);
         rle.SetNum( offset, false );
 
-        for ( int r=0; r<rows; ++r )
+        for ( int32 r=0; r<rows; ++r )
         {
             uint32* pOffset = (uint32*) &rle[ sizeof(uint32)*(r+1) ];
             *pOffset = offset;
@@ -214,7 +221,8 @@ namespace mu
         if (destDataSize<(uint32)rle.Num())
         {
             // Failed
-            return 0;
+			OutCompressedSize = 0;
+            return;
         }
 
         uint32* pTotalSize = (uint32*) &rle[ 0 ];
@@ -223,18 +231,18 @@ namespace mu
         FMemory::Memmove( destData, &rle[0], offset );
 
         // succeded
-        return offset;
+		OutCompressedSize = offset;
     }
 
 
     //---------------------------------------------------------------------------------------------
-    uint32 UncompressRLE_L1( int width, int rows, const uint8* pStartBaseData, uint8* pDestData )
+    uint32 UncompressRLE_L1( int32 width, int32 rows, const uint8* pStartBaseData, uint8* pDestData )
     {
         const uint8* pBaseData = pStartBaseData;
         pBaseData += sizeof(uint32); // Total mip size
         pBaseData += rows*sizeof(uint32);
 
-        for ( int r=0; r<rows; ++r )
+        for ( int32 r=0; r<rows; ++r )
         {
             const uint8* pDestRowEnd = pDestData + width;
             while ( pDestData!=pDestRowEnd )
@@ -270,7 +278,7 @@ namespace mu
 
 
     //---------------------------------------------------------------------------------------------
-    void CompressRLE_RGBA( int width, int rows,
+    void CompressRLE_RGBA( int32 width, int32 rows,
                            const uint8* pBaseDataByte,
                            TArray<uint8>& destData )
     {
@@ -281,7 +289,7 @@ namespace mu
         const uint32* pBaseData = (const uint32*)pBaseDataByte;
         rle.SetNum( rows*4, false);
         uint32 offset = sizeof(uint32)*rows;
-        for ( int r=0; r<rows; ++r )
+        for ( int32 r=0; r<rows; ++r )
         {
             uint32* pOffset = (uint32*) &rle[ sizeof(uint32)*r ];
             FMemory::Memmove(pOffset, &offset, sizeof(uint32));
@@ -351,15 +359,15 @@ namespace mu
 
 
     //---------------------------------------------------------------------------------------------
-    void UncompressRLE_RGBA( int width, int rows, const uint8* pBaseData, uint8* pDestDataB )
+    void UncompressRLE_RGBA( int32 width, int32 rows, const uint8* pBaseData, uint8* pDestDataB )
     {
         uint32* pDestData = reinterpret_cast<uint32*>( pDestDataB );
 
         pBaseData += rows*sizeof(uint32);
 
-        int pendingPixels = width*rows;
+        int32 pendingPixels = width*rows;
 
-        for ( int r=0; r<rows; ++r )
+        for ( int32 r=0; r<rows; ++r )
         {
             const uint32* pDestRowEnd = pDestData + width;
             while ( pDestData!=pDestRowEnd )
@@ -379,7 +387,7 @@ namespace mu
 
                 check((equal+different)*4<=pendingPixels);
 
-                for ( int e=0; e<equal*4; ++e )
+                for ( int32 e=0; e<equal*4; ++e )
                 {
                     FMemory::Memmove( pDestData, &equalPixel, 4 );
                     ++pDestData;
@@ -422,7 +430,7 @@ namespace mu
 
 
     //---------------------------------------------------------------------------------------------
-    void CompressRLE_RGB( int width, int rows,
+    void CompressRLE_RGB( int32 width, int32 rows,
                           const uint8* pBaseDataByte,
                           TArray<uint8>& destData )
     {
@@ -432,7 +440,7 @@ namespace mu
         const UINT24* pBaseData = (const UINT24*)pBaseDataByte;
         rle.SetNum( rows*4, false );
         uint32 offset = sizeof(uint32)*rows;
-        for ( int r=0; r<rows; ++r )
+        for ( int32 r=0; r<rows; ++r )
         {
             uint32* pOffset = (uint32*) &rle[ sizeof(uint32)*r ];
             *pOffset = offset;
@@ -502,13 +510,13 @@ namespace mu
 
 
     //---------------------------------------------------------------------------------------------
-    void UncompressRLE_RGB( int width, int rows, const uint8* pBaseData, uint8* pDestDataB )
+    void UncompressRLE_RGB( int32 width, int32 rows, const uint8* pBaseData, uint8* pDestDataB )
     {
         UINT24* pDestData = reinterpret_cast<UINT24*>( pDestDataB );
 
         pBaseData += rows*sizeof(uint32);
 
-        for ( int r=0; r<rows; ++r )
+        for ( int32 r=0; r<rows; ++r )
         {
             const UINT24* pDestRowEnd = pDestData + width;
             while ( pDestData!=pDestRowEnd )
@@ -523,7 +531,7 @@ namespace mu
                 UINT24 equalPixel = *(const UINT24*)pBaseData;
                 pBaseData += 4;
 
-                for ( int e=0; e<equal*4; ++e )
+                for ( int32 e=0; e<equal*4; ++e )
                 {
                     FMemory::Memmove( pDestData, &equalPixel, 3 );
                     ++pDestData;
diff --git a/Engine/Plugins/Experimental/Mutable/Source/MutableRuntime/Private/MuR/ImageRLE.h b/Engine/Plugins/Experimental/Mutable/Source/MutableRuntime/Private/MuR/ImageRLE.h
--- a/Engine/Plugins/Experimental/Mutable/Source/MutableRuntime/Private/MuR/ImageRLE.h
+++ b/Engine/Plugins/Experimental/Mutable/Source/MutableRuntime/Private/MuR/ImageRLE.h
@@ -28,15 +28,14 @@ namespace mu
     //!		- the value of the "equal" pixels
     //! After the header there are as many pixel values as different pixels.
     //---------------------------------------------------------------------------------------------
-	MUTABLERUNTIME_API extern uint32_t CompressRLE_L( int width, int rows,
-                                   const uint8_t* pBaseData,
-                                   uint8_t* destData,
-                                   uint32_t destDataSize );
+	MUTABLERUNTIME_API extern void CompressRLE_L( uint32& OutCompressedSize, int32 width, int32 rows,
+                                   const uint8* BaseData,
+                                   uint8* DestData, uint32 DestDataSize );
 
     //---------------------------------------------------------------------------------------------
-	MUTABLERUNTIME_API extern uint32_t UncompressRLE_L( int width, int rows,
-                                     const uint8_t* pBaseData,
-                                     uint8_t* pDestData );
+	MUTABLERUNTIME_API extern uint32 UncompressRLE_L( int32 width, int32 rows,
+                                     const uint8* BaseData,
+                                     uint8* pDestData );
 
     //---------------------------------------------------------------------------------------------
     //! This RLE format compresses the data per row. At the beginning of the data there are the row
@@ -45,15 +44,15 @@ namespace mu
     //!     - a UINT16 with how many 0 pixels are there
     //!     - a UINT16 with how many 1 pixels follow the 0 pixels
     //---------------------------------------------------------------------------------------------
-	MUTABLERUNTIME_API extern uint32_t CompressRLE_L1( int width, int rows,
-                                    const uint8_t* pBaseData,
-                                    uint8_t* destData,
-                                    uint32_t destDataSize );
+	MUTABLERUNTIME_API extern void CompressRLE_L1( uint32& OutCompressedSize, int32 width, int32 rows,
+                                    const uint8* BaseData,
+                                    uint8* DestData,
+                                    uint32 DestDataSize );
 
     //---------------------------------------------------------------------------------------------
-	MUTABLERUNTIME_API extern uint32_t UncompressRLE_L1( int width, int rows,
-                                      const uint8_t* pBaseData,
-                                      uint8_t* pDestData );
+	MUTABLERUNTIME_API extern uint32 UncompressRLE_L1(int32 width, int32 rows,
+                                      const uint8* pBaseData,
+                                      uint8* pDestData );
 
 
     //---------------------------------------------------------------------------------------------
@@ -65,14 +64,14 @@ namespace mu
     //!     - 4 UINT8 with the value of the "equal" pixels
     //! After the header there are as many pixel values as different pixels.
     //---------------------------------------------------------------------------------------------
-	MUTABLERUNTIME_API extern void CompressRLE_RGBA( int width, int rows,
-									const uint8_t* pBaseDataByte,
-									TArray<uint8_t>& destData );
+	MUTABLERUNTIME_API extern void CompressRLE_RGBA(int32 width, int32 rows,
+									const uint8* pBaseDataByte,
+									TArray<uint8>& destData );
 
     //---------------------------------------------------------------------------------------------
-	MUTABLERUNTIME_API extern void UncompressRLE_RGBA( int width, int rows,
-                                    const uint8_t* pBaseData,
-                                    uint8_t* pDestDataB );
+	MUTABLERUNTIME_API extern void UncompressRLE_RGBA(int32 width, int32 rows,
+                                    const uint8* pBaseData,
+                                    uint8* pDestDataB );
 
 
     //---------------------------------------------------------------------------------------------
@@ -84,14 +83,14 @@ namespace mu
     //!     - 4 UINT8 with the value of the "equal" pixels
     //! After the header there are as many pixel values as different pixels.
     //---------------------------------------------------------------------------------------------
-	MUTABLERUNTIME_API extern void CompressRLE_RGB( int width, int rows,
-									const uint8_t* pBaseDataByte,
-									TArray<uint8_t>& destData );
+	MUTABLERUNTIME_API extern void CompressRLE_RGB(int32 width, int32 rows,
+									const uint8* pBaseDataByte,
+									TArray<uint8>& destData );
 
     //---------------------------------------------------------------------------------------------
-	MUTABLERUNTIME_API extern void UncompressRLE_RGB( int width, int rows,
-                                   const uint8_t* pBaseData,
-                                   uint8_t* pDestDataB );
+	MUTABLERUNTIME_API extern void UncompressRLE_RGB(int32 width, int32 rows,
+                                   const uint8* pBaseData,
+                                   uint8* pDestDataB );
 
 
 	//---------------------------------------------------------------------------------------------
@@ -101,14 +100,14 @@ namespace mu
 	{
 		int sizeX = pBase->GetSizeX();
 		int sizeY = pBase->GetSizeY();
-		const uint8_t* pData = pBase->GetData();
-		uint8_t* pDestData = pDest->GetData();
-		for (int lod = 0; lod < pBase->GetLODCount(); ++lod)
+		const uint8* pData = pBase->GetData();
+		uint8* pDestData = pDest->GetData();
+		for (int32 Lod = 0; Lod < pBase->GetLODCount(); ++Lod)
 		{
 			uint32_t compressedSize = UncompressRLE_L(sizeX, sizeY, pData, pDestData);
 			pData += compressedSize;
 
-			pDestData += pDest->GetLODDataSize(lod);
+			pDestData += pDest->GetLODDataSize(Lod);
 			sizeX = FMath::Max(1, FMath::DivideAndRoundUp(sizeX, 2));
 			sizeY = FMath::Max(1, FMath::DivideAndRoundUp(sizeY, 2));
 		}
@@ -118,32 +117,37 @@ namespace mu
     //---------------------------------------------------------------------------------------------
     //!
     //---------------------------------------------------------------------------------------------
-    inline bool CompressRLE_L( const Image* pBase, Image* pDest )
+    inline void CompressRLE_L(bool& bOutSuccess, const Image* pBase, Image* pDest )
     {
         int sizeX = pBase->GetSizeX();
         int sizeY = pBase->GetSizeY();
-        const uint8_t* pData = pBase->GetData();
-        uint32_t destDataSize = pDest->GetDataSize();
-        uint8_t* pDestData = pDest->GetData();
-        for (int lod=0;lod<pBase->GetLODCount();++lod)
+        const uint8* pData = pBase->GetData();
+        uint32 DestDataSize = pDest->GetDataSize();
+        uint8* pDestData = pDest->GetData();
+        for (int Lod=0;Lod<pBase->GetLODCount();++Lod)
         {
-            uint32_t compressedSize = CompressRLE_L( sizeX, sizeY, pData,
-                                                     pDestData,
-                                                     destDataSize );
+			uint32 compressedSize = 0;
+			CompressRLE_L( compressedSize, sizeX, sizeY,
+				pData,
+				pDestData, DestDataSize );
+
             if (!compressedSize)
             {
-                return false;
+				bOutSuccess = false;
+                return;
             }
+			check(compressedSize<=DestDataSize);
 
-            destDataSize -= compressedSize;
+            DestDataSize -= compressedSize;
             pDestData += compressedSize;
 
-            pData += pBase->GetLODDataSize(lod);
+            pData += pBase->GetLODDataSize(Lod);
             sizeX = FMath::Max(1, FMath::DivideAndRoundUp(sizeX,2));
             sizeY = FMath::Max(1, FMath::DivideAndRoundUp(sizeY,2));
         }
 
-        pDest->m_data.SetNum( pDestData-pDest->GetData() );
+		uint32 TotalDataSize = pDestData - pDest->GetData();
+        pDest->m_data.SetNum( TotalDataSize );
 
 #ifdef MUTABLE_DEBUG_RLE		
 		{
@@ -153,33 +157,33 @@ namespace mu
 		}
 #endif
 
-        return true;
-    }
+		bOutSuccess = true;
+	}
 
 
     //---------------------------------------------------------------------------------------------
     //! This assumes the maximum data size has been preallocateds in pDest. It will return false
     //! if the RLE-compressed data doesn't fit.
     //---------------------------------------------------------------------------------------------
-    inline bool CompressRLE_L1( const Image* pBase, Image* pDest )
+    inline void CompressRLE_L1(bool& bOutSuccess, const Image* pBase, Image* pDest )
     {
         int sizeX = pBase->GetSizeX();
         int sizeY = pBase->GetSizeY();
-        const uint8_t* pData = pBase->GetData();
-        uint32_t destDataSize = pDest->GetDataSize();
-        uint8_t* pDestData = pDest->GetData();
+        const uint8* pData = pBase->GetData();
+        uint32 destDataSize = pDest->GetDataSize();
+        uint8* pDestData = pDest->GetData();
         for (int lod=0;lod<pBase->GetLODCount();++lod)
         {
-            uint32_t compressedSize = CompressRLE_L1( sizeX, sizeY, pData,
-                                                      pDestData,
-                                                      destDataSize );
-            if (!compressedSize)
-            {
-                return false;
+			uint32 CompressedSize=0;
+            CompressRLE_L1( CompressedSize, sizeX, sizeY, pData, pDestData, destDataSize );
+			if (!CompressedSize)
+			{
+				bOutSuccess = false;
+                return;
             }
 
-            destDataSize -= compressedSize;
-            pDestData += compressedSize;
+            destDataSize -= CompressedSize;
+            pDestData += CompressedSize;
 
             pData += pBase->GetLODDataSize(lod);
 			sizeX = FMath::Max(1, FMath::DivideAndRoundUp(sizeX, 2));
@@ -188,7 +192,7 @@ namespace mu
 
         pDest->m_data.SetNum( pDestData-pDest->GetData() );
 
-        return true;
+		bOutSuccess = true;
     }
 
 
@@ -197,8 +201,8 @@ namespace mu
     {
         int sizeX = pBase->GetSizeX();
         int sizeY = pBase->GetSizeY();
-        const uint8_t* pData = pBase->GetData();
-        uint8_t* pDestData = pDest->GetData();
+        const uint8* pData = pBase->GetData();
+        uint8* pDestData = pDest->GetData();
         for (int lod=0;lod<pBase->GetLODCount();++lod)
         {
             uint32_t compressedSize = UncompressRLE_L1( sizeX, sizeY, pData, pDestData );
diff --git a/Engine/Plugins/Experimental/Mutable/Source/MutableRuntime/Private/MuR/Instance.cpp b/Engine/Plugins/Experimental/Mutable/Source/MutableRuntime/Private/MuR/Instance.cpp
--- a/Engine/Plugins/Experimental/Mutable/Source/MutableRuntime/Private/MuR/Instance.cpp
+++ b/Engine/Plugins/Experimental/Mutable/Source/MutableRuntime/Private/MuR/Instance.cpp
@@ -50,6 +50,13 @@ namespace mu
     }
 
 
+	//---------------------------------------------------------------------------------------------
+	int32 Instance::GetDataSize() const
+	{
+		return 16 + sizeof(Private) + m_pD->m_lods.GetAllocatedSize() + m_pD->m_extensionData.GetAllocatedSize();
+	}
+
+
     //---------------------------------------------------------------------------------------------
     Instance::ID Instance::GetId() const
     {
diff --git a/Engine/Plugins/Experimental/Mutable/Source/MutableRuntime/Private/MuR/Mesh.cpp b/Engine/Plugins/Experimental/Mutable/Source/MutableRuntime/Private/MuR/Mesh.cpp
--- a/Engine/Plugins/Experimental/Mutable/Source/MutableRuntime/Private/MuR/Mesh.cpp
+++ b/Engine/Plugins/Experimental/Mutable/Source/MutableRuntime/Private/MuR/Mesh.cpp
@@ -604,11 +604,12 @@ void Mesh::AddSkeletonID(int32 SkeletonID)
 //---------------------------------------------------------------------------------------------
 //---------------------------------------------------------------------------------------------
 //---------------------------------------------------------------------------------------------
-size_t Mesh::GetDataSize() const
+int32 Mesh::GetDataSize() const
 {
-    return m_IndexBuffers.GetDataSize()
-            + m_VertexBuffers.GetDataSize()
-            + m_FaceBuffers.GetDataSize();
+    return sizeof(Mesh)
+		+ m_IndexBuffers.GetDataSize()
+		+ m_VertexBuffers.GetDataSize()
+		+ m_FaceBuffers.GetDataSize();
 }
 
 
diff --git a/Engine/Plugins/Experimental/Mutable/Source/MutableRuntime/Private/MuR/MeshBufferSet.cpp b/Engine/Plugins/Experimental/Mutable/Source/MutableRuntime/Private/MuR/MeshBufferSet.cpp
--- a/Engine/Plugins/Experimental/Mutable/Source/MutableRuntime/Private/MuR/MeshBufferSet.cpp
+++ b/Engine/Plugins/Experimental/Mutable/Source/MutableRuntime/Private/MuR/MeshBufferSet.cpp
@@ -380,9 +380,9 @@ namespace mu
 
 
 	//---------------------------------------------------------------------------------------------
-	size_t FMeshBufferSet::GetDataSize() const
+	int32 FMeshBufferSet::GetDataSize() const
 	{
-		size_t res = 0;
+		int32 res = 0;
 
 		for (int32 b = 0; b < m_buffers.Num(); ++b)
 		{
diff --git a/Engine/Plugins/Experimental/Mutable/Source/MutableRuntime/Private/MuR/ModelPrivate.h b/Engine/Plugins/Experimental/Mutable/Source/MutableRuntime/Private/MuR/ModelPrivate.h
--- a/Engine/Plugins/Experimental/Mutable/Source/MutableRuntime/Private/MuR/ModelPrivate.h
+++ b/Engine/Plugins/Experimental/Mutable/Source/MutableRuntime/Private/MuR/ModelPrivate.h
@@ -213,7 +213,7 @@ namespace mu
         //! Location in the m_byteCode of the beginning of each operation
 		TArray<uint32> m_opAddress;
 
-        //! Byte-coded representation of the program, using flexible-sized ops.
+        //! Byte-coded representation of the program, using variable-sized op data.
 		TArray<uint8> m_byteCode;
 
         //!
@@ -249,13 +249,13 @@ namespace mu
         //! Constant matrices, usually used for transforms
 		TArray<mat4f> m_constantMatrices;
 
-		//! Constant projectors
+		//! Constant shapes
 		TArray<FShape> m_constantShapes;
 
         //! Constant curves
 		TArray<Curve> m_constantCurves;
 
-        //! Constant curves
+        //! Constant skeletons
 		TArray<Ptr<const Skeleton>> m_constantSkeletons;
 
 		//! Constant Physics Bodies
diff --git a/Engine/Plugins/Experimental/Mutable/Source/MutableRuntime/Private/MuR/MutableString.cpp b/Engine/Plugins/Experimental/Mutable/Source/MutableRuntime/Private/MuR/MutableString.cpp
--- a/Engine/Plugins/Experimental/Mutable/Source/MutableRuntime/Private/MuR/MutableString.cpp
+++ b/Engine/Plugins/Experimental/Mutable/Source/MutableRuntime/Private/MuR/MutableString.cpp
@@ -41,6 +41,13 @@ namespace mu {
 	}
 
 
+	//---------------------------------------------------------------------------------------------
+	int32 String::GetDataSize() const
+	{
+		return sizeof(String) + m_value.capacity();
+	}
+
+
 	//---------------------------------------------------------------------------------------------
 	const char* String::GetValue() const
 	{
diff --git a/Engine/Plugins/Experimental/Mutable/Source/MutableRuntime/Private/MuR/OpImageBinarise.h b/Engine/Plugins/Experimental/Mutable/Source/MutableRuntime/Private/MuR/OpImageBinarise.h
--- a/Engine/Plugins/Experimental/Mutable/Source/MutableRuntime/Private/MuR/OpImageBinarise.h
+++ b/Engine/Plugins/Experimental/Mutable/Source/MutableRuntime/Private/MuR/OpImageBinarise.h
@@ -13,15 +13,13 @@ namespace mu
 	//---------------------------------------------------------------------------------------------
 	//!
 	//---------------------------------------------------------------------------------------------
-	inline Ptr<Image> ImageBinarise( const Image* pA, float threshold )
+	inline void ImageBinarise( Image* pDest, const Image* pA, float threshold )
 	{
-		if (!pA)
+		if (!pA || !pDest)
 		{
-			return nullptr;
+			return;
 		}
 
-		Ptr<Image> pDest = new Image( pA->GetSizeX(), pA->GetSizeY(), pA->GetLODCount(), EImageFormat::IF_L_UBYTE );
-
         uint8* pDestBuf = pDest->GetData();
         const uint8* pABuf = pA->GetData();
 
@@ -72,8 +70,6 @@ namespace mu
 		default:
 			check(false);
 		}
-
-		return pDest;
 	}
 
 }
diff --git a/Engine/Plugins/Experimental/Mutable/Source/MutableRuntime/Private/MuR/OpImageColourMap.h b/Engine/Plugins/Experimental/Mutable/Source/MutableRuntime/Private/MuR/OpImageColourMap.h
--- a/Engine/Plugins/Experimental/Mutable/Source/MutableRuntime/Private/MuR/OpImageColourMap.h
+++ b/Engine/Plugins/Experimental/Mutable/Source/MutableRuntime/Private/MuR/OpImageColourMap.h
@@ -8,17 +8,13 @@
 namespace mu
 {
 
-	inline Ptr<Image> ImageColourMap( const Image* pSource, const Image* pMask, const Image* pMap, bool bOnlyOneMip )
+	inline void ImageColourMap( Image* pDest, const Image* pSource, const Image* pMask, const Image* pMap, bool bOnlyOneMip )
 	{
 		check( pSource->GetSizeX() == pMask->GetSizeX() );
 		check( pSource->GetSizeY() == pMask->GetSizeY() );
 		check( pSource->GetLODCount() == pMask->GetLODCount() || bOnlyOneMip );
 		check( pMask->GetFormat() == EImageFormat::IF_L_UBYTE );
 
-		Ptr<Image> pDest = new Image( pSource->GetSizeX(), pSource->GetSizeY(), pSource->GetLODCount(),
-									  pSource->GetFormat(),
-									  EInitializationType::NotInitialized );
-
         uint8* pDestBuf = pDest->GetData();
         const uint8* pSourceBuf = pSource->GetData();
         const uint8* pMaskBuf = pMask->GetData();
@@ -134,8 +130,6 @@ namespace mu
 		default:
 			check(false);
 		}
-
-		return pDest;
 	}
 
 }
diff --git a/Engine/Plugins/Experimental/Mutable/Source/MutableRuntime/Private/MuR/OpImageCompose.h b/Engine/Plugins/Experimental/Mutable/Source/MutableRuntime/Private/MuR/OpImageCompose.h
--- a/Engine/Plugins/Experimental/Mutable/Source/MutableRuntime/Private/MuR/OpImageCompose.h
+++ b/Engine/Plugins/Experimental/Mutable/Source/MutableRuntime/Private/MuR/OpImageCompose.h
@@ -4,7 +4,6 @@
 
 #include "MuR/ImagePrivate.h"
 #include "MuR/MutableMath.h"
-#include "MuR/OpImagePixelFormat.h"
 
 
 namespace mu
diff --git a/Engine/Plugins/Experimental/Mutable/Source/MutableRuntime/Private/MuR/OpImageCrop.h b/Engine/Plugins/Experimental/Mutable/Source/MutableRuntime/Private/MuR/OpImageCrop.h
--- a/Engine/Plugins/Experimental/Mutable/Source/MutableRuntime/Private/MuR/OpImageCrop.h
+++ b/Engine/Plugins/Experimental/Mutable/Source/MutableRuntime/Private/MuR/OpImageCrop.h
@@ -5,82 +5,81 @@
 #include "MuR/ImagePrivate.h"
 #include "MuR/MutableMath.h"
 #include "MuR/Platform.h"
-#include "MuR/OpImagePixelFormat.h"
 
 
 namespace mu
 {
 
-	inline ImagePtr ImageCrop( int imageCompressionQuality, const Image* pBaseReal, const box< vec2<int> >& rect )
+	inline bool ImageCrop( Image* InCropped, int32 CompressionQuality, const Image* InBase, const box< vec2<int32> >& Rect )
 	{
-		ImagePtrConst pBase = pBaseReal;
+		Ptr<const Image> Base = InBase;
+		Ptr<Image> Cropped = InCropped;
 
-		// RLE are not supported yet
-		if ( pBaseReal->GetFormat()==EImageFormat::IF_RGBA_UBYTE_RLE )
-		{
-            pBase = ImagePixelFormat( imageCompressionQuality, pBase.get(), EImageFormat::IF_RGBA_UBYTE );
-        }
-		else if ( pBaseReal->GetFormat()==EImageFormat::IF_RGB_UBYTE_RLE )
-		{
-            pBase = ImagePixelFormat( imageCompressionQuality, pBase.get(), EImageFormat::IF_RGB_UBYTE );
-        }
-		else if ( pBaseReal->GetFormat()==EImageFormat::IF_L_UBYTE_RLE )
+		EImageFormat BaseFormat = Base->GetFormat();
+		EImageFormat UncompressedFormat = GetUncompressedFormat(BaseFormat);
+
+		if (BaseFormat != UncompressedFormat)
 		{
-            pBase = ImagePixelFormat( imageCompressionQuality, pBase.get(), EImageFormat::IF_L_UBYTE );
+			// Compressed formats need decompression + compression after crop			
+			// \TODO: This may use some additional untracked memory locally in this function.
+			Base = ImagePixelFormat(CompressionQuality, Base.get(), UncompressedFormat);
+			Cropped = new Image( InCropped->GetSizeX(), InCropped->GetSizeY(), InCropped->GetLODCount(), UncompressedFormat, EInitializationType::NotInitialized );
         }
 
-		check( pBase );
+		const FImageFormatData& finfo = GetImageFormatData(UncompressedFormat);
 
-		ImagePtr pCropped = new Image
-				(
-					rect.size[0], rect.size[1],
-					1,
-					pBase->GetFormat()
-					);
+		check(Base && Cropped);
+		check(Cropped->GetSizeX() == Rect.size[0]);
+		check(Cropped->GetSizeY() == Rect.size[1]);
 
-//		check( rect.min[0]>=0 && rect.min[1]>=0 );
-//		check( rect.size[0]>=0 && rect.size[1]>=0 );
-//		check( rect.min[0]+rect.size[0]<=pBase->GetSizeX() );
-//		check( rect.min[1]+rect.size[1]<=pBase->GetSizeY() );
 		// TODO: better error control. This happens if some layouts are corrupt.
-		bool correct =
-				( rect.min[0]>=0 && rect.min[1]>=0 ) &&
-				( rect.size[0]>=0 && rect.size[1]>=0 ) &&
-				( rect.min[0]+rect.size[0]<=pBase->GetSizeX() ) &&
-				( rect.min[1]+rect.size[1]<=pBase->GetSizeY() );
-		if (!correct)
+		bool bCorrect =
+				( Rect.min[0]>=0 && Rect.min[1]>=0 ) &&
+				( Rect.size[0]>=0 && Rect.size[1]>=0 ) &&
+				( Rect.min[0]+Rect.size[0]<=Base->GetSizeX() ) &&
+				( Rect.min[1]+Rect.size[1]<=Base->GetSizeY() );
+		if (!bCorrect)
 		{
-			return pCropped;
+			return false;
 		}
 
-
-		const FImageFormatData& finfo = GetImageFormatData( pBase->GetFormat() );
-
-		// For now
+		// Block images are not supported for now
 		check( finfo.m_pixelsPerBlockX == 1 );
 		check( finfo.m_pixelsPerBlockY == 1 );
 
-		checkf( rect.min[0] % finfo.m_pixelsPerBlockX == 0, TEXT("Rect must snap to blocks.") );
-		checkf( rect.min[1] % finfo.m_pixelsPerBlockY == 0, TEXT("Rect must snap to blocks.") );
-		checkf( rect.size[0] % finfo.m_pixelsPerBlockX == 0, TEXT("Rect must snap to blocks.") );
-		checkf( rect.size[1] % finfo.m_pixelsPerBlockY == 0, TEXT("Rect must snap to blocks.") );
+		checkf( Rect.min[0] % finfo.m_pixelsPerBlockX == 0, TEXT("Rect must snap to blocks.") );
+		checkf( Rect.min[1] % finfo.m_pixelsPerBlockY == 0, TEXT("Rect must snap to blocks.") );
+		checkf( Rect.size[0] % finfo.m_pixelsPerBlockX == 0, TEXT("Rect must snap to blocks.") );
+		checkf( Rect.size[1] % finfo.m_pixelsPerBlockY == 0, TEXT("Rect must snap to blocks.") );
 
-		int baseRowSize = finfo.m_bytesPerBlock * pBase->GetSizeX() / finfo.m_pixelsPerBlockX;
-		int cropRowSize = finfo.m_bytesPerBlock * rect.size[0] / finfo.m_pixelsPerBlockX;
+		int baseRowSize = finfo.m_bytesPerBlock * Base->GetSizeX() / finfo.m_pixelsPerBlockX;
+		int cropRowSize = finfo.m_bytesPerBlock * Rect.size[0] / finfo.m_pixelsPerBlockX;
 
-        const uint8_t* pBaseBuf = pBase->GetData();
-        uint8_t* pCropBuf = pCropped->GetData();
+        const uint8_t* pBaseBuf = Base->GetData();
+        uint8_t* pCropBuf = Cropped->GetData();
 
-		int skipPixels = pBase->GetSizeX() * rect.min[1] + rect.min[0];
+		int skipPixels = Base->GetSizeX() * Rect.min[1] + Rect.min[0];
 		pBaseBuf += finfo.m_bytesPerBlock * skipPixels / finfo.m_pixelsPerBlockX;
-		for ( int y=0; y<rect.size[1]; ++y )
+		for ( int y=0; y<Rect.size[1]; ++y )
 		{
-			memcpy( pCropBuf, pBaseBuf, cropRowSize );
+			FMemory::Memcpy( pCropBuf, pBaseBuf, cropRowSize );
 			pCropBuf += cropRowSize;
 			pBaseBuf += baseRowSize;
 		}
 
-		return pCropped;
+		if (BaseFormat != UncompressedFormat)
+		{
+			bool bSuccess = false;
+			int32 DataSize = Cropped->m_data.Num();
+			while (!bSuccess)
+			{
+				InCropped->m_data.SetNumUninitialized(DataSize);
+				ImagePixelFormat(bSuccess, CompressionQuality, InCropped, Cropped.get());
+				DataSize = (DataSize + 16) * 2;
+			}
+		}
+
+		return true;
 	}
 
 }
diff --git a/Engine/Plugins/Experimental/Mutable/Source/MutableRuntime/Private/MuR/OpImageGradient.h b/Engine/Plugins/Experimental/Mutable/Source/MutableRuntime/Private/MuR/OpImageGradient.h
--- a/Engine/Plugins/Experimental/Mutable/Source/MutableRuntime/Private/MuR/OpImageGradient.h
+++ b/Engine/Plugins/Experimental/Mutable/Source/MutableRuntime/Private/MuR/OpImageGradient.h
@@ -6,38 +6,36 @@
 #include "MuR/Platform.h"
 
 
-
 namespace mu
 {
 
-	inline ImagePtr ImageGradient( const FVector4f& c0,
-								   const FVector4f& c1,
-								   int sizeX, int sizeY )
+	inline void ImageGradient( Image* pDest, const FVector4f& c0, const FVector4f& c1 )
 	{
-		ImagePtr pDest = new Image((uint16)sizeX, (uint16)sizeY, 1, EImageFormat::IF_RGB_UBYTE );
+		check(pDest && pDest->GetFormat()==EImageFormat::IF_RGBA_UBYTE);
+
+        uint8* pDestBuf = pDest->GetData();
 
-        uint8_t* pDestBuf = pDest->GetData();
+		int32 sizeX = pDest->GetSizeX();
+		int32 sizeY = pDest->GetSizeX();
 
-		for ( int i=0; i<sizeX; ++i )
+		for ( int i=0; i< sizeX; ++i )
 		{
-			float delta = float(i)/float(sizeX-1);
+			float delta = float(i)/float(sizeX -1);
 			vec3<float> c = c0 * (1.0f-delta) + c1 * delta;
 
-            uint8_t colour[4];
-            colour[0] = (uint8_t)FMath::Max( 0, FMath::Min( 255, int(c[0]*255.0f) ) );
-            colour[1] = (uint8_t)FMath::Max( 0, FMath::Min( 255, int(c[1]*255.0f) ) );
-            colour[2] = (uint8_t)FMath::Max( 0, FMath::Min( 255, int(c[2]*255.0f) ) );
+            uint8 colour[4];
+            colour[0] = (uint8)FMath::Max( 0, FMath::Min( 255, int(c[0]*255.0f) ) );
+            colour[1] = (uint8)FMath::Max( 0, FMath::Min( 255, int(c[1]*255.0f) ) );
+            colour[2] = (uint8)FMath::Max( 0, FMath::Min( 255, int(c[2]*255.0f) ) );
 			colour[3] = 255;
 
-			for ( int j=0; j<sizeY; ++j )
+			for ( int j=0; j< sizeY; ++j )
 			{
 				pDestBuf[ (sizeX*j+i)*3 + 0 ] = colour[0];
 				pDestBuf[ (sizeX*j+i)*3 + 1 ] = colour[1];
 				pDestBuf[ (sizeX*j+i)*3 + 2 ] = colour[2];
 			}
 		}
-
-		return pDest;
 	}
 
 }
diff --git a/Engine/Plugins/Experimental/Mutable/Source/MutableRuntime/Private/MuR/OpImageGrow.h b/Engine/Plugins/Experimental/Mutable/Source/MutableRuntime/Private/MuR/OpImageGrow.h
--- a/Engine/Plugins/Experimental/Mutable/Source/MutableRuntime/Private/MuR/OpImageGrow.h
+++ b/Engine/Plugins/Experimental/Mutable/Source/MutableRuntime/Private/MuR/OpImageGrow.h
@@ -15,7 +15,7 @@ namespace mu
 	//---------------------------------------------------------------------------------------------
 	inline void ImageGrow( Image* pImage )
 	{
-		check( pImage->GetFormat()== EImageFormat::IF_L_UBYTE );
+		check(pImage && pImage->GetFormat()== EImageFormat::IF_L_UBYTE );
 
 		MUTABLE_CPUPROFILER_SCOPE(ImageGrow1)
 
diff --git a/Engine/Plugins/Experimental/Mutable/Source/MutableRuntime/Private/MuR/OpImageInterpolate.h b/Engine/Plugins/Experimental/Mutable/Source/MutableRuntime/Private/MuR/OpImageInterpolate.h
--- a/Engine/Plugins/Experimental/Mutable/Source/MutableRuntime/Private/MuR/OpImageInterpolate.h
+++ b/Engine/Plugins/Experimental/Mutable/Source/MutableRuntime/Private/MuR/OpImageInterpolate.h
@@ -18,6 +18,7 @@ namespace mu
 	{
 		MUTABLE_CPUPROFILER_SCOPE(ImageInterpolate)
 
+		check(pDest && pA && pB);
         check( pA->GetSizeX() == pDest->GetSizeX() );
         check( pA->GetSizeY() == pDest->GetSizeY() );
         check( pA->GetFormat() == pDest->GetFormat() );
@@ -87,8 +88,7 @@ namespace mu
 	//---------------------------------------------------------------------------------------------
 	//!
 	//---------------------------------------------------------------------------------------------
-	inline ImagePtr ImageInterpolate( const Image* p0, const Image* p1, const Image* p2,
-									  float factor1, float factor2 )
+	inline void ImageInterpolate( Image* Dest, const Image* p0, const Image* p1, const Image* p2, float factor1, float factor2 )
 	{
 		MUTABLE_CPUPROFILER_SCOPE(ImageInterpolate2)
 
@@ -99,26 +99,22 @@ namespace mu
 		check( p0->GetSizeY() == p2->GetSizeY() );
 		check( p0->GetFormat() == p2->GetFormat() );
 
-		// Clamp the factor
-		factor1 = FMath::Max( 0.0f, FMath::Min( 1.0f, factor1 ) );
-		factor2 = FMath::Max( 0.0f, FMath::Min( 1.0f-factor1, factor2 ) );
-		float factor0 = FMath::Max( 0.0f, 1.0f-factor1-factor2 );
-
-        ImagePtr pDest = new Image( p0->GetSizeX(), p0->GetSizeY(),
-                                    p0->GetLODCount(),
-                                    p0->GetFormat() );
+		// Clamp the factors
+		factor1 = FMath::Clamp(factor1, 0.0f, 1.0f );
+		factor2 = FMath::Min( 1.0f-factor1, factor2 );
+		float factor0 = 1.0f-factor1-factor2;
 
-        uint8_t* pDestBuf = pDest->GetData();
-        const uint8_t* p0Buf = p0->GetData();
-        const uint8_t* p1Buf = p1->GetData();
-        const uint8_t* p2Buf = p2->GetData();
+        uint8* DestBuf = Dest->GetData();
+        const uint8* p0Buf = p0->GetData();
+        const uint8* p1Buf = p1->GetData();
+        const uint8* p2Buf = p2->GetData();
 
 		// Generic implementation
 		int pixelCount = (int)p0->CalculatePixelCount();
 
-        uint32_t w0_8 = (uint32_t)(factor0*255);
-        uint32_t w1_8 = (uint32_t)(factor1*255);
-        uint32_t w2_8 = (uint32_t)(factor2*255);
+        uint32 w0_8 = (uint32)(factor0*255);
+        uint32 w1_8 = (uint32)(factor1*255);
+        uint32 w2_8 = (uint32)(factor2*255);
 
 		switch ( p0->GetFormat() )
 		{
@@ -126,11 +122,11 @@ namespace mu
 		{
 			for ( int i=0; i<pixelCount; ++i )
 			{
-                uint32_t t0_8 = *p0Buf++;
-                uint32_t t1_8 = *p1Buf++;
-                uint32_t t2_8 = *p2Buf++;
-                uint32_t i_16 = t0_8 * w0_8 + t1_8 * w1_8 + t2_8 * w2_8;
-                *pDestBuf++ = (uint8_t) ( i_16>>8 );
+                uint32 t0_8 = *p0Buf++;
+                uint32 t1_8 = *p1Buf++;
+                uint32 t2_8 = *p2Buf++;
+                uint32 i_16 = t0_8 * w0_8 + t1_8 * w1_8 + t2_8 * w2_8;
+                *DestBuf++ = (uint8) ( i_16>>8 );
 			}
 			break;
 		}
@@ -139,11 +135,11 @@ namespace mu
 		{
 			for ( int i=0; i<pixelCount*3; ++i )
 			{
-                uint32_t t0_8 = *p0Buf++;
-                uint32_t t1_8 = *p1Buf++;
-                uint32_t t2_8 = *p2Buf++;
-                uint32_t i_16 = t0_8 * w0_8 + t1_8 * w1_8 + t2_8 * w2_8;
-                *pDestBuf++ = (uint8_t) ( i_16>>8 );
+                uint32 t0_8 = *p0Buf++;
+                uint32 t1_8 = *p1Buf++;
+                uint32 t2_8 = *p2Buf++;
+                uint32 i_16 = t0_8 * w0_8 + t1_8 * w1_8 + t2_8 * w2_8;
+                *DestBuf++ = (uint8) ( i_16>>8 );
 			}
 			break;
 		}
@@ -153,11 +149,11 @@ namespace mu
         {
 			for ( int i=0; i<pixelCount*4; ++i )
 			{
-                uint32_t t0_8 = *p0Buf++;
-                uint32_t t1_8 = *p1Buf++;
-                uint32_t t2_8 = *p2Buf++;
-                uint32_t i_16 = t0_8 * w0_8 + t1_8 * w1_8 + t2_8 * w2_8;
-                *pDestBuf++ = (uint8_t) (i_16>>8);
+                uint32 t0_8 = *p0Buf++;
+                uint32 t1_8 = *p1Buf++;
+                uint32 t2_8 = *p2Buf++;
+                uint32 i_16 = t0_8 * w0_8 + t1_8 * w1_8 + t2_8 * w2_8;
+                *DestBuf++ = (uint8) (i_16>>8);
 			}
 			break;
 		}
@@ -165,8 +161,6 @@ namespace mu
 		default:
 			check(false);
 		}
-
-		return pDest;
 	}
 
 }
diff --git a/Engine/Plugins/Experimental/Mutable/Source/MutableRuntime/Private/MuR/OpImageLuminance.h b/Engine/Plugins/Experimental/Mutable/Source/MutableRuntime/Private/MuR/OpImageLuminance.h
--- a/Engine/Plugins/Experimental/Mutable/Source/MutableRuntime/Private/MuR/OpImageLuminance.h
+++ b/Engine/Plugins/Experimental/Mutable/Source/MutableRuntime/Private/MuR/OpImageLuminance.h
@@ -8,9 +8,9 @@
 namespace mu
 {
 
-	inline ImagePtr ImageLuminance( const Image* pA )
+	inline void ImageLuminance( Image* pDest, const Image* pA )
 	{
-        ImagePtr pDest = new Image( pA->GetSizeX(), pA->GetSizeY(), pA->GetLODCount(), EImageFormat::IF_L_UBYTE );
+		check(pDest && pA && pDest->GetFormat()==EImageFormat::IF_L_UBYTE);
 
         uint8_t* pDestBuf = pDest->GetData();
         const uint8_t* pABuf = pA->GetData();
@@ -54,8 +54,6 @@ namespace mu
 			check(false);
 			break;
 		}
-
-		return pDest;
 	}
 
 }
diff --git a/Engine/Plugins/Experimental/Mutable/Source/MutableRuntime/Private/MuR/OpImageMipmap.cpp b/Engine/Plugins/Experimental/Mutable/Source/MutableRuntime/Private/MuR/OpImageMipmap.cpp
new file mode 100644
--- /dev/null
+++ b/Engine/Plugins/Experimental/Mutable/Source/MutableRuntime/Private/MuR/OpImageMipmap.cpp
@@ -0,0 +1,790 @@
+// Copyright Epic Games, Inc. All Rights Reserved.
+
+#pragma once
+
+#include "MuR/OpImageMipmap.h"
+#include "MuR/ImagePrivate.h"
+#include "MuR/SystemPrivate.h"
+#include "Async/ParallelFor.h"
+
+namespace mu
+{
+
+	namespace OpImageMipmap_Detail
+	{
+
+		struct SharpenKernelStorage8
+		{
+			constexpr static size_t SIZE = 8;
+			float m_storage[SIZE * SIZE];
+		};
+
+		template<size_t ITERS, size_t N>
+		inline void BlurKernel(float(&inOutKernelStorage)[N])
+		{
+			static_assert(N > 2);
+
+			float temp[N];
+
+			float* inBuf = &inOutKernelStorage[0];
+			float* outBuf = &temp[0];
+
+			for (int32 i = 0; i < ITERS; ++i)
+			{
+				constexpr float oneOverThree = 1.0f / 3.0f;
+
+				outBuf[0] = (inBuf[0] + inBuf[1]) * oneOverThree;
+
+				for (int32 j = 1; j < N - 1; ++j)
+				{
+					outBuf[j] = (inBuf[j - 1] + inBuf[j] + inBuf[j + 1]) * oneOverThree;
+				}
+
+				outBuf[N - 1] = (inBuf[N - 2] + inBuf[N - 1]) * oneOverThree;
+
+				// Ping-pong buffers.
+				std::swap(outBuf, inBuf);
+			}
+
+			if (ITERS % 2 == 1)
+			{
+				for (uint32 i = 0; i < N; ++i)
+				{
+					inOutKernelStorage[i] = temp[i];
+				}
+			}
+		}
+
+		template<size_t N>
+		inline void BuildGaussianBlurKernel(float(&kernel)[N], float variance)
+		{
+			constexpr float sqrtTau = 2.5066282746310002f;
+			const auto normalDist = [sdSqrtTauInv = 1.0f / (FMath::Sqrt(variance) * sqrtTau),
+				twoVInv = 1.0f / (2.0f * variance)]
+				(float x)
+				-> float
+			{
+				return FMath::Exp(-x * x * twoVInv) * sdSqrtTauInv;
+			};
+
+			const float center = static_cast<float>(N) * 0.5f;
+			float sum = 0;
+			for (int32 i = 0; i < N; ++i)
+			{
+				const float dist = normalDist(static_cast<float>(i) - center + 0.5f);
+				kernel[i] = dist;
+				sum += dist;
+			}
+
+			const float normFactor = 1.0f / sum;
+			for (int32 i = 0; i < N; ++i)
+			{
+				kernel[i] *= normFactor;
+			}
+		}
+
+		template<size_t N>
+		inline void BuildSharpenKernel(float(&kernel)[N], float factor)
+		{
+			for (uint32 i = 0; i < N; ++i)
+			{
+				kernel[i] = 0.0f;
+			}
+
+			kernel[(N >> 1) - 1] = factor * 0.5f;
+			kernel[(N >> 1)] = factor * 0.5f;
+		}
+
+		template<size_t N>
+		inline void BuildKernel2DFromSeparable(float(&kernel1D)[N], SharpenKernelStorage8& outKernel)
+		{
+			static_assert(N == SharpenKernelStorage8::SIZE);
+
+			for (int32 y = 0; y < N; ++y)
+			{
+				for (int32 x = 0; x < N; ++x)
+				{
+					outKernel.m_storage[y * N + x] = kernel1D[y] * kernel1D[x];
+				}
+			}
+		}
+
+		inline SharpenKernelStorage8 MakeMipGaussianSharpenKernel(float factor)
+		{
+			SharpenKernelStorage8 kernel;
+
+			float kernel1DStorage0[SharpenKernelStorage8::SIZE];
+			float kernel1DStorage1[SharpenKernelStorage8::SIZE];
+
+			// negative factors indicate blur where the -factor is the variance.
+			if (factor < 0)
+			{
+				BuildGaussianBlurKernel(kernel1DStorage0, -factor);
+
+				BuildKernel2DFromSeparable(kernel1DStorage0, kernel);
+
+				return kernel;
+			}
+
+			BuildSharpenKernel(kernel1DStorage0, 1.0f + factor);
+			BuildSharpenKernel(kernel1DStorage1, -factor);
+
+			BlurKernel<1>(kernel1DStorage0);
+			BlurKernel<3>(kernel1DStorage1);
+
+			for (int32 i = 0; i < SharpenKernelStorage8::SIZE; ++i)
+			{
+				kernel1DStorage0[i] += kernel1DStorage1[i];
+			}
+
+			BuildKernel2DFromSeparable(kernel1DStorage0, kernel);
+
+			return kernel;
+		}
+
+		template<int32 PIXEL_SIZE, typename CHANNEL_TYPE, EAddressMode AD_MODE = EAddressMode::AM_NONE>
+		class ImageAccessor
+		{
+			const FIntVector2 m_dim;
+			CHANNEL_TYPE const* const m_pImageBuf;
+
+		public:
+			ImageAccessor(const CHANNEL_TYPE* const pImageBuf, FIntVector2 dim)
+				: m_dim(dim)
+				, m_pImageBuf(pImageBuf)
+			{
+			}
+
+			uint8 operator()(int32 x, int32 y, int32 c) const
+			{
+				if constexpr (AD_MODE == EAddressMode::AM_WRAP)
+				{
+					// (y & (m_dim.Y - 1)) * m_dim.Y + (x & (m_dim.X - 1))
+					return m_pImageBuf[((y % m_dim.Y) * m_dim.X + (x % m_dim.X)) * PIXEL_SIZE + c];
+				}
+				else if (AD_MODE == EAddressMode::AM_CLAMP)
+				{
+					return m_pImageBuf[(mu::clamp(y, 0, m_dim.Y - 1) * m_dim.X +
+						mu::clamp(x, 0, m_dim.X - 1))
+						* PIXEL_SIZE + c];
+				}
+				else if (AD_MODE == EAddressMode::AM_BLACK_BORDER)
+				{
+					return ((x < m_dim.X) && (y < m_dim.Y) && (x >= 0) && (y >= 0))
+						? m_pImageBuf[(y * m_dim.X + x) * PIXEL_SIZE + c]
+						: 0;
+				}
+				else
+				{
+					return m_pImageBuf[(y * m_dim.X + x) * PIXEL_SIZE + c];
+				}
+			}
+		};
+
+
+		template<int32 PIXEL_SIZE, EAddressMode AD, typename PIXEL_TYPE>
+		void GenerateMipSharpenedRegion(
+			const SharpenKernelStorage8& kernel,
+			const ImageAccessor<PIXEL_SIZE, PIXEL_TYPE, AD>& sourceAccessor,
+			PIXEL_TYPE* Dest, FIntVector2 destSize,
+			FIntVector2 regionStart, FIntVector2 regionEnd)
+		{
+			// Only kernels of size 8 supported.
+			constexpr int32 kernelSize = SharpenKernelStorage8::SIZE;
+			static_assert(kernelSize == 8);
+
+			constexpr int32 kernelCenter = (kernelSize >> 1) - 1;
+			for (int32 y = regionStart.Y; y < regionEnd.Y; ++y)
+			{
+				for (int32 x = regionStart.X; x < regionEnd.X; ++x)
+				{
+					float stagingStorage[kernelSize * kernelSize * PIXEL_SIZE];
+
+					const FIntVector2 sourceCoord = FIntVector2(x << 1, y << 1);
+
+
+					// Convert Image kernel stamp to float. This conversion will be done multiple times.
+					// and should be done once before applying the kernel, or operate directly in unorm8/snorm8.
+					for (int32 ky = 0; ky < kernelSize; ++ky)
+					{
+						for (int32 kx = 0; kx < kernelSize; ++kx)
+						{
+							for (int32 c = 0; c < PIXEL_SIZE; ++c)
+							{
+								stagingStorage[(ky * kernelSize + kx) + (kernelSize * kernelSize * c)] = static_cast<float>(
+									sourceAccessor(sourceCoord.X + kx - kernelCenter,
+										sourceCoord.Y + ky - kernelCenter,
+										c)) / 255.0f;
+							}
+						}
+					}
+
+					// Apply kernel. Although the kernel is separable, the cost would not differ much from the not separable one.
+					// That is because to make it separable a partial convolution of the source is needed. 
+					// Since here we are only applying to half the image and the kernel used is small (8x8) the benefits are not 
+					// clear.
+
+					for (int32 c = 0; c < PIXEL_SIZE; ++c)
+					{
+						float sampleValue = 0.0f;
+						float sum = 0.0f;
+						for (int32 k = 0; k < kernelSize * kernelSize; ++k)
+						{
+							sampleValue += kernel.m_storage[k] * stagingStorage[k + (kernelSize * kernelSize * c)];
+							sum += kernel.m_storage[k];
+						}
+
+						Dest[(y * destSize.X + x) * PIXEL_SIZE + c] = static_cast<uint8>(mu::clamp(sampleValue, 0.0f, 1.0f) * 255.0f);
+					}
+				}
+			}
+		}
+
+		template<size_t PIXEL_SIZE, EAddressMode AD_MODE>
+		inline void GenerateMipmapUint8Sharpen(int mips,
+			const uint8* pSource, uint8* Dest,
+			FIntVector2 sourceSize,
+			float sharpeningFactor)
+		{
+			SharpenKernelStorage8 kernel = MakeMipGaussianSharpenKernel(sharpeningFactor);
+
+			size_t totalDestSize = 0;
+
+			for (; mips >= 0; --mips)
+			{
+				ImageAccessor<PIXEL_SIZE, uint8, AD_MODE> sourceSampler(pSource, sourceSize);
+				ImageAccessor<PIXEL_SIZE, uint8, EAddressMode::AM_NONE> sourceSamplerAddressModeNone(pSource, sourceSize);
+
+				FIntVector2 destSize = FIntVector2(FMath::DivideAndRoundUp(sourceSize.X, 2), FMath::DivideAndRoundUp(sourceSize.Y, 2));
+
+				// Core image, will never sample outside the image so we don't need to care about address mode.
+				// Only sample source pixels where the kernel can be fully applied.
+				// Border cases are treated separately.
+				GenerateMipSharpenedRegion(
+					kernel,
+					sourceSamplerAddressModeNone,
+					Dest, destSize,
+					FIntVector2(2, 2), FIntVector2(destSize.X - 2, destSize.Y - 2));
+
+				// Horizontal borders.
+				GenerateMipSharpenedRegion(
+					kernel,
+					sourceSampler,
+					Dest, destSize,
+					FIntVector2(0, 0), FIntVector2(destSize.X, FMath::Min(2, destSize.Y)));
+
+				GenerateMipSharpenedRegion(
+					kernel,
+					sourceSampler,
+					Dest, destSize,
+					FIntVector2(0, FMath::Max(destSize.Y - 2, 0)), FIntVector2(destSize.X, destSize.Y));
+
+				// Vertical borders
+				GenerateMipSharpenedRegion(
+					kernel,
+					sourceSampler,
+					Dest, destSize,
+					FIntVector2(0, FMath::Min(2, destSize.Y)), FIntVector2(FMath::Min(2, destSize.X), destSize.Y - 2));
+
+				GenerateMipSharpenedRegion(
+					kernel,
+					sourceSampler,
+					Dest, destSize,
+					FIntVector2(destSize.X - 2, FMath::Min(2, destSize.Y)), FIntVector2(destSize.X, destSize.Y - 2));
+
+				sourceSize = destSize;
+
+				pSource = Dest;
+				Dest = Dest + destSize.X * destSize.Y * PIXEL_SIZE;
+
+				totalDestSize += sourceSize.X * sourceSize.Y * PIXEL_SIZE;
+			}
+		}
+
+		template<int32 PIXEL_SIZE>
+		inline void GenerateMipmapUint8Unfiltered(
+			int mips,
+			const uint8* pSource, uint8* Dest,
+			FIntVector2 sourceSize)
+		{
+			for (; mips >= 0; --mips)
+			{
+				FIntVector2 destSize = FIntVector2(FMath::DivideAndRoundUp(sourceSize.X, 2), FMath::DivideAndRoundUp(sourceSize.Y, 2));
+
+				for (int32 y = 0; y < destSize.Y; ++y)
+				{
+					for (int32 x = 0; x < destSize.X; ++x)
+					{
+						for (int32 c = 0; c < PIXEL_SIZE; ++c)
+						{
+							Dest[(y * destSize.X + x) * PIXEL_SIZE + c] =
+								pSource[((y << 1) * sourceSize.X + (x << 1)) * PIXEL_SIZE + c];
+						}
+					}
+				}
+
+				sourceSize = destSize;
+				pSource = Dest;
+				Dest = Dest + destSize.X * destSize.Y * PIXEL_SIZE;
+			}
+		}
+
+		template<int32 PIXEL_SIZE>
+		inline void GenerateMipmapsUint8SimpleAverage(
+			int mips,
+			const uint8* pSource, uint8* Dest,
+			FIntVector2 sourceSize)
+		{
+			const uint8* pMipSource = pSource;
+			uint8* pMipDest = Dest;
+
+			FIntVector2 destSize = sourceSize;
+			for (int m = 0; m < mips; ++m)
+			{
+				check(destSize[0] > 1 || destSize[1] > 1);
+
+				sourceSize = destSize;
+
+				int fullColumns = destSize[0] / 2;
+				bool strayColumn = (destSize[0] % 2) != 0;
+				int fullRows = destSize[1] / 2;
+				bool strayRow = (destSize[1] % 2) != 0;
+
+				destSize[0] = FMath::DivideAndRoundUp(destSize[0], 2);
+				destSize[1] = FMath::DivideAndRoundUp(destSize[1], 2);
+
+				int sourceStride = sourceSize[0] * PIXEL_SIZE;
+				int destStride = destSize[0] * PIXEL_SIZE;
+
+				const auto ProcessRow = [
+					pMipDest, pMipSource, fullColumns, strayColumn, sourceStride, destStride
+				] (uint32 y)
+				{
+					const uint8* pSourceRow0 = pMipSource + 2 * y * sourceStride;
+					const uint8* pSourceRow1 = pSourceRow0 + sourceStride;
+					uint8* pDestRow = pMipDest + y * destStride;
+
+					for (int x = 0; x < fullColumns; ++x)
+					{
+						for (int c = 0; c < PIXEL_SIZE; ++c)
+						{
+							int p = pSourceRow0[c] + pSourceRow0[PIXEL_SIZE + c] + pSourceRow1[c] +
+								pSourceRow1[PIXEL_SIZE + c];
+							pDestRow[c] = (uint8)(p >> 2);
+						}
+
+						pSourceRow0 += 2 * PIXEL_SIZE;
+						pSourceRow1 += 2 * PIXEL_SIZE;
+						pDestRow += PIXEL_SIZE;
+					}
+
+					if (strayColumn)
+					{
+						for (int c = 0; c < PIXEL_SIZE; ++c)
+						{
+							int p = pSourceRow0[c] + pSourceRow1[c];
+							pDestRow[c] = (uint8)(p >> 1);
+						}
+					}
+				};
+
+					constexpr int PixelConcurrencyThreshold = 0xffff;
+					if (destSize[0] * destSize[1] < PixelConcurrencyThreshold)
+					{
+						for (int y = 0; y < fullRows; ++y)
+						{
+							ProcessRow(y);
+						}
+					}
+					else
+					{
+						ParallelFor(fullRows, ProcessRow);
+					}
+
+					if (strayRow)
+					{
+						const uint8* pSourceRow0 = pMipSource + 2 * fullRows * sourceStride;
+						const uint8* pSourceRow1 = pSourceRow0 + sourceStride;
+						uint8* pDestRow = pMipDest + fullRows * destStride;
+
+						for (int x = 0; x < fullColumns; ++x)
+						{
+							for (int c = 0; c < PIXEL_SIZE; ++c)
+							{
+								int p = pSourceRow0[c] + pSourceRow0[PIXEL_SIZE + c];
+								pDestRow[c] = (uint8)(p >> 1);
+							}
+
+							pSourceRow0 += 2 * PIXEL_SIZE;
+							pDestRow += PIXEL_SIZE;
+						}
+
+						if (strayColumn)
+						{
+							for (int c = 0; c < PIXEL_SIZE; ++c)
+							{
+								pDestRow[c] = pSourceRow0[c];
+							}
+						}
+					}
+
+					// Reset the source pointer for the next mip, to use the dest that we have just
+					// generated.
+					pMipSource = pMipDest;
+					pMipDest += destSize[0] * destSize[1] * PIXEL_SIZE;
+			}
+		}
+
+	} // namespace OpImageMipmap_Detail
+
+
+	/** Generate the mipmaps for byte-based images of whatever number of channels.
+	* \param mips number of additional levels to build from the source.
+	*/
+	template<int32 PIXEL_SIZE>
+	inline void GenerateMipmapsUint8(int32 mips,
+		const uint8* pSource, uint8* Dest,
+		FIntVector2 sourceSize,
+		const FMipmapGenerationSettings& settings)
+	{
+		using namespace OpImageMipmap_Detail;
+
+		switch (settings.m_filterType)
+		{
+		case EMipmapFilterType::MFT_SimpleAverage:
+		{
+			GenerateMipmapsUint8SimpleAverage<PIXEL_SIZE>(mips, pSource, Dest, sourceSize);
+			break;
+		}
+		case EMipmapFilterType::MFT_Sharpen:
+		{
+			switch (settings.m_addressMode)
+			{
+			case EAddressMode::AM_CLAMP:
+			{
+				GenerateMipmapUint8Sharpen<PIXEL_SIZE, EAddressMode::AM_CLAMP>(
+					mips, pSource, Dest, sourceSize, settings.m_sharpenFactor);
+				break;
+			}
+			case EAddressMode::AM_WRAP:
+			{
+				GenerateMipmapUint8Sharpen<PIXEL_SIZE, EAddressMode::AM_WRAP>(
+					mips, pSource, Dest, sourceSize, settings.m_sharpenFactor);
+				break;
+			}
+			case EAddressMode::AM_BLACK_BORDER:
+			{
+				GenerateMipmapUint8Sharpen<PIXEL_SIZE, EAddressMode::AM_BLACK_BORDER>(
+					mips, pSource, Dest, sourceSize, settings.m_sharpenFactor);
+				break;
+			}
+			default:
+			{
+				GenerateMipmapUint8Sharpen<PIXEL_SIZE, EAddressMode::AM_NONE>(
+					mips, pSource, Dest, sourceSize, settings.m_sharpenFactor);
+				break;
+			}
+			}
+
+			break;
+		}
+		case EMipmapFilterType::MFT_Unfiltered:
+		{
+			GenerateMipmapUint8Unfiltered<PIXEL_SIZE>(mips, pSource, Dest, sourceSize);
+			break;
+		}
+		default:
+		{
+			check(false);
+			break;
+		}
+		}
+	}
+
+
+
+    //---------------------------------------------------------------------------------------------
+    void ImageMipmap_PrepareScratch(FWorkingMemoryManager& Manager, Image* Dest, const Image* Base, int32 LevelCount, FScratchImageMipmap& Scratch )
+    {
+        int StartLevel = Base->GetLODCount() - 1;
+
+        check(Dest->GetLODCount() == LevelCount);
+        check(Dest->GetSizeX() == Base->GetSizeX());
+        check(Dest->GetSizeY() == Base->GetSizeY());
+        check(Dest->GetFormat() == Base->GetFormat());
+
+        switch (Base->GetFormat())
+        {
+        case EImageFormat::IF_L_UBYTE:
+        case EImageFormat::IF_RGB_UBYTE:
+        case EImageFormat::IF_RGBA_UBYTE:
+        case EImageFormat::IF_BGRA_UBYTE:
+            break;
+
+        // Bad cases: we need to decompress, mip and then recompress. It may be necessary to
+        // generate the latests mips after composing, or in unoptimised code.
+        // TODO: Make sure that the code optimisations avoid this cases generating separate
+        // operations to generate the mip tail.
+        case EImageFormat::IF_BC1:
+        case EImageFormat::IF_BC2:
+        case EImageFormat::IF_BC3:
+        case EImageFormat::IF_BC4:
+        case EImageFormat::IF_BC5:
+        case EImageFormat::IF_ASTC_4x4_RGB_LDR:
+        case EImageFormat::IF_ASTC_4x4_RGBA_LDR:
+        case EImageFormat::IF_ASTC_4x4_RG_LDR:
+        {
+			// Uncompress the last mip that we already have
+			FIntVector2 UncompressedSize = Base->CalculateMipSize(StartLevel);
+			Scratch.Uncompressed = Manager.CreateImage( 
+				(uint16)UncompressedSize[0], (uint16)UncompressedSize[1], 
+				1, 
+				EImageFormat::IF_RGBA_UBYTE);
+
+			FIntVector2 UncompressedMipsSize = Base->CalculateMipSize(StartLevel + 1);
+			// Generate the mipmaps from there on
+			Scratch.UncompressedMips = Manager.CreateImage( 
+				(uint16)UncompressedMipsSize[0], (uint16)UncompressedMipsSize[1],
+				FMath::Max(1, LevelCount - StartLevel - 1), 
+				EImageFormat::IF_RGBA_UBYTE);
+
+			// Compress the mipmapped image
+			Scratch.CompressedMips = Manager.CreateImage(
+				(uint16)UncompressedMipsSize[0], (uint16)UncompressedMipsSize[1],
+				Scratch.UncompressedMips->GetLODCount(),
+				Base->GetFormat());
+
+            break;
+        }
+
+        case EImageFormat::IF_L_UBIT_RLE:
+        case EImageFormat::IF_L_UBYTE_RLE:
+        {
+            // Uncompress the last mip that we already have
+			FIntVector2 UncompressedSize = Base->CalculateMipSize(StartLevel);
+            Scratch.Uncompressed = Manager.CreateImage(
+                (uint16)UncompressedSize[0], (uint16)UncompressedSize[1],
+                1,
+                EImageFormat::IF_L_UBYTE);
+
+
+			FIntVector2 UncompressedMipsSize = Base->CalculateMipSize(StartLevel + 1);
+            // Generate the mipmaps from there on
+            Scratch.UncompressedMips = Manager.CreateImage(
+                (uint16)UncompressedMipsSize[0], (uint16)UncompressedMipsSize[1],
+				FMath::Max(1, LevelCount - StartLevel - 1),
+                EImageFormat::IF_L_UBYTE);
+
+
+            // Compress the mipmapped image
+            Scratch.CompressedMips = Manager.CreateImage(
+                (uint16)UncompressedMipsSize[0], (uint16)UncompressedMipsSize[1],
+                Scratch.UncompressedMips->GetLODCount(),
+                Base->GetFormat());
+
+            // Preallocate ample memory for the compressed data
+			uint32 TotalMemory = Scratch.UncompressedMips->GetDataSize();
+            Scratch.CompressedMips->m_data.SetNumUninitialized(TotalMemory);
+
+            // Preallocate ample memory for the destination data
+			TotalMemory = Base->GetDataSize() + Scratch.UncompressedMips->GetDataSize();
+            Dest->m_data.SetNumUninitialized(TotalMemory);
+
+            break;
+        }
+
+        default:
+            checkf( false, TEXT("Format not implemented in mipmap generation."));
+        }
+    }
+
+
+	void ImageMipmap_ReleaseScratch(FWorkingMemoryManager& Manager, FScratchImageMipmap& Scratch)
+	{
+		Manager.Release(Scratch.Uncompressed);
+		Manager.Release(Scratch.UncompressedMips);
+		Manager.Release(Scratch.CompressedMips);
+	}
+
+
+	void ImageMipmap(FScratchImageMipmap& Scratch, int32 CompressionQuality, Image* Dest, const Image* Base,
+		int32 LevelCount,
+		const FMipmapGenerationSettings& Settings, bool bGenerateOnlyTail)
+	{
+		int32 StartLevel = Base->GetLODCount() - 1;
+
+		check(!(Base->m_flags & Image::IF_CANNOT_BE_SCALED));
+
+
+		if (!bGenerateOnlyTail)
+		{
+			check(Dest->GetLODCount() == LevelCount);
+			check(Dest->GetSizeX() == Base->GetSizeX());
+			check(Dest->GetSizeY() == Base->GetSizeY());
+		}
+		else
+		{
+			check(Dest->GetLODCount() + Base->GetLODCount() == LevelCount);
+
+			check(
+				[&]() -> bool
+				{
+					const FIntVector2 BaseImageNextMipSize = Base->CalculateMipSize(StartLevel + 1);
+					return BaseImageNextMipSize.X == Dest->GetSizeX() && BaseImageNextMipSize.Y == Dest->GetSizeY();
+				}());
+		}
+
+		check(Dest->GetFormat() == Base->GetFormat());
+
+		// Calculate the data size, since the base could have more data allocated.
+		int32 BaseMipDataSize = 0;
+		if (!bGenerateOnlyTail)
+		{
+			BaseMipDataSize = Base->GetMipsDataSize();
+			check(Dest->GetDataSize() >= BaseMipDataSize);
+			FMemory::Memcpy(Dest->GetData(), Base->GetData(), BaseMipDataSize);
+		}
+
+		int mipsToBuild = LevelCount - StartLevel - 1;
+		if (!mipsToBuild)
+		{
+			return;
+		}
+
+		const uint8* pSourceBuf = !bGenerateOnlyTail ? Dest->GetMipData(StartLevel) : Base->GetMipData(StartLevel);
+
+		uint8* pDestBuf = !bGenerateOnlyTail ? Dest->GetMipData(StartLevel + 1) : Dest->GetMipData(0);
+
+		FIntVector2 sourceSize = Base->CalculateMipSize(StartLevel);
+
+		switch (Base->GetFormat())
+		{
+		case EImageFormat::IF_L_UBYTE:
+			GenerateMipmapsUint8<1>(LevelCount - StartLevel - 1, pSourceBuf, pDestBuf, sourceSize, Settings);
+			break;
+
+		case EImageFormat::IF_RGB_UBYTE:
+			GenerateMipmapsUint8<3>(LevelCount - StartLevel - 1, pSourceBuf, pDestBuf, sourceSize, Settings);
+			break;
+
+		case EImageFormat::IF_BGRA_UBYTE:
+		case EImageFormat::IF_RGBA_UBYTE:
+			GenerateMipmapsUint8<4>(LevelCount - StartLevel - 1, pSourceBuf, pDestBuf, sourceSize, Settings);
+			break;
+
+			// Bad cases: we need to decompress, mip and then recompress. It may be necessary to
+			// generate the latests mips after composing, or in unoptimised code.
+			// TODO: Make sure that the code optimisations avoid this cases generating separate
+			// operations to generate the mip tail.
+		case EImageFormat::IF_BC1:
+		case EImageFormat::IF_BC2:
+		case EImageFormat::IF_BC3:
+		case EImageFormat::IF_BC4:
+		case EImageFormat::IF_BC5:
+		case EImageFormat::IF_ASTC_4x4_RGB_LDR:
+		case EImageFormat::IF_ASTC_4x4_RGBA_LDR:
+		case EImageFormat::IF_ASTC_4x4_RG_LDR:
+		case EImageFormat::IF_L_UBIT_RLE:
+		case EImageFormat::IF_L_UBYTE_RLE:
+		{
+			// Uncompress the last mip that we already have
+			bool bSuccess = false;
+			ImagePixelFormat(bSuccess, CompressionQuality, Scratch.Uncompressed.get(), Base, StartLevel);
+			check(bSuccess);
+
+			// Generate the mipmaps from there on
+
+			constexpr bool bGenerateOnlyTailForCompressed = true;
+			ImageMipmap(Scratch, CompressionQuality, Scratch.UncompressedMips.get(),
+				Scratch.Uncompressed.get(), LevelCount - StartLevel, Settings, bGenerateOnlyTailForCompressed);
+
+			// Compress the mipmapped image
+			bSuccess = false;
+			ImagePixelFormat(bSuccess, CompressionQuality, Scratch.CompressedMips.get(), Scratch.UncompressedMips.get());
+			int32 ExcessDataSize = FMath::Max(2, Scratch.CompressedMips->GetDataSize());
+			while (!bSuccess)
+			{
+				// Bad case: this should almost never happen.
+				MUTABLE_CPUPROFILER_SCOPE(Mipmap_Recompression_OutOfSpace);
+
+				Scratch.CompressedMips->m_data.SetNumUninitialized(ExcessDataSize);
+				bSuccess = false;
+				ImagePixelFormat(bSuccess, CompressionQuality, Scratch.CompressedMips.get(), Scratch.UncompressedMips.get());
+				ExcessDataSize *= 4;
+			}
+
+			check(!bGenerateOnlyTail || BaseMipDataSize == 0);
+			const int32 FinalDestSize = BaseMipDataSize + Scratch.CompressedMips->GetDataSize();
+
+			if (FinalDestSize > Dest->GetDataSize())
+			{
+				// Bad case: this should almost never happen.
+				Dest->m_data.SetNumUninitialized(FinalDestSize);
+				pDestBuf = !bGenerateOnlyTail ? Dest->GetMipData(StartLevel + 1) : Dest->GetMipData(0);
+			}
+
+			FMemory::Memcpy(pDestBuf, Scratch.CompressedMips->GetData(), Scratch.CompressedMips->GetDataSize());
+			break;
+		}
+
+		default:
+			checkf(false, TEXT("Format not implemented in mipmap generation."));
+		}
+	}
+
+
+	void ImageMipmap(FWorkingMemoryManager& Manager, int32 CompressionQuality, Image* Dest, const Image* Base,
+		int32 LevelCount,
+		const FMipmapGenerationSettings& Settings, bool bGenerateOnlyTail)
+	{
+		FScratchImageMipmap Scratch;
+		ImageMipmap_PrepareScratch(Manager, Dest, Base, LevelCount, Scratch);
+
+		ImageMipmap(Scratch, CompressionQuality, Dest, Base, LevelCount, Settings, bGenerateOnlyTail);
+
+		ImageMipmap_ReleaseScratch(Manager, Scratch);
+	}
+
+
+	/** Update all the mipmaps in the image from the data in the base one. 
+	* Only the mipmaps already existing in the image are updated.
+	*/
+	void ImageMipmapInPlace(int32 InImageCompressionQuality, Image* InBase, const FMipmapGenerationSettings& InSettings)
+	{
+		int32 StartLevel = 0;
+		int32 LevelCount = InBase->GetLODCount();
+
+		check(!(InBase->m_flags & Image::IF_CANNOT_BE_SCALED));
+
+		int32 MipsToBuild = InBase->GetLODCount()-1;
+		if (!MipsToBuild)
+		{
+			return;
+		}
+
+		const uint8* pSourceBuf = InBase->GetMipData(StartLevel);
+		uint8* pDestBuf = InBase->GetMipData(StartLevel + 1);
+
+		FIntVector2 sourceSize = InBase->CalculateMipSize(StartLevel);
+
+		switch (InBase->GetFormat())
+		{
+		case EImageFormat::IF_L_UBYTE:
+			GenerateMipmapsUint8<1>(MipsToBuild, pSourceBuf, pDestBuf, sourceSize, InSettings);
+			break;
+
+		case EImageFormat::IF_RGB_UBYTE:
+			GenerateMipmapsUint8<3>(MipsToBuild, pSourceBuf, pDestBuf, sourceSize, InSettings);
+			break;
+
+		case EImageFormat::IF_BGRA_UBYTE:
+		case EImageFormat::IF_RGBA_UBYTE:
+			GenerateMipmapsUint8<4>(MipsToBuild, pSourceBuf, pDestBuf, sourceSize, InSettings);
+			break;
+
+		default:
+			checkf(false, TEXT("Format not implemented in mipmap generation."));
+		}
+	}
+
+}
