[mutable] Implement memory budget control for image operations.

diff --git a/Engine/Plugins/Experimental/Mutable/Source/MutableRuntime/Private/MuR/OpImageMipmap.h b/Engine/Plugins/Experimental/Mutable/Source/MutableRuntime/Private/MuR/OpImageMipmap.h
--- a/Engine/Plugins/Experimental/Mutable/Source/MutableRuntime/Private/MuR/OpImageMipmap.h
+++ b/Engine/Plugins/Experimental/Mutable/Source/MutableRuntime/Private/MuR/OpImageMipmap.h
@@ -3,781 +3,36 @@
 #pragma once
 
 #include "MuR/ImagePrivate.h"
-#include "MuR/OpImagePixelFormat.h"
-#include "Async/ParallelFor.h"
 
 namespace mu
 {
-
-	namespace OpImageMipmap_Detail
-	{
-
-		struct SharpenKernelStorage8
-		{
-			constexpr static size_t SIZE = 8;
-			float m_storage[SIZE * SIZE];
-		};
-
-		template<size_t ITERS, size_t N>
-		inline void BlurKernel(float(&inOutKernelStorage)[N])
-		{
-			static_assert(N > 2);
-
-			float temp[N];
-
-			float* inBuf = &inOutKernelStorage[0];
-			float* outBuf = &temp[0];
-
-			for (int32 i = 0; i < ITERS; ++i)
-			{
-				constexpr float oneOverThree = 1.0f / 3.0f;
-
-				outBuf[0] = (inBuf[0] + inBuf[1]) * oneOverThree;
-
-				for (int32 j = 1; j < N - 1; ++j)
-				{
-					outBuf[j] = (inBuf[j - 1] + inBuf[j] + inBuf[j + 1]) * oneOverThree;
-				}
-
-				outBuf[N - 1] = (inBuf[N - 2] + inBuf[N - 1]) * oneOverThree;
-
-				// Ping-pong buffers.
-				std::swap(outBuf, inBuf);
-			}
-
-			if (ITERS % 2 == 1)
-			{
-				for (uint32 i = 0; i < N; ++i)
-				{
-					inOutKernelStorage[i] = temp[i];
-				}
-			}
-		}
-
-		template<size_t N>
-		inline void BuildGaussianBlurKernel(float(&kernel)[N], float variance)
-		{
-			constexpr float sqrtTau = 2.5066282746310002f;
-			const auto normalDist = [sdSqrtTauInv = 1.0f / (FMath::Sqrt(variance) * sqrtTau),
-				twoVInv = 1.0f / (2.0f * variance)]
-				(float x)
-				-> float
-			{
-				return FMath::Exp(-x * x * twoVInv) * sdSqrtTauInv;
-			};
-
-			const float center = static_cast<float>(N) * 0.5f;
-			float sum = 0;
-			for (int32 i = 0; i < N; ++i)
-			{
-				const float dist = normalDist(static_cast<float>(i) - center + 0.5f);
-				kernel[i] = dist;
-				sum += dist;
-			}
-
-			const float normFactor = 1.0f / sum;
-			for (int32 i = 0; i < N; ++i)
-			{
-				kernel[i] *= normFactor;
-			}
-		}
-
-		template<size_t N>
-		inline void BuildSharpenKernel(float(&kernel)[N], float factor)
-		{
-			for (uint32 i = 0; i < N; ++i)
-			{
-				kernel[i] = 0.0f;
-			}
-
-			kernel[(N >> 1) - 1] = factor * 0.5f;
-			kernel[(N >> 1)] = factor * 0.5f;
-		}
-
-		template<size_t N>
-		inline void BuildKernel2DFromSeparable(float(&kernel1D)[N], SharpenKernelStorage8& outKernel)
-		{
-			static_assert(N == SharpenKernelStorage8::SIZE);
-
-			for (int32 y = 0; y < N; ++y)
-			{
-				for (int32 x = 0; x < N; ++x)
-				{
-					outKernel.m_storage[y * N + x] = kernel1D[y] * kernel1D[x];
-				}
-			}
-		}
-
-		inline SharpenKernelStorage8 MakeMipGaussianSharpenKernel(float factor)
-		{
-			SharpenKernelStorage8 kernel;
-
-			float kernel1DStorage0[SharpenKernelStorage8::SIZE];
-			float kernel1DStorage1[SharpenKernelStorage8::SIZE];
-
-			// negative factors indicate blur where the -factor is the variance.
-			if (factor < 0)
-			{
-				BuildGaussianBlurKernel(kernel1DStorage0, -factor);
-
-				BuildKernel2DFromSeparable(kernel1DStorage0, kernel);
-
-				return kernel;
-			}
-
-			BuildSharpenKernel(kernel1DStorage0, 1.0f + factor);
-			BuildSharpenKernel(kernel1DStorage1, -factor);
-
-			BlurKernel<1>(kernel1DStorage0);
-			BlurKernel<3>(kernel1DStorage1);
-
-			for (int32 i = 0; i < SharpenKernelStorage8::SIZE; ++i)
-			{
-				kernel1DStorage0[i] += kernel1DStorage1[i];
-			}
-
-			BuildKernel2DFromSeparable(kernel1DStorage0, kernel);
-
-			return kernel;
-		}
-
-		template<int32 PIXEL_SIZE, typename CHANNEL_TYPE, EAddressMode AD_MODE = EAddressMode::AM_NONE>
-		class ImageAccessor
-		{
-			const FIntVector2 m_dim;
-			CHANNEL_TYPE const* const m_pImageBuf;
-
-		public:
-			ImageAccessor(const CHANNEL_TYPE* const pImageBuf, FIntVector2 dim)
-				: m_dim(dim)
-				, m_pImageBuf(pImageBuf)
-			{
-			}
-
-			uint8 operator()(int32 x, int32 y, int32 c) const
-			{
-				if constexpr (AD_MODE == EAddressMode::AM_WRAP)
-				{
-					// (y & (m_dim.Y - 1)) * m_dim.Y + (x & (m_dim.X - 1))
-					return m_pImageBuf[((y % m_dim.Y) * m_dim.X + (x % m_dim.X)) * PIXEL_SIZE + c];
-				}
-				else if (AD_MODE == EAddressMode::AM_CLAMP)
-				{
-					return m_pImageBuf[(mu::clamp(y, 0, m_dim.Y - 1) * m_dim.X +
-						mu::clamp(x, 0, m_dim.X - 1))
-						* PIXEL_SIZE + c];
-				}
-				else if (AD_MODE == EAddressMode::AM_BLACK_BORDER)
-				{
-					return ((x < m_dim.X) && (y < m_dim.Y) && (x >= 0) && (y >= 0))
-						? m_pImageBuf[(y * m_dim.X + x) * PIXEL_SIZE + c]
-						: 0;
-				}
-				else
-				{
-					return m_pImageBuf[(y * m_dim.X + x) * PIXEL_SIZE + c];
-				}
-			}
-		};
-
-
-		template<int32 PIXEL_SIZE, EAddressMode AD, typename PIXEL_TYPE>
-		void GenerateMipSharpenedRegion(
-			const SharpenKernelStorage8& kernel,
-			const ImageAccessor<PIXEL_SIZE, PIXEL_TYPE, AD>& sourceAccessor,
-			PIXEL_TYPE* pDest, FIntVector2 destSize,
-			FIntVector2 regionStart, FIntVector2 regionEnd)
-		{
-			// Only kernels of size 8 supported.
-			constexpr int32 kernelSize = SharpenKernelStorage8::SIZE;
-			static_assert(kernelSize == 8);
-
-			constexpr int32 kernelCenter = (kernelSize >> 1) - 1;
-			for (int32 y = regionStart.Y; y < regionEnd.Y; ++y)
-			{
-				for (int32 x = regionStart.X; x < regionEnd.X; ++x)
-				{
-					float stagingStorage[kernelSize * kernelSize * PIXEL_SIZE];
-
-					const FIntVector2 sourceCoord = FIntVector2(x << 1, y << 1);
-
-
-					// Convert Image kernel stamp to float. This conversion will be done multiple times.
-					// and should be done once before applying the kernel, or operate directly in unorm8/snorm8.
-					for (int32 ky = 0; ky < kernelSize; ++ky)
-					{
-						for (int32 kx = 0; kx < kernelSize; ++kx)
-						{
-							for (int32 c = 0; c < PIXEL_SIZE; ++c)
-							{
-								stagingStorage[(ky * kernelSize + kx) + (kernelSize * kernelSize * c)] = static_cast<float>(
-									sourceAccessor(sourceCoord.X + kx - kernelCenter,
-										sourceCoord.Y + ky - kernelCenter,
-										c)) / 255.0f;
-							}
-						}
-					}
-
-					// Apply kernel. Although the kernel is separable, the cost would not differ much from the not separable one.
-					// That is because to make it separable a partial convolution of the source is needed. 
-					// Since here we are only applying to half the image and the kernel used is small (8x8) the benefits are not 
-					// clear.
-
-					for (int32 c = 0; c < PIXEL_SIZE; ++c)
-					{
-						float sampleValue = 0.0f;
-						float sum = 0.0f;
-						for (int32 k = 0; k < kernelSize * kernelSize; ++k)
-						{
-							sampleValue += kernel.m_storage[k] * stagingStorage[k + (kernelSize * kernelSize * c)];
-							sum += kernel.m_storage[k];
-						}
-
-						pDest[(y * destSize.X + x) * PIXEL_SIZE + c] = static_cast<uint8>(mu::clamp(sampleValue, 0.0f, 1.0f) * 255.0f);
-					}
-				}
-			}
-		}
-
-		template<size_t PIXEL_SIZE, EAddressMode AD_MODE>
-		inline void GenerateMipmapUint8Sharpen(int mips,
-			const uint8* pSource, uint8* pDest,
-			FIntVector2 sourceSize,
-			float sharpeningFactor)
-		{
-			SharpenKernelStorage8 kernel = MakeMipGaussianSharpenKernel(sharpeningFactor);
-
-			size_t totalDestSize = 0;
-
-			for (; mips >= 0; --mips)
-			{
-				ImageAccessor<PIXEL_SIZE, uint8, AD_MODE> sourceSampler(pSource, sourceSize);
-				ImageAccessor<PIXEL_SIZE, uint8, EAddressMode::AM_NONE> sourceSamplerAddressModeNone(pSource, sourceSize);
-
-				FIntVector2 destSize = FIntVector2(FMath::DivideAndRoundUp(sourceSize.X, 2), FMath::DivideAndRoundUp(sourceSize.Y, 2));
-
-				// Core image, will never sample outside the image so we don't need to care about address mode.
-				// Only sample source pixels where the kernel can be fully applied.
-				// Border cases are treated separately.
-				GenerateMipSharpenedRegion(
-					kernel,
-					sourceSamplerAddressModeNone,
-					pDest, destSize,
-					FIntVector2(2, 2), FIntVector2(destSize.X - 2, destSize.Y - 2));
-
-				// Horizontal borders.
-				GenerateMipSharpenedRegion(
-					kernel,
-					sourceSampler,
-					pDest, destSize,
-					FIntVector2(0, 0), FIntVector2(destSize.X, FMath::Min(2, destSize.Y)));
-
-				GenerateMipSharpenedRegion(
-					kernel,
-					sourceSampler,
-					pDest, destSize,
-					FIntVector2(0, FMath::Max(destSize.Y - 2, 0)), FIntVector2(destSize.X, destSize.Y));
-
-				// Vertical borders
-				GenerateMipSharpenedRegion(
-					kernel,
-					sourceSampler,
-					pDest, destSize,
-					FIntVector2(0, FMath::Min(2, destSize.Y)), FIntVector2(FMath::Min(2, destSize.X), destSize.Y - 2));
-
-				GenerateMipSharpenedRegion(
-					kernel,
-					sourceSampler,
-					pDest, destSize,
-					FIntVector2(destSize.X - 2, FMath::Min(2, destSize.Y)), FIntVector2(destSize.X, destSize.Y - 2));
-
-				sourceSize = destSize;
-
-				pSource = pDest;
-				pDest = pDest + destSize.X * destSize.Y * PIXEL_SIZE;
-
-				totalDestSize += sourceSize.X * sourceSize.Y * PIXEL_SIZE;
-			}
-		}
-
-		template<int32 PIXEL_SIZE>
-		inline void GenerateMipmapUint8Unfiltered(
-			int mips,
-			const uint8* pSource, uint8* pDest,
-			FIntVector2 sourceSize)
-		{
-			for (; mips >= 0; --mips)
-			{
-				FIntVector2 destSize = FIntVector2(FMath::DivideAndRoundUp(sourceSize.X, 2), FMath::DivideAndRoundUp(sourceSize.Y, 2));
-
-				for (int32 y = 0; y < destSize.Y; ++y)
-				{
-					for (int32 x = 0; x < destSize.X; ++x)
-					{
-						for (int32 c = 0; c < PIXEL_SIZE; ++c)
-						{
-							pDest[(y * destSize.X + x) * PIXEL_SIZE + c] =
-								pSource[((y << 1) * sourceSize.X + (x << 1)) * PIXEL_SIZE + c];
-						}
-					}
-				}
-
-				sourceSize = destSize;
-				pSource = pDest;
-				pDest = pDest + destSize.X * destSize.Y * PIXEL_SIZE;
-			}
-		}
-
-		template<int32 PIXEL_SIZE>
-		inline void GenerateMipmapsUint8SimpleAverage(
-			int mips,
-			const uint8* pSource, uint8* pDest,
-			FIntVector2 sourceSize)
-		{
-			const uint8* pMipSource = pSource;
-			uint8* pMipDest = pDest;
-
-			FIntVector2 destSize = sourceSize;
-			for (int m = 0; m < mips; ++m)
-			{
-				check(destSize[0] > 1 || destSize[1] > 1);
-
-				sourceSize = destSize;
-
-				int fullColumns = destSize[0] / 2;
-				bool strayColumn = (destSize[0] % 2) != 0;
-				int fullRows = destSize[1] / 2;
-				bool strayRow = (destSize[1] % 2) != 0;
-
-				destSize[0] = FMath::DivideAndRoundUp(destSize[0], 2);
-				destSize[1] = FMath::DivideAndRoundUp(destSize[1], 2);
-
-				int sourceStride = sourceSize[0] * PIXEL_SIZE;
-				int destStride = destSize[0] * PIXEL_SIZE;
-
-				const auto ProcessRow = [
-					pMipDest, pMipSource, fullColumns, strayColumn, sourceStride, destStride
-				] (uint32 y)
-				{
-					const uint8* pSourceRow0 = pMipSource + 2 * y * sourceStride;
-					const uint8* pSourceRow1 = pSourceRow0 + sourceStride;
-					uint8* pDestRow = pMipDest + y * destStride;
-
-					for (int x = 0; x < fullColumns; ++x)
-					{
-						for (int c = 0; c < PIXEL_SIZE; ++c)
-						{
-							int p = pSourceRow0[c] + pSourceRow0[PIXEL_SIZE + c] + pSourceRow1[c] +
-								pSourceRow1[PIXEL_SIZE + c];
-							pDestRow[c] = (uint8)(p >> 2);
-						}
-
-						pSourceRow0 += 2 * PIXEL_SIZE;
-						pSourceRow1 += 2 * PIXEL_SIZE;
-						pDestRow += PIXEL_SIZE;
-					}
-
-					if (strayColumn)
-					{
-						for (int c = 0; c < PIXEL_SIZE; ++c)
-						{
-							int p = pSourceRow0[c] + pSourceRow1[c];
-							pDestRow[c] = (uint8)(p >> 1);
-						}
-					}
-				};
-
-					constexpr int PixelConcurrencyThreshold = 0xffff;
-					if (destSize[0] * destSize[1] < PixelConcurrencyThreshold)
-					{
-						for (int y = 0; y < fullRows; ++y)
-						{
-							ProcessRow(y);
-						}
-					}
-					else
-					{
-						ParallelFor(fullRows, ProcessRow);
-					}
-
-					if (strayRow)
-					{
-						const uint8* pSourceRow0 = pMipSource + 2 * fullRows * sourceStride;
-						const uint8* pSourceRow1 = pSourceRow0 + sourceStride;
-						uint8* pDestRow = pMipDest + fullRows * destStride;
-
-						for (int x = 0; x < fullColumns; ++x)
-						{
-							for (int c = 0; c < PIXEL_SIZE; ++c)
-							{
-								int p = pSourceRow0[c] + pSourceRow0[PIXEL_SIZE + c];
-								pDestRow[c] = (uint8)(p >> 1);
-							}
-
-							pSourceRow0 += 2 * PIXEL_SIZE;
-							pDestRow += PIXEL_SIZE;
-						}
-
-						if (strayColumn)
-						{
-							for (int c = 0; c < PIXEL_SIZE; ++c)
-							{
-								pDestRow[c] = pSourceRow0[c];
-							}
-						}
-					}
-
-					// Reset the source pointer for the next mip, to use the dest that we have just
-					// generated.
-					pMipSource = pMipDest;
-					pMipDest += destSize[0] * destSize[1] * PIXEL_SIZE;
-			}
-		}
-
-	} // namespace OpImageMipmap_Detail
-
-
-		//---------------------------------------------------------------------------------------------
-		//! Generate the mipmaps for byte-based images of whatever number of channels.
-		//! \param mips number of additional levels to build from the source.
-		//---------------------------------------------------------------------------------------------
-	template<int32 PIXEL_SIZE>
-	inline void GenerateMipmapsUint8(int mips,
-		const uint8* pSource, uint8* pDest,
-		FIntVector2 sourceSize,
-		const FMipmapGenerationSettings& settings)
-	{
-		using namespace OpImageMipmap_Detail;
-
-		switch (settings.m_filterType)
-		{
-		case EMipmapFilterType::MFT_SimpleAverage:
-		{
-			GenerateMipmapsUint8SimpleAverage<PIXEL_SIZE>(mips, pSource, pDest, sourceSize);
-			break;
-		}
-		case EMipmapFilterType::MFT_Sharpen:
-		{
-			switch (settings.m_addressMode)
-			{
-			case EAddressMode::AM_CLAMP:
-			{
-				GenerateMipmapUint8Sharpen<PIXEL_SIZE, EAddressMode::AM_CLAMP>(
-					mips, pSource, pDest, sourceSize, settings.m_sharpenFactor);
-				break;
-			}
-			case EAddressMode::AM_WRAP:
-			{
-				GenerateMipmapUint8Sharpen<PIXEL_SIZE, EAddressMode::AM_WRAP>(
-					mips, pSource, pDest, sourceSize, settings.m_sharpenFactor);
-				break;
-			}
-			case EAddressMode::AM_BLACK_BORDER:
-			{
-				GenerateMipmapUint8Sharpen<PIXEL_SIZE, EAddressMode::AM_BLACK_BORDER>(
-					mips, pSource, pDest, sourceSize, settings.m_sharpenFactor);
-				break;
-			}
-			default:
-			{
-				GenerateMipmapUint8Sharpen<PIXEL_SIZE, EAddressMode::AM_NONE>(
-					mips, pSource, pDest, sourceSize, settings.m_sharpenFactor);
-				break;
-			}
-			}
-
-			break;
-		}
-		case EMipmapFilterType::MFT_Unfiltered:
-		{
-			GenerateMipmapUint8Unfiltered<PIXEL_SIZE>(mips, pSource, pDest, sourceSize);
-			break;
-		}
-		default:
-		{
-			check(false);
-			break;
-		}
-		}
-	}
-
-	struct SCRATCH_IMAGE_MIPMAP
+	/** Support struct to keep preallocated data required for some mipmap operations. */
+	struct FScratchImageMipmap
 	{
-		ImagePtr pUncompressed;
-		ImagePtr pUncompressedMips;
-		ImagePtr pCompressedMips;
+		Ptr<Image> Uncompressed;
+		Ptr<Image> UncompressedMips;
+		Ptr<Image> CompressedMips;
 	};
 
-
-    //---------------------------------------------------------------------------------------------
-    inline void ImageMipmap_PrepareScratch( Image* pDest, const Image* pBase, int levelCount,
-                                            SCRATCH_IMAGE_MIPMAP* scratch )
-    {
-        int startLevel = pBase->GetLODCount() - 1;
-
-        check(pDest->GetLODCount() == levelCount);
-        check(pDest->GetSizeX() == pBase->GetSizeX());
-        check(pDest->GetSizeY() == pBase->GetSizeY());
-        check(pDest->GetFormat() == pBase->GetFormat());
-		(void)pDest;
-
-        switch (pBase->GetFormat())
-        {
-        case EImageFormat::IF_L_UBYTE:
-        case EImageFormat::IF_RGB_UBYTE:
-        case EImageFormat::IF_RGBA_UBYTE:
-        case EImageFormat::IF_BGRA_UBYTE:
-            break;
-
-        // Bad cases: we need to decompress, mip and then recompress. It may be necessary to
-        // generate the latests mips after composing, or in unoptimised code.
-        // TODO: Make sure that the code optimisations avoid this cases generating separate
-        // operations to generate the mip tail.
-        case EImageFormat::IF_BC1:
-        case EImageFormat::IF_BC2:
-        case EImageFormat::IF_BC3:
-        case EImageFormat::IF_BC4:
-        case EImageFormat::IF_BC5:
-        case EImageFormat::IF_ASTC_4x4_RGB_LDR:
-        case EImageFormat::IF_ASTC_4x4_RGBA_LDR:
-        case EImageFormat::IF_ASTC_4x4_RG_LDR:
-        {
-		// Uncompress the last mip that we already have
-			FIntVector2 uncompressedSize = pBase->CalculateMipSize(startLevel);
-			scratch->pUncompressed = new Image(
-				(uint16)uncompressedSize[0],
-				(uint16)uncompressedSize[1],
-				1,
-				EImageFormat::IF_RGBA_UBYTE);
-
-			FIntVector2 uncompressedMipsSize = pBase->CalculateMipSize(startLevel + 1);
-			// Generate the mipmaps from there on
-			scratch->pUncompressedMips = new Image(
-				(uint16)uncompressedMipsSize[0],
-				(uint16)uncompressedMipsSize[1],
-				FMath::Max(1, levelCount - startLevel - 1),
-				EImageFormat::IF_RGBA_UBYTE);
-
-			// Compress the mipmapped image
-			scratch->pCompressedMips = new Image(
-				(uint16)uncompressedMipsSize[0],
-				(uint16)uncompressedMipsSize[1],
-				scratch->pUncompressedMips->GetLODCount(),
-				pBase->GetFormat());
-
-            break;
-        }
-
-        case EImageFormat::IF_L_UBIT_RLE:
-        case EImageFormat::IF_L_UBYTE_RLE:
-        {
-            // Uncompress the last mip that we already have
-			FIntVector2 uncompressedSize = pBase->CalculateMipSize(startLevel);
-            scratch->pUncompressed = new Image(
-                (uint16)uncompressedSize[0],
-                (uint16)uncompressedSize[1],
-                1,
-                EImageFormat::IF_L_UBYTE);
-
-
-			FIntVector2 uncompressedMipsSize = pBase->CalculateMipSize(startLevel + 1);
-            // Generate the mipmaps from there on
-            scratch->pUncompressedMips = new Image(
-                (uint16)uncompressedMipsSize[0],
-                (uint16)uncompressedMipsSize[1],
-				FMath::Max(1, levelCount - startLevel - 1),
-                EImageFormat::IF_L_UBYTE);
-
-
-            // Compress the mipmapped image
-            scratch->pCompressedMips = new Image(
-                (uint16)uncompressedMipsSize[0],
-                (uint16)uncompressedMipsSize[1],
-                scratch->pUncompressedMips->GetLODCount(),
-                pBase->GetFormat());
-
-            // Preallocate ample memory for the compressed data
-            scratch->pCompressedMips->m_data.SetNumUninitialized(scratch->pUncompressedMips->GetDataSize());
-
-            // Preallocate ample memory for the destination data
-            pDest->m_data.SetNumUninitialized(pBase->GetDataSize() + scratch->pUncompressedMips->GetDataSize());
-
-            break;
-        }
-
-        default:
-            checkf( false, TEXT("Format not implemented in mipmap generation."));
-        }
-    }
-
-    //---------------------------------------------------------------------------------------------
-    //! Generate the mipmaps for images.
-	//! if bGenerateOnlyTail is true, generates the mips missing from pBase to levelCount and sets
-	//! them in pDest (the full chain is spit in two images). Otherwise generate the mips missing 
-	//! from pBase up to levelCount and append them in pDest to the already generated pBase's mips.
-    //---------------------------------------------------------------------------------------------
-    inline void ImageMipmap( int imageCompressionQuality, Image* pDest, const Image* pBase,
-                             int levelCount, SCRATCH_IMAGE_MIPMAP* scratch, 
-							 const FMipmapGenerationSettings& settings, bool bGenerateOnlyTail = false )
-    {
-        int startLevel = pBase->GetLODCount() - 1;
-
-		check(!(pBase->m_flags & Image::IF_CANNOT_BE_SCALED));
-
-	
-		if (!bGenerateOnlyTail)
-		{
-			check(pDest->GetLODCount() == levelCount);
-			check(pDest->GetSizeX() == pBase->GetSizeX());
-			check(pDest->GetSizeY() == pBase->GetSizeY());
-		}
-		else
-		{	
-			check(pDest->GetLODCount() + pBase->GetLODCount() == levelCount);
-
-			check(
-				[&]() -> bool
-				{
-					const FIntVector2 BaseImageNextMipSize = pBase->CalculateMipSize(startLevel + 1);
-					return BaseImageNextMipSize.X == pDest->GetSizeX() && BaseImageNextMipSize.Y == pDest->GetSizeY();
-				}());
-		}
-
-        check(pDest->GetFormat() == pBase->GetFormat());
-
-        // Calculate the data size, since the base could have more data allocated.
-		int32 BaseMipDataSize = 0;
-		if (!bGenerateOnlyTail)
-		{
-			BaseMipDataSize = pBase->GetMipsDataSize();
-			check(pDest->GetDataSize() >= BaseMipDataSize);
-			FMemory::Memcpy(pDest->GetData(), pBase->GetData(), BaseMipDataSize);
-		}
-
-        int mipsToBuild = levelCount - startLevel - 1;
-        if (!mipsToBuild)
-        {
-            return;
-        }
-
-        const uint8* pSourceBuf = !bGenerateOnlyTail ? pDest->GetMipData(startLevel) : pBase->GetMipData(startLevel);
-
-        uint8* pDestBuf = !bGenerateOnlyTail ? pDest->GetMipData(startLevel + 1) : pDest->GetMipData(0);
-	
-		FIntVector2 sourceSize = pBase->CalculateMipSize(startLevel);
-
-        switch (pBase->GetFormat())
-        {
-        case EImageFormat::IF_L_UBYTE:
-            GenerateMipmapsUint8<1>(levelCount - startLevel - 1, pSourceBuf, pDestBuf, sourceSize, settings);
-            break; 
-
-        case EImageFormat::IF_RGB_UBYTE:
-            GenerateMipmapsUint8<3>(levelCount - startLevel - 1, pSourceBuf, pDestBuf, sourceSize, settings);
-            break;
-
-        case EImageFormat::IF_BGRA_UBYTE:
-        case EImageFormat::IF_RGBA_UBYTE:
-            GenerateMipmapsUint8<4>(levelCount - startLevel - 1, pSourceBuf, pDestBuf, sourceSize, settings);
-            break;
-
-        // Bad cases: we need to decompress, mip and then recompress. It may be necessary to
-        // generate the latests mips after composing, or in unoptimised code.
-        // TODO: Make sure that the code optimisations avoid this cases generating separate
-        // operations to generate the mip tail.
-        case EImageFormat::IF_BC1:
-        case EImageFormat::IF_BC2:
-        case EImageFormat::IF_BC3:
-        case EImageFormat::IF_BC4:
-        case EImageFormat::IF_BC5:
-        case EImageFormat::IF_ASTC_4x4_RGB_LDR:
-        case EImageFormat::IF_ASTC_4x4_RGBA_LDR:
-        case EImageFormat::IF_ASTC_4x4_RG_LDR:
-        case EImageFormat::IF_L_UBIT_RLE:
-        case EImageFormat::IF_L_UBYTE_RLE:
-        {
-            // Uncompress the last mip that we already have
-            bool bSuccess = ImagePixelFormatInPlace(imageCompressionQuality, scratch->pUncompressed.get(), pBase, startLevel);
-			check(bSuccess);
-
-            // Generate the mipmaps from there on
-
-			constexpr bool bGenerateOnlyTailForCompressed = true;
-            ImageMipmap(imageCompressionQuality, scratch->pUncompressedMips.get(),
-                        scratch->pUncompressed.get(), levelCount - startLevel, 0, settings, bGenerateOnlyTailForCompressed);
-
-            // Compress the mipmapped image
-			bSuccess = ImagePixelFormatInPlace(imageCompressionQuality, scratch->pCompressedMips.get(), scratch->pUncompressedMips.get());
-			int32 ExcessDataSize = FMath::Max(2, scratch->pCompressedMips->GetDataSize());
-			while (!bSuccess)
-			{
-				// Bad case: this should almost never happen.
-				MUTABLE_CPUPROFILER_SCOPE(Mipmap_Recompression_OutOfSpace);
-
-				scratch->pCompressedMips->m_data.SetNumUninitialized(ExcessDataSize);
-				bSuccess = ImagePixelFormatInPlace(imageCompressionQuality, scratch->pCompressedMips.get(), scratch->pUncompressedMips.get());
-				ExcessDataSize *= 4;
-			}
-
-			check(!bGenerateOnlyTail || BaseMipDataSize == 0);
-			const int32 FinalDestSize = BaseMipDataSize + scratch->pCompressedMips->GetDataSize();
-
-			if (FinalDestSize > pDest->GetDataSize())
-			{
-				// Bad case: this should almost never happen.
-				pDest->m_data.SetNumUninitialized(FinalDestSize);
-				pDestBuf = !bGenerateOnlyTail ? pDest->GetMipData(startLevel + 1) : pDest->GetMipData(0);
-			}
-
-            FMemory::Memcpy(pDestBuf, scratch->pCompressedMips->GetData(), scratch->pCompressedMips->GetDataSize());
-            break;
-        }
-
-        default:
-            checkf( false, TEXT("Format not implemented in mipmap generation."));
-        }
-    }
-
+    /** Generate the mipmaps for images.
+	* if bGenerateOnlyTail is true, generates the mips missing from Base to LevelCount and sets
+	* them in Dest (the full chain is spit in two images). Otherwise generate the mips missing 
+	* from Base up to LevelCount and append them in Dest to the already generated Base's mips.
+    */
+	void ImageMipmap(struct FWorkingMemoryManager&, int32 CompressionQuality, Image* Dest, const Image* Base,
+		int32 LevelCount,
+		const FMipmapGenerationSettings&, bool bGenerateOnlyTail = false);
+
+	/** Mipmap separating the worst case treatment in 3 steps to manage allocations of temp data. */
+	void ImageMipmap_PrepareScratch(FWorkingMemoryManager&, Image* Dest, const Image* Base, int32 LevelCount, FScratchImageMipmap&);
+	void ImageMipmap(FScratchImageMipmap&, int32 CompressionQuality, Image* Dest, const Image* Base,
+		int32 LevelCount,
+		const FMipmapGenerationSettings&, bool bGenerateOnlyTail=false);
+	void ImageMipmap_ReleaseScratch(FWorkingMemoryManager&, FScratchImageMipmap&);
 
 	/** Update all the mipmaps in the image from the data in the base one. 
 	* Only the mipmaps already existing in the image are updated.
 	*/
-	inline void ImageMipmapInPlace(int InImageCompressionQuality, Image* InBase, const FMipmapGenerationSettings& InSettings)
-	{
-		int32 StartLevel = 0;
-		int32 LevelCount = InBase->GetLODCount();
-
-		check(!(InBase->m_flags & Image::IF_CANNOT_BE_SCALED));
-
-		int32 MipsToBuild = InBase->GetLODCount()-1;
-		if (!MipsToBuild)
-		{
-			return;
-		}
-
-		const uint8* pSourceBuf = InBase->GetMipData(StartLevel);
-		uint8* pDestBuf = InBase->GetMipData(StartLevel + 1);
-
-		FIntVector2 sourceSize = InBase->CalculateMipSize(StartLevel);
-
-		switch (InBase->GetFormat())
-		{
-		case EImageFormat::IF_L_UBYTE:
-			GenerateMipmapsUint8<1>(MipsToBuild, pSourceBuf, pDestBuf, sourceSize, InSettings);
-			break;
-
-		case EImageFormat::IF_RGB_UBYTE:
-			GenerateMipmapsUint8<3>(MipsToBuild, pSourceBuf, pDestBuf, sourceSize, InSettings);
-			break;
-
-		case EImageFormat::IF_BGRA_UBYTE:
-		case EImageFormat::IF_RGBA_UBYTE:
-			GenerateMipmapsUint8<4>(MipsToBuild, pSourceBuf, pDestBuf, sourceSize, InSettings);
-			break;
-
-		default:
-			checkf(false, TEXT("Format not implemented in mipmap generation."));
-		}
-	}
+	void ImageMipmapInPlace(int32 CompressionQuality, Image* Base, const FMipmapGenerationSettings& );
 
 }
diff --git a/Engine/Plugins/Experimental/Mutable/Source/MutableRuntime/Private/MuR/OpImagePixelFormat.cpp b/Engine/Plugins/Experimental/Mutable/Source/MutableRuntime/Private/MuR/OpImagePixelFormat.cpp
--- a/Engine/Plugins/Experimental/Mutable/Source/MutableRuntime/Private/MuR/OpImagePixelFormat.cpp
+++ b/Engine/Plugins/Experimental/Mutable/Source/MutableRuntime/Private/MuR/OpImagePixelFormat.cpp
@@ -1,7 +1,7 @@
 // Copyright Epic Games, Inc. All Rights Reserved.
 
 
-#include "MuR/OpImagePixelFormat.h"
+#include "MuR/Image.h"
 
 #include "Containers/Array.h"
 #include "HAL/UnrealMemory.h"
@@ -19,39 +19,39 @@ namespace mu
 
 
 //-------------------------------------------------------------------------------------------------
-ImagePtr ImagePixelFormat( int imageCompressionQuality, const Image* pBase,
-                               EImageFormat targetFormat, int onlyLOD )
+Ptr<Image> ImagePixelFormat( int32 CompressionQuality, const Image* Base, EImageFormat TargetFormat, int32 OnlyLOD )
 {
 	MUTABLE_CPUPROFILER_SCOPE(ImagePixelFormat);
 
-	TRACE_CPUPROFILER_EVENT_SCOPE_TEXT(*FString::Printf(TEXT("From %d to %d"), int32(pBase->GetFormat()), int32(targetFormat)));
+	TRACE_CPUPROFILER_EVENT_SCOPE_TEXT(*FString::Printf(TEXT("From %d to %d"), int32(Base->GetFormat()), int32(TargetFormat)));
 
 
 	FIntVector2 resultSize;
     int resultLODCount = 0;
 
-    if (onlyLOD==-1)
+    if (OnlyLOD==-1)
     {
-        resultSize[0] = pBase->GetSizeX();
-        resultSize[1] = pBase->GetSizeY();
-        resultLODCount = pBase->GetLODCount();
+        resultSize[0] = Base->GetSizeX();
+        resultSize[1] = Base->GetSizeY();
+        resultLODCount = Base->GetLODCount();
     }
     else
     {
-        resultSize = pBase->CalculateMipSize( onlyLOD );
+        resultSize = Base->CalculateMipSize( OnlyLOD );
         resultLODCount = 1;
     }
 
-    ImagePtr result = new Image( (uint16)resultSize[0], (uint16)resultSize[1], resultLODCount, targetFormat, EInitializationType::NotInitialized );
-	result->m_flags = pBase->m_flags;
+    Ptr<Image> Result = new Image( (uint16)resultSize[0], (uint16)resultSize[1], resultLODCount, TargetFormat, EInitializationType::NotInitialized );
+	Result->m_flags = Base->m_flags;
 
-	if (pBase->GetSizeX()<=0 ||pBase->GetSizeY()<=0)
+	if (Base->GetSizeX()<=0 ||Base->GetSizeY()<=0)
 	{
-		return result;
+		return Result;
 	}
 
-	bool bSuccess = ImagePixelFormatInPlace( imageCompressionQuality, result.get(), pBase, onlyLOD );
-	int32 OriginalDataSize = FMath::Max(result->m_data.Num(), pBase->m_data.Num());
+	bool bSuccess = false;
+	ImagePixelFormat( bSuccess, CompressionQuality, Result.get(), Base, OnlyLOD );
+	int32 OriginalDataSize = FMath::Max(Result->m_data.Num(), Base->m_data.Num());
 	int32 ExcessDataSize = OriginalDataSize * 4 + 4;
 	while (!bSuccess)
     {
@@ -59,21 +59,20 @@ ImagePtr ImagePixelFormat( int imageCompressionQuality, const Image* pBase,
 
         // Bad case, where the RLE compressed data requires more memory than the uncompressed data.
 		// We need to support it anyway for small mips or scaled images.
-		result->m_data.SetNum(ExcessDataSize);
-		bSuccess = ImagePixelFormatInPlace(imageCompressionQuality, result.get(), pBase, onlyLOD);
+		Result->m_data.SetNum(ExcessDataSize);
+		ImagePixelFormat(bSuccess, CompressionQuality, Result.get(), Base, OnlyLOD);
 		ExcessDataSize *= 4;
 	}
-    return result;
+    return Result;
 }
 
 
 //-------------------------------------------------------------------------------------------------
-bool ImagePixelFormatInPlace( int imageCompressionQuality, Image* pResult, const Image* pBase,
-                                  int onlyLOD )
+void ImagePixelFormat( bool& bOutSuccess, int32 CompressionQuality, Image* pResult, const Image* Base, int32 OnlyLOD )
 {
 	MUTABLE_CPUPROFILER_SCOPE(ImagePixelFormatInPlace);
 
-    bool success = true;
+	bOutSuccess = true;
 
 	FIntVector2 resultSize;
     int resultLODCount = 0;
@@ -81,32 +80,32 @@ bool ImagePixelFormatInPlace( int imageCompressionQuality, Image* pResult, const
     int pixelCount = 0;
     int baseLOD = 0;
 
-    if (onlyLOD==-1)
+    if (OnlyLOD==-1)
     {
-        pixelCount = (int)pBase->CalculatePixelCount();
-        pBaseBuf = pBase->GetData();
-        resultSize[0] = pBase->GetSizeX();
-        resultSize[1] = pBase->GetSizeY();
-        resultLODCount = pBase->GetLODCount();
+        pixelCount = (int)Base->CalculatePixelCount();
+        pBaseBuf = Base->GetData();
+        resultSize[0] = Base->GetSizeX();
+        resultSize[1] = Base->GetSizeY();
+        resultLODCount = Base->GetLODCount();
         baseLOD = 0;
     }
     else
     {
-        pixelCount = (int)pBase->CalculatePixelCount( onlyLOD );
-        pBaseBuf = pBase->GetMipData( onlyLOD );
-        resultSize = pBase->CalculateMipSize( onlyLOD );
+        pixelCount = (int)Base->CalculatePixelCount( OnlyLOD );
+        pBaseBuf = Base->GetMipData( OnlyLOD );
+        resultSize = Base->CalculateMipSize( OnlyLOD );
         resultLODCount = 1;
-        baseLOD = onlyLOD;
+        baseLOD = OnlyLOD;
     }
 
     check( pResult->GetSizeX()==resultSize[0] && pResult->GetSizeY()==resultSize[1] );
     check( pResult->GetLODCount()==resultLODCount );
 
-    if ( pResult->GetFormat() == pBase->GetFormat()
-         && onlyLOD==-1 )
+    if ( pResult->GetFormat() == Base->GetFormat()
+         && OnlyLOD==-1 )
     {
         // Shouldn't really happen
-        pResult->m_data = pBase->m_data;
+        pResult->m_data = Base->m_data;
     }
     else
     {
@@ -117,7 +116,7 @@ bool ImagePixelFormatInPlace( int imageCompressionQuality, Image* pResult, const
         {
             uint8_t* pDestBuf = pResult->GetData();
 
-            switch ( pBase->GetFormat() )
+            switch ( Base->GetFormat() )
             {
 
             case EImageFormat::IF_L_UBYTE:
@@ -128,9 +127,9 @@ bool ImagePixelFormatInPlace( int imageCompressionQuality, Image* pResult, const
 
             case EImageFormat::IF_L_UBYTE_RLE:
             {
-                if( onlyLOD == -1 )
+                if( OnlyLOD == -1 )
                 {
-                    UncompressRLE_L( pBase, pResult );
+                    UncompressRLE_L( Base, pResult );
                 }
                 else
                 {
@@ -141,9 +140,9 @@ bool ImagePixelFormatInPlace( int imageCompressionQuality, Image* pResult, const
 
             case EImageFormat::IF_L_UBIT_RLE:
             {
-                if( onlyLOD == -1 )
+                if( OnlyLOD == -1 )
                 {
-                    UncompressRLE_L1( pBase, pResult );
+                    UncompressRLE_L1( Base, pResult );
                 }
                 else
                 {
@@ -156,10 +155,10 @@ bool ImagePixelFormatInPlace( int imageCompressionQuality, Image* pResult, const
             {
                 for ( int i=0; i<pixelCount; ++i )
                 {
-                    unsigned result = 76 * pBaseBuf[3*i+0]
+                    unsigned Result = 76 * pBaseBuf[3*i+0]
                             + 150 * pBaseBuf[3*i+1]
                             + 29 * pBaseBuf[3*i+2];
-                    pDestBuf[i] = (uint8_t)FMath::Min( 255u, result >> 8 );
+                    pDestBuf[i] = (uint8_t)FMath::Min( 255u, Result >> 8 );
                 }
                 break;
             }
@@ -167,21 +166,21 @@ bool ImagePixelFormatInPlace( int imageCompressionQuality, Image* pResult, const
             case EImageFormat::IF_RGB_UBYTE_RLE:
             {
                 ImagePtr pTempBase =
-                    ImagePixelFormat( imageCompressionQuality, pBase, EImageFormat::IF_RGB_UBYTE );
-                if (onlyLOD==-1)
+                    ImagePixelFormat( CompressionQuality, Base, EImageFormat::IF_RGB_UBYTE );
+                if (OnlyLOD==-1)
                 {
                     pBaseBuf = pTempBase->GetData();
                 }
                 else
                 {
-                    pBaseBuf = pTempBase->GetMipData( onlyLOD );
+                    pBaseBuf = pTempBase->GetMipData( OnlyLOD );
                 }
                 for ( int i=0; i<pixelCount; ++i )
                 {
-                    unsigned result = 76 * pBaseBuf[3*i+0]
+                    unsigned Result = 76 * pBaseBuf[3*i+0]
                             + 150 * pBaseBuf[3*i+1]
                             + 29 * pBaseBuf[3*i+2];
-                    pDestBuf[i] = (uint8_t)FMath::Min( 255u, result >> 8 );
+                    pDestBuf[i] = (uint8_t)FMath::Min( 255u, Result >> 8 );
                 }
                 break;
             }
@@ -190,10 +189,10 @@ bool ImagePixelFormatInPlace( int imageCompressionQuality, Image* pResult, const
             {
                 for ( int i=0; i<pixelCount; ++i )
                 {
-                    unsigned result = 76 * pBaseBuf[4*i+0]
+                    unsigned Result = 76 * pBaseBuf[4*i+0]
                             + 150 * pBaseBuf[4*i+1]
                             + 29 * pBaseBuf[4*i+2];
-                    pDestBuf[i] = (uint8_t)FMath::Min( 255u, result >> 8 );
+                    pDestBuf[i] = (uint8_t)FMath::Min( 255u, Result >> 8 );
                 }
                 break;
             }
@@ -201,21 +200,21 @@ bool ImagePixelFormatInPlace( int imageCompressionQuality, Image* pResult, const
             case EImageFormat::IF_RGBA_UBYTE_RLE:
             {
                 ImagePtr pTempBase =
-                    ImagePixelFormat( imageCompressionQuality, pBase, EImageFormat::IF_RGBA_UBYTE );
-                if (onlyLOD==-1)
+                    ImagePixelFormat( CompressionQuality, Base, EImageFormat::IF_RGBA_UBYTE );
+                if (OnlyLOD==-1)
                 {
                     pBaseBuf = pTempBase->GetData();
                 }
                 else
                 {
-                    pBaseBuf = pTempBase->GetMipData( onlyLOD );
+                    pBaseBuf = pTempBase->GetMipData( OnlyLOD );
                 }
                 for ( int i=0; i<pixelCount; ++i )
                 {
-                    unsigned result = 76 * pBaseBuf[4*i+0]
+                    unsigned Result = 76 * pBaseBuf[4*i+0]
                             + 150 * pBaseBuf[4*i+1]
                             + 29 * pBaseBuf[4*i+2];
-                    pDestBuf[i] = (uint8_t)FMath::Min( 255u, result >> 8 );
+                    pDestBuf[i] = (uint8_t)FMath::Min( 255u, Result >> 8 );
                 }
                 break;
             }
@@ -224,10 +223,10 @@ bool ImagePixelFormatInPlace( int imageCompressionQuality, Image* pResult, const
             {
                 for ( int i=0; i<pixelCount; ++i )
                 {
-                    unsigned result = 76 * pBaseBuf[4*i+2]
+                    unsigned Result = 76 * pBaseBuf[4*i+2]
                             + 150 * pBaseBuf[4*i+1]
                             + 29 * pBaseBuf[4*i+0];
-                    pDestBuf[i] = (uint8_t)FMath::Min( 255u, result >> 8 );
+                    pDestBuf[i] = (uint8_t)FMath::Min( 255u, Result >> 8 );
                 }
                 break;
             }
@@ -239,20 +238,20 @@ bool ImagePixelFormatInPlace( int imageCompressionQuality, Image* pResult, const
                 for (int m = 0; m < resultLODCount; ++m)
                 {
                     ImagePtr pTempBase =
-                        ImagePixelFormat( imageCompressionQuality, pBase, EImageFormat::IF_RGB_UBYTE );
-                    if (onlyLOD == -1)
+                        ImagePixelFormat( CompressionQuality, Base, EImageFormat::IF_RGB_UBYTE );
+                    if (OnlyLOD == -1)
                     {
                         pBaseBuf = pTempBase->GetData();
                     }
                     else
                     {
-                        pBaseBuf = pTempBase->GetMipData(onlyLOD);
+                        pBaseBuf = pTempBase->GetMipData(OnlyLOD);
                     }
                     for (int i = 0; i < pixelCount; ++i)
                     {
-                        unsigned result = 76 * pBaseBuf[3 * i + 0] + 150 * pBaseBuf[3 * i + 1] +
+                        unsigned Result = 76 * pBaseBuf[3 * i + 0] + 150 * pBaseBuf[3 * i + 1] +
                                           29 * pBaseBuf[3 * i + 2];
-                        pDestBuf[i] = (uint8_t)FMath::Min(255u, result >> 8);
+                        pDestBuf[i] = (uint8_t)FMath::Min(255u, Result >> 8);
                     }
                 }
                 break;
@@ -262,7 +261,7 @@ bool ImagePixelFormatInPlace( int imageCompressionQuality, Image* pResult, const
                 for ( int m=0; m<resultLODCount; ++m )
                 {
 					FIntVector2 mipSize = pResult->CalculateMipSize( m );
-                    miro::BC4_to_L( mipSize[0], mipSize[1], pBase->GetMipData(baseLOD+m), pResult->GetMipData(m) );
+                    miro::BC4_to_L( mipSize[0], mipSize[1], Base->GetMipData(baseLOD+m), pResult->GetMipData(m) );
                 }
                 break;
 			}
@@ -278,28 +277,29 @@ bool ImagePixelFormatInPlace( int imageCompressionQuality, Image* pResult, const
 
         case EImageFormat::IF_L_UBYTE_RLE:
         {
-            switch ( pBase->GetFormat() )
+            switch ( Base->GetFormat() )
             {
 
             case EImageFormat::IF_L_UBYTE:
             {
-                check( onlyLOD == -1 );
+                check( OnlyLOD == -1 );
 
                 // Try to compress
 				if (!pResult->m_data.Num())
 				{
 					// Allocate memory for the compressed data. TODO: Smaller?
-					pResult->m_data.SetNum(pBase->m_data.Num());
+					uint32 TotalMemory = Base->m_data.Num();
+					pResult->m_data.SetNum(TotalMemory);
 				}
 
-                success = CompressRLE_L( pBase, pResult );
+                CompressRLE_L(bOutSuccess, Base, pResult);
 
 #ifdef MUTABLE_DEBUG_RLE		
-				if (success)
+				if (bOutSuccess)
                 {
                     // verify
-                    ImagePtr pTest = ImagePixelFormat( imageCompressionQuality, pResult, EImageFormat::IF_L_UBYTE );
-					check(!FMemory::Memcmp(pTest->GetData(), pBase->GetData(), pBase->GetDataSize()));
+                    ImagePtr pTest = ImagePixelFormat( CompressionQuality, pResult, EImageFormat::IF_L_UBYTE );
+					check(!FMemory::Memcmp(pTest->GetData(), Base->GetData(), Base->GetDataSize()));
                 }
 #endif
 
@@ -310,8 +310,8 @@ bool ImagePixelFormatInPlace( int imageCompressionQuality, Image* pResult, const
             case EImageFormat::IF_RGBA_UBYTE:
             case EImageFormat::IF_BGRA_UBYTE:
             {
-                check( onlyLOD == -1 );
-                ImagePtr pTemp = ImagePixelFormat( imageCompressionQuality, pBase, EImageFormat::IF_L_UBYTE );
+                check( OnlyLOD == -1 );
+                ImagePtr pTemp = ImagePixelFormat( CompressionQuality, Base, EImageFormat::IF_L_UBYTE );
  
 				// Try to compress
 				if (!pResult->m_data.Num())
@@ -320,13 +320,13 @@ bool ImagePixelFormatInPlace( int imageCompressionQuality, Image* pResult, const
 					pResult->m_data.SetNum(pTemp->m_data.Num());
 				}
 
-                success = CompressRLE_L( pTemp.get(), pResult );
+				CompressRLE_L(bOutSuccess, pTemp.get(), pResult );
 
 #ifdef MUTABLE_DEBUG_RLE		
-				if (success)
+				if (bOutSuccess)
                 {
                     // verify
-                    ImagePtr pTest = ImagePixelFormat( imageCompressionQuality, pResult, EImageFormat::IF_L_UBYTE );
+                    ImagePtr pTest = ImagePixelFormat( CompressionQuality, pResult, EImageFormat::IF_L_UBYTE );
                     check( !FMemory::Memcmp(pTest->GetData(), pTemp->GetData(), pTemp->GetDataSize() ) );
                 }
 #endif
@@ -345,28 +345,28 @@ bool ImagePixelFormatInPlace( int imageCompressionQuality, Image* pResult, const
 
         case EImageFormat::IF_L_UBIT_RLE:
         {
-            switch ( pBase->GetFormat() )
+            switch ( Base->GetFormat() )
             {
 
             case EImageFormat::IF_L_UBYTE:
             {
-                check( onlyLOD == -1 );
+                check( OnlyLOD == -1 );
 
 				// Allocate memory for the compressed data. TODO: Smaller?
 				if (!pResult->m_data.Num())
 				{
 					// Allocate memory for the compressed data. TODO: Smaller?
-					pResult->m_data.SetNum(pBase->m_data.Num());
+					pResult->m_data.SetNum(Base->m_data.Num());
 				}
 				
-				success = CompressRLE_L1( pBase, pResult );
+				CompressRLE_L1(bOutSuccess, Base, pResult );
                 
 #ifdef MUTABLE_DEBUG_RLE		
-				if (success)
+				if (bOutSuccess)
                 {
 					// Verify
-                    ImagePtr pTest = ImagePixelFormat( imageCompressionQuality, pResult, EImageFormat::IF_L_UBYTE );
-                    check( !FMemory::Memcmp(pTest->GetData(), pBase->GetData(), pBase->GetDataSize() ) );
+                    ImagePtr pTest = ImagePixelFormat( CompressionQuality, pResult, EImageFormat::IF_L_UBYTE );
+                    check( !FMemory::Memcmp(pTest->GetData(), Base->GetData(), Base->GetDataSize() ) );
                 }
 #endif
 
@@ -377,15 +377,15 @@ bool ImagePixelFormatInPlace( int imageCompressionQuality, Image* pResult, const
             case EImageFormat::IF_RGBA_UBYTE:
             case EImageFormat::IF_BGRA_UBYTE:
             {
-                check( onlyLOD == -1 );
-                ImagePtr pTemp = ImagePixelFormat( imageCompressionQuality, pBase, EImageFormat::IF_L_UBYTE );
-				success = CompressRLE_L1( pTemp.get(), pResult );
+                check( OnlyLOD == -1 );
+                ImagePtr pTemp = ImagePixelFormat( CompressionQuality, Base, EImageFormat::IF_L_UBYTE );
+				CompressRLE_L1(bOutSuccess, pTemp.get(), pResult );
 
-#ifdef MUTABLE_DEBUG_RLE		
-				if (success)
+#ifdef bOutSuccess		
+				if (bOutSuccess)
 				{
 					// Verify
-					ImagePtr pTest = ImagePixelFormat(imageCompressionQuality, pResult, EImageFormat::IF_L_UBYTE);
+					ImagePtr pTest = ImagePixelFormat(CompressionQuality, pResult, EImageFormat::IF_L_UBYTE);
 					check(!FMemory::Memcmp(pTest->GetData(), pTemp->GetData(), pTemp->GetDataSize()));
 				}
 #endif
@@ -406,7 +406,7 @@ bool ImagePixelFormatInPlace( int imageCompressionQuality, Image* pResult, const
         {
             uint8_t* pDestBuf = pResult->GetData();
 
-            switch ( pBase->GetFormat() )
+            switch ( Base->GetFormat() )
             {
 
             case EImageFormat::IF_L_UBYTE:
@@ -450,14 +450,14 @@ bool ImagePixelFormatInPlace( int imageCompressionQuality, Image* pResult, const
 
             case EImageFormat::IF_RGB_UBYTE_RLE:
             {
-                UncompressRLE_RGB( pBase, pResult );
+                UncompressRLE_RGB( Base, pResult );
                 break;
             }
 
             case EImageFormat::IF_RGBA_UBYTE_RLE:
             {
-                check( onlyLOD == -1 );
-                ImagePtr pTemp = ImagePixelFormat( imageCompressionQuality, pBase, EImageFormat::IF_RGBA_UBYTE );
+                check( OnlyLOD == -1 );
+                ImagePtr pTemp = ImagePixelFormat( CompressionQuality, Base, EImageFormat::IF_RGBA_UBYTE );
                 const uint8_t* pTempBuf = pTemp->GetData();
                 for ( int i=0; i<pixelCount; ++i )
                 {
@@ -474,7 +474,7 @@ bool ImagePixelFormatInPlace( int imageCompressionQuality, Image* pResult, const
                 {
 					FIntVector2 mipSize = pResult->CalculateMipSize( m );
                     miro::BC1_to_RGB( mipSize[0], mipSize[1],
-                                        pBase->GetMipData(baseLOD+m),
+                                        Base->GetMipData(baseLOD+m),
                                         pResult->GetMipData(m) );
                 }
                 break;
@@ -486,7 +486,7 @@ bool ImagePixelFormatInPlace( int imageCompressionQuality, Image* pResult, const
                 {
 					FIntVector2 mipSize = pResult->CalculateMipSize( m );
                     miro::BC2_to_RGB( mipSize[0], mipSize[1],
-                                        pBase->GetMipData(baseLOD+m),
+                                        Base->GetMipData(baseLOD+m),
                                         pResult->GetMipData(m) );
                 }
                 break;
@@ -498,7 +498,7 @@ bool ImagePixelFormatInPlace( int imageCompressionQuality, Image* pResult, const
                 {
 					FIntVector2 mipSize = pResult->CalculateMipSize( m );
                     miro::BC3_to_RGB( mipSize[0], mipSize[1],
-                                        pBase->GetMipData(baseLOD+m),
+                                        Base->GetMipData(baseLOD+m),
                                         pResult->GetMipData(m) );
                 }
                 break;
@@ -510,7 +510,7 @@ bool ImagePixelFormatInPlace( int imageCompressionQuality, Image* pResult, const
                 {
 					FIntVector2 mipSize = pResult->CalculateMipSize( m );
                     miro::BC4_to_RGB( mipSize[0], mipSize[1],
-                                      pBase->GetMipData(baseLOD+m),
+                                      Base->GetMipData(baseLOD+m),
                                       pResult->GetMipData(m) );
                 }
                 break;
@@ -522,7 +522,7 @@ bool ImagePixelFormatInPlace( int imageCompressionQuality, Image* pResult, const
                 {
 					FIntVector2 mipSize = pResult->CalculateMipSize( m );
                     miro::BC5_to_RGB( mipSize[0], mipSize[1],
-                                      pBase->GetMipData(baseLOD+m),
+                                      Base->GetMipData(baseLOD+m),
                                       pResult->GetMipData(m) );
                 }
                 break;
@@ -534,7 +534,7 @@ bool ImagePixelFormatInPlace( int imageCompressionQuality, Image* pResult, const
                 {
 					FIntVector2 mipSize = pResult->CalculateMipSize( m );
                     miro::ASTC4x4RGBL_to_RGB( mipSize[0], mipSize[1],
-                                      pBase->GetMipData(baseLOD+m),
+                                      Base->GetMipData(baseLOD+m),
                                       pResult->GetMipData(m) );
                 }
                 break;
@@ -546,7 +546,7 @@ bool ImagePixelFormatInPlace( int imageCompressionQuality, Image* pResult, const
                 {
 					FIntVector2 mipSize = pResult->CalculateMipSize( m );
                     miro::ASTC4x4RGBAL_to_RGB( mipSize[0], mipSize[1],
-                                      pBase->GetMipData(baseLOD+m),
+                                      Base->GetMipData(baseLOD+m),
                                       pResult->GetMipData(m) );
                 }
                 break;
@@ -558,7 +558,7 @@ bool ImagePixelFormatInPlace( int imageCompressionQuality, Image* pResult, const
                 {
 					FIntVector2 mipSize = pResult->CalculateMipSize( m );
                     miro::ASTC4x4RGL_to_RGB( mipSize[0], mipSize[1],
-                                      pBase->GetMipData(baseLOD+m),
+                                      Base->GetMipData(baseLOD+m),
                                       pResult->GetMipData(m) );
                 }
                 break;
@@ -570,7 +570,7 @@ bool ImagePixelFormatInPlace( int imageCompressionQuality, Image* pResult, const
 				{
 					FIntVector2 mipSize = pResult->CalculateMipSize(m);
 					miro::ASTC8x8RGBL_to_RGB(mipSize[0], mipSize[1],
-						pBase->GetMipData(baseLOD + m),
+						Base->GetMipData(baseLOD + m),
 						pResult->GetMipData(m));
 				}
 				break;
@@ -582,7 +582,7 @@ bool ImagePixelFormatInPlace( int imageCompressionQuality, Image* pResult, const
 				{
 					FIntVector2 mipSize = pResult->CalculateMipSize(m);
 					miro::ASTC8x8RGBAL_to_RGB(mipSize[0], mipSize[1],
-						pBase->GetMipData(baseLOD + m),
+						Base->GetMipData(baseLOD + m),
 						pResult->GetMipData(m));
 				}
 				break;
@@ -594,7 +594,7 @@ bool ImagePixelFormatInPlace( int imageCompressionQuality, Image* pResult, const
 				{
 					FIntVector2 mipSize = pResult->CalculateMipSize(m);
 					miro::ASTC8x8RGL_to_RGB(mipSize[0], mipSize[1],
-						pBase->GetMipData(baseLOD + m),
+						Base->GetMipData(baseLOD + m),
 						pResult->GetMipData(m));
 				}
 				break;
@@ -606,7 +606,7 @@ bool ImagePixelFormatInPlace( int imageCompressionQuality, Image* pResult, const
 				{
 					FIntVector2 mipSize = pResult->CalculateMipSize(m);
 					miro::ASTC12x12RGBL_to_RGB(mipSize[0], mipSize[1],
-						pBase->GetMipData(baseLOD + m),
+						Base->GetMipData(baseLOD + m),
 						pResult->GetMipData(m));
 				}
 				break;
@@ -618,7 +618,7 @@ bool ImagePixelFormatInPlace( int imageCompressionQuality, Image* pResult, const
 				{
 					FIntVector2 mipSize = pResult->CalculateMipSize(m);
 					miro::ASTC12x12RGBAL_to_RGB(mipSize[0], mipSize[1],
-						pBase->GetMipData(baseLOD + m),
+						Base->GetMipData(baseLOD + m),
 						pResult->GetMipData(m));
 				}
 				break;
@@ -630,7 +630,7 @@ bool ImagePixelFormatInPlace( int imageCompressionQuality, Image* pResult, const
 				{
 					FIntVector2 mipSize = pResult->CalculateMipSize(m);
 					miro::ASTC12x12RGL_to_RGB(mipSize[0], mipSize[1],
-						pBase->GetMipData(baseLOD + m),
+						Base->GetMipData(baseLOD + m),
 						pResult->GetMipData(m));
 				}
 				break;
@@ -649,8 +649,8 @@ bool ImagePixelFormatInPlace( int imageCompressionQuality, Image* pResult, const
         {
             // TODO: Optimise
             pResult->m_format = EImageFormat::IF_RGBA_UBYTE;
-            bool bSuccess = ImagePixelFormatInPlace( imageCompressionQuality, pResult, pBase, onlyLOD );
-			check(bSuccess);
+            ImagePixelFormat(bOutSuccess, CompressionQuality, pResult, Base, OnlyLOD );
+			check(bOutSuccess);
             pResult->m_format = EImageFormat::IF_BGRA_UBYTE;
 
             uint8_t* pDestBuf = pResult->GetData();
@@ -667,13 +667,13 @@ bool ImagePixelFormatInPlace( int imageCompressionQuality, Image* pResult, const
         {
             uint8_t* pDestBuf = pResult->GetData();
 
-            switch ( pBase->GetFormat() )
+            switch ( Base->GetFormat() )
             {
 
             case EImageFormat::IF_RGBA_UBYTE_RLE:
             {
-                check( onlyLOD == -1 );
-                UncompressRLE_RGBA( pBase, pResult );
+                check( OnlyLOD == -1 );
+                UncompressRLE_RGBA( Base, pResult );
                 break;
             }
 
@@ -725,7 +725,7 @@ bool ImagePixelFormatInPlace( int imageCompressionQuality, Image* pResult, const
                 {
 					FIntVector2 mipSize = pResult->CalculateMipSize( m );
                     miro::BC1_to_RGBA( mipSize[0], mipSize[1],
-                                        pBase->GetMipData(baseLOD+m),
+                                        Base->GetMipData(baseLOD+m),
                                         pResult->GetMipData(m) );
                 }
                 break;
@@ -737,7 +737,7 @@ bool ImagePixelFormatInPlace( int imageCompressionQuality, Image* pResult, const
                 {
 					FIntVector2 mipSize = pResult->CalculateMipSize( m );
                     miro::BC2_to_RGBA( mipSize[0], mipSize[1],
-                                        pBase->GetMipData(baseLOD+m),
+                                        Base->GetMipData(baseLOD+m),
                                         pResult->GetMipData(m) );
                 }
                 break;
@@ -749,7 +749,7 @@ bool ImagePixelFormatInPlace( int imageCompressionQuality, Image* pResult, const
                 {
 					FIntVector2 mipSize = pResult->CalculateMipSize( m );
                     miro::BC3_to_RGBA( mipSize[0], mipSize[1],
-                                        pBase->GetMipData(baseLOD+m),
+                                        Base->GetMipData(baseLOD+m),
                                         pResult->GetMipData(m) );
                 }
                 break;
@@ -761,7 +761,7 @@ bool ImagePixelFormatInPlace( int imageCompressionQuality, Image* pResult, const
                 {
 					FIntVector2 mipSize = pResult->CalculateMipSize( m );
                     miro::BC4_to_RGBA( mipSize[0], mipSize[1],
-                                       pBase->GetMipData(baseLOD+m),
+                                       Base->GetMipData(baseLOD+m),
                                        pResult->GetMipData(m) );
                 }
                 break;
@@ -773,7 +773,7 @@ bool ImagePixelFormatInPlace( int imageCompressionQuality, Image* pResult, const
                 {
                     FIntVector2 mipSize = pResult->CalculateMipSize( m );
                     miro::BC5_to_RGBA( mipSize[0], mipSize[1],
-                                        pBase->GetMipData(baseLOD+m),
+                                        Base->GetMipData(baseLOD+m),
                                         pResult->GetMipData(m) );
                 }
                 break;
@@ -785,7 +785,7 @@ bool ImagePixelFormatInPlace( int imageCompressionQuality, Image* pResult, const
                 {
                     FIntVector2 mipSize = pResult->CalculateMipSize( m );
                     miro::ASTC4x4RGBL_to_RGBA( mipSize[0], mipSize[1],
-                                        pBase->GetMipData(baseLOD+m),
+                                        Base->GetMipData(baseLOD+m),
                                         pResult->GetMipData(m) );
                 }
                 break;
@@ -797,7 +797,7 @@ bool ImagePixelFormatInPlace( int imageCompressionQuality, Image* pResult, const
                 {
                     FIntVector2 mipSize = pResult->CalculateMipSize( m );
                     miro::ASTC4x4RGBAL_to_RGBA( mipSize[0], mipSize[1],
-                                        pBase->GetMipData(baseLOD+m),
+                                        Base->GetMipData(baseLOD+m),
                                         pResult->GetMipData(m) );
                 }
                 break;
@@ -809,7 +809,7 @@ bool ImagePixelFormatInPlace( int imageCompressionQuality, Image* pResult, const
                 {
                     FIntVector2 mipSize = pResult->CalculateMipSize( m );
                     miro::ASTC4x4RGL_to_RGBA( mipSize[0], mipSize[1],
-                                        pBase->GetMipData(baseLOD+m),
+                                        Base->GetMipData(baseLOD+m),
                                         pResult->GetMipData(m) );
                 }
                 break;
@@ -821,7 +821,7 @@ bool ImagePixelFormatInPlace( int imageCompressionQuality, Image* pResult, const
 				{
 					FIntVector2 mipSize = pResult->CalculateMipSize(m);
 					miro::ASTC8x8RGBL_to_RGBA(mipSize[0], mipSize[1],
-						pBase->GetMipData(baseLOD + m),
+						Base->GetMipData(baseLOD + m),
 						pResult->GetMipData(m));
 				}
 				break;
@@ -833,7 +833,7 @@ bool ImagePixelFormatInPlace( int imageCompressionQuality, Image* pResult, const
 				{
 					FIntVector2 mipSize = pResult->CalculateMipSize(m);
 					miro::ASTC8x8RGBAL_to_RGBA(mipSize[0], mipSize[1],
-						pBase->GetMipData(baseLOD + m),
+						Base->GetMipData(baseLOD + m),
 						pResult->GetMipData(m));
 				}
 				break;
@@ -845,7 +845,7 @@ bool ImagePixelFormatInPlace( int imageCompressionQuality, Image* pResult, const
 				{
 					FIntVector2 mipSize = pResult->CalculateMipSize(m);
 					miro::ASTC8x8RGL_to_RGBA(mipSize[0], mipSize[1],
-						pBase->GetMipData(baseLOD + m),
+						Base->GetMipData(baseLOD + m),
 						pResult->GetMipData(m));
 				}
 				break;
@@ -857,7 +857,7 @@ bool ImagePixelFormatInPlace( int imageCompressionQuality, Image* pResult, const
 				{
 					FIntVector2 mipSize = pResult->CalculateMipSize(m);
 					miro::ASTC12x12RGBL_to_RGBA(mipSize[0], mipSize[1],
-						pBase->GetMipData(baseLOD + m),
+						Base->GetMipData(baseLOD + m),
 						pResult->GetMipData(m));
 				}
 				break;
@@ -869,7 +869,7 @@ bool ImagePixelFormatInPlace( int imageCompressionQuality, Image* pResult, const
 				{
 					FIntVector2 mipSize = pResult->CalculateMipSize(m);
 					miro::ASTC12x12RGBAL_to_RGBA(mipSize[0], mipSize[1],
-						pBase->GetMipData(baseLOD + m),
+						Base->GetMipData(baseLOD + m),
 						pResult->GetMipData(m));
 				}
 				break;
@@ -881,7 +881,7 @@ bool ImagePixelFormatInPlace( int imageCompressionQuality, Image* pResult, const
 				{
 					FIntVector2 mipSize = pResult->CalculateMipSize(m);
 					miro::ASTC12x12RGL_to_RGBA(mipSize[0], mipSize[1],
-						pBase->GetMipData(baseLOD + m),
+						Base->GetMipData(baseLOD + m),
 						pResult->GetMipData(m));
 				}
 				break;
@@ -889,8 +889,8 @@ bool ImagePixelFormatInPlace( int imageCompressionQuality, Image* pResult, const
 
 			case EImageFormat::IF_RGB_UBYTE_RLE:
 			{
-				check(onlyLOD == -1);
-				ImagePtr pTemp = ImagePixelFormat(imageCompressionQuality, pBase, EImageFormat::IF_RGB_UBYTE);
+				check(OnlyLOD == -1);
+				ImagePtr pTemp = ImagePixelFormat(CompressionQuality, Base, EImageFormat::IF_RGB_UBYTE);
 				const uint8_t* pTempData = pTemp->GetData();
 				for (int i = 0; i < pixelCount; ++i)
 				{
@@ -904,8 +904,8 @@ bool ImagePixelFormatInPlace( int imageCompressionQuality, Image* pResult, const
 
 			case EImageFormat::IF_L_UBYTE_RLE:
 			{
-				check(onlyLOD == -1);
-				ImagePtr pTemp = ImagePixelFormat(imageCompressionQuality, pBase, EImageFormat::IF_L_UBYTE);
+				check(OnlyLOD == -1);
+				ImagePtr pTemp = ImagePixelFormat(CompressionQuality, Base, EImageFormat::IF_L_UBYTE);
 				const uint8_t* pTempData = pTemp->GetData();
 				for (int i = 0; i < pixelCount; ++i)
 				{
@@ -929,46 +929,46 @@ bool ImagePixelFormatInPlace( int imageCompressionQuality, Image* pResult, const
 
         case EImageFormat::IF_RGBA_UBYTE_RLE:
         {
-            switch ( pBase->GetFormat() )
+            switch ( Base->GetFormat() )
             {
 
             case EImageFormat::IF_RGBA_UBYTE:
             {
-                check( onlyLOD == -1 );
-                CompressRLE_RGBA( pBase, pResult );
+                check( OnlyLOD == -1 );
+                CompressRLE_RGBA( Base, pResult );
 
                 // TODO: TEST BLEH
-                //                ImagePtr pTest = ImagePixelFormat( imageCompressionQuality,
+                //                ImagePtr pTest = ImagePixelFormat( CompressionQuality,
                 //                pResult, EImageFormat::IF_RGBA_UBYTE ); check(
-                //                !memcmp(pTest->GetData(), pBase->GetData(),
-                //                pBase->GetDataSize() ) );
+                //                !memcmp(pTest->GetData(), Base->GetData(),
+                //                Base->GetDataSize() ) );
                 break;
             }
 
             case EImageFormat::IF_RGB_UBYTE:
             {
-                check( onlyLOD == -1 );
+                check( OnlyLOD == -1 );
 
                 // \todo: optimise
-                ImagePtr pTemp = ImagePixelFormat( imageCompressionQuality, pBase, EImageFormat::IF_RGBA_UBYTE );
+                ImagePtr pTemp = ImagePixelFormat( CompressionQuality, Base, EImageFormat::IF_RGBA_UBYTE );
                 CompressRLE_RGBA( pTemp.get(), pResult );
 
                 // TODO: TEST BLEH
-                //                ImagePtr pTest = ImagePixelFormat( imageCompressionQuality,
+                //                ImagePtr pTest = ImagePixelFormat( CompressionQuality,
                 //                pResult, EImageFormat::IF_RGBA_UBYTE ); check(
-                //                !memcmp(pTest->GetData(), pBase->GetData(),
-                //                pBase->GetDataSize() ) );
+                //                !memcmp(pTest->GetData(), Base->GetData(),
+                //                Base->GetDataSize() ) );
                 break;
             }
 
             case EImageFormat::IF_RGB_UBYTE_RLE:
             {
-                check( onlyLOD == -1 );
+                check( OnlyLOD == -1 );
 
                 // \todo: optimise
-                ImagePtr pTemp1 = ImagePixelFormat( imageCompressionQuality, pBase, EImageFormat::IF_RGB_UBYTE );
+                ImagePtr pTemp1 = ImagePixelFormat( CompressionQuality, Base, EImageFormat::IF_RGB_UBYTE );
                 ImagePtr pTemp2 =
-                    ImagePixelFormat( imageCompressionQuality, pTemp1.get(), EImageFormat::IF_RGBA_UBYTE );
+                    ImagePixelFormat( CompressionQuality, pTemp1.get(), EImageFormat::IF_RGBA_UBYTE );
                 CompressRLE_RGBA( pTemp2.get(), pResult );
                 break;
             }
@@ -985,29 +985,29 @@ bool ImagePixelFormatInPlace( int imageCompressionQuality, Image* pResult, const
 
         case EImageFormat::IF_RGB_UBYTE_RLE:
         {
-            switch ( pBase->GetFormat() )
+            switch ( Base->GetFormat() )
             {
 
             case EImageFormat::IF_RGB_UBYTE:
             {
-                check( onlyLOD == -1 );
-                CompressRLE_RGB( pBase, pResult );
+                check( OnlyLOD == -1 );
+                CompressRLE_RGB( Base, pResult );
 
                 // TODO: TEST BLEH
-                // ImagePtr pTest = ImagePixelFormat( imageCompressionQuality, pDest.get(),
-                // EImageFormat::IF_RGB_UBYTE ); check( !memcmp(pTest->GetData(), pBase->GetData(),
-                // pBase->GetDataSize() ) );
+                // ImagePtr pTest = ImagePixelFormat( CompressionQuality, pDest.get(),
+                // EImageFormat::IF_RGB_UBYTE ); check( !memcmp(pTest->GetData(), Base->GetData(),
+                // Base->GetDataSize() ) );
                 break;
             }
 
             case EImageFormat::IF_RGBA_UBYTE_RLE:
             {
-                check( onlyLOD == -1 );
+                check( OnlyLOD == -1 );
 
                 // \todo: optimise
-                ImagePtr pTemp1 = ImagePixelFormat( imageCompressionQuality, pBase, EImageFormat::IF_RGBA_UBYTE );
+                ImagePtr pTemp1 = ImagePixelFormat( CompressionQuality, Base, EImageFormat::IF_RGBA_UBYTE );
                 ImagePtr pTemp2 =
-                    ImagePixelFormat( imageCompressionQuality, pTemp1.get(), EImageFormat::IF_RGB_UBYTE );
+                    ImagePixelFormat( CompressionQuality, pTemp1.get(), EImageFormat::IF_RGB_UBYTE );
                 CompressRLE_RGB( pTemp2.get(), pResult );
                 break;
             }
@@ -1023,15 +1023,15 @@ bool ImagePixelFormatInPlace( int imageCompressionQuality, Image* pResult, const
 
         case EImageFormat::IF_BC1:
         {
-            switch ( pBase->GetFormat() )
+            switch ( Base->GetFormat() )
             {
             case EImageFormat::IF_RGB_UBYTE:
                 for ( int m=0; m<resultLODCount; ++m )
                 {
                     FIntVector2 mipSize = pResult->CalculateMipSize( m );
                     miro::RGB_to_BC1(
-                        mipSize[0], mipSize[1], pBase->GetMipData( baseLOD + m ),
-                        pResult->GetMipData( m ), imageCompressionQuality );
+                        mipSize[0], mipSize[1], Base->GetMipData( baseLOD + m ),
+                        pResult->GetMipData( m ), CompressionQuality );
                 }
                 break;
             case EImageFormat::IF_RGBA_UBYTE:
@@ -1039,33 +1039,33 @@ bool ImagePixelFormatInPlace( int imageCompressionQuality, Image* pResult, const
                 {
                     FIntVector2 mipSize = pResult->CalculateMipSize( m );
                     miro::RGBA_to_BC1(
-                        mipSize[0], mipSize[1], pBase->GetMipData( baseLOD + m ),
-                        pResult->GetMipData( m ), imageCompressionQuality );
+                        mipSize[0], mipSize[1], Base->GetMipData( baseLOD + m ),
+                        pResult->GetMipData( m ), CompressionQuality );
                 }
                 break;
 
             case EImageFormat::IF_L_UBYTE:
             {
-                ImagePtr pTemp = ImagePixelFormat( imageCompressionQuality, pBase, EImageFormat::IF_RGB_UBYTE );
+                ImagePtr pTemp = ImagePixelFormat( CompressionQuality, Base, EImageFormat::IF_RGB_UBYTE );
                 for ( int m=0; m<resultLODCount; ++m )
                 {
                     FIntVector2 mipSize = pResult->CalculateMipSize( m );
                     miro::RGB_to_BC1(
                         mipSize[0], mipSize[1], pTemp->GetMipData( baseLOD + m ),
-                        pResult->GetMipData( m ), imageCompressionQuality );
+                        pResult->GetMipData( m ), CompressionQuality );
                 }
                 break;
             }
 
             case EImageFormat::IF_BC3:
             {
-                ImagePtr pTemp = ImagePixelFormat( imageCompressionQuality, pBase, EImageFormat::IF_RGBA_UBYTE );
+                ImagePtr pTemp = ImagePixelFormat( CompressionQuality, Base, EImageFormat::IF_RGBA_UBYTE );
                 for (int m = 0; m < resultLODCount; ++m)
                 {
                     FIntVector2 mipSize = pResult->CalculateMipSize(m);
                     miro::RGBA_to_BC1(
                         mipSize[0], mipSize[1], pTemp->GetMipData( baseLOD + m ),
-                        pResult->GetMipData( m ), imageCompressionQuality );
+                        pResult->GetMipData( m ), CompressionQuality );
                 }
                 break;
             }
@@ -1081,15 +1081,15 @@ bool ImagePixelFormatInPlace( int imageCompressionQuality, Image* pResult, const
 
         case EImageFormat::IF_BC2:
         {
-            switch ( pBase->GetFormat() )
+            switch ( Base->GetFormat() )
             {
             case EImageFormat::IF_RGBA_UBYTE:
                 for ( int m=0; m<resultLODCount; ++m )
                 {
                     FIntVector2 mipSize = pResult->CalculateMipSize( m );
                     miro::RGBA_to_BC2(
-                        mipSize[0], mipSize[1], pBase->GetMipData( baseLOD + m ),
-                        pResult->GetMipData( m ), imageCompressionQuality );
+                        mipSize[0], mipSize[1], Base->GetMipData( baseLOD + m ),
+                        pResult->GetMipData( m ), CompressionQuality );
                 }
                 break;
 
@@ -1104,15 +1104,15 @@ bool ImagePixelFormatInPlace( int imageCompressionQuality, Image* pResult, const
 
         case EImageFormat::IF_BC3:
         {
-            switch ( pBase->GetFormat() )
+            switch ( Base->GetFormat() )
             {
             case EImageFormat::IF_RGBA_UBYTE:
                 for ( int m=0; m<resultLODCount; ++m )
                 {
                     FIntVector2 mipSize = pResult->CalculateMipSize( m );
                     miro::RGBA_to_BC3(
-                        mipSize[0], mipSize[1], pBase->GetMipData( baseLOD + m ),
-                        pResult->GetMipData( m ), imageCompressionQuality );
+                        mipSize[0], mipSize[1], Base->GetMipData( baseLOD + m ),
+                        pResult->GetMipData( m ), CompressionQuality );
                 }
                 break;
 
@@ -1121,8 +1121,8 @@ bool ImagePixelFormatInPlace( int imageCompressionQuality, Image* pResult, const
                 {
                     FIntVector2 mipSize = pResult->CalculateMipSize( m );
                     miro::RGB_to_BC3(
-                        mipSize[0], mipSize[1], pBase->GetMipData( baseLOD + m ),
-                        pResult->GetMipData( m ), imageCompressionQuality );
+                        mipSize[0], mipSize[1], Base->GetMipData( baseLOD + m ),
+                        pResult->GetMipData( m ), CompressionQuality );
                 }
                 break;
 
@@ -1134,8 +1134,8 @@ bool ImagePixelFormatInPlace( int imageCompressionQuality, Image* pResult, const
                 {
                     FIntVector2 mipSize = pResult->CalculateMipSize( m );
                     miro::BC1_to_BC3(
-                        mipSize[0], mipSize[1], pBase->GetMipData( baseLOD + m ),
-                        pResult->GetMipData( m ), imageCompressionQuality );
+                        mipSize[0], mipSize[1], Base->GetMipData( baseLOD + m ),
+                        pResult->GetMipData( m ), CompressionQuality );
                 }
                 break;
             }
@@ -1151,7 +1151,7 @@ bool ImagePixelFormatInPlace( int imageCompressionQuality, Image* pResult, const
 
         case EImageFormat::IF_BC4:
         {
-            switch ( pBase->GetFormat() )
+            switch ( Base->GetFormat() )
             {
             case EImageFormat::IF_L_UBYTE:
 			{
@@ -1159,8 +1159,8 @@ bool ImagePixelFormatInPlace( int imageCompressionQuality, Image* pResult, const
 				{
 					FIntVector2 mipSize = pResult->CalculateMipSize(m);
 					miro::L_to_BC4(
-						mipSize[0], mipSize[1], pBase->GetMipData(baseLOD + m),
-						pResult->GetMipData(m), imageCompressionQuality);
+						mipSize[0], mipSize[1], Base->GetMipData(baseLOD + m),
+						pResult->GetMipData(m), CompressionQuality);
 				}
 				break;
 			}
@@ -1169,14 +1169,14 @@ bool ImagePixelFormatInPlace( int imageCompressionQuality, Image* pResult, const
 			case EImageFormat::IF_BGRA_UBYTE:
 			{	
                 ImagePtr pTempBase =
-                        ImagePixelFormat( imageCompressionQuality, pBase, EImageFormat::IF_L_UBYTE );
+                        ImagePixelFormat( CompressionQuality, Base, EImageFormat::IF_L_UBYTE );
                 for (int m = 0; m < resultLODCount; ++m)
                 {
                     FIntVector2 mipSize = pTempBase->CalculateMipSize( m );
 
                     miro::L_to_BC4(
-                        mipSize[0], mipSize[1], pBase->GetMipData( baseLOD + m ),
-                        pResult->GetMipData( m ), imageCompressionQuality );
+                        mipSize[0], mipSize[1], Base->GetMipData( baseLOD + m ),
+                        pResult->GetMipData( m ), CompressionQuality );
                 }
 				break;
 			}
@@ -1192,15 +1192,15 @@ bool ImagePixelFormatInPlace( int imageCompressionQuality, Image* pResult, const
 
         case EImageFormat::IF_BC5:
         {
-            switch ( pBase->GetFormat() )
+            switch ( Base->GetFormat() )
             {
             case EImageFormat::IF_RGBA_UBYTE:
                 for ( int m=0; m<resultLODCount; ++m )
                 {
                     FIntVector2 mipSize = pResult->CalculateMipSize( m );
                     miro::RGBA_to_BC5(
-                        mipSize[0], mipSize[1], pBase->GetMipData( baseLOD + m ),
-                        pResult->GetMipData( m ), imageCompressionQuality );
+                        mipSize[0], mipSize[1], Base->GetMipData( baseLOD + m ),
+                        pResult->GetMipData( m ), CompressionQuality );
                 }
                 break;
 
@@ -1209,8 +1209,8 @@ bool ImagePixelFormatInPlace( int imageCompressionQuality, Image* pResult, const
                 {
                     FIntVector2 mipSize = pResult->CalculateMipSize( m );
                     miro::RGB_to_BC5(
-                        mipSize[0], mipSize[1], pBase->GetMipData( baseLOD + m ),
-                        pResult->GetMipData( m ), imageCompressionQuality );
+                        mipSize[0], mipSize[1], Base->GetMipData( baseLOD + m ),
+                        pResult->GetMipData( m ), CompressionQuality );
                 }
                 break;
 
@@ -1225,15 +1225,15 @@ bool ImagePixelFormatInPlace( int imageCompressionQuality, Image* pResult, const
 
         case EImageFormat::IF_ASTC_4x4_RGB_LDR:
         {
-            switch ( pBase->GetFormat() )
+            switch ( Base->GetFormat() )
             {
             case EImageFormat::IF_RGBA_UBYTE:
                 for ( int m=0; m<resultLODCount; ++m )
                 {
                     FIntVector2 mipSize = pResult->CalculateMipSize( m );
                     miro::RGBA_to_ASTC4x4RGBL(
-                        mipSize[0], mipSize[1], pBase->GetMipData( baseLOD + m ),
-                        pResult->GetMipData( m ), imageCompressionQuality );
+                        mipSize[0], mipSize[1], Base->GetMipData( baseLOD + m ),
+                        pResult->GetMipData( m ), CompressionQuality );
                 }
                 break;
 
@@ -1242,8 +1242,8 @@ bool ImagePixelFormatInPlace( int imageCompressionQuality, Image* pResult, const
                 {
                     FIntVector2 mipSize = pResult->CalculateMipSize( m );
                     miro::RGB_to_ASTC4x4RGBL(
-                        mipSize[0], mipSize[1], pBase->GetMipData( baseLOD + m ),
-                        pResult->GetMipData( m ), imageCompressionQuality );
+                        mipSize[0], mipSize[1], Base->GetMipData( baseLOD + m ),
+                        pResult->GetMipData( m ), CompressionQuality );
                 }
                 break;
 
@@ -1252,7 +1252,7 @@ bool ImagePixelFormatInPlace( int imageCompressionQuality, Image* pResult, const
                 {
                     FIntVector2 mipSize = pResult->CalculateMipSize( m );
                     miro::ASTC4x4RGBAL_to_ASTC4x4RGBL( mipSize[0], mipSize[1],
-                                       pBase->GetMipData(baseLOD+m),
+                                       Base->GetMipData(baseLOD+m),
                                        pResult->GetMipData(m) );
                 }
                 break;
@@ -1262,8 +1262,8 @@ bool ImagePixelFormatInPlace( int imageCompressionQuality, Image* pResult, const
 				{
 					FIntVector2 mipSize = pResult->CalculateMipSize(m);
 					miro::L_to_ASTC4x4RGBL(mipSize[0], mipSize[1],
-						pBase->GetMipData(baseLOD + m),
-						pResult->GetMipData(m), imageCompressionQuality);
+						Base->GetMipData(baseLOD + m),
+						pResult->GetMipData(m), CompressionQuality);
 				}
 				break;
 
@@ -1278,15 +1278,15 @@ bool ImagePixelFormatInPlace( int imageCompressionQuality, Image* pResult, const
 
         case EImageFormat::IF_ASTC_4x4_RGBA_LDR:
         {
-            switch ( pBase->GetFormat() )
+            switch ( Base->GetFormat() )
             {
             case EImageFormat::IF_RGBA_UBYTE:
                 for ( int m=0; m<resultLODCount; ++m )
                 {
                     FIntVector2 mipSize = pResult->CalculateMipSize( m );
                     miro::RGBA_to_ASTC4x4RGBAL(
-                        mipSize[0], mipSize[1], pBase->GetMipData( baseLOD + m ),
-                        pResult->GetMipData( m ), imageCompressionQuality );
+                        mipSize[0], mipSize[1], Base->GetMipData( baseLOD + m ),
+                        pResult->GetMipData( m ), CompressionQuality );
                 }
                 break;
 
@@ -1295,8 +1295,8 @@ bool ImagePixelFormatInPlace( int imageCompressionQuality, Image* pResult, const
                 {
                     FIntVector2 mipSize = pResult->CalculateMipSize( m );
                     miro::RGB_to_ASTC4x4RGBAL(
-                        mipSize[0], mipSize[1], pBase->GetMipData( baseLOD + m ),
-                        pResult->GetMipData( m ), imageCompressionQuality );
+                        mipSize[0], mipSize[1], Base->GetMipData( baseLOD + m ),
+                        pResult->GetMipData( m ), CompressionQuality );
                 }
                 break;
 
@@ -1305,7 +1305,7 @@ bool ImagePixelFormatInPlace( int imageCompressionQuality, Image* pResult, const
                 {
                     FIntVector2 mipSize = pResult->CalculateMipSize( m );
                     miro::ASTC4x4RGBL_to_ASTC4x4RGBAL( mipSize[0], mipSize[1],
-                                       pBase->GetMipData(baseLOD+m),
+                                       Base->GetMipData(baseLOD+m),
                                        pResult->GetMipData(m) );
                 }
                 break;
@@ -1321,15 +1321,15 @@ bool ImagePixelFormatInPlace( int imageCompressionQuality, Image* pResult, const
 
         case EImageFormat::IF_ASTC_4x4_RG_LDR:
         {
-            switch ( pBase->GetFormat() )
+            switch ( Base->GetFormat() )
             {
             case EImageFormat::IF_RGBA_UBYTE:
                 for ( int m=0; m<resultLODCount; ++m )
                 {
                     FIntVector2 mipSize = pResult->CalculateMipSize( m );
                     miro::RGBA_to_ASTC4x4RGL(
-                        mipSize[0], mipSize[1], pBase->GetMipData( baseLOD + m ),
-                        pResult->GetMipData( m ), imageCompressionQuality );
+                        mipSize[0], mipSize[1], Base->GetMipData( baseLOD + m ),
+                        pResult->GetMipData( m ), CompressionQuality );
                 }
                 break;
 
@@ -1338,8 +1338,8 @@ bool ImagePixelFormatInPlace( int imageCompressionQuality, Image* pResult, const
                 {
                     FIntVector2 mipSize = pResult->CalculateMipSize( m );
                     miro::RGB_to_ASTC4x4RGL(
-                        mipSize[0], mipSize[1], pBase->GetMipData( baseLOD + m ),
-                        pResult->GetMipData( m ), imageCompressionQuality );
+                        mipSize[0], mipSize[1], Base->GetMipData( baseLOD + m ),
+                        pResult->GetMipData( m ), CompressionQuality );
                 }
                 break;
 
@@ -1349,7 +1349,7 @@ bool ImagePixelFormatInPlace( int imageCompressionQuality, Image* pResult, const
                     FIntVector2 mipSize = pResult->CalculateMipSize( m );
                     // Hack that actually works because of block size.
                     miro::ASTC4x4RGBAL_to_ASTC4x4RGBL( mipSize[0], mipSize[1],
-                                       pBase->GetMipData(baseLOD+m),
+                                       Base->GetMipData(baseLOD+m),
                                        pResult->GetMipData(m) );
                 }
                 break;
@@ -1369,7 +1369,6 @@ bool ImagePixelFormatInPlace( int imageCompressionQuality, Image* pResult, const
         }
     }
 
-    return success;
 }
 
 
diff --git a/Engine/Plugins/Experimental/Mutable/Source/MutableRuntime/Private/MuR/OpImageResize.cpp b/Engine/Plugins/Experimental/Mutable/Source/MutableRuntime/Private/MuR/OpImageResize.cpp
new file mode 100644
--- /dev/null
+++ b/Engine/Plugins/Experimental/Mutable/Source/MutableRuntime/Private/MuR/OpImageResize.cpp
@@ -0,0 +1,170 @@
+// Copyright Epic Games, Inc. All Rights Reserved.
+
+#include "MuR/OpImageResize.h"
+
+namespace mu
+{
+
+	Ptr<Image> ImageResizeLinear(int32 imageCompressionQuality, const Image* pBasePtr, FImageSize destSize)
+	{
+		MUTABLE_CPUPROFILER_SCOPE(ImageResizeLinear);
+
+		if (pBasePtr->GetSize() == destSize)
+		{
+			return pBasePtr->Clone();
+		}
+
+		check(!(pBasePtr->m_flags & Image::IF_CANNOT_BE_SCALED));
+
+		ImagePtrConst pBase = pBasePtr;
+
+		// Shouldn't happen! But if it does...
+		EImageFormat sourceFormat = pBase->GetFormat();
+		EImageFormat uncompressedFormat = GetUncompressedFormat(sourceFormat);
+		if (sourceFormat != uncompressedFormat)
+		{
+			pBase = ImagePixelFormat(imageCompressionQuality, pBasePtr, uncompressedFormat);
+		}
+
+		FImageSize baseSize = FImageSize(pBase->GetSizeX(), pBase->GetSizeY());
+
+		ImagePtr pDest = new Image(destSize[0], destSize[1], 1, pBase->GetFormat());
+		if (!destSize[0] || !destSize[1] || !baseSize[0] || !baseSize[1])
+		{
+			return pDest;
+		}
+
+		// First resize X
+		ImagePtr pTemp;
+		if (destSize[0] > baseSize[0])
+		{
+			pTemp = new Image(destSize[0], baseSize[1], 1, pBase->GetFormat());
+			ImageMagnifyX(pTemp.get(), pBase.get());
+		}
+		else if (destSize[0] < baseSize[0])
+		{
+			pTemp = new Image(destSize[0], baseSize[1], 1, pBase->GetFormat());
+			ImageMinifyX(pTemp.get(), pBase.get());
+		}
+		else
+		{
+			pTemp = pBase->Clone();
+		}
+
+		// Now resize Y
+		if (destSize[1] > baseSize[1])
+		{
+			ImageMagnifyY(pDest.get(), pTemp.get());
+		}
+		else if (destSize[1] < baseSize[1])
+		{
+			ImageMinifyY(pDest.get(), pTemp.get());
+		}
+		else
+		{
+			pDest = pTemp;
+		}
+
+
+		// Reset format if it was changed to scale
+		if (sourceFormat != uncompressedFormat)
+		{
+			pDest = ImagePixelFormat(imageCompressionQuality, pDest.get(), sourceFormat);
+		}
+
+		// Update the relevancy data of the image.
+		if (pBase->m_flags & Image::EImageFlags::IF_HAS_RELEVANCY_MAP)
+		{
+			pDest->m_flags |= Image::EImageFlags::IF_HAS_RELEVANCY_MAP;
+
+			float FactorY = float(destSize[1]) / float(baseSize[1]);
+
+			pDest->RelevancyMinY = uint16(FMath::FloorToFloat(pBase->RelevancyMinY * FactorY));
+			pDest->RelevancyMaxY = uint16(FMath::Min((int32)FMath::CeilToFloat(pBase->RelevancyMinY * FactorY), pDest->GetSizeY() - 1));
+		}
+
+		return pDest;
+	}
+
+
+	void ImageResizeLinear(Image* pDest, int32 imageCompressionQuality, const Image* pBasePtr)
+	{
+		MUTABLE_CPUPROFILER_SCOPE(ImageResizeLinear);
+
+		check(!(pBasePtr->m_flags & Image::IF_CANNOT_BE_SCALED));
+
+		ImagePtrConst pBase = pBasePtr;
+
+		// Shouldn't happen! But if it does...
+		EImageFormat sourceFormat = pBase->GetFormat();
+		EImageFormat uncompressedFormat = GetUncompressedFormat(sourceFormat);
+		if (sourceFormat != uncompressedFormat)
+		{
+			pBase = ImagePixelFormat(imageCompressionQuality, pBasePtr, uncompressedFormat);
+		}
+
+		FImageSize baseSize = FImageSize(pBase->GetSizeX(), pBase->GetSizeY());
+		FImageSize destSize = FImageSize(pDest->GetSizeX(), pDest->GetSizeY());
+		if (!destSize[0] || !destSize[1] || !baseSize[0] || !baseSize[1])
+		{
+			return;
+		}
+
+		// First resize X
+		ImagePtr pTemp;
+		if (destSize[0] > baseSize[0])
+		{
+			pTemp = new Image(destSize[0], baseSize[1], 1, pBase->GetFormat());
+			ImageMagnifyX(pTemp.get(), pBase.get());
+		}
+		else if (destSize[0] < baseSize[0])
+		{
+			pTemp = new Image(destSize[0], baseSize[1], 1, pBase->GetFormat());
+			ImageMinifyX(pTemp.get(), pBase.get());
+		}
+		else
+		{
+			pTemp = pBase->Clone();
+		}
+
+		// Now resize Y
+		ImagePtr pTemp2;
+		if (destSize[1] > baseSize[1])
+		{
+			pTemp2 = new Image(destSize[0], destSize[1], 1, pBase->GetFormat());
+			ImageMagnifyY(pTemp2.get(), pTemp.get());
+		}
+		else if (destSize[1] < baseSize[1])
+		{
+			pTemp2 = new Image(destSize[0], destSize[1], 1, pBase->GetFormat());
+			ImageMinifyY(pTemp2.get(), pTemp.get());
+		}
+		else
+		{
+			pTemp2 = pTemp;
+		}
+		pTemp = nullptr;
+
+
+		// Reset format if it was changed to scale
+		if (sourceFormat != uncompressedFormat)
+		{
+			pTemp2 = ImagePixelFormat(imageCompressionQuality, pTemp2.get(), sourceFormat);
+		}
+
+		pDest->CopyMove(pTemp2.get());
+		pTemp2 = nullptr;
+
+		// Update the relevancy data of the image.
+		if (pBase->m_flags & Image::EImageFlags::IF_HAS_RELEVANCY_MAP)
+		{
+			pDest->m_flags |= Image::EImageFlags::IF_HAS_RELEVANCY_MAP;
+
+			float FactorY = float(destSize[1]) / float(baseSize[1]);
+
+			pDest->RelevancyMinY = uint16(FMath::FloorToFloat(pBase->RelevancyMinY * FactorY));
+			pDest->RelevancyMaxY = uint16(FMath::Min((int32)FMath::CeilToFloat(pBase->RelevancyMinY * FactorY), pDest->GetSizeY() - 1));
+		}
+	}
+
+}
\ No newline at end of file
diff --git a/Engine/Plugins/Experimental/Mutable/Source/MutableRuntime/Private/MuR/OpImageResize.h b/Engine/Plugins/Experimental/Mutable/Source/MutableRuntime/Private/MuR/OpImageResize.h
--- a/Engine/Plugins/Experimental/Mutable/Source/MutableRuntime/Private/MuR/OpImageResize.h
+++ b/Engine/Plugins/Experimental/Mutable/Source/MutableRuntime/Private/MuR/OpImageResize.h
@@ -3,7 +3,6 @@
 #pragma once
 
 #include "MuR/ImagePrivate.h"
-#include "MuR/OpImagePixelFormat.h"
 #include "MuR/MutableTrace.h"
 #include "Async/ParallelFor.h"
 
@@ -15,7 +14,7 @@ namespace mu
     //! Point-filter image resize
     //! TODO: Optimise
     //---------------------------------------------------------------------------------------------
-    inline ImagePtr ImageResize( const Image* pBase, FImageSize destSize )
+    inline Ptr<Image> ImageResize( const Image* pBase, FImageSize destSize )
     {
 		MUTABLE_CPUPROFILER_SCOPE(ImageResizePoint)
 
@@ -631,9 +630,9 @@ namespace mu
         int destSizeY = pDest->GetSizeY();
         int sizeX = pBase->GetSizeX();
 
-        uint32_t dy_16 = (uint32_t(baseSizeY)<<16) / destSizeY;
+        uint32 dy_16 = (uint32(baseSizeY)<<16) / destSizeY;
 
-        const uint8_t* pBaseBuf = pBase->GetData();
+        const uint8* pBaseBuf = pBase->GetData();
 
         // Linear filtering
 		//for (int x = 0; x < sizeX; ++x)
@@ -641,22 +640,22 @@ namespace mu
 			pDest, pBaseBuf, sizeX, destSizeY, dy_16
 		] (uint32 x) 
 		{
-			uint8_t* pDestBuf = pDest->GetData() + x * NC;
-			uint32_t py_16 = 0;
+			uint8* pDestBuf = pDest->GetData() + x * NC;
+			uint32 py_16 = 0;
 			for (int y = 0; y < destSizeY; ++y)
 			{
-				uint32_t r_16[NC];
+				uint32 r_16[NC];
 				for (int c = 0; c < NC; ++c)
 				{
 					r_16[c] = 0;
 				}
 
-				uint32_t epy_16 = py_16 + dy_16;
-				uint32_t py = py_16 >> 16;
-				uint32_t epy = epy_16 >> 16;
+				uint32 epy_16 = py_16 + dy_16;
+				uint32 py = py_16 >> 16;
+				uint32 epy = epy_16 >> 16;
 
 				// First fraction
-				uint32_t frac0 = py_16 & 0xffff;
+				uint32 frac0 = py_16 & 0xffff;
 				if (frac0)
 				{
 					for (int c = 0; c < NC; ++c)
@@ -672,14 +671,14 @@ namespace mu
 				{
 					for (int c = 0; c < NC; ++c)
 					{
-						r_16[c] += uint32_t(pBaseBuf[(py * sizeX + x) * NC + c]) << 16;
+						r_16[c] += uint32(pBaseBuf[(py * sizeX + x) * NC + c]) << 16;
 					}
 
 					++py;
 				}
 
 				// Second fraction
-				uint32_t frac1 = epy_16 & 0xffff;
+				uint32 frac1 = epy_16 & 0xffff;
 				if (frac1)
 				{
 					for (int c = 0; c < NC; ++c)
@@ -690,7 +689,7 @@ namespace mu
 
 				for (int c = 0; c < NC; ++c)
 				{
-					pDestBuf[c] = (uint8_t)(r_16[c] / dy_16);
+					pDestBuf[c] = (uint8)(r_16[c] / dy_16);
 				}
 
 				py_16 = epy_16;
@@ -821,133 +820,8 @@ namespace mu
     }
 
 
-    //---------------------------------------------------------------------------------------------
-    //! Bilinear filter image resize.
-    //---------------------------------------------------------------------------------------------
-    inline ImagePtr ImageResizeLinear( int imageCompressionQuality, const Image* pBasePtr,
-                                       FImageSize destSize )
-    {
-		MUTABLE_CPUPROFILER_SCOPE(ImageResizeLinear);
-		
-		if(pBasePtr->GetSize()==destSize)
-		{
-			return pBasePtr->Clone();
-		}
-
-		check(!(pBasePtr->m_flags & Image::IF_CANNOT_BE_SCALED));
-
-        ImagePtrConst pBase = pBasePtr;
-
-        // Shouldn't happen! But if it does...
-        EImageFormat sourceFormat = pBase->GetFormat();
-		EImageFormat uncompressedFormat = GetUncompressedFormat( sourceFormat );
-        if ( sourceFormat!=uncompressedFormat )
-        {
-            pBase = ImagePixelFormat( imageCompressionQuality, pBasePtr, uncompressedFormat );
-        }
-
-        FImageSize baseSize = FImageSize( pBase->GetSizeX(), pBase->GetSizeY() );
-
-        ImagePtr pDest = new Image( destSize[0], destSize[1], 1, pBase->GetFormat() );
-        if (!destSize[0] || !destSize[1] || !baseSize[0] || !baseSize[1])
-        {
-            return pDest;
-        }
-
-        // First resize X
-        ImagePtr pTemp;
-        if ( destSize[0] > baseSize[0] )
-        {
-            pTemp = new Image( destSize[0], baseSize[1], 1, pBase->GetFormat() );
-            ImageMagnifyX( pTemp.get(), pBase.get() );
-        }
-        else if ( destSize[0] < baseSize[0] )
-        {
-            pTemp = new Image( destSize[0], baseSize[1], 1, pBase->GetFormat() );
-            ImageMinifyX( pTemp.get(), pBase.get() );
-        }
-        else
-        {
-            pTemp = pBase->Clone();
-        }
+	/** Bilinear filter image resize. */
+	extern Ptr<Image> ImageResizeLinear(int32 imageCompressionQuality, const Image* pBasePtr, FImageSize destSize);
+	extern void ImageResizeLinear(Image* pDest, int32 imageCompressionQuality, const Image* pBasePtr);
 
-        // Now resize Y
-        if ( destSize[1] > baseSize[1] )
-        {
-            ImageMagnifyY( pDest.get(), pTemp.get() );
-        }
-        else if ( destSize[1] < baseSize[1] )
-        {
-            ImageMinifyY( pDest.get(), pTemp.get() );
-        }
-        else
-        {
-            pDest = pTemp;
-        }
-
-
-        // Reset format if it was changed to scale
-        if ( sourceFormat!=uncompressedFormat )
-        {
-            pDest = ImagePixelFormat( imageCompressionQuality, pDest.get(), sourceFormat );
-        }
-
-		// Update the relevancy data of the image.
-		if (pBase->m_flags & Image::EImageFlags::IF_HAS_RELEVANCY_MAP)
-		{
-			pDest->m_flags |= Image::EImageFlags::IF_HAS_RELEVANCY_MAP;
-
-			float FactorY = float(destSize[1]) / float(baseSize[1]);
-
-			pDest->RelevancyMinY = uint16( FMath::FloorToFloat(pBase->RelevancyMinY*FactorY) );
-			pDest->RelevancyMaxY = uint16( FMath::Min((int32)FMath::CeilToFloat(pBase->RelevancyMinY * FactorY), pDest->GetSizeY() - 1) );
-		}
-
-        return pDest;
-    }
-
-
-	inline mu::ImagePtr ImageReduceSize( const Image* pImage, const vec2<int32>& size, int32 imageCompressionQuality)
-	{
-		check( size[0] < pImage->GetSizeX() || size[1] < pImage->GetSizeY() );
-
-		check(!(pImage->m_flags & Image::IF_CANNOT_BE_SCALED));
-
-		mu::ImagePtr Result;
-		int32 ResizeByDroppingMips = 0;
-
-		// Is it smaller
-		if (size[0] < pImage->GetSizeX() && size[1] < pImage->GetSizeY()
-			&&
-			// It is a multiple
-			(pImage->GetSizeX() % size[0]) == 0 && (pImage->GetSizeY() % size[1]) == 0
-			&&
-			// It is the same multiple
-			(pImage->GetSizeX() / size[0]) == (pImage->GetSizeY() / size[1]))
-		{
-			ResizeByDroppingMips = FMath::CountTrailingZeros64(pImage->GetSizeX() / size[0]);
-		}
-
-		if (ResizeByDroppingMips > 0 && pImage->GetLODCount() > ResizeByDroppingMips)
-		{
-			Result = pImage->Clone();
-			Result->ReduceLODs(ResizeByDroppingMips);
-		}
-		else
-		{
-			if (IsCompressedFormat(pImage->GetFormat()))
-			{
-				EImageFormat format = GetUncompressedFormat(pImage->GetFormat());
-				Result = ImagePixelFormat(
-					imageCompressionQuality,
-					pImage,
-					format);
-			}
-
-			FImageSize blockSize((uint16)size[0], (uint16)size[1]);
-			Result = ImageResizeLinear(imageCompressionQuality, pImage, blockSize);
-		}
-
-		return Result;
-	}
 }
diff --git a/Engine/Plugins/Experimental/Mutable/Source/MutableRuntime/Private/MuR/OpImageSaturate.h b/Engine/Plugins/Experimental/Mutable/Source/MutableRuntime/Private/MuR/OpImageSaturate.h
--- a/Engine/Plugins/Experimental/Mutable/Source/MutableRuntime/Private/MuR/OpImageSaturate.h
+++ b/Engine/Plugins/Experimental/Mutable/Source/MutableRuntime/Private/MuR/OpImageSaturate.h
@@ -9,85 +9,6 @@
 namespace mu
 {
 
-	inline ImagePtr ImageSaturate( const Image* pA, float factor )
-	{
-		// Clamp the factor
-		// TODO: See what happens if we don't
-		factor = FMath::Max( 0.0f, factor );
-        int32_t f_8 = (int32_t)(factor*255);
-
-        ImagePtr pDest = new Image( pA->GetSizeX(), pA->GetSizeY(),
-                                    pA->GetLODCount(),
-                                    pA->GetFormat() );
-
-        uint8_t* pDestBuf = pDest->GetData();
-        const uint8_t* pABuf = pA->GetData();
-
-		// Generic implementation
-		int32 PixelCount = (int)pA->CalculatePixelCount();
-
-		switch ( pA->GetFormat() )
-		{
-		case EImageFormat::IF_RGB_UBYTE:
-		{
-			for ( int i=0; i<PixelCount; ++i )
-			{
-                int32_t l_16 = 76 * pABuf[3*i+0] + 150 * pABuf[3*i+1] + 29 * pABuf[3*i+2];
-
-				for ( int c=0; c<3; ++c )
-				{
-                    int32_t d_16 = (pABuf[3*i+c]<<8) - l_16;
-                    int32_t r_16 = l_16 + ( (d_16 * f_8) >> 8 );
-                    pDestBuf[3*i+c] = (uint8_t)FMath::Min( (uint32_t)r_16>>8, 255u );
-				}
-			}
-	
-			break;
-		}
-
-        case EImageFormat::IF_RGBA_UBYTE:
-        {
-            for ( int i=0; i<PixelCount; ++i )
-            {
-                int32_t l_16 = 76 * pABuf[4*i+0] + 150 * pABuf[4*i+1] + 29 * pABuf[4*i+2];
-
-                for ( int c=0; c<3; ++c )
-                {
-                    int32_t d_16 = (pABuf[4*i+c]<<8) - l_16;
-                    int32_t r_16 = l_16 + ( (d_16 * f_8) >> 8 );
-                    pDestBuf[4*i+c] = (uint8_t)FMath::Min( (uint32_t)r_16>>8, 255u );
-                }
-
-                pDestBuf[4*i+3] = pABuf[4*i+3];
-            }
-            break;
-        }
-
-        case EImageFormat::IF_BGRA_UBYTE:
-        {
-            for ( int i=0; i<PixelCount; ++i )
-            {
-                int32_t l_16 = 76 * pABuf[4*i+2] + 150 * pABuf[4*i+1] + 29 * pABuf[4*i+0];
-
-                for ( int c=0; c<3; ++c )
-                {
-                    int32_t d_16 = (pABuf[4*i+c]<<8) - l_16;
-                    int32_t r_16 = l_16 + ( (d_16 * f_8) >> 8 );
-                    pDestBuf[4*i+c] = (uint8_t)FMath::Min( (uint32_t)r_16>>8, 255u );
-                }
-
-                pDestBuf[4*i+3] = pABuf[4*i+3];
-            }
-            break;
-        }
-
-		default:
-			check(false);
-		}
-
-		return pDest;
-	}
-
 	template<bool bUseVectorImpl = false>
 	inline void ImageSaturateInPlace(Image* pA, float Factor)
 	{
