[mutable] Implement memory budget control for image operations.

diff --git a/Engine/Plugins/Experimental/Mutable/Source/MutableRuntime/Private/MuR/OpImageSwizzle.cpp b/Engine/Plugins/Experimental/Mutable/Source/MutableRuntime/Private/MuR/OpImageSwizzle.cpp
--- a/Engine/Plugins/Experimental/Mutable/Source/MutableRuntime/Private/MuR/OpImageSwizzle.cpp
+++ b/Engine/Plugins/Experimental/Mutable/Source/MutableRuntime/Private/MuR/OpImageSwizzle.cpp
@@ -1,42 +1,33 @@
 // Copyright Epic Games, Inc. All Rights Reserved.
 
-#include "MuR/OpImageSwizzle.h"
-
 #include "MuR/ImagePrivate.h"
 #include "Async/ParallelFor.h"
 
 namespace mu
 {
-	
-	Ptr<Image> ImageSwizzle
-		(
-			EImageFormat format,
-			const Ptr<const Image> pSources[],
-            const uint8 channels[]
-		)
+
+	void ImageSwizzle( Image* Result, const Ptr<const Image> Sources[], const uint8 Channels[] )
 	{
 		MUTABLE_CPUPROFILER_SCOPE(ImageSwizzle);
 
-		if (!pSources[0])
+		if (!Sources[0])
 		{
-			return nullptr;
+			return;
 		}
 
-        ImagePtr pDest = new Image( pSources[0]->GetSizeX(), pSources[0]->GetSizeY(),
-                                    pSources[0]->GetLODCount(),
-                                    format );
-
 		// Very slow generic implementations
-        int32 PixelCount = pDest->CalculatePixelCount();
+        int32 PixelCount = Result->CalculatePixelCount();
+
+		EImageFormat Format = Result->GetFormat();
 
         // Pixelcount should already match, but due to bugs it may not be the case. Try to detect it,
         // but avoid crashing below:
-        uint16 numChannels = GetImageFormatData(format).m_channels;
+        uint16 numChannels = GetImageFormatData(Format).m_channels;
         for (uint16 c=0;c<numChannels; ++c)
         {
-            if (pSources[c])
+            if (Sources[c])
             {
-                size_t sourcePixelCount = pSources[c]->CalculatePixelCount();
+                size_t sourcePixelCount = Sources[c]->CalculatePixelCount();
                 if (PixelCount>sourcePixelCount)
                 {
                     check(false);
@@ -49,7 +40,7 @@ namespace mu
 
 		int NumDestChannels = 0;
 
-		switch ( format )
+		switch (Format)
 		{
 		case EImageFormat::IF_L_UBYTE:
 			NumDestChannels = 1;
@@ -70,17 +61,17 @@ namespace mu
 
 		for (int i = 0; i < NumDestChannels; ++i)
 		{
-			uint8* pDestBuf = pDest->GetData() + i;
+			uint8* pDestBuf = Result->GetData() + i;
 
-			if (format == EImageFormat::IF_BGRA_UBYTE)
+			if (Format == EImageFormat::IF_BGRA_UBYTE)
 			{
 				if (i == 0)
 				{
-					pDestBuf = pDest->GetData() + 2;
+					pDestBuf = Result->GetData() + 2;
 				}
 				else if (i == 2)
 				{
-					pDestBuf = pDest->GetData() + 0;
+					pDestBuf = Result->GetData() + 0;
 				}
 			}
 
@@ -89,14 +80,14 @@ namespace mu
 			constexpr int32 NumBatchElems = 4096*2;
 			const int32 NumBatches = FMath::DivideAndRoundUp(PixelCount, NumBatchElems);
 
-			if (pSources[i])
+			if (Sources[i])
 			{
-				const uint8* pSourceBuf = pSources[i]->GetData() + channels[i];
+				const uint8* pSourceBuf = Sources[i]->GetData() + Channels[i];
 
-				switch (pSources[i]->GetFormat())
+				switch (Sources[i]->GetFormat())
 				{
 				case EImageFormat::IF_L_UBYTE:
-					if (channels[i] < 1)
+					if (Channels[i] < 1)
 					{
 						auto ProcessBatch = [pDestBuf, pSourceBuf, NumDestChannels, PixelCount, NumBatchElems ](int32 BatchId)
 						{
@@ -123,7 +114,7 @@ namespace mu
 					break;
 
 				case EImageFormat::IF_RGB_UBYTE:
-					if (channels[i] < 3)
+					if (Channels[i] < 3)
 					{
 						auto ProcessBatch = [pDestBuf, pSourceBuf, NumDestChannels, PixelCount, NumBatchElems](int32 BatchId)
 						{
@@ -150,7 +141,7 @@ namespace mu
 					break;
 
 				case EImageFormat::IF_RGBA_UBYTE:
-					if (channels[i] < 4)
+					if (Channels[i] < 4)
 					{	
 						auto ProcessBatch = [pDestBuf, pSourceBuf, NumDestChannels, PixelCount, NumBatchElems](int32 BatchId)
 						{
@@ -177,15 +168,15 @@ namespace mu
 					break;
 
 				case EImageFormat::IF_BGRA_UBYTE:
-					if (channels[i] == 0)
+					if (Channels[i] == 0)
 					{
-						pSourceBuf = pSources[i]->GetData() + 2;
+						pSourceBuf = Sources[i]->GetData() + 2;
 					}
-					else if (channels[i] == 2)
+					else if (Channels[i] == 2)
 					{
-						pSourceBuf = pSources[i]->GetData() + 0;
+						pSourceBuf = Sources[i]->GetData() + 0;
 					}
-					if (channels[i] < 4)
+					if (Channels[i] < 4)
 					{
 						auto ProcessBatch = [pDestBuf, pSourceBuf, NumDestChannels, PixelCount, NumBatchElems](int32 BatchId)
 						{
@@ -238,8 +229,26 @@ namespace mu
 				}
 			}
 		}
+	}
+
+
+	Ptr<Image> ImageSwizzle( EImageFormat Format, const Ptr<const Image> Sources[], const uint8 Channels[] )
+	{
+		MUTABLE_CPUPROFILER_SCOPE(ImageSwizzle);
+
+		if (!Sources[0])
+		{
+			return nullptr;
+		}
 
-		return pDest;
+		ImagePtr Dest = new Image(Sources[0]->GetSizeX(), Sources[0]->GetSizeY(),
+			Sources[0]->GetLODCount(),
+			Format);
+
+		ImageSwizzle(Dest.get(), Sources, Channels);
+
+		return Dest;
 	}
-	
+
+
 }
diff --git a/Engine/Plugins/Experimental/Mutable/Source/MutableRuntime/Private/MuR/OpImageTransform.h b/Engine/Plugins/Experimental/Mutable/Source/MutableRuntime/Private/MuR/OpImageTransform.h
--- a/Engine/Plugins/Experimental/Mutable/Source/MutableRuntime/Private/MuR/OpImageTransform.h
+++ b/Engine/Plugins/Experimental/Mutable/Source/MutableRuntime/Private/MuR/OpImageTransform.h
@@ -3,7 +3,6 @@
 #pragma once
 
 #include "MuR/ImagePrivate.h"
-#include "MuR/OpImagePixelFormat.h"
 #include "Async/ParallelFor.h"
 #include "MuR/MutableTrace.h"
 
diff --git a/Engine/Plugins/Experimental/Mutable/Source/MutableRuntime/Private/MuR/Operations.cpp b/Engine/Plugins/Experimental/Mutable/Source/MutableRuntime/Private/MuR/Operations.cpp
--- a/Engine/Plugins/Experimental/Mutable/Source/MutableRuntime/Private/MuR/Operations.cpp
+++ b/Engine/Plugins/Experimental/Mutable/Source/MutableRuntime/Private/MuR/Operations.cpp
@@ -85,13 +85,11 @@ namespace mu
 		{ DT_IMAGE,			true,	{ 0, 1, 1, 1, 0, 0, 0, 0, 0, 0, 0, 0, 0, 0, 0, 0, 0, 0, 0, 0, 0, 0, 0, 0, 0, 0, 0, 0, 0 }		},	// IM_RESIZEREL
 		{ DT_IMAGE,			true,   { 1, 1, 1, 1, 1, 1, 1, 1, 1, 1, 1, 1, 0, 0, 0, 0, 0, 0, 0, 0, 0, 0, 0, 0, 0, 0, 0, 0, 0 }		},	// IM_BLANKLAYOUT
 		{ DT_IMAGE,			true,	{ 1, 1, 1, 1, 1, 1, 1, 1, 1, 1, 1, 0, 0, 0, 1, 1, 1, 1, 1, 1, 1, 1, 0, 0, 1, 1, 0, 0, 0 }		},	// IM_COMPOSE
-		{ DT_IMAGE,			true,	{ 0, 1, 1, 0, 0, 0, 0, 0, 0, 0, 0, 0, 0, 0, 0, 0, 0, 0, 0, 0, 0, 0, 0, 0, 0, 0, 0, 0, 0 }		},	// IM_DIFFERENCE
 		{ DT_IMAGE,			true,	{ 0, 1, 1, 1, 0, 0, 0, 0, 0, 0, 0, 0, 0, 0, 0, 0, 0, 0, 0, 0, 0, 0, 0, 0, 0, 0, 0, 0, 0 }		},	// IM_INTERPOLATE
 		{ DT_IMAGE,			true,	{ 0, 1, 1, 1, 0, 0, 0, 0, 0, 0, 0, 0, 0, 0, 0, 0, 0, 0, 0, 0, 0, 0, 0, 0, 0, 0, 0, 0, 0 }		},	// IM_INTERPOLATE3
 		{ DT_IMAGE,			true,	{ 0, 1, 1, 0, 0, 0, 0, 0, 0, 0, 0, 0, 0, 0, 0, 0, 0, 0, 0, 0, 0, 0, 0, 0, 0, 0, 0, 0, 0 }		},	// IM_SATURATE
 		{ DT_IMAGE,			true,	{ 0, 1, 1, 0, 0, 0, 0, 0, 0, 0, 0, 0, 0, 0, 0, 0, 0, 0, 0, 0, 0, 0, 0, 0, 0, 0, 0, 0, 0 }		},	// IM_LUMINANCE
 		{ DT_IMAGE,			true,	{ 0, 1, 1, 1, 0, 0, 0, 0, 0, 0, 0, 0, 0, 0, 0, 0, 0, 0, 0, 0, 0, 0, 0, 0, 0, 0, 0, 0, 0 }		},	// IM_SWIZZLE
-		{ DT_IMAGE,			true,	{ 0, 1, 1, 1, 0, 0, 0, 0, 0, 0, 0, 0, 0, 0, 0, 0, 0, 0, 0, 0, 0, 0, 0, 0, 0, 0, 0, 0, 0 }		},	// IM_SELECTCOLOUR
 		{ DT_IMAGE,			true,	{ 0, 1, 1, 1, 0, 0, 0, 0, 0, 0, 0, 0, 0, 0, 0, 0, 0, 0, 0, 0, 0, 0, 0, 0, 0, 0, 0, 0, 0 }		},	// IM_COLOURMAP
 		{ DT_IMAGE,			true,	{ 0, 0, 0, 0, 0, 0, 0, 0, 0, 0, 0, 0, 0, 0, 0, 0, 0, 0, 0, 0, 0, 0, 0, 0, 0, 0, 0, 0, 0 }		},	// IM_GRADIENT
 		{ DT_IMAGE,			true,	{ 0, 1, 1, 1, 0, 0, 0, 0, 0, 0, 0, 0, 0, 0, 0, 0, 0, 0, 0, 0, 0, 0, 0, 0, 0, 0, 0, 0, 0 }		},	// IM_BINARISE
@@ -269,11 +267,6 @@ namespace mu
             f(&op.args.ImageBlankLayout.layout );
             break;
 
-        case OP_TYPE::IM_DIFFERENCE:
-            f(&op.args.ImageDifference.a );
-            f(&op.args.ImageDifference.b );
-            break;
-
         case OP_TYPE::IM_INTERPOLATE:
             f(&op.args.ImageInterpolate.factor );
 
@@ -300,11 +293,6 @@ namespace mu
             f(&op.args.ImageLuminance.base );
             break;
 
-        case OP_TYPE::IM_SELECTCOLOUR:
-            f(&op.args.ImageSelectColour.base );
-            f(&op.args.ImageSelectColour.colour );
-            break;
-
         case OP_TYPE::IM_COLOURMAP:
             f(&op.args.ImageColourMap.base );
             f(&op.args.ImageColourMap.mask );
@@ -694,14 +682,6 @@ namespace mu
             break;
         }
 
-        case OP_TYPE::IM_DIFFERENCE:
-        {
-			OP::ImageDifferenceArgs args = program.GetOpArgs<OP::ImageDifferenceArgs>(at);
-            f(args.a );
-            f(args.b );
-            break;
-        }
-
         case OP_TYPE::IM_INTERPOLATE:
         {
 			OP::ImageInterpolateArgs args = program.GetOpArgs<OP::ImageInterpolateArgs>(at);
@@ -750,14 +730,6 @@ namespace mu
             break;
         }
 
-        case OP_TYPE::IM_SELECTCOLOUR:
-        {
-			OP::ImageSelectColourArgs args = program.GetOpArgs<OP::ImageSelectColourArgs>(at);
-            f(args.base );
-            f(args.colour );
-            break;
-        }
-
         case OP_TYPE::IM_COLOURMAP:
         {
 			OP::ImageColourMapArgs args = program.GetOpArgs<OP::ImageColourMapArgs>(at);
diff --git a/Engine/Plugins/Experimental/Mutable/Source/MutableRuntime/Private/MuR/Operations.h b/Engine/Plugins/Experimental/Mutable/Source/MutableRuntime/Private/MuR/Operations.h
--- a/Engine/Plugins/Experimental/Mutable/Source/MutableRuntime/Private/MuR/Operations.h
+++ b/Engine/Plugins/Experimental/Mutable/Source/MutableRuntime/Private/MuR/Operations.h
@@ -167,9 +167,6 @@ namespace mu
         //! Copy an image into a rect of another one.
         IM_COMPOSE,
 
-        //! Compare two images and generate a one channel image with 1 in the different pixels.
-        IM_DIFFERENCE,
-
         //! Interpolate between 2 images taken from a row of targets (2 consecutive targets).
         IM_INTERPOLATE,
 
@@ -185,9 +182,6 @@ namespace mu
         //! Recombine the channels of several images into one.
         IM_SWIZZLE,
 
-        //! Create a black and white image with the pixels of another image that match a colour.
-        IM_SELECTCOLOUR,
-
         //! Convert the source image colours using a "palette" image sampled with the source
         //! grey-level.
         IM_COLOURMAP,
@@ -720,12 +714,6 @@ namespace mu
             ADDRESS sources[ MUTABLE_OP_MAX_SWIZZLE_CHANNELS ];
         };
 
-        struct ImageSelectColourArgs
-        {
-            ADDRESS base;
-            ADDRESS colour;
-        };
-
         struct ImageColourMapArgs
         {
             ADDRESS base;
@@ -1149,7 +1137,6 @@ namespace mu
             ImageInterpolate3Args ImageInterpolate3;
             ImageSaturateArgs ImageSaturate;
             ImageLuminanceArgs ImageLuminance;
-            ImageSelectColourArgs ImageSelectColour;
             ImageColourMapArgs ImageColourMap;
             ImageGradientArgs ImageGradient;
             ImageBinariseArgs ImageBinarise;
diff --git a/Engine/Plugins/Experimental/Mutable/Source/MutableRuntime/Private/MuR/Settings.cpp b/Engine/Plugins/Experimental/Mutable/Source/MutableRuntime/Private/MuR/Settings.cpp
--- a/Engine/Plugins/Experimental/Mutable/Source/MutableRuntime/Private/MuR/Settings.cpp
+++ b/Engine/Plugins/Experimental/Mutable/Source/MutableRuntime/Private/MuR/Settings.cpp
@@ -3,53 +3,38 @@
 #include "MuR/Settings.h"
 
 #include "Misc/AssertionMacros.h"
-#include "MuR/SettingsPrivate.h"
 
 namespace mu
 {
-    //---------------------------------------------------------------------------------------------
-    //!
-    //---------------------------------------------------------------------------------------------
-    Settings::Settings()
+  
+	Settings::Settings()
     {
-        m_pD = new Settings::Private();
     }
 
 
-    //---------------------------------------------------------------------------------------------
     Settings::~Settings()
     {
-        check( m_pD );
-        delete m_pD;
-        m_pD = 0;
-    }
-
-
-    //---------------------------------------------------------------------------------------------
-    Settings::Private* Settings::GetPrivate() const
-    {
-        return m_pD;
     }
 
 
     //---------------------------------------------------------------------------------------------
     void Settings::SetProfile( bool bEnabled )
     {
-        m_pD->m_profile = bEnabled;
+        bProfile = bEnabled;
     }
 
 
     //---------------------------------------------------------------------------------------------
-    void Settings::SetStreamingCache( uint64 bytes )
+    void Settings::SetWorkingMemoryBytes( uint64 Bytes )
     {
-        m_pD->m_streamingCacheBytes = bytes;
+        WorkingMemoryBytes = Bytes;
     }
 
 
     //---------------------------------------------------------------------------------------------
-    void Settings::SetImageCompressionQuality( int quality )
+    void Settings::SetImageCompressionQuality( int32 Quality )
     {
-        m_pD->m_imageCompressionQuality = quality;
+        ImageCompressionQuality = Quality;
     }
 
 }
diff --git a/Engine/Plugins/Experimental/Mutable/Source/MutableRuntime/Private/MuR/System.cpp b/Engine/Plugins/Experimental/Mutable/Source/MutableRuntime/Private/MuR/System.cpp
--- a/Engine/Plugins/Experimental/Mutable/Source/MutableRuntime/Private/MuR/System.cpp
+++ b/Engine/Plugins/Experimental/Mutable/Source/MutableRuntime/Private/MuR/System.cpp
@@ -27,7 +27,6 @@
 #include "MuR/ParametersPrivate.h"
 #include "MuR/Platform.h"
 #include "MuR/Serialisation.h"
-#include "MuR/SettingsPrivate.h"
 #include "MuR/SystemPrivate.h"
 #include "Templates/SharedPointer.h"
 #include "Templates/Tuple.h"
@@ -42,15 +41,26 @@ namespace mu
 
 	MUTABLE_IMPLEMENT_ENUM_SERIALISABLE(ETextureCompressionStrategy);
 
+	TRACE_DECLARE_INT_COUNTER(MutableRuntime_LiveInstances,		TEXT("MutableRuntime/LiveInstances"));
+	TRACE_DECLARE_INT_COUNTER(MutableRuntime_Updates,			TEXT("MutableRuntime/Updates"));
+	TRACE_DECLARE_INT_COUNTER(MutableRuntime_StreamingBytes,	TEXT("MutableRuntime/StreamingBytes"));
+
+	TRACE_DECLARE_INT_COUNTER(MutableRuntime_MemInternal,	TEXT("MutableRuntime/MemInternal"));
+	TRACE_DECLARE_INT_COUNTER(MutableRuntime_MemTemp,		TEXT("MutableRuntime/MemTemp"));
+	TRACE_DECLARE_INT_COUNTER(MutableRuntime_MemPool,		TEXT("MutableRuntime/MemPool"));
+	TRACE_DECLARE_INT_COUNTER(MutableRuntime_MemCache,		TEXT("MutableRuntime/MemCache"));
+	TRACE_DECLARE_INT_COUNTER(MutableRuntime_MemRom,		TEXT("MutableRuntime/MemRom"));
+	TRACE_DECLARE_INT_COUNTER(MutableRuntime_MemStream,		TEXT("MutableRuntime/MemStream"));
+	TRACE_DECLARE_INT_COUNTER(MutableRuntime_MemTotal,		TEXT("MutableRuntime/MemTotal"));
+	TRACE_DECLARE_INT_COUNTER(MutableRuntime_MemBudget,		TEXT("MutableRuntime/MemBudget"));
+
 
     //---------------------------------------------------------------------------------------------
-    //---------------------------------------------------------------------------------------------
-    //---------------------------------------------------------------------------------------------
-    System::System(const SettingsPtr& pInSettings, ExtensionDataStreamer* DataStreamer)
+    System::System(const Ptr<Settings>& InSettings, const TSharedPtr<ExtensionDataStreamer>& DataStreamer)
     {
 		LLM_SCOPE_BYNAME(TEXT("MutableRuntime"));
 
-        SettingsPtr pSettings = pInSettings;
+        Ptr<Settings> pSettings = InSettings;
 
         if ( !pSettings )
         {
@@ -83,22 +93,19 @@ namespace mu
 
 
     //---------------------------------------------------------------------------------------------
-    System::Private::Private(SettingsPtr pSettings, ExtensionDataStreamer* DataStreamer)
+    System::Private::Private( Ptr<mu::Settings> InSettings, const TSharedPtr<mu::ExtensionDataStreamer>& InDataStreamer)
     {
 		LLM_SCOPE_BYNAME(TEXT("MutableRuntime"));
 
-        m_pSettings = pSettings;
-        m_pStreamInterface = nullptr;
-        m_pImageParameterGenerator = nullptr;
-        m_maxMemory = 0;
+        Settings = InSettings;
 	
-		m_modelCache.m_romBudget = pSettings->GetPrivate()->m_streamingCacheBytes;
+		WorkingMemoryManager.BudgetBytes = Settings->WorkingMemoryBytes;
 
-		if (!DataStreamer)
+		ExtensionDataStreamer = InDataStreamer;
+		if (!ExtensionDataStreamer)
 		{
-			DataStreamer = new NullExtensionDataStreamer();
+			ExtensionDataStreamer = MakeShared<NullExtensionDataStreamer>();
 		}
-		m_ExtensionDataStreamer = DataStreamer;
 	}
 
 
@@ -108,127 +115,117 @@ namespace mu
 		LLM_SCOPE_BYNAME(TEXT("MutableRuntime"));
 		MUTABLE_CPUPROFILER_SCOPE(SystemPrivateDestructor);
 
-        delete m_pStreamInterface;
-        m_pStreamInterface = nullptr;
-
-		delete m_ExtensionDataStreamer;
-		m_ExtensionDataStreamer = nullptr;
-
-        delete m_pImageParameterGenerator;
-        m_pImageParameterGenerator = nullptr;
+		// Make it explicit to try to capture metrics
+        StreamInterface = nullptr;
+		ExtensionDataStreamer = nullptr;
+        ImageParameterGenerator = nullptr;
     }
 
 
     //---------------------------------------------------------------------------------------------
-    void System::SetStreamingInterface( ModelStreamer *pInterface )
+    void System::SetStreamingInterface(const TSharedPtr<ModelStreamer>& InInterface )
     {
 		LLM_SCOPE_BYNAME(TEXT("MutableRuntime"));
-
-        (void)pInterface;
        
-        delete m_pD->m_pStreamInterface;
-        m_pD->m_pStreamInterface = pInterface;
+        m_pD->StreamInterface = InInterface;
     }
 
 
     //---------------------------------------------------------------------------------------------
-    void System::SetStreamingCache( uint64 bytes )
+    void System::SetWorkingMemoryBytes( uint64 InBytes )
     {
 		LLM_SCOPE_BYNAME(TEXT("MutableRuntime"));
+		MUTABLE_CPUPROFILER_SCOPE(SetWorkingMemoryBytes);
 
-        m_pD->SetStreamingCache( bytes );
-    }
+		m_pD->WorkingMemoryManager.BudgetBytes = InBytes;
+		m_pD->WorkingMemoryManager.EnsureBudgetBelow(0);
+	}
 
 	
 	//---------------------------------------------------------------------------------------------
-	void System::ClearStreamingCache()
-    {
-		LLM_SCOPE_BYNAME(TEXT("MutableRuntime"));
-
-        m_pD->ClearStreamingCache();
-   }
-
-	
-    //---------------------------------------------------------------------------------------------
-    void System::SetImageParameterGenerator( ImageParameterGenerator* pInterface )
+	void System::ClearWorkingMemory()
     {
 		LLM_SCOPE_BYNAME(TEXT("MutableRuntime"));
 
-        delete m_pD->m_pImageParameterGenerator;
-        m_pD->m_pImageParameterGenerator = pInterface;
-    }
-
-
-    //---------------------------------------------------------------------------------------------
-    void System::SetMemoryLimit( uint32 mem )
-    {
-		LLM_SCOPE_BYNAME(TEXT("MutableRuntime"));
+		UE_LOG(LogMutableCore, Log, TEXT("Forcing a working memory clear."));
 
-        m_pD->m_maxMemory = mem;
+		// rom caches
+		for (FWorkingMemoryManager::FModelCacheEntry& ModelCache : m_pD->WorkingMemoryManager.CachePerModel)
+		{
+			if (const TSharedPtr<const Model> CacheModel = ModelCache.Model.Pin())
+			{
+				FProgram& Program = CacheModel->GetPrivate()->m_program;
 
-        // TODO: Clear cache if we are over the new limit
-    }
+				for (int32 RomIndex = 0; RomIndex < Program.m_roms.Num(); ++RomIndex)
+				{
+					Program.UnloadRom(RomIndex);
+				}
+			}
+		}
 
+		m_pD->WorkingMemoryManager.PooledImages.Empty();
+		m_pD->WorkingMemoryManager.CacheResources.Empty();
+		check(m_pD->WorkingMemoryManager.TempImages.IsEmpty());
+	}
 
+	
     //---------------------------------------------------------------------------------------------
-    void System::ClearCaches()
+    void System::SetImageParameterGenerator(const TSharedPtr<ImageParameterGenerator>& InInterface )
     {
 		LLM_SCOPE_BYNAME(TEXT("MutableRuntime"));
 
-        // TODO
+        m_pD->ImageParameterGenerator = InInterface;
     }
 
-	TRACE_DECLARE_INT_COUNTER(MutableRuntime_LiveInstances, TEXT("MutableRuntime/LiveInstances"));
-	TRACE_DECLARE_INT_COUNTER(MutableRuntime_Updates, TEXT("MutableRuntime/Updates"));
 
     //---------------------------------------------------------------------------------------------
-    Instance::ID System::NewInstance( const TSharedPtr<const Model>& pModel )
+    Instance::ID System::NewInstance( const TSharedPtr<const Model>& InModel )
     {
 		LLM_SCOPE_BYNAME(TEXT("MutableRuntime"));
 		MUTABLE_CPUPROFILER_SCOPE(NewInstance);
 
-		Private::FLiveInstance instanceData;
-		instanceData.m_instanceID = ++m_pD->m_lastInstanceID;
-		instanceData.m_pInstance = nullptr;
-		instanceData.m_pModel = pModel;
-		instanceData.m_state = -1;
-		instanceData.m_memory = MakeShared<FProgramCache>();
-		m_pD->m_liveInstances.Add(instanceData);
+		FLiveInstance instanceData;
+		instanceData.InstanceID = ++m_pD->LastInstanceID;
+		instanceData.Instance = nullptr;
+		instanceData.Model = InModel;
+		instanceData.State = -1;
+		instanceData.Cache = MakeShared<FProgramCache>();
+		m_pD->WorkingMemoryManager.LiveInstances.Add(instanceData);
 
-		TRACE_COUNTER_SET(MutableRuntime_LiveInstances, m_pD->m_liveInstances.Num());
+		TRACE_COUNTER_SET(MutableRuntime_LiveInstances, m_pD->WorkingMemoryManager.LiveInstances.Num());
 
-		return instanceData.m_instanceID;
+		return instanceData.InstanceID;
 	}
 
 
     //---------------------------------------------------------------------------------------------
-    const Instance* System::BeginUpdate( Instance::ID instanceID,
-                                     const ParametersPtrConst& pParams,
-                                     int32 stateIndex,
-                                     uint32 lodMask )
+    const Instance* System::BeginUpdate( Instance::ID InInstanceID,
+                                     const Ptr<const Parameters>& InParams,
+                                     int32 InStateIndex,
+                                     uint32 InLodMask )
     {
 		LLM_SCOPE_BYNAME(TEXT("MutableRuntime"));
 		MUTABLE_CPUPROFILER_SCOPE(SystemBeginUpdate);
 		TRACE_COUNTER_INCREMENT(MutableRuntime_Updates);
 
-		if (!pParams)
+		if (!InParams)
 		{
 			UE_LOG(LogMutableCore, Error, TEXT("Invalid parameters in mutable update."));
 			return nullptr;
 		}
 
-		Private::FLiveInstance* pLiveInstance = m_pD->FindLiveInstance(instanceID);
+		FLiveInstance* pLiveInstance = m_pD->FindLiveInstance(InInstanceID);
 		if (!pLiveInstance)
 		{
 			UE_LOG(LogMutableCore, Error, TEXT("Invalid instance id in mutable update."));
 			return nullptr;
 		}
 
-		m_pD->m_memory = pLiveInstance->m_memory;
+		m_pD->WorkingMemoryManager.CurrentInstanceCache = pLiveInstance->Cache;
 
-		FProgram& program = pLiveInstance->m_pModel->GetPrivate()->m_program;
+		FProgram& program = pLiveInstance->Model->GetPrivate()->m_program;
 
-		bool validState = stateIndex >= 0 && stateIndex < (int)program.m_states.Num();
+		bool validState = InStateIndex >= 0 && InStateIndex < (int)program.m_states.Num();
 		if (!validState)
 		{
 			UE_LOG(LogMutableCore, Error, TEXT("Invalid state in mutable update."));
@@ -236,69 +233,73 @@ namespace mu
 		}
 
 		// This may free resources that allow us to use less memory.
-		pLiveInstance->m_pInstance = nullptr;
+		pLiveInstance->Instance = nullptr;
 
-		bool fullBuild = (stateIndex != pLiveInstance->m_state);
+		bool fullBuild = (InStateIndex != pLiveInstance->State);
 
-		pLiveInstance->m_state = stateIndex;
+		pLiveInstance->State = InStateIndex;
 
 		// If we changed parameters that are not in this state, we need to rebuild all.
 		if (!fullBuild)
 		{
-			fullBuild = m_pD->CheckUpdatedParameters(pLiveInstance, pParams.get(), pLiveInstance->m_updatedParameters);
+			fullBuild = m_pD->CheckUpdatedParameters(pLiveInstance, InParams.get(), pLiveInstance->UpdatedParameters);
 		}
 
 		// Remove cached data
-		pLiveInstance->m_memory->ClearCacheLayer0();
+		m_pD->WorkingMemoryManager.ClearCacheLayer0();
 		if (fullBuild)
 		{
-			pLiveInstance->m_memory->ClearCacheLayer1();
+			m_pD->WorkingMemoryManager.ClearCacheLayer1();
 		}
 
-		OP::ADDRESS rootAt = pLiveInstance->m_pModel->GetPrivate()->m_program.m_states[stateIndex].m_root;
+		m_pD->WorkingMemoryManager.BeginRunnerThread();
+
+		OP::ADDRESS rootAt = pLiveInstance->Model->GetPrivate()->m_program.m_states[InStateIndex].m_root;
 
-		m_pD->PrepareCache(pLiveInstance->m_pModel.Get(), stateIndex);
-		pLiveInstance->m_pOldParameters = pParams->Clone();
+		m_pD->PrepareCache(pLiveInstance->Model.Get(), InStateIndex);
+		pLiveInstance->OldParameters = InParams->Clone();
 
-		m_pD->RunCode(pLiveInstance->m_pModel, pParams.get(), rootAt, lodMask);
+		m_pD->RunCode(pLiveInstance->Model, InParams.get(), rootAt, InLodMask);
 
-		InstancePtrConst pResult = pLiveInstance->m_memory->GetInstance(FCacheAddress(rootAt, 0, 0));
+		Ptr<const Instance> Result = pLiveInstance->Cache->GetInstance(FCacheAddress(rootAt, 0, 0));
 
 		// Debug check to see if we managed the op-hit-counts correctly
-		pLiveInstance->m_memory->CheckHitCountsCleared();
+		pLiveInstance->Cache->CheckHitCountsCleared();
 
-		pLiveInstance->m_pInstance = pResult;
-		if (pResult)
+		pLiveInstance->Instance = Result;
+		if (Result)
 		{
-			pResult->GetPrivate()->m_id = pLiveInstance->m_instanceID;
+			Result->GetPrivate()->m_id = pLiveInstance->InstanceID;
 		}
 
-		m_pD->m_memory = nullptr;
+		m_pD->WorkingMemoryManager.EndRunnerThread();
 
-		return pResult.get();
+		m_pD->WorkingMemoryManager.CurrentInstanceCache = nullptr;
+
+		return Result.get();
 	}
 
 
 	//---------------------------------------------------------------------------------------------
-	ImagePtrConst System::GetImage(Instance::ID instanceID, RESOURCE_ID imageId, int32 MipsToSkip, int32 InImageLOD)
+	Ptr<const Image> System::GetImage(Instance::ID instanceID, RESOURCE_ID imageId, int32 MipsToSkip, int32 InImageLOD)
 	{
 		LLM_SCOPE_BYNAME(TEXT("MutableRuntime"));
 		MUTABLE_CPUPROFILER_SCOPE(SystemGetImage);
 
-		ImagePtrConst pResult;
+		Ptr<const Image> pResult;
 
 		// Find the live instance
-		Private::FLiveInstance* pLiveInstance = m_pD->FindLiveInstance(instanceID);
+		FLiveInstance* pLiveInstance = m_pD->FindLiveInstance(instanceID);
 		check(pLiveInstance);
-		m_pD->m_memory = pLiveInstance->m_memory;
+		m_pD->WorkingMemoryManager.CurrentInstanceCache = pLiveInstance->Cache;
 
 		// Find the resource id in the model's resource cache
-		for (const Model::Private::RESOURCE_KEY& res : pLiveInstance->m_pModel->GetPrivate()->m_generatedResources)
+		for (const Model::Private::RESOURCE_KEY& res : pLiveInstance->Model->GetPrivate()->m_generatedResources)
 		{
 			if (res.m_id == imageId)
 			{
-				pResult = m_pD->BuildImage(pLiveInstance->m_pModel,
-					pLiveInstance->m_pOldParameters.get(),
+				pResult = m_pD->BuildImage(pLiveInstance->Model,
+					pLiveInstance->OldParameters.get(),
 					res.m_rootAddress, MipsToSkip, InImageLOD);
 
 				// We always need to return something valid.
@@ -311,7 +312,7 @@ namespace mu
 			}
 		}
 
-		m_pD->m_memory = nullptr;
+		m_pD->WorkingMemoryManager.CurrentInstanceCache = nullptr;
 
 		return pResult;
 	}
@@ -335,32 +336,32 @@ namespace mu
 		OutDesc = FImageDesc();
 
 		// Find the live instance
-		Private::FLiveInstance* pLiveInstance = m_pD->FindLiveInstance(instanceID);
+		FLiveInstance* pLiveInstance = m_pD->FindLiveInstance(instanceID);
 		check(pLiveInstance);
-		m_pD->m_memory = pLiveInstance->m_memory;
+		m_pD->WorkingMemoryManager.CurrentInstanceCache = pLiveInstance->Cache;
 
 		// Find the resource id in the model's resource cache
-		for (const Model::Private::RESOURCE_KEY& res : pLiveInstance->m_pModel->GetPrivate()->m_generatedResources)
+		for (const Model::Private::RESOURCE_KEY& res : pLiveInstance->Model->GetPrivate()->m_generatedResources)
 		{
 			if (res.m_id == imageId)
 			{
-				const mu::Model* Model = pLiveInstance->m_pModel.Get();
+				const mu::Model* Model = pLiveInstance->Model.Get();
 				const mu::FProgram& program = Model->GetPrivate()->m_program;
 
 				int32 VarValue = CVarClearImageDescCache.GetValueOnAnyThread();
 				if (VarValue != 0)
 				{
-					m_pD->m_memory->m_descCache.Reset();
+					m_pD->WorkingMemoryManager.CurrentInstanceCache->m_descCache.Reset();
 				}
 
-				m_pD->m_memory->m_descCache.SetNum(program.m_opAddress.Num());
+				m_pD->WorkingMemoryManager.CurrentInstanceCache->m_descCache.SetNum(program.m_opAddress.Num());
 
 				OP::ADDRESS at = res.m_rootAddress;
 				mu::OP_TYPE opType = program.GetOpType(at);
 				if (GetOpDataType(opType) == DT_IMAGE)
 				{
 					int8 executionOptions = 0;
-					CodeRunner Runner(m_pD->m_pSettings, m_pD, EExecutionStrategy::MinimizeMemory, pLiveInstance->m_pModel, pLiveInstance->m_pOldParameters.get(), at, System::AllLODs, executionOptions, 0, FScheduledOp::EType::ImageDesc);
+					CodeRunner Runner(m_pD->Settings, m_pD, EExecutionStrategy::MinimizeMemory, pLiveInstance->Model, pLiveInstance->OldParameters.get(), at, System::AllLODs, executionOptions, 0, FScheduledOp::EType::ImageDesc);
 					Runner.Run();
 					Runner.GetImageDescResult(OutDesc);
 				}
@@ -369,7 +370,7 @@ namespace mu
 			}
 		}
 
-		m_pD->m_memory = nullptr;
+		m_pD->WorkingMemoryManager.CurrentInstanceCache = nullptr;
 	}
 
 
@@ -383,17 +384,17 @@ namespace mu
 		MeshPtrConst pResult;
 
 		// Find the live instance
-		Private::FLiveInstance* pLiveInstance = m_pD->FindLiveInstance(instanceID);
+		FLiveInstance* pLiveInstance = m_pD->FindLiveInstance(instanceID);
 		check(pLiveInstance);
-		m_pD->m_memory = pLiveInstance->m_memory;
+		m_pD->WorkingMemoryManager.CurrentInstanceCache = pLiveInstance->Cache;
 
 		// Find the resource id in the model's resource cache
-		for (const Model::Private::RESOURCE_KEY& res : pLiveInstance->m_pModel->GetPrivate()->m_generatedResources)
+		for (const Model::Private::RESOURCE_KEY& res : pLiveInstance->Model->GetPrivate()->m_generatedResources)
 		{
 			if (res.m_id == meshId)
 			{
-				pResult = m_pD->BuildMesh(pLiveInstance->m_pModel,
-					pLiveInstance->m_pOldParameters.get(),
+				pResult = m_pD->BuildMesh(pLiveInstance->Model,
+					pLiveInstance->OldParameters.get(),
 					res.m_rootAddress);
 
 				// If the mesh is null it means empty, but we still need to return a valid one
@@ -406,7 +407,7 @@ namespace mu
 			}
 		}
 
-		m_pD->m_memory = nullptr;
+		m_pD->WorkingMemoryManager.CurrentInstanceCache = nullptr;
 		return pResult;
 	}
 
@@ -417,25 +418,29 @@ namespace mu
 		LLM_SCOPE_BYNAME(TEXT("MutableRuntime"));
 		MUTABLE_CPUPROFILER_SCOPE(EndUpdate);
 
-		// Reduce the cache until it fits the limit.
-		uint64 totalMemory = m_pD->m_modelCache.EnsureCacheBelowBudget(0, [](const Model*, int) {return false;});
-
-		Private::FLiveInstance* pLiveInstance = m_pD->FindLiveInstance(instanceID);
+		FLiveInstance* pLiveInstance = m_pD->FindLiveInstance(instanceID);
 		if (pLiveInstance)
 		{
-			pLiveInstance->m_pInstance = nullptr;
+			pLiveInstance->Instance = nullptr;
 			
 			// Debug check to see if we managed the op-hit-counts correctly
-			pLiveInstance->m_memory->CheckHitCountsCleared();
+			pLiveInstance->Cache->CheckHitCountsCleared();
+
+			m_pD->WorkingMemoryManager.CurrentInstanceCache = pLiveInstance->Cache;
 
 			// We don't want to clear the cache layer 1 because it contains data that can be useful for a 
 			// future update (same states, just runtime parameters changed).
-			//pLiveInstance->m_memory->ClearCacheLayer1();
+			//m_pD->WorkingMemoryManager.ClearCacheLayer1();
 
 			// We need to clear the layer 0 cache, because it contains data that is only valid for the current 
 			// parameter values (unless it is data marked as state cache)
-			pLiveInstance->m_memory->ClearCacheLayer0();
+			m_pD->WorkingMemoryManager.ClearCacheLayer0();
+
+			m_pD->WorkingMemoryManager.CurrentInstanceCache = nullptr;
 		}
+
+		// Reduce the cache until it fits the limit.
+		m_pD->WorkingMemoryManager.EnsureBudgetBelow(0);
 	}
 
 
@@ -445,13 +450,36 @@ namespace mu
 		LLM_SCOPE_BYNAME(TEXT("MutableRuntime"));
 		MUTABLE_CPUPROFILER_SCOPE(ReleaseInstance);
 
-		int Removed = m_pD->m_liveInstances.RemoveAllSwap(
-			[instanceID](const Private::FLiveInstance& Instance)
+		for (int32 Index = 0; Index < m_pD->WorkingMemoryManager.LiveInstances.Num(); ++Index)
+		{
+			FLiveInstance& Instance = m_pD->WorkingMemoryManager.LiveInstances[Index];
+			if (Instance.InstanceID == instanceID)
+			{
+				// Make sure all the resources cached in the instance are removed from the tracking list
+				CodeContainer<TPair<int32, Ptr<const Resource>>>::iterator it = Instance.Cache->m_resources.begin();
+				for (; it.IsValid(); ++it)
+				{
+					Ptr<const Resource> Value = (*it).Value.get();
+					if (!Value)
+					{
+						continue;
+					}
+
+					m_pD->WorkingMemoryManager.CacheResources.Remove(Value);
+				}
+
+				m_pD->WorkingMemoryManager.LiveInstances.RemoveAtSwap(Index);
+				break;
+			}
+		}
+
+ 		int Removed = m_pD->WorkingMemoryManager.LiveInstances.RemoveAllSwap(
+			[instanceID](const FLiveInstance& Instance)
 			{
-				return (Instance.m_instanceID == instanceID);
+				return (Instance.InstanceID == instanceID);
 			});
 
-		TRACE_COUNTER_SET(MutableRuntime_LiveInstances, m_pD->m_liveInstances.Num());
+		TRACE_COUNTER_SET(MutableRuntime_LiveInstances, m_pD->WorkingMemoryManager.LiveInstances.Num());
 
 	}
 
@@ -463,18 +491,18 @@ namespace mu
 
         RelevantParameterVisitor
             (
-                System::Private* pSystem,
-				const TSharedPtr<const Model>& pModel,
-                const Ptr<const Parameters>& pParams,
-                bool* pFlags
+                System::Private* InSystem,
+				const TSharedPtr<const Model>& InModel,
+                const Ptr<const Parameters>& InParams,
+                bool* InFlags
             )
-            : UniqueDiscreteCoveredCodeVisitor<>( pSystem, pModel, pParams, 0xffffffff )
+            : UniqueDiscreteCoveredCodeVisitor<>( InSystem, InModel, InParams, System::AllLODs )
         {
-            m_pFlags = pFlags;
+            Flags = InFlags;
 
-            FMemory::Memset( pFlags, 0, sizeof(bool)*pParams->GetCount() );
+            FMemory::Memset( Flags, 0, sizeof(bool)*InParams->GetCount() );
 
-            OP::ADDRESS at = pModel->GetPrivate()->m_program.m_states[0].m_root;
+            OP::ADDRESS at = InModel->GetPrivate()->m_program.m_states[0].m_root;
 
             Run( at );
         }
@@ -493,7 +521,7 @@ namespace mu
             {
 				OP::ParameterArgs args = program.GetOpArgs<OP::ParameterArgs>(at);
 				OP::ADDRESS param = args.variable;
-                m_pFlags[param] = true;
+                Flags[param] = true;
                 break;
             }
 
@@ -508,38 +536,50 @@ namespace mu
     private:
 
         //! Non-owned result buffer
-        bool* m_pFlags;
+        bool* Flags;
     };
 
 
     //---------------------------------------------------------------------------------------------
     void System::GetParameterRelevancy( Instance::ID InstanceID,
-                                        const ParametersPtrConst& Parameters,
+                                        const Ptr<const Parameters>& Parameters,
                                         bool* Flags )
     {
 		LLM_SCOPE_BYNAME(TEXT("MutableRuntime"));
 
 		// Find the live instance
-		Private::FLiveInstance* pLiveInstance = m_pD->FindLiveInstance(InstanceID);
+		FLiveInstance* pLiveInstance = m_pD->FindLiveInstance(InstanceID);
 		check(pLiveInstance);
-		m_pD->m_memory = pLiveInstance->m_memory;
+		m_pD->WorkingMemoryManager.CurrentInstanceCache = pLiveInstance->Cache;
 		
-		RelevantParameterVisitor visitor( m_pD, pLiveInstance->m_pModel, Parameters, Flags );
+		RelevantParameterVisitor visitor( m_pD, pLiveInstance->Model, Parameters, Flags );
 
-		m_pD->m_memory = nullptr;
+		m_pD->WorkingMemoryManager.CurrentInstanceCache = nullptr;
     }
 
 
     //---------------------------------------------------------------------------------------------
-    //---------------------------------------------------------------------------------------------
-    //---------------------------------------------------------------------------------------------
-    bool System::Private::CheckUpdatedParameters( const FLiveInstance* LiveInstance,
+	inline FLiveInstance* System::Private::FindLiveInstance(Instance::ID id)
+	{
+		for (int32 i = 0; i < WorkingMemoryManager.LiveInstances.Num(); ++i)
+		{
+			if (WorkingMemoryManager.LiveInstances[i].InstanceID == id)
+			{
+				return &WorkingMemoryManager.LiveInstances[i];
+			}
+		}
+		return nullptr;
+	}
+
+
+	//---------------------------------------------------------------------------------------------
+	bool System::Private::CheckUpdatedParameters( const FLiveInstance* LiveInstance,
                                  const Ptr<const Parameters>& Params,
                                  uint64& UpdatedParameters)
     {
         bool bFullBuild = false;
 
-		if (!LiveInstance->m_pOldParameters)
+		if (!LiveInstance->OldParameters)
 		{
 			UpdatedParameters = AllParametersMask;
 			return true;
@@ -547,18 +587,18 @@ namespace mu
 
         // check what parameters have changed
 		UpdatedParameters = 0;
-        const FProgram& program = LiveInstance->m_pModel->GetPrivate()->m_program;
-        const TArray<int>& runtimeParams = program.m_states[ LiveInstance->m_state ].m_runtimeParameters;
+        const FProgram& program = LiveInstance->Model->GetPrivate()->m_program;
+        const TArray<int>& runtimeParams = program.m_states[ LiveInstance->State ].m_runtimeParameters;
 
         check( Params->GetCount() == (int)program.m_parameters.Num() );
-        check( !LiveInstance->m_pOldParameters
+        check( !LiveInstance->OldParameters
 			||
-			Params->GetCount() == LiveInstance->m_pOldParameters->GetCount() );
+			Params->GetCount() == LiveInstance->OldParameters->GetCount() );
 
         for ( int32 p=0; p<program.m_parameters.Num() && !bFullBuild; ++p )
         {
             bool isRuntime = runtimeParams.Contains( p );
-            bool changed = !Params->HasSameValue( p, LiveInstance->m_pOldParameters, p );
+            bool changed = !Params->HasSameValue( p, LiveInstance->OldParameters, p );
 
             if (changed && isRuntime)
             {
@@ -578,42 +618,13 @@ namespace mu
     }
 
 
-	//---------------------------------------------------------------------------------------------
-	void System::Private::SetStreamingCache(uint64 bytes)
-	{
-		m_modelCache.m_romBudget = bytes;
-		m_modelCache.EnsureCacheBelowBudget(0);
-	}
-
-
-	//---------------------------------------------------------------------------------------------
-	void System::Private::ClearStreamingCache()
-	{
-		for (FModelCache::FModelCacheEntry& ModelCache : m_modelCache.m_cachePerModel)
-	    {
-		    if (const TSharedPtr<const Model> CacheModel = ModelCache.m_pModel.Pin())
-			{
-				FProgram& Program = CacheModel->GetPrivate()->m_program;
-
-				for (int32 RomIndex=0; RomIndex < Program.m_roms.Num(); ++RomIndex)
-				{
-					Program.UnloadRom(RomIndex);		
-    			}
-			}
- 		}
-	}
-	
-
 	//---------------------------------------------------------------------------------------------
 	void System::Private::BeginBuild(const TSharedPtr<const Model>& pModel)
 	{
 		// We don't have a FLiveInstance, let's create the memory
 		// \TODO: There is no clear moment to remove this... EndBuild?
-		m_memory = MakeShared<FProgramCache>();
-		m_memory->Init(pModel->GetPrivate()->m_program.m_opAddress.Num());
-
-		// Remove previously results cached from previous builds.
-		m_memory->ClearCacheLayer0();
+		WorkingMemoryManager.CurrentInstanceCache = MakeShared<FProgramCache>();
+		WorkingMemoryManager.CurrentInstanceCache->Init(pModel->GetPrivate()->m_program.m_opAddress.Num());
 
 		PrepareCache(pModel.Get(), -1);
 	}
@@ -622,7 +633,7 @@ namespace mu
 	//---------------------------------------------------------------------------------------------
 	void System::Private::EndBuild()
 	{
-		m_memory = nullptr;
+		WorkingMemoryManager.CurrentInstanceCache = nullptr;
 	}
 
 
@@ -630,7 +641,7 @@ namespace mu
 	void System::Private::RunCode(const TSharedPtr<const Model>& InModel,
 		const Parameters* InParameters, OP::ADDRESS InCodeRoot, uint32 InLODs, uint8 executionOptions, int32 InImageLOD)
 	{
-		CodeRunner Runner(m_pSettings, this, EExecutionStrategy::MinimizeMemory, InModel, InParameters, InCodeRoot, InLODs,
+		CodeRunner Runner(Settings, this, EExecutionStrategy::MinimizeMemory, InModel, InParameters, InCodeRoot, InLODs,
 			executionOptions, InImageLOD, FScheduledOp::EType::Full);
 		Runner.Run();
 		bUnrecoverableError = Runner.bUnrecoverableError;
@@ -638,149 +649,160 @@ namespace mu
 
 
 	//---------------------------------------------------------------------------------------------
-	bool System::Private::BuildBool(const TSharedPtr<const Model>& pModel,
-		const Parameters* pParams,
-		OP::ADDRESS at)
+	bool System::Private::BuildBool(const TSharedPtr<const Model>& pModel, const Parameters* pParams, OP::ADDRESS at)
 	{
+		WorkingMemoryManager.BeginRunnerThread();
+
 		RunCode(pModel, pParams, at);
-		if (bUnrecoverableError)
+
+		bool bResult = false;
+		if (!bUnrecoverableError)
 		{
-			return false;
+			bResult = WorkingMemoryManager.CurrentInstanceCache->GetBool(FCacheAddress(at, 0, 0));
 		}
-		return m_memory->GetBool(FCacheAddress(at, 0, 0));
+
+		WorkingMemoryManager.EndRunnerThread();
+
+		return bResult;
 	}
 
 
 	//---------------------------------------------------------------------------------------------
-	float System::Private::BuildScalar(const TSharedPtr<const Model>& pModel,
-		const Parameters* pParams,
-		OP::ADDRESS at)
+	float System::Private::BuildScalar(const TSharedPtr<const Model>& pModel, const Parameters* pParams, OP::ADDRESS at)
 	{
+		WorkingMemoryManager.BeginRunnerThread();
+
 		RunCode(pModel, pParams, at);
-		if (bUnrecoverableError)
+
+		float Result = 0.0f;		
+		if (!bUnrecoverableError)
 		{
-			return 0.0f;
+			Result = WorkingMemoryManager.CurrentInstanceCache->GetScalar(FCacheAddress(at, 0, 0));
 		}
-		return m_memory->GetScalar(FCacheAddress(at, 0, 0));
+
+		WorkingMemoryManager.EndRunnerThread();
+
+		return Result;
 	}
 
 
 	//---------------------------------------------------------------------------------------------
-	int System::Private::BuildInt(const TSharedPtr<const Model>& pModel,
-		const Parameters* pParams,
-		OP::ADDRESS at)
+	int32 System::Private::BuildInt(const TSharedPtr<const Model>& pModel, const Parameters* pParams, OP::ADDRESS at)
 	{
+		WorkingMemoryManager.BeginRunnerThread();
+
 		RunCode(pModel, pParams, at);
-		if (bUnrecoverableError)
+
+		int32 Result = 0;
+		if (!bUnrecoverableError)
 		{
-			return 0;
+			Result = WorkingMemoryManager.CurrentInstanceCache->GetInt(FCacheAddress(at, 0, 0));;
 		}
 
-		return m_memory->GetInt(FCacheAddress(at, 0, 0));
+		WorkingMemoryManager.EndRunnerThread();
+
+		return Result;
 	}
 
 
 	//---------------------------------------------------------------------------------------------
-	void System::Private::BuildColour(const TSharedPtr<const Model>& pModel,
-		const Parameters* pParams,
-		OP::ADDRESS at,
-		float* pR,
-		float* pG,
-		float* pB,
-		float* pA)
+	FVector4f System::Private::BuildColour(const TSharedPtr<const Model>& pModel, const Parameters* pParams, OP::ADDRESS at)
 	{
-		FVector4f col;
+		WorkingMemoryManager.BeginRunnerThread();
+
+		FVector4f Result(0,0,0,1);
 
 		mu::OP_TYPE opType = pModel->GetPrivate()->m_program.GetOpType(at);
 		if (GetOpDataType(opType) == DT_COLOUR)
 		{
 			RunCode(pModel, pParams, at);
-			if (bUnrecoverableError)
+			if (!bUnrecoverableError)
 			{
-				if (pR) *pR = 0.0f;
-				if (pG) *pG = 0.0f;
-				if (pB) *pB = 0.0f;
-				if (pA) *pA = 1.0f;
+				Result = WorkingMemoryManager.CurrentInstanceCache->GetColour(FCacheAddress(at, 0, 0));
 			}
-
-			col = m_memory->GetColour(FCacheAddress(at, 0, 0));
 		}
 
-		if (pR) *pR = col[0];
-		if (pG) *pG = col[1];
-		if (pB) *pB = col[2];
-		if (pA) *pA = col[3];
+		WorkingMemoryManager.EndRunnerThread();
+
+		return Result;
 	}
 
 	
 	//---------------------------------------------------------------------------------------------
 	Ptr<const Projector> System::Private::BuildProjector(const TSharedPtr<const Model>& pModel, const Parameters* pParams, OP::ADDRESS at)
 	{
+		WorkingMemoryManager.BeginRunnerThread();
+
     	RunCode(pModel, pParams, at);
-		if (bUnrecoverableError)
+
+		Ptr<const Projector> Result;
+		if (!bUnrecoverableError)
 		{
-			return nullptr;
+			Result = WorkingMemoryManager.CurrentInstanceCache->GetProjector(FCacheAddress(at, 0, 0));
 		}
-    	return m_memory->GetProjector(FCacheAddress(at, 0, 0));
+
+		WorkingMemoryManager.EndRunnerThread();
+
+		return Result;
 	}
 
 	
 	//---------------------------------------------------------------------------------------------
 	Ptr<const Image> System::Private::BuildImage(const TSharedPtr<const Model>& pModel,
-		const Parameters* pParams,
-		OP::ADDRESS at, int32 MipsToSkip, int32 InImageLOD)
+		const Parameters* pParams, OP::ADDRESS at, int32 MipsToSkip, int32 InImageLOD)
 	{
+		WorkingMemoryManager.BeginRunnerThread();
+
+		Ptr<const Image> Result;
+
 		mu::OP_TYPE opType = pModel->GetPrivate()->m_program.GetOpType(at);
 		if (GetOpDataType(opType) == DT_IMAGE)
 		{
 			RunCode(pModel, pParams, at, System::AllLODs, uint8(MipsToSkip), InImageLOD);
-			if (bUnrecoverableError)
+			if (!bUnrecoverableError)
 			{
-				return nullptr;
-			}
-			ImagePtrConst Result = m_memory->GetImage(FCacheAddress(at, 0, MipsToSkip));
-
-			// Debug check to see if we managed the op-hit-counts correctly
-			m_memory->CheckHitCountsCleared();
-
-			return Result;
+				Result = WorkingMemoryManager.LoadImage(FCacheAddress(at, 0, MipsToSkip), true);
+			}			
 		}
 
-		return nullptr;
+		WorkingMemoryManager.EndRunnerThread();
+
+		return Result;
 	}
 
 
 	//---------------------------------------------------------------------------------------------
-	MeshPtrConst System::Private::BuildMesh(const TSharedPtr<const Model>& pModel,
-		const Parameters* pParams,
-		OP::ADDRESS at)
+	MeshPtrConst System::Private::BuildMesh(const TSharedPtr<const Model>& pModel, const Parameters* pParams, OP::ADDRESS at)
 	{
+		WorkingMemoryManager.BeginRunnerThread();
+
+		Ptr<const Mesh> Result;
+
 		mu::OP_TYPE opType = pModel->GetPrivate()->m_program.GetOpType(at);
 		if (GetOpDataType(opType) == DT_MESH)
 		{
 			RunCode(pModel, pParams, at);
-			if (bUnrecoverableError)
+			if (!bUnrecoverableError)
 			{
-				return nullptr;
+				Result = WorkingMemoryManager.CurrentInstanceCache->GetMesh(FCacheAddress(at, 0, 0));
 			}
-			MeshPtrConst pResult = m_memory->GetMesh(FCacheAddress(at, 0, 0));
 
 			// Debug check to see if we managed the op-hit-counts correctly
-			m_memory->CheckHitCountsCleared();
-
-			return pResult;
+			WorkingMemoryManager.CurrentInstanceCache->CheckHitCountsCleared();
 		}
 
-		return nullptr;
+		WorkingMemoryManager.EndRunnerThread();
+
+		return Result;
 	}
 
 
 	//---------------------------------------------------------------------------------------------
-	LayoutPtrConst System::Private::BuildLayout(const TSharedPtr<const Model>& pModel,
-		const Parameters* pParams,
-		OP::ADDRESS at)
+	Ptr<const Layout> System::Private::BuildLayout(const TSharedPtr<const Model>& pModel, const Parameters* pParams, OP::ADDRESS at)
 	{
-		LayoutPtrConst  pResult;
+		WorkingMemoryManager.BeginRunnerThread();
+
+		Ptr<const Layout> Result;
 
 		if (pModel->GetPrivate()->m_program.m_states[0].m_root)
 		{
@@ -788,24 +810,25 @@ namespace mu
 			if (GetOpDataType(opType) == DT_LAYOUT)
 			{
 				RunCode(pModel, pParams, at);
-				if (bUnrecoverableError)
+				if (!bUnrecoverableError)
 				{
-					return nullptr;
+					Result = WorkingMemoryManager.CurrentInstanceCache->GetLayout(FCacheAddress(at, 0, 0));
 				}
-				pResult = m_memory->GetLayout(FCacheAddress(at, 0, 0));
 			}
 		}
 
-		return pResult;
+		WorkingMemoryManager.EndRunnerThread();
+
+		return Result;
 	}
 
 
 	//---------------------------------------------------------------------------------------------
-	Ptr<const String> System::Private::BuildString(const TSharedPtr<const Model>& pModel,
-		const Parameters* pParams,
-		OP::ADDRESS at)
+	Ptr<const String> System::Private::BuildString(const TSharedPtr<const Model>& pModel, const Parameters* pParams, OP::ADDRESS at)
 	{
-		Ptr<const String> pResult;
+		WorkingMemoryManager.BeginRunnerThread();
+
+		Ptr<const String> Result;
 
 		if (pModel->GetPrivate()->m_program.m_states[0].m_root)
 		{
@@ -813,50 +836,187 @@ namespace mu
 			if (GetOpDataType(opType) == DT_STRING)
 			{
 				RunCode(pModel, pParams, at);
-				if (bUnrecoverableError)
+				if (!bUnrecoverableError)
 				{
-					return nullptr;
+					Result = WorkingMemoryManager.CurrentInstanceCache->GetString(FCacheAddress(at, 0, 0));
 				}
-				pResult = m_memory->GetString(FCacheAddress(at, 0, 0));
 			}
 		}
 
-		return pResult;
+		WorkingMemoryManager.EndRunnerThread();
+
+		return Result;
 	}
 
 
 	//---------------------------------------------------------------------------------------------
-	void System::Private::PrepareCache( const Model* pModel, int state)
+	void System::Private::PrepareCache( const Model* InModel, int32 InState)
 	{
 		MUTABLE_CPUPROFILER_SCOPE(PrepareCache);
 
-		FProgram& program = pModel->GetPrivate()->m_program;
-		size_t opCount = program.m_opAddress.Num();
-		m_memory->m_opHitCount.clear();
-		m_memory->Init(opCount);
+		FProgram& program = InModel->GetPrivate()->m_program;
+		int32 opCount = program.m_opAddress.Num();
+		WorkingMemoryManager.CurrentInstanceCache->m_opHitCount.clear();
+		WorkingMemoryManager.CurrentInstanceCache->Init(opCount);
 
 		// Mark the resources that have to be cached to update the instance in this state
-		if (state >= 0)
+		if (InState >= 0)
 		{
-			const FProgram::FState& s = program.m_states[state];
+			const FProgram::FState& s = program.m_states[InState];
 			for (uint32 a : s.m_updateCache)
 			{
-				m_memory->SetForceCached(a);
+				WorkingMemoryManager.CurrentInstanceCache->SetForceCached(a);
 			}
 		}
 	}
 
 
+	//---------------------------------------------------------------------------------------------
+	void FWorkingMemoryManager::LogWorkingMemory(const CodeRunner* CurrentRunner) const
+	{
+#if !(UE_BUILD_SHIPPING || UE_BUILD_TEST)
+
+		// For now, we calculate these for every log. We will later track on resource creation, destruction or state change.
+		const uint32 RomBytes = GetRomBytes();
+		const uint32 Cache0Bytes = GetCache0Bytes();
+		const uint32 Cache1Bytes = GetCache1Bytes();
+		const uint32 TrackedCacheBytes = GetTrackedCacheBytes();
+		const uint32 PoolBytes = GetPooledBytes();
+		const uint32 TempBytes = GetTempBytes();
+		uint32 StreamingBufferBytes = 0;
+		uint32 InternalBytes = 0;
+
+		InternalBytes += PooledImages.GetAllocatedSize();
+		InternalBytes += TempImages.GetAllocatedSize();
+		InternalBytes += CacheResources.GetAllocatedSize();
+		InternalBytes += CachePerModel.GetAllocatedSize();
+		InternalBytes += LiveInstances.GetAllocatedSize();
+
+		if (CurrentRunner)
+		{
+			// Data reserved for ongoing streaming operations
+			for ( const CodeRunner::FRomLoadOp& Op : CurrentRunner->m_romLoadOps )
+			{
+				StreamingBufferBytes += Op.m_streamBuffer.Num();
+			}
+
+			InternalBytes += CurrentRunner->GetInternalMemoryBytes();
+		}
+
+		TArray<const Model*> Models;
+		for (const FLiveInstance& Instance: LiveInstances)
+		{
+			Models.AddUnique(Instance.Model.Get());
+
+			InternalBytes += Instance.Cache->m_resources.GetAllocatedSize()
+				+ Instance.Cache->m_descCache.GetAllocatedSize()
+				+ Instance.Cache->m_opHitCount.GetAllocatedSize()
+				+ Instance.Cache->m_usedRangeIndices.GetAllocatedSize();
+
+			for (const FWorkingMemoryManager::FModelCacheEntry& c : CachePerModel)
+			{
+				InternalBytes += c.RomWeights.GetAllocatedSize();
+			}			
+		}
+
+		// Data stored per-model, but related to instance construction
+		for (const Model* Model : Models)
+		{
+			// Keys for resources generated so far
+			for (const Model::Private::RESOURCE_KEY& Key : Model->GetPrivate()->m_generatedResources)
+			{
+				InternalBytes += Key.m_parameterValuesBlob.GetAllocatedSize();
+			}
+		}
+
+		uint32 TotalBytes = InternalBytes 
+			+ TempBytes 
+			+ TrackedCacheBytes
+			+ PoolBytes
+			+ StreamingBufferBytes
+			+ RomBytes;
+
+		UE_LOG(LogMutableCore, Log, TEXT("Mem KB: Internal %7d   Temp %7d   Pool %7d   Cache0+1 %7d   CacheTracked %7d   Rom %7d   Stream %7d   Total %7d / %7d"),
+			InternalBytes / 1024,
+			TempBytes / 1024,
+			PoolBytes / 1024,
+			(Cache0Bytes + Cache1Bytes) / 1024,
+			TrackedCacheBytes / 1024,
+			RomBytes / 1024,
+			StreamingBufferBytes / 1024,
+			TotalBytes / 1024,
+			BudgetBytes / 1024
+			);
+
+		// Detailed log for debug.
+		//if (TotalBytes > 30 * 1024 * 1024)
+		if (false)
+		{
+			UE_LOG(LogMutableCore, Log, TEXT("Mem detailed report:"));
+
+			for (const FLiveInstance& Instance : LiveInstances)
+			{
+				UE_LOG(LogMutableCore, Log, TEXT("  instance:"));
+
+				TSet<const Resource*> InstanceCache0Unique;
+				TSet<const Resource*> InstanceCache1Unique;
+
+				CodeContainer<int>::iterator it = Instance.Cache->m_opHitCount.begin();
+				for (; it.IsValid(); ++it)
+				{
+					const Resource* Value = Instance.Cache->m_resources[it.get_address()].Value.get();
+					if (!Value)
+					{
+						continue;
+					}
+
+					int32 Count = *it;
+					if (Count < UE_MUTABLE_CACHE_COUNT_LIMIT)
+					{
+						InstanceCache0Unique.Add(Value);
+					}
+					else
+					{
+						InstanceCache1Unique.Add(Value);
+					}
+				}
+
+				UE_LOG(LogMutableCore, Log, TEXT("    cache level 0:"));
+
+				for (const Resource* Value : InstanceCache0Unique)
+				{
+					UE_LOG(LogMutableCore, Log, TEXT("      resource: %d"), Value->GetDataSize());
+				}
+
+				UE_LOG(LogMutableCore, Log, TEXT("    cache level 1:"));
+
+				for (const Resource* Value : InstanceCache1Unique)
+				{
+					UE_LOG(LogMutableCore, Log, TEXT("      resource: %d"), Value->GetDataSize());
+				}
+			}
+		}
+
+		TRACE_COUNTER_SET(MutableRuntime_MemInternal, InternalBytes);
+		TRACE_COUNTER_SET(MutableRuntime_MemTemp, TempBytes);
+		TRACE_COUNTER_SET(MutableRuntime_MemPool, PoolBytes);
+		TRACE_COUNTER_SET(MutableRuntime_MemCache, TrackedCacheBytes);
+		TRACE_COUNTER_SET(MutableRuntime_MemRom, RomBytes);
+		TRACE_COUNTER_SET(MutableRuntime_MemStream, StreamingBufferBytes);
+		TRACE_COUNTER_SET(MutableRuntime_MemTotal, TotalBytes);
+		TRACE_COUNTER_SET(MutableRuntime_MemBudget, BudgetBytes);
+#endif
+	}
+
+
     //---------------------------------------------------------------------------------------------
-    //---------------------------------------------------------------------------------------------
-    //---------------------------------------------------------------------------------------------
-	FModelCache::FModelCacheEntry& FModelCache::GetModelCache(const TSharedPtr<const Model>& InModel )
+	FWorkingMemoryManager::FModelCacheEntry& FWorkingMemoryManager::GetModelCache(const TSharedPtr<const Model>& InModel )
     {
         check(InModel);
 
-        for(FModelCacheEntry& c:m_cachePerModel)
+        for(FModelCacheEntry& c:CachePerModel)
         {
-			TSharedPtr<const Model> pCandidate = c.m_pModel.Pin();
+			TSharedPtr<const Model> pCandidate = c.Model.Pin();
             if (pCandidate)
             {
                 if (pCandidate==InModel)
@@ -867,158 +1027,201 @@ namespace mu
             else
             {
                 // Free stray data. TODO: remove vector entry.
-                c.m_romWeight.Empty();
+                c.RomWeights.Empty();
             }
         }
 
         // Not found. Add new
 		FModelCacheEntry n;
-        n.m_pModel = TWeakPtr<const Model>(InModel);
-        m_cachePerModel.Add(n);
-        return m_cachePerModel.Last();
+        n.Model = TWeakPtr<const Model>(InModel);
+		n.PendingOpsPerRom.SetNum(InModel->GetPrivate()->m_program.m_roms.Num());
+        CachePerModel.Add(n);
+        return CachePerModel.Last();
     }
 
-	/** */
-	TRACE_DECLARE_INT_COUNTER(MutableRuntime_StreamingBytes, TEXT("MutableRuntime/StreamingBytes"));
 
     //---------------------------------------------------------------------------------------------
-    uint64 FModelCache::EnsureCacheBelowBudget( uint64 additionalMemory,
-												TFunctionRef<bool(const Model*,int)> isRomLockedFunc )
+    bool FWorkingMemoryManager::EnsureBudgetBelow( uint64 AdditionalMemory )
     {
-		uint64 totalMemory = 0;
-        for (FModelCacheEntry& m : m_cachePerModel)
-        {
-			TSharedPtr<const Model> pCacheModel = m.m_pModel.Pin();
-            if (pCacheModel)
-            {
-				mu::FProgram& program = pCacheModel->GetPrivate()->m_program;
-				for (int32 RomIndex = 0; RomIndex < program.m_roms.Num(); ++RomIndex)
+		MUTABLE_CPUPROFILER_SCOPE(EnsureBudgetBelow);
+
+		// If we have 0 budget it means we have unlimited budget
+		if (BudgetBytes == 0)
+		{
+			return true;
+		}
+
+		uint32 RomBytes = 0;
+		uint32 PoolBytes = 0;
+		uint32 TempBytes = 0;
+		//uint32 Cache0Bytes = 0;
+		//uint32 Cache1Bytes = 0;
+		uint32 TrackedCacheBytes = 0;
+		{
+			MUTABLE_CPUPROFILER_SCOPE(EnsureBudgetBelow_Measure);
+			RomBytes = GetRomBytes();
+			PoolBytes = GetPooledBytes();
+			TempBytes = GetTempBytes();
+			//Cache0Bytes = GetCache0Bytes();
+			//Cache1Bytes = GetCache1Bytes();
+			TrackedCacheBytes = GetTrackedCacheBytes();
+		}
+		// These types of memory are not tracked here yet.
+		//uint32 StreamingBufferBytes = 0;
+		//uint32 InternalBytes = 0;
+
+		uint64 TotalBytes = 0;
+		TotalBytes += TempBytes;
+		TotalBytes += RomBytes;
+		TotalBytes += PoolBytes;
+		TotalBytes += TrackedCacheBytes;
+
+		// Add the extra memory that we are trying to allocate when we return.
+		TotalBytes += AdditionalMemory;
+
+        bool bFinished = TotalBytes <= BudgetBytes;
+
+		// Try to free pooled resources first
+		{
+			MUTABLE_CPUPROFILER_SCOPE(EnsureBudgetBelow_FreePooled);
+			while (PooledImages.Num() && !bFinished)
+			{
+				// TODO: Actually advancing index if possible after swap may be better to remove the oldest in the pool first.
+				int32 PooledResourceSize = PooledImages[0]->GetDataSize();
+				TotalBytes -= PooledResourceSize;
+				PooledImages.RemoveAtSwap(0);
+				bFinished = TotalBytes <= BudgetBytes;
+			}
+		}
+		
+		// Try to free a loaded roms
+		{
+			MUTABLE_CPUPROFILER_SCOPE(EnsureBudgetBelow_FreeRoms);
+
+			while (!bFinished)
+			{
+				TSharedPtr<const Model> LowestPriorityModel;
+				int32 LowestPriorityRom = -1;
+				float LowestPriority = 0.0f;
+				for (FModelCacheEntry& ModelCache : CachePerModel)
 				{
-					if (program.IsRomLoaded(RomIndex))
+					TSharedPtr<const Model> CacheModel = ModelCache.Model.Pin();
+					if (CacheModel)
 					{
-						totalMemory += program.m_roms[RomIndex].Size;
+						mu::FProgram& program = CacheModel->GetPrivate()->m_program;
+						check(ModelCache.RomWeights.Num() == program.m_roms.Num());
+
+						for (int32 RomIndex = 0; RomIndex < program.m_roms.Num(); ++RomIndex)
+						{
+							const FRomData& Rom = program.m_roms[RomIndex];
+							bool bIsLoaded = program.IsRomLoaded(RomIndex);
+
+							// We cannot unload a rom if some operation is expecting it.
+							bool bIsRomLocked = ModelCache.PendingOpsPerRom.IsValidIndex(RomIndex)
+								&& 
+								ModelCache.PendingOpsPerRom[RomIndex]>0;
+							if (bIsLoaded && !bIsRomLocked)
+							{
+								constexpr float FactorWeight = 100.0f;
+								constexpr float FactorTime = -1.0f;
+								float Priority = FactorWeight * float(ModelCache.RomWeights[RomIndex].Key)
+									+
+									FactorTime * float((RomTick - ModelCache.RomWeights[RomIndex].Value));
+
+								if (LowestPriorityRom < 0 || Priority < LowestPriority)
+								{
+									LowestPriorityRom = RomIndex;
+									LowestPriority = Priority;
+									LowestPriorityModel = CacheModel;
+								}
+							}
+						}
 					}
 				}
+
+				if (LowestPriorityRom < 0)
+				{
+					// If we reached this it means we couldn't find more streaming data to free
+					// Try to continue anyway.
+					break;
+				}
+				else
+				{
+					MUTABLE_CPUPROFILER_SCOPE(EnsureBudgetBelow_UnloadRom);
+
+					//UE_LOG(LogMutableCore,Log, "Unloading rom because of memory budget: %d.", lowestPriorityRom);
+					LowestPriorityModel->GetPrivate()->m_program.UnloadRom(LowestPriorityRom);
+					TotalBytes -= LowestPriorityModel->GetPrivate()->m_program.m_roms[LowestPriorityRom].Size;
+					bFinished = TotalBytes <= BudgetBytes;
+				}
 			}
-        }
+		}
 
-    	// If budget is zero, we don't unload anything here, and we assume it is managed
-        // somewhere else.
-  		if (!m_romBudget)
+		// Try to free cache 1 memory (from other models first?)
+		if (!bFinished)
 		{
-			return totalMemory;
+			MUTABLE_CPUPROFILER_SCOPE(EnsureBudgetBelow_FreeCached);
+
+			// TODO.
+			// from other models first?
 		}
-    	
-        if (totalMemory>0)
-        {
-            totalMemory += additionalMemory;
 
-            bool finished = totalMemory < m_romBudget;
-            while (!finished)
-            {
-				TSharedPtr<const Model> lowestPriorityModel;
-                int32 lowestPriorityRom = -1;
-                float lowestPriority = 0.0f;
-                for (FModelCacheEntry& modelCache : m_cachePerModel)
-                {
-					TSharedPtr<const Model> pCacheModel = modelCache.m_pModel.Pin();
-                    if (pCacheModel)
-                    {
-						mu::FProgram& program = pCacheModel->GetPrivate()->m_program;
-                        check( modelCache.m_romWeight.Num() == program.m_roms.Num());
-
-                        for (int32 RomIndex=0; RomIndex <program.m_roms.Num(); ++RomIndex)
-                        {
-							const FRomData& Rom = program.m_roms[RomIndex];
-							bool bIsLoaded = program.IsRomLoaded(RomIndex);
+		TRACE_COUNTER_SET(MutableRuntime_MemTemp, TempBytes);
+		TRACE_COUNTER_SET(MutableRuntime_MemPool, PoolBytes);
+		TRACE_COUNTER_SET(MutableRuntime_MemCache, TrackedCacheBytes);
+		TRACE_COUNTER_SET(MutableRuntime_MemRom, RomBytes);
+		TRACE_COUNTER_SET(MutableRuntime_MemTotal, TotalBytes);
+		TRACE_COUNTER_SET(MutableRuntime_MemBudget, BudgetBytes);
 
-                            if (bIsLoaded
-                                &&
-                                (!isRomLockedFunc(pCacheModel.Get(),RomIndex)) )
-                            {
-                                constexpr float factorWeight = 100.0f;
-                                constexpr float factorTime = -1.0f;
-                                float priority =
-                                        factorWeight * float(modelCache.m_romWeight[RomIndex].Key)
-                                        +
-                                        factorTime * float((m_romTick-modelCache.m_romWeight[RomIndex].Value));
-
-                                if (lowestPriorityRom<0 || priority<lowestPriority)
-                                {
-                                    lowestPriorityRom = RomIndex;
-                                    lowestPriority = priority;
-                                    lowestPriorityModel = pCacheModel;
-                                }
-                            }
-                        }
-                    }
-                }
+		if (!bFinished)
+		{
+			uint64 ExcessBytes = TotalBytes - BudgetBytes;
 
-                if (lowestPriorityRom<0)
-                {
-                    // None found
-                    finished = true;
+			if (ExcessBytes > BudgetExcessBytes)
+			{
+				BudgetExcessBytes = ExcessBytes;
 
-					// If we reached this it means we need to use more memory for streaming data that was given to mutable.
-					// Try to continue anyway.
-                    //UE_LOG(LogMutableCore,Log, TEXT("EnsureCacheBelowBudget failed to free all the necessary memory."));
-                }
-                else
-                {
-                    //UE_LOG(LogMutableCore,Log, "Unloading rom because of memory budget: %d.", lowestPriorityRom);
-					const FRomData& Rom = lowestPriorityModel->GetPrivate()->m_program.m_roms[lowestPriorityRom];
-                    lowestPriorityModel->GetPrivate()->m_program.UnloadRom(lowestPriorityRom);
-                    totalMemory -= lowestPriorityModel->GetPrivate()->m_program.m_roms[lowestPriorityRom].Size;
-                    finished = totalMemory < m_romBudget;
-                }
-            }
-        }
+				// We failed to free enough memory. Log this, but try to continue anyway.
+				// This is a good place to insert a brakpoint to detect callstacks with memory peaks
+				UE_LOG(LogMutableCore, Log, TEXT("Failed to keep memory budget. Budget: %d, Current: %d, New: %d"),
+					BudgetBytes / 1024, (TotalBytes - AdditionalMemory) / 1024, AdditionalMemory / 1024);
 
-		TRACE_COUNTER_SET(MutableRuntime_StreamingBytes, totalMemory);
+				// We won't show correct internal or streaming buffer memory.
+				LogWorkingMemory(nullptr);
+			}
+		}
 
-        return totalMemory;
+        return bFinished;
     }
 
 
     //---------------------------------------------------------------------------------------------
-    void FModelCache::MarkRomUsed( int romIndex, const TSharedPtr<const Model>& pModel )
+    void FWorkingMemoryManager::MarkRomUsed( int32 romIndex, const TSharedPtr<const Model>& pModel )
     {
         check(pModel);
 
         FProgram& program = pModel->GetPrivate()->m_program;
 
-        // If budget is zero, we don't unload anything here, and we assume it is managed
-        // somewhere else.
-        if ( !m_romBudget )
+        // If budget is zero, we don't unload anything here, and we assume it is managed somewhere else.
+        if (!BudgetBytes)
         {
             return;
         }
 
-        ++m_romTick;
+        ++RomTick;
 
         // Update current cache
         {
-			FModelCacheEntry& modelCache = GetModelCache(pModel);
+			FModelCacheEntry& ModelCache = GetModelCache(pModel);
 
-            while (modelCache.m_romWeight.Num()<program.m_roms.Num())
+            while (ModelCache.RomWeights.Num()<program.m_roms.Num())
             {
-				modelCache.m_romWeight.Add({ 0,0 });
+				ModelCache.RomWeights.Add({ 0,0 });
             }
 
-            modelCache.m_romWeight[romIndex].Key++;
-            modelCache.m_romWeight[romIndex].Value = m_romTick;
+            ModelCache.RomWeights[romIndex].Key++;
+            ModelCache.RomWeights[romIndex].Value = RomTick;
         }
     }
 
-
-    //---------------------------------------------------------------------------------------------
-    void FModelCache::UpdateForLoad( int romIndex, const TSharedPtr<const Model>& pModel,
-                                     TFunctionRef<bool(const Model*,int)> isRomLockedFunc )
-    {
-		MarkRomUsed( romIndex, pModel);
-
-		FProgram& program = pModel->GetPrivate()->m_program;
-        EnsureCacheBelowBudget(program.m_roms[romIndex].Size, isRomLockedFunc);
-    }
-
 }
