[mutable] Implement memory budget control for image operations.

diff --git a/Engine/Plugins/Experimental/Mutable/Source/MutableRuntime/Private/MuR/SystemPrivate.h b/Engine/Plugins/Experimental/Mutable/Source/MutableRuntime/Private/MuR/SystemPrivate.h
--- a/Engine/Plugins/Experimental/Mutable/Source/MutableRuntime/Private/MuR/SystemPrivate.h
+++ b/Engine/Plugins/Experimental/Mutable/Source/MutableRuntime/Private/MuR/SystemPrivate.h
@@ -8,11 +8,18 @@
 #include "MuR/Operations.h"
 #include "MuR/ImagePrivate.h"
 #include "MuR/MutableString.h"
+#include "MuR/ModelPrivate.h"
 #include "MuR/MeshPrivate.h"
 #include "MuR/InstancePrivate.h"
 #include "MuR/ParametersPrivate.h"
 #include "MuR/MutableTrace.h"
 
+#if !(UE_BUILD_SHIPPING || UE_BUILD_TEST)
+#include "HAL/Thread.h"
+#include "HAL/PlatformTLS.h"
+#endif
+
+#define UE_MUTABLE_CACHE_COUNT_LIMIT	3000000
 
 namespace mu
 {
@@ -36,58 +43,64 @@ namespace mu
 
 
 	//! Reference-counted colour to be stored in the cache
-	class Colour : public RefCounted
+	class Colour : public Resource
 	{
 	public:
 		Colour(FVector4f v = FVector4f()) : m_colour(v) {}
 		FVector4f m_colour;
+		int32 GetDataSize() const override { return sizeof(Colour); }
 	};
 	typedef Ptr<Colour> ColourPtr;
 
 
 	//! Reference-counted bool to be stored in the cache
-	class Bool : public RefCounted
+	class Bool : public Resource
 	{
 	public:
 		Bool(bool v = false) : m_value(v) {}
 		bool m_value;
+		int32 GetDataSize() const override { return sizeof(Bool); }
 	};
 	typedef Ptr<Bool> BoolPtr;
 
 
 	//! Reference-counted scalar to be stored in the cache
-	class Scalar : public RefCounted
+	class Scalar : public Resource
 	{
 	public:
 		Scalar(float v = 0.0f) : m_value(v) {}
 		float m_value;
+		int32 GetDataSize() const override { return sizeof(Scalar); }
 	};
 	typedef Ptr<Scalar> ScalarPtr;
 
 
 	//! Reference-counted scalar to be stored in the cache
-	class Int : public RefCounted
+	class Int : public Resource
 	{
 	public:
 		Int(int32 v = 0) : m_value(v) {}
 		int32 m_value;
+		int32 GetDataSize() const override { return sizeof(Int); }
 	};
 	typedef Ptr<Int> IntPtr;
 
 
 	//! Reference-counted scalar to be stored in the cache
-	class Projector : public RefCounted
+	class Projector : public Resource
 	{
 	public:
 		FProjector m_value;
+		int32 GetDataSize() const override { return sizeof(Projector); }
 	};
 	typedef Ptr<Projector> ProjectorPtr;
 
 
-	//! ExecutinIndex stores the location inside all ranges of the execution of a specific
-	//! operation. The first integer on each pair is the dimension/range index in the program
-	//! array of ranges, and the second integer is the value inside that range.
-	//! The vector order is undefined.
+	/** ExecutinIndex stores the location inside all ranges of the execution of a specific
+	* operation. The first integer on each pair is the dimension/range index in the program
+	* array of ranges, and the second integer is the value inside that range.
+	* The vector order is undefined.
+	*/
 	class ExecutionIndex : public  TArray<TPair<int32, int32>>
 	{
 	public:
@@ -194,7 +207,7 @@ namespace mu
 	}
 
 
-	//! A cache address is the operation plus the context of execution (iteration indices, etc...).
+	/** A cache address is the operation plus the context of execution(iteration indices, etc...). */
 	struct FCacheAddress
 	{
 		/** The meaning of all these fields is the same than the FScheduledOp struct. */
@@ -236,6 +249,12 @@ namespace mu
 	}
 
 
+	inline uint32 GetTypeHash(const Ptr<const Resource>& V)
+	{
+		return ::GetTypeHash(V.get());
+	}
+
+
 	//! Container that stores data per executable code operation (indexed by address and execution
 	//! index).
 	template<class DATA>
@@ -414,6 +433,12 @@ namespace mu
 			return it;
 		}
 
+		inline int32 GetAllocatedSize() const
+		{
+			return m_index0.GetAllocatedSize()
+				+ m_otherIndex.GetAllocatedSize();
+		}
+
 	private:
 		// For index 0
 		TArray<DATA> m_index0;
@@ -436,9 +461,9 @@ namespace mu
 		* 1 : valid, not worth freeing for memory
 		* 2 : valid, worth freeing
 		*/
-		CodeContainer< TPair<int32, Ptr<const RefCounted>> > m_resources;
+		CodeContainer< TPair<int32, Ptr<const Resource>> > m_resources;
 
-		//! Addressed with OP::ADDRESS. It is true if value for an image desc is valid.
+		/** Addressed with OP::ADDRESS.It is true if value for an image desc is valid. */
 		TArray<bool> m_descCache;
 
 		/** The number operation stages waiting for the output of a specific operation.
@@ -447,7 +472,7 @@ namespace mu
 		CodeContainer<int32> m_opHitCount;
 
 
-		inline const ExecutionIndex& GetRageIndex(uint32_t i)
+		inline const ExecutionIndex& GetRangeIndex(uint32_t i)
 		{
 			// Make sure we have the default element.
 			if (m_usedRangeIndices.IsEmpty())
@@ -460,7 +485,7 @@ namespace mu
 		}
 
 		//!
-		inline uint32_t GetRageIndexIndex(const ExecutionIndex& rangeIndex)
+		inline uint32_t GetRangeIndexIndex(const ExecutionIndex& rangeIndex)
 		{
 			if (rangeIndex.IsEmpty())
 			{
@@ -507,7 +532,7 @@ namespace mu
 		}
 
 
-		bool IsValid(FCacheAddress at)
+		bool IsValid(FCacheAddress at) const
 		{
 			if (at.At == 0) return false;
 
@@ -543,7 +568,7 @@ namespace mu
 			//for (; it.IsValid(); ++it)
 			//{
 			//	int32 Count = *it;
-			//	if (Count>0 && Count < 3000000)
+			//	if (Count>0 && Count < UE_MUTABLE_CACHE_COUNT_LIMIT)
 			//	{
 			//		// We don't manage the hitcounts of small types that don't use much memory
 			//		if (m_resources[it.get_address()].Key==2)
@@ -564,48 +589,6 @@ namespace mu
 		}
 
 
-		/** Remove all intermediate data (big and small) from the memory except for the one that has been explicitely
-		* marked as state cache.
-		*/
-		void ClearCacheLayer0()
-		{
-			MUTABLE_CPUPROFILER_SCOPE(ClearLayer0);
-
-			CodeContainer<int>::iterator it = m_opHitCount.begin();
-			for (; it.IsValid(); ++it)
-			{
-				int32 Count = *it;
-				if (Count < 3000000)
-				{
-					// SetUnused only clears if the data size is relevant (mesh or image) but we need to 
-					// clear it in all cases here, because it may have become invalid because of parameter changes.
-					// SetUnused(it.get_address());
-					m_resources[it.get_address()].Value = nullptr;
-					m_resources[it.get_address()].Key = 0;
-					*it = 0;
-				}
-			}
-		}
-
-
-		/** Remove all intermediate data (big and small) from the memory including the one that has been explicitely
-		* marked as state cache.
-		*/
-		void ClearCacheLayer1()
-		{
-			MUTABLE_CPUPROFILER_SCOPE(ClearLayer1);
-
-			CodeContainer< TPair<int32, Ptr<const RefCounted>> >::iterator it = m_resources.begin();
-			for (; it.IsValid(); ++it)
-			{
-				(*it).Value = nullptr;
-				(*it).Key = 0;
-			}
-
-			m_descCache.SetNum(0);
-		}
-
-
 		void Clear()
 		{
 			MUTABLE_CPUPROFILER_SCOPE(ProgramCacheClear);
@@ -755,7 +738,7 @@ namespace mu
 				return nullptr;
 			}
 			
-			const TPair<int, Ptr<const RefCounted>>* d = m_resources.get_ptr(at);
+			const TPair<int, Ptr<const Resource>>* d = m_resources.get_ptr(at);
 			if (!d)
 			{
 				return nullptr;
@@ -775,7 +758,7 @@ namespace mu
 
 			Ptr<Bool> pResult = new Bool;
 			pResult->m_value = v;
-			m_resources[at] = TPair<int, Ptr<const RefCounted>>(1, pResult);
+			m_resources[at] = TPair<int, Ptr<const Resource>>(1, pResult);
 		}
 
 		void SetValidDesc(FCacheAddress at)
@@ -792,7 +775,7 @@ namespace mu
 
 			Ptr<Scalar> pResult = new Scalar;
 			pResult->m_value = v;
-			m_resources[at] = TPair<int, Ptr<const RefCounted>>(1, pResult);
+			m_resources[at] = TPair<int, Ptr<const Resource>>(1, pResult);
 		}
 
 		void SetInt(FCacheAddress at, int v)
@@ -801,7 +784,7 @@ namespace mu
 
 			Ptr<Int> pResult = new Int;
 			pResult->m_value = v;
-			m_resources[at] = TPair<int, Ptr<const RefCounted>>(1, pResult);
+			m_resources[at] = TPair<int, Ptr<const Resource>>(1, pResult);
 		}
 
 		void SetColour(FCacheAddress at, const FVector4f& v)
@@ -810,31 +793,31 @@ namespace mu
 
 			Ptr<Colour> pResult = new Colour;
 			pResult->m_colour = v;
-			m_resources[at] = TPair<int, Ptr<const RefCounted>>(1, pResult);
+			m_resources[at] = TPair<int, Ptr<const Resource>>(1, pResult);
 		}
 
 		void SetProjector(FCacheAddress at, Ptr<const Projector> v)
 		{
 			check(at.At < m_resources.size_code());
-			m_resources[at] = TPair<int, Ptr<const RefCounted>>(1, v);
+			m_resources[at] = TPair<int, Ptr<const Resource>>(1, v);
 		}
 
 		void SetInstance(FCacheAddress at, Ptr<const Instance> v)
 		{
 			check(at.At < m_resources.size_code());
-			m_resources[at] = TPair<int, Ptr<const RefCounted>>(2, v);
+			m_resources[at] = TPair<int, Ptr<const Resource>>(2, v);
 		}
 
 		void SetExtensionData(FCacheAddress at, Ptr<const ExtensionData> v)
 		{
 			check(at.At < m_resources.size_code());
-			m_resources[at] = TPair<int, Ptr<const RefCounted>>(1, v);
+			m_resources[at] = TPair<int, Ptr<const Resource>>(1, v);
 		}
 
 		void SetImage(FCacheAddress at, Ptr<const Image> v)
 		{
 			check(at.At < m_resources.size_code());
-			m_resources[at] = TPair<int, Ptr<const RefCounted>>(2, v);
+			m_resources[at] = TPair<int, Ptr<const Resource>>(2, v);
 
 			mu::UpdateLLMStats();
 		}
@@ -848,7 +831,7 @@ namespace mu
 //            }
 
 			check(at.At < m_resources.size_code());
-			m_resources[at] = TPair<int, Ptr<const RefCounted>>(2, v);
+			m_resources[at] = TPair<int, Ptr<const Resource>>(2, v);
 
 			mu::UpdateLLMStats();
 		}
@@ -856,7 +839,7 @@ namespace mu
 		void SetLayout(FCacheAddress at, Ptr<const Layout> v)
 		{
 			check(at.At < m_resources.size_code());
-			m_resources[at] = TPair<int, Ptr<const RefCounted>>(1, v);
+			m_resources[at] = TPair<int, Ptr<const Resource>>(1, v);
 
 			mu::UpdateLLMStats();
 		}
@@ -864,7 +847,7 @@ namespace mu
 		void SetString(FCacheAddress at, Ptr<const String> v)
 		{
 			check(at.At < m_resources.size_code());
-			m_resources[at] = TPair<int, Ptr<const RefCounted>>(1, v);
+			m_resources[at] = TPair<int, Ptr<const Resource>>(1, v);
 		}
 
 
@@ -886,16 +869,18 @@ namespace mu
 		}
 
 
-		Ptr<const Image> GetImage(FCacheAddress at)
+		Ptr<const Image> GetImage(FCacheAddress at, bool& bIsLastReference)
 		{
+			bIsLastReference = false;
+
 			if (!at.At) return nullptr;
 			if (size_t(at.At) >= m_resources.size_code()) return nullptr;
 			auto d = m_resources.get_ptr(at);
 			if (!d) return nullptr;
 
-			Ptr<const Image> pResult = (const Image*)d->Value.get();
+			Ptr<const Image> Result = (const Image*)d->Value.get();
 
-			// We need to decrease the hit count even if the result is null.
+			// We need to decrease the hit count (even if the result is null).
 			// Lower hit counts means we shouldn't clear the value
 			if ( m_opHitCount[at] > 0)
 			{
@@ -903,10 +888,11 @@ namespace mu
 				if (m_opHitCount[at] <= 0)
 				{
 					SetUnused(at);
+					bIsLastReference = true;
 				}
 			}
 
-			return pResult;
+			return Result;
 		}
 
 		Ptr<const Mesh> GetMesh(FCacheAddress at)
@@ -957,55 +943,535 @@ namespace mu
 		return a.Type < b.Type;
 	}
 
+	/** Data for an instance that is currently being processed in the mutable system. This means it is
+	* between a BeginUpdate and EndUpdate, or during an "atomic" operation (like generate a single resource).
+	*/
+	struct FLiveInstance
+	{
+		Instance::ID InstanceID;
+		int32 State = 0;
+		Ptr<const Instance> Instance;
+		TSharedPtr<const Model> Model;
+
+		Ptr<Parameters> OldParameters;
 
+		/** Mask of the parameters that have changed since the last update.
+		* Every bit represents a state parameter.
+		*/
+		uint64 UpdatedParameters = 0;
 
-    //---------------------------------------------------------------------------------------------
-    //! Struct to manage models for which we have streamed rom data.
-    //---------------------------------------------------------------------------------------------
-    struct FModelCache
-    {
-		~FModelCache()
+		/** Cached data for the generation of this instance. */
+		TSharedPtr<FProgramCache> Cache;
+
+		~FLiveInstance()
 		{
-			MUTABLE_CPUPROFILER_SCOPE(FModelCacheDestructor);
+			// Manually done to trace mem deallocations
+			MUTABLE_CPUPROFILER_SCOPE(LiveInstanceDestructor);
+			Cache = nullptr;
+			OldParameters = nullptr;
+			Instance = nullptr;
+			Model = nullptr;
 		}
+	};
 
+
+    /** Struct to manage all the memory allocated for resources used during mutable operation. */
+    struct FWorkingMemoryManager
+    {
+		/** Cached traking for streamed model data for one model. */
         struct FModelCacheEntry
         {
-            //!
-            TWeakPtr<const Model> m_pModel;
+            /** Model who's data is being tracked. */
+            TWeakPtr<const Model> Model;
+
+            /** For each model rom, the last time its streamed data was used. */
+            TArray< TPair<uint64,uint64> > RomWeights;
 
-            // Rom streaming management
-            TArray< TPair<uint64,uint64> > m_romWeight;
+			/** Count of pending operations for every rom index. */
+			TArray<uint16> PendingOpsPerRom;
         };
 
-        // Rom streaming management
-        uint64 m_romBudget = 0;
+		/** Maximum working memory that mutable should be using. */
+		uint64 BudgetBytes = 0;
 
-        uint64 m_romTick = 0;
-        TArray< FModelCacheEntry > m_cachePerModel;
+		/** Maximum excess memory reached suring the current operation. */
+		uint64 BudgetExcessBytes = 0;
 
-        //! Make sure the cached memory is below the internal budget, even counting with the
-        //! passed additional memory.
-		uint64 EnsureCacheBelowBudget(uint64 AdditionalMemory, 
-			TFunctionRef<bool(const Model*, int)> IsRomLockedFunc = [](const Model*, int) {return false;});
+		/** This value is used to track the order of loading of roms. */
+        uint64 RomTick = 0;
 
-        //!
-        void UpdateForLoad( int romIndex, const TSharedPtr<const Model>& pModel, TFunctionRef<bool(const Model*,int)> isRomLockedFunc );
-        void MarkRomUsed( int romIndex, const TSharedPtr<const Model>& pModel );
+		/** Control info for the per-model cache of streamed data. */
+        TArray< FModelCacheEntry > CachePerModel;
+
+		/** Data for each mutable instance that is being updated. */
+		TArray<FLiveInstance> LiveInstances;
 
-        //! Private helper
+		/** Temporary reference to the memory of the current instance being updated. Only valid during a mutable "atomic" operation, like a BeginUpdate or a GetImage. */
+		TSharedPtr<FProgramCache> CurrentInstanceCache;
+
+		/** Resources that have been used in the past, but haven't been deallocated because they still fitted the memory budget and they could be reused. */
+		TArray<Ptr<Image>> PooledImages;
+
+		/** List of intermediate resources that are not soterd anywhere yet. They are still locally referenced by code. */
+		TArray<Ptr<const Image>> TempImages;
+
+		/** */
+		TMap<Ptr<const Resource>, int32> CacheResources;
+
+		/** Given a mutable model, find or create its rom cache. */
 		FModelCacheEntry& GetModelCache(const TSharedPtr<const Model>&);
-    };
 
+        /** Make sure the working memory is below the internal budget, even counting with the passed additional memory. 
+		* An optional function can be passed to "block" the unload of certain roms of data.
+		* Return true if it succeeded, false otherwise.
+		*/
+		bool EnsureBudgetBelow(uint64 AdditionalMemory);
+
+        /** Register that a specific rom has been requested and update the heuristics to keep it in memory. */
+        void MarkRomUsed( int32 RomIndex, const TSharedPtr<const Model>& );
+
+		/** */
+		Ptr<Image> CreateImage(uint32 SizeX, uint32 SizeY, uint32 Lods, EImageFormat Format, EInitializationType Init = EInitializationType::Black)
+		{
+			CheckRunnerThread();
+
+			uint32 DataSize = Image::CalculateDataSize(SizeX, SizeY, Lods, Format);
+
+			// Look for an unused image in the pool that can be reused
+			int32 PooledImageCount = PooledImages.Num();
+			for (int Index = 0; DataSize>0 && Index<PooledImageCount; ++Index)
+			{
+				Ptr<Image>& Candidate = PooledImages[Index];
+				if (Candidate->GetFormat() == Format
+					&& Candidate->GetSizeX() == SizeX
+					&& Candidate->GetSizeY() == SizeY
+					&& Candidate->GetLODCount() == Lods
+					)
+				{
+					Ptr<Image> Result = Candidate;
+					PooledImages.RemoveAtSwap(Index);
+					if (Init == EInitializationType::Black)
+					{
+						Result->InitToBlack();
+					}
+					return Result;
+				}
+			}
+
+			// Make room in the budget
+			EnsureBudgetBelow(DataSize);
+
+			// Create it
+			Ptr<Image> Result = new Image(SizeX, SizeY, Lods, Format, Init);
+
+			TempImages.Add(Result);
+			return Result;
+		}
 
-    //---------------------------------------------------------------------------------------------
-    //! Abstract private system interface
-    //---------------------------------------------------------------------------------------------
+		/** Ref will be nulled and relesed in any case. */
+		Ptr<Image> CloneOrTakeOver(Ptr<const Image>& Resource)
+		{
+			CheckRunnerThread();
+
+			TempImages.RemoveSingle(Resource);
+			check(!TempImages.Contains(Resource));
+			check(!PooledImages.Contains(Resource));
+
+			Ptr<Image> Result;
+			if (!Resource->IsUnique())
+			{
+				// TODO: try to grab from the pool
+
+				uint32 DataSize = Resource->GetDataSize();
+				EnsureBudgetBelow(DataSize);
+
+				Result = Resource->Clone();
+				Release(Resource);
+			}
+			else
+			{
+				Result = const_cast<Image*>(Resource.get());
+				Resource = nullptr;
+			}
+
+			return Result;
+		}
+
+		/** */
+		void Release(Ptr<const Image>& Resource)
+		{
+			CheckRunnerThread();
+
+			if (!Resource)
+			{
+				return;
+			}
+
+			TempImages.RemoveSingle(Resource);
+			check(!TempImages.Contains(Resource));
+			check(!PooledImages.Contains(Resource));
+
+			if (IsBudgetTemp(Resource))
+			{
+				// Check if we are exceeding the budget
+				bool bInBudget = EnsureBudgetBelow(Resource->GetDataSize());
+				if (bInBudget)
+				{
+					PooledImages.Add(const_cast<Image*>(Resource.get()));
+				}
+			}
+			else
+			{
+				// Check if we are exceeding the budget
+				bool bInBudget = EnsureBudgetBelow(0);
+			}
+
+			Resource = nullptr;
+		}
+
+
+		/** */
+		void Release(Ptr<Image>& Resource)
+		{
+			CheckRunnerThread();
+
+			if (!Resource)
+			{
+				return;
+			}
+
+			TempImages.RemoveSingle(Resource);
+			check(!TempImages.Contains(Resource));
+			check(!PooledImages.Contains(Resource));
+
+			if (IsBudgetTemp(Resource))
+			{
+				// Check if we are exceeding the budget
+				bool bInBudget = EnsureBudgetBelow(Resource->GetDataSize());
+				if (bInBudget)
+				{
+					PooledImages.Add(Resource.get());
+				}
+			}
+			else
+			{
+				// Check if we are exceeding the budget
+				bool bInBudget = EnsureBudgetBelow(0);
+			}
+
+			Resource = nullptr;
+		}
+
+		/** */
+		Ptr<const Image> LoadImage(const FCacheAddress& From, bool bTakeOwnership=false)
+		{
+			bool bIsLastReference = false;
+			Ptr<const Image> Result = CurrentInstanceCache->GetImage(From, bIsLastReference);
+			if (!Result)
+			{
+				return nullptr;
+			}
+
+			// If we retrieved the last reference to this resource in "From" cache position (it could still be in other cache positions as well)
+			if (bIsLastReference)
+			{
+				int32* CountPtr = CacheResources.Find(Result);
+				check(CountPtr);
+				*CountPtr = (*CountPtr) - 1;
+				if (!*CountPtr)
+				{
+					CacheResources.FindAndRemoveChecked(Result);
+				}
+			}
+
+
+			if (!bTakeOwnership && Result->IsUnique())
+			{
+				TempImages.Add(Result);
+			}
+
+			return Result;
+		}
+
+		/** */
+		void StoreImage(const FCacheAddress& To, Ptr<const Image> Resource)
+		{
+			if (Resource)
+			{
+				TempImages.RemoveSingle(Resource);
+				check(!TempImages.Contains(Resource));
+
+				int32& Count = CacheResources.FindOrAdd(Resource, 0);
+				++Count;
+			}
+
+			CurrentInstanceCache->SetImage(To, Resource);
+		}
+
+
+		/** Return true if the resource is not in any cache (0,1,rom). */
+		bool IsBudgetTemp(const Ptr<const Image>& Resource)
+		{
+			if (!Resource)
+			{
+				return false;
+			}
+			bool bIsTemp = Resource->IsUnique();
+			return bIsTemp;
+		}
+
+
+		int32 GetPooledBytes() const
+		{
+			int32 Result = 0;
+			for (const Ptr<Image>& Value : PooledImages)
+			{
+				Result += Value->GetDataSize();
+			}
+
+			return Result;
+		}
+
+
+		int32 GetTempBytes() const
+		{
+			int32 Result = 0;
+			for (const Ptr<const Image>& Value : TempImages)
+			{
+				Result += Value->GetDataSize();
+			}
+
+			return Result;
+		}
+
+
+		int32 GetRomBytes() const
+		{
+			int32 Result = 0;
+
+			TArray<const Model*> Models;
+			for (const FLiveInstance& Instance : LiveInstances)
+			{
+				Models.AddUnique(Instance.Model.Get());
+			}
+
+			// Data stored per-model, but related to instance construction
+			for (const Model* Model : Models)
+			{
+				// Count streamable and currently-loaded resources
+				const FProgram& Program = Model->GetPrivate()->m_program;
+				for (const TPair<int32, Ptr<const Image>>& Rom : Program.m_constantImageLODs)
+				{
+					if (Rom.Value && Rom.Key >= 0)
+					{
+						Result += Rom.Value->GetDataSize();
+					}
+				}
+				for (const TPair<int32, Ptr<const Mesh>>& Rom : Program.m_constantMeshes)
+				{
+					if (Rom.Value && Rom.Key >= 0)
+					{
+						Result += Rom.Value->GetDataSize();
+					}
+				}
+			}
+
+			return Result;
+		}
+
+
+		int32 GetTrackedCacheBytes() const
+		{
+			int32 Result = 0;
+			for (const TTuple<Ptr<const Resource>,int32>& It : CacheResources)
+			{
+				Result += It.Key->GetDataSize();
+			}
+
+			return Result;
+		}
+
+
+		int32 GetCache0Bytes() const
+		{
+			if (!CurrentInstanceCache)
+			{
+				return 0;
+			}
+
+			int32 Result = 0;
+			TSet<const Resource*> Cache0Unique;
+
+			CodeContainer<int>::iterator it = CurrentInstanceCache->m_opHitCount.begin();
+			for (; it.IsValid(); ++it)
+			{
+				const Resource* Value = CurrentInstanceCache->m_resources[it.get_address()].Value.get();
+				if (!Value)
+				{
+					continue;
+				}
+
+				int32 Count = *it;
+				if (Count < UE_MUTABLE_CACHE_COUNT_LIMIT)
+				{
+					Cache0Unique.Add(Value);
+				}
+			}
+
+			for (const Resource* Value : Cache0Unique)
+			{
+				Result += Value->GetDataSize();
+			}
+
+			return Result;
+		}
+
+		int32 GetCache1Bytes() const
+		{
+			int32 Result = 0;
+			TSet<const Resource*> Cache1Unique;
+
+			for (const FLiveInstance& Instance : LiveInstances)
+			{
+				CodeContainer<int>::iterator it = Instance.Cache->m_opHitCount.begin();
+				for (; it.IsValid(); ++it)
+				{
+					const Resource* Value = Instance.Cache->m_resources[it.get_address()].Value.get();
+					if (!Value)
+					{
+						continue;
+					}
+
+					int32 Count = *it;
+					if (Count >= UE_MUTABLE_CACHE_COUNT_LIMIT)
+					{
+						Cache1Unique.Add(Value);
+					}
+				}
+			}
+
+			for (const Resource* Value : Cache1Unique)
+			{
+				Result += Value->GetDataSize();
+			}
+
+			return Result;
+		}
+
+
+		/** Remove all intermediate data (big and small) from the memory except for the one that has been explicitely
+		* marked as state cache.
+		*/
+		void ClearCacheLayer0()
+		{
+			check(CurrentInstanceCache);
+
+			MUTABLE_CPUPROFILER_SCOPE(ClearLayer0);
+
+			CodeContainer<int>::iterator it = CurrentInstanceCache->m_opHitCount.begin();
+			for (; it.IsValid(); ++it)
+			{
+				int32 Count = *it;
+				if (Count < UE_MUTABLE_CACHE_COUNT_LIMIT)
+				{
+					TPair<int32, Ptr<const Resource>>& PairRef = CurrentInstanceCache->m_resources[it.get_address()];
+
+					//CacheResources.FindAndRemoveChecked(PairRef.Value);
+					CacheResources.Remove(PairRef.Value);
+
+					// SetUnused only clears if the data size is relevant (mesh or image) but we need to 
+					// clear it in all cases here, because it may have become invalid because of parameter changes.
+					// SetUnused(it.get_address());
+					PairRef.Value = nullptr;
+					PairRef.Key = 0;
+					*it = 0;
+				}
+			}
+		}
+
+
+		/** Remove all intermediate data (big and small) from the memory including the one that has been explicitely
+		* marked as state cache.
+		*/
+		void ClearCacheLayer1()
+		{
+			MUTABLE_CPUPROFILER_SCOPE(ClearLayer1);
+
+			CodeContainer< TPair<int32, Ptr<const Resource>> >::iterator it = CurrentInstanceCache->m_resources.begin();
+			for (; it.IsValid(); ++it)
+			{
+				//CacheResources.FindAndRemoveChecked((*it).Value);
+				CacheResources.Remove((*it).Value);
+
+				(*it).Value = nullptr;
+				(*it).Key = 0;
+			}
+
+			CurrentInstanceCache->m_descCache.SetNum(0);
+		}
+
+
+		void LogWorkingMemory(const class CodeRunner* CurrentRunner) const;
+
+
+#if !(UE_BUILD_SHIPPING || UE_BUILD_TEST)
+		/** Temp variable with the ID of the thread running code, for debugging. */
+		uint32 DebugRunnerThreadID = FThread::InvalidThreadId;
+#endif
+
+		/** This is a development-only check to make sure calls to resource management happen in the correct thread. */
+		inline void BeginRunnerThread()
+		{
+#if !(UE_BUILD_SHIPPING || UE_BUILD_TEST)
+			// If this check fails it means have not set up correctly all the paths to debug threading for resource management.
+			check(DebugRunnerThreadID == FThread::InvalidThreadId);
+
+			DebugRunnerThreadID = FPlatformTLS::GetCurrentThreadId();
+#endif
+		}
+
+
+		/** This is a development-only check to make sure calls to resource management happen in the correct thread. */
+		inline void CheckRunnerThread()
+		{
+#if !(UE_BUILD_SHIPPING || UE_BUILD_TEST)
+			// If this check fails it means have not set up correctly all the paths to debug threading for resource management.
+			check(DebugRunnerThreadID != FThread::InvalidThreadId);
+
+			// If this check fails it means we are doing resource management from a thread that is not the CodeRunner::RunCode
+			// thread, and this is not allowed.
+			check(DebugRunnerThreadID== FPlatformTLS::GetCurrentThreadId());
+#endif
+		}
+
+
+		/** This is a development-only check to make sure calls to resource management happen in the correct thread. */
+		inline void EndRunnerThread()
+		{
+			CurrentInstanceCache->CheckHitCountsCleared();
+
+			// If this check fails it means some operation is not correctly handling resource management and didn't release
+			// a resource it created.
+			// This should be reported and reviewed, but it is not fatal. Some unnecessary memory may be used temporarily.
+			ensure(TempImages.Num()==0);
+
+#if !(UE_BUILD_SHIPPING || UE_BUILD_TEST)
+			// If this check fails it means have not set up correctly all the paths to debug threading for resource management.
+			check(DebugRunnerThreadID != FThread::InvalidThreadId);
+
+			DebugRunnerThreadID = FThread::InvalidThreadId;
+#endif
+		}
+
+	};
+
+
+	/** */
     class System::Private : public Base
     {
     public:
 
-        Private(SettingsPtr pSettings, ExtensionDataStreamer* DataStreamer);
+        Private( Ptr<Settings>, const TSharedPtr<ExtensionDataStreamer>& );
         virtual ~Private();
 
         //-----------------------------------------------------------------------------------------
@@ -1016,116 +1482,48 @@ namespace mu
 		MUTABLERUNTIME_API void BeginBuild(const TSharedPtr<const Model>&);
 		MUTABLERUNTIME_API void EndBuild();
 
-		MUTABLERUNTIME_API bool BuildBool(const TSharedPtr<const Model>&, const Parameters* pParams, OP::ADDRESS at) ;
-		MUTABLERUNTIME_API int BuildInt(const TSharedPtr<const Model>&, const Parameters* pParams, OP::ADDRESS at) ;
-		MUTABLERUNTIME_API float BuildScalar(const TSharedPtr<const Model>&, const Parameters* pParams, OP::ADDRESS at) ;
-		MUTABLERUNTIME_API void BuildColour(const TSharedPtr<const Model>&, const Parameters*, OP::ADDRESS, float* OutR, float* OutG, float* OutB, float* OutA) ;
-		MUTABLERUNTIME_API Ptr<const String> BuildString(const TSharedPtr<const Model>&, const Parameters* pParams, OP::ADDRESS at) ;
-		MUTABLERUNTIME_API Ptr<const Image> BuildImage(const TSharedPtr<const Model>&, const Parameters* pParams, OP::ADDRESS at, int32 MipsToSkip, int32 LOD) ;
-		MUTABLERUNTIME_API Ptr<const Mesh> BuildMesh(const TSharedPtr<const Model>&, const Parameters* pParams, OP::ADDRESS at) ;
-		MUTABLERUNTIME_API Ptr<const Layout> BuildLayout(const TSharedPtr<const Model>&, const Parameters* pParams, OP::ADDRESS at) ;
-    	MUTABLERUNTIME_API Ptr<const Projector> BuildProjector(const TSharedPtr<const Model>&, const Parameters* pParams, OP::ADDRESS at) ;
-		void SetStreamingCache(uint64 bytes) ;
-		void ClearStreamingCache() ;
+		MUTABLERUNTIME_API bool BuildBool(const TSharedPtr<const Model>&, const Parameters*, OP::ADDRESS) ;
+		MUTABLERUNTIME_API int32 BuildInt(const TSharedPtr<const Model>&, const Parameters*, OP::ADDRESS) ;
+		MUTABLERUNTIME_API float BuildScalar(const TSharedPtr<const Model>&, const Parameters*, OP::ADDRESS) ;
+		MUTABLERUNTIME_API FVector4f BuildColour(const TSharedPtr<const Model>&, const Parameters*, OP::ADDRESS) ;
+		MUTABLERUNTIME_API Ptr<const String> BuildString(const TSharedPtr<const Model>&, const Parameters*, OP::ADDRESS) ;
+		MUTABLERUNTIME_API Ptr<const Image> BuildImage(const TSharedPtr<const Model>&, const Parameters*, OP::ADDRESS, int32 MipsToSkip, int32 LOD) ;
+		MUTABLERUNTIME_API Ptr<const Mesh> BuildMesh(const TSharedPtr<const Model>&, const Parameters*, OP::ADDRESS) ;
+		MUTABLERUNTIME_API Ptr<const Layout> BuildLayout(const TSharedPtr<const Model>&, const Parameters*, OP::ADDRESS) ;
+    	MUTABLERUNTIME_API Ptr<const Projector> BuildProjector(const TSharedPtr<const Model>&, const Parameters*, OP::ADDRESS) ;
 
-		ExtensionDataStreamer* GetExtensionDataStreamer() const { return m_ExtensionDataStreamer; }
+		ExtensionDataStreamer* GetExtensionDataStreamer() const { return ExtensionDataStreamer.Get(); }
 
         //!
-        SettingsPtrConst m_pSettings;
+        Ptr<const Settings> Settings;
 
         //! Data streaming interface, if any.
-        //! It is owned by this system
-        ModelStreamer* m_pStreamInterface;
-
-        //! It is not owned by this system
-        ImageParameterGenerator* m_pImageParameterGenerator;
-
-        //! Maximum amount of memory (bytes) if a limit has been set. If 0, it means there is no
-        //! limit.
-        uint32 m_maxMemory;
-
-        //! Current instances management
-        //! @{
-
-		/** Data for an instance that is currently being processed in the mutable system. This means it is
-		* between a BeginUpdate and EndUpdate, or during an atmoic operation (like generate a single resource).
-		*/
-        struct FLiveInstance
-        {
-            Instance::ID m_instanceID;
-            InstancePtrConst m_pInstance;
-            TSharedPtr<const Model> m_pModel;
-
-            int m_state;
-            ParametersPtr m_pOldParameters;
+		TSharedPtr<ModelStreamer> StreamInterface;
 
-			//! Mask of the parameters that have changed since the last update.
-			//! Every bit represents a state parameter.
-			uint64 m_updatedParameters = 0;
+		TSharedPtr<ImageParameterGenerator> ImageParameterGenerator;
 
-			//! An entry for every instruction in the program to cache resources (meshes, images) if necessary.
-			TSharedPtr<FProgramCache> m_memory;
-
-			~FLiveInstance()
-			{
-				// Manually done to trace mem deallocations
-				MUTABLE_CPUPROFILER_SCOPE(LiveInstanceDestructor);
-				m_memory = nullptr;
-				m_pOldParameters = nullptr;
-				m_pInstance = nullptr;
-				m_pModel = nullptr;
-			}
-        };
-
-        TArray<FLiveInstance> m_liveInstances;
-
-		/** Temporary reference to the memory of the current instance being updated. */
-		TSharedPtr<FProgramCache> m_memory;
+		/** */
+		FWorkingMemoryManager WorkingMemoryManager;
 
 		/** Counter used to generate unique IDs for every new instance created in the system. */
-        Instance::ID m_lastInstanceID = 0;
+        Instance::ID LastInstanceID = 0;
 
 		/** The pointer returned by this function is only valid for the duration of the current mutable operation. */
-        inline FLiveInstance* FindLiveInstance( Instance::ID id )
-        {
-            for ( int32 i=0; i<m_liveInstances.Num(); ++i )
-            {
-                if (m_liveInstances[i].m_instanceID==id)
-                {
-                    return &m_liveInstances[i];
-                }
-            }
-            return nullptr;
-        }
-
-        //! @}
+		inline FLiveInstance* FindLiveInstance(Instance::ID id);
 
         //!
-        bool CheckUpdatedParameters( const FLiveInstance* pLiveInstance,
-			const Ptr<const Parameters>& pParams,
-			uint64& OutUpdatedParameters );
+        bool CheckUpdatedParameters( const FLiveInstance*, const Ptr<const Parameters>&, uint64& OutUpdatedParameters );
 
-	public:
 
-		void RunCode(const TSharedPtr<const Model>& pModel,
-			const Parameters* pParams,
-			OP::ADDRESS at, uint32 LODs = System::AllLODs, uint8 executionOptions = 0, int32 LOD = 0);
+		void RunCode(const TSharedPtr<const Model>&, const Parameters*, OP::ADDRESS at, uint32 LODs = System::AllLODs, uint8 executionOptions = 0, int32 LOD = 0);
 
 		//!
 		void PrepareCache(const Model*, int32 State);
 
-		//! Cache related members
-		//! @{
-
-		//! Data for each used model.
-		FModelCache m_modelCache;
-
-		//! @}
-
 	private:
 
 		/** Owned by this system. */
-		ExtensionDataStreamer* m_ExtensionDataStreamer = nullptr;
+		TSharedPtr<ExtensionDataStreamer> ExtensionDataStreamer = nullptr;
 
 		/** This flag is turned on when a streaming error or similar happens. Results are not usable.
 		* This should only happen in-editor.
diff --git a/Engine/Plugins/Experimental/Mutable/Source/MutableRuntime/Public/MuR/ExtensionData.h b/Engine/Plugins/Experimental/Mutable/Source/MutableRuntime/Public/MuR/ExtensionData.h
--- a/Engine/Plugins/Experimental/Mutable/Source/MutableRuntime/Public/MuR/ExtensionData.h
+++ b/Engine/Plugins/Experimental/Mutable/Source/MutableRuntime/Public/MuR/ExtensionData.h
@@ -18,12 +18,15 @@ namespace mu
 	typedef Ptr<ExtensionData> ExtensionDataPtr;
 	typedef Ptr<const ExtensionData> ExtensionDataPtrConst;
 
-	class MUTABLERUNTIME_API ExtensionData : public RefCounted
+	class MUTABLERUNTIME_API ExtensionData : public Resource
 	{
 	public:
 		static void Serialise(const ExtensionData* Data, OutputArchive& Archive);
 		static ExtensionDataPtr StaticUnserialise(InputArchive& Archive);
 
+		// Resource interface
+		int32 GetDataSize() const override;
+
 		//! A stable hash of the contents
 		uint32 Hash() const;
 
diff --git a/Engine/Plugins/Experimental/Mutable/Source/MutableRuntime/Public/MuR/Image.h b/Engine/Plugins/Experimental/Mutable/Source/MutableRuntime/Public/MuR/Image.h
--- a/Engine/Plugins/Experimental/Mutable/Source/MutableRuntime/Public/MuR/Image.h
+++ b/Engine/Plugins/Experimental/Mutable/Source/MutableRuntime/Public/MuR/Image.h
@@ -227,23 +227,27 @@ namespace mu
 		
 	MUTABLERUNTIME_API const FImageFormatData& GetImageFormatData(EImageFormat format );
 
-	//---------------------------------------------------------------------------------------------
-	//! Convert an image to another pixel format.
-	//! Allocates the destination image.
-	//! \warning Not all format conversions are implemented.
-	//! \param onlyLOD If different than -1, only the specified lod level will be converted in the
-	//! returned image.
-	//! \return false if the conversion failed, usually because not enough memory was allocated in
-	//!     the result. This is only checked for RLE compression.
-	//---------------------------------------------------------------------------------------------
-	MUTABLERUNTIME_API ImagePtr ImagePixelFormat( int quality, const Image* pBase, EImageFormat targetFormat,
-	                                      int onlyLOD = -1 );
-
-	MUTABLERUNTIME_API Ptr<Image> ImageSwizzle(EImageFormat format, const Ptr<const Image> pSources[], const uint8 channels[]);
+	/** Convert an image to another pixel format.
+	* Allocates the destination image.
+	* \warning Not all format conversions are implemented.
+	* \param onlyLOD If different than -1, only the specified lod level will be converted in the returned image.
+	* \return nullptr if the conversion failed, usually because not enough memory was allocated in the result. This is only checked for RLE compression.
+	*/
+	MUTABLERUNTIME_API Ptr<Image> ImagePixelFormat(int32 Quality, const Image* Base, EImageFormat TargetFormat, int32 OnlyLOD = -1);
 	
+	/** Convert an image to another pixel format.
+	* \warning Not all format conversions are implemented.
+	* \param onlyLOD If different than -1, only the specified lod level will be converted in the returned image.
+	* \return false if the conversion failed, usually because not enough memory was allocated in the result. This is only checked for RLE compression.
+	*/
+	MUTABLERUNTIME_API void ImagePixelFormat(bool& bOutSuccess, int32 Quality, Image* Result, const Image* Base, int32 OnlyLOD = -1);
+
+	MUTABLERUNTIME_API Ptr<Image> ImageSwizzle(EImageFormat Format, const Ptr<const Image> Sources[], const uint8 Channels[]);
+	MUTABLERUNTIME_API void ImageSwizzle(Image* Result, const Ptr<const Image> Sources[], const uint8 Channels[]);
+
     //! \brief 2D image resource with mipmaps.
 	//! \ingroup runtime
-    class MUTABLERUNTIME_API Image : public RefCounted
+    class MUTABLERUNTIME_API Image : public Resource
 	{
 	public:
 
@@ -267,11 +271,14 @@ namespace mu
 		static Ptr<Image> CreateAsReference( uint32 ID );
 
 		/** */
-		ImagePtr ExtractMip(int32 Mip) const;
+		Ptr<Image> ExtractMip(int32 Mip) const;
 
 		//! Serialisation
 		static void Serialise( const Image* p, OutputArchive& arch );
-		static ImagePtr StaticUnserialise( InputArchive& arch );
+		static Ptr<Image> StaticUnserialise( InputArchive& arch );
+
+		// Resource interface
+		int32 GetDataSize() const override;
 
 		//-----------------------------------------------------------------------------------------
 		// Own interface
@@ -297,9 +304,6 @@ namespace mu
         const uint8* GetData() const;
         uint8* GetData();
 
-        //! Return the size in bytes of all the LODs of the image.
-        int32 GetDataSize() const;
-
         //! Return the size in bytes of a specific LOD of the image.
         int32 GetLODDataSize( int lod ) const;
 
@@ -309,6 +313,9 @@ namespace mu
 		/** Return the id of the engine referenced texture. Only valid if IsReference. */
 		uint32 GetReferencedTexture() const;
 
+		/** Clear the image to black colour. */
+		void InitToBlack();
+
 	protected:
 
 		//! Forbidden. Manage with the Ptr<> template.
@@ -366,12 +373,12 @@ namespace mu
 		//-----------------------------------------------------------------------------------------
 
 		//! Deep clone.
-		ImagePtr Clone() const
+		Ptr<Image> Clone() const
 		{
 			LLM_SCOPE_BYNAME(TEXT("MutableRuntime"));
 			MUTABLE_CPUPROFILER_SCOPE(ImageClone)
 
-			ImagePtr pResult = new Image();
+			Ptr<Image> pResult = new Image();
 
 			pResult->m_size = m_size;
 			pResult->m_format = m_format;
@@ -385,6 +392,36 @@ namespace mu
 			return pResult;
 		}
 
+		//! Copy another image.
+		void CopyMove( Image* Other )
+		{
+			LLM_SCOPE_BYNAME(TEXT("MutableRuntime"));
+			MUTABLE_CPUPROFILER_SCOPE(CopyMove);
+
+			if (Other == this)
+			{
+				return;
+			}
+
+			m_size = Other->m_size;
+			m_format = Other->m_format;
+			m_lods = Other->m_lods;
+			m_flags = Other->m_flags;
+			RelevancyMinY = Other->RelevancyMinY;
+			RelevancyMaxY = Other->RelevancyMaxY;
+			m_data = MoveTemp(Other->m_data);
+			ReferenceID = Other->ReferenceID;
+			
+			Other->m_size = FImageSize(0,0);
+			Other->m_format = EImageFormat::IF_NONE;
+			Other->m_lods = 0;
+			Other->m_flags = 0;
+			Other->RelevancyMinY = 0;
+			Other->RelevancyMaxY = 0;
+			Other->m_data.SetNum(0);
+			Other->ReferenceID=0;
+		}
+
 
 		//!
 		void Serialise(OutputArchive& arch) const;
@@ -411,6 +448,8 @@ namespace mu
 		//! m_data, only using the image descriptions. For non-block-compressed images, it returns
 		//! 0.
 		int32 CalculateDataSize() const;
+		static int32 CalculateDataSize(int32 SizeX, int32 SizeY, int32 LodCount, EImageFormat Format);
+
 
 		//! Calculate the size of a lod of the image data in bytes, regardless of what is allocated
 		//! in m_data, only using the image descriptions. For non-block-compressed images, it
diff --git a/Engine/Plugins/Experimental/Mutable/Source/MutableRuntime/Public/MuR/Instance.h b/Engine/Plugins/Experimental/Mutable/Source/MutableRuntime/Public/MuR/Instance.h
--- a/Engine/Plugins/Experimental/Mutable/Source/MutableRuntime/Public/MuR/Instance.h
+++ b/Engine/Plugins/Experimental/Mutable/Source/MutableRuntime/Public/MuR/Instance.h
@@ -22,7 +22,7 @@ namespace mu
     //! It corresponds to an "engine object" but the contents of its data depends on the Model, and
     //! it may contain any number of LODs, components, surfaces, meshes and images, even none.
 	//! \ingroup runtime
-	class MUTABLERUNTIME_API Instance : public RefCounted
+	class MUTABLERUNTIME_API Instance : public Resource
 	{
 	public:
 
@@ -35,6 +35,8 @@ namespace mu
         //! Clone this instance
         InstancePtr Clone() const;
 
+		// Resource interface
+		int32 GetDataSize() const override;
 
 		//-----------------------------------------------------------------------------------------
 		// Own interface
diff --git a/Engine/Plugins/Experimental/Mutable/Source/MutableRuntime/Public/MuR/Layout.h b/Engine/Plugins/Experimental/Mutable/Source/MutableRuntime/Public/MuR/Layout.h
--- a/Engine/Plugins/Experimental/Mutable/Source/MutableRuntime/Public/MuR/Layout.h
+++ b/Engine/Plugins/Experimental/Mutable/Source/MutableRuntime/Public/MuR/Layout.h
@@ -39,7 +39,7 @@ namespace mu
     //! coordinates set from a mesh).
     //! It is usually not necessary to use this objects, except for some advanced cases.
 	//! \ingroup runtime
-    class MUTABLERUNTIME_API Layout : public RefCounted
+    class MUTABLERUNTIME_API Layout : public Resource
 	{
 	public:
 
@@ -59,6 +59,9 @@ namespace mu
 		//! Full compare
 		bool operator==( const Layout& other ) const;
 
+		// Resource interface
+		int32 GetDataSize() const override;
+
 		//-----------------------------------------------------------------------------------------
 		// Own interface
 		//-----------------------------------------------------------------------------------------
diff --git a/Engine/Plugins/Experimental/Mutable/Source/MutableRuntime/Public/MuR/Mesh.h b/Engine/Plugins/Experimental/Mutable/Source/MutableRuntime/Public/MuR/Mesh.h
--- a/Engine/Plugins/Experimental/Mutable/Source/MutableRuntime/Public/MuR/Mesh.h
+++ b/Engine/Plugins/Experimental/Mutable/Source/MutableRuntime/Public/MuR/Mesh.h
@@ -145,7 +145,7 @@ namespace mu
     //! it can be ignored.
     //! The meshes are always assumed to be triangle list primitives.
     //! \ingroup runtime
-    class MUTABLERUNTIME_API Mesh : public RefCounted
+    class MUTABLERUNTIME_API Mesh : public Resource
     {
     public:
 
@@ -163,6 +163,9 @@ namespace mu
         static void Serialise( const Mesh* p, OutputArchive& arch );
         static MeshPtr StaticUnserialise( InputArchive& arch );
 
+		// Resource interface
+		int32 GetDataSize() const override;
+
         //-----------------------------------------------------------------------------------------
         // Own interface
         //-----------------------------------------------------------------------------------------
@@ -527,9 +530,6 @@ namespace mu
 		//! Update the flags identifying the mesh format as some of the optimised formats.
 		void ResetStaticFormatFlags() const;
 
-		//! Get the total memory size of the buffers
-		size_t GetDataSize() const;
-
 		//! Create the surface data if not present.
 		void EnsureSurfaceData();
 
diff --git a/Engine/Plugins/Experimental/Mutable/Source/MutableRuntime/Public/MuR/MeshBufferSet.h b/Engine/Plugins/Experimental/Mutable/Source/MutableRuntime/Public/MuR/MeshBufferSet.h
--- a/Engine/Plugins/Experimental/Mutable/Source/MutableRuntime/Public/MuR/MeshBufferSet.h
+++ b/Engine/Plugins/Experimental/Mutable/Source/MutableRuntime/Public/MuR/MeshBufferSet.h
@@ -352,8 +352,8 @@ namespace mu
 		//! Compare the format of the two buffers at index buffer and return true if they match.
 		bool HasSameFormat(int32 buffer, const FMeshBufferSet& pOther) const;
 
-		//! Get the total memory size of the buffers
-		size_t GetDataSize() const;
+		//! Get the total memory size of the buffers and this struct
+		int32 GetDataSize() const;
 
 		//! Compare the mesh buffer with another one, but ignore internal data like generated
 		//! vertex indices.
diff --git a/Engine/Plugins/Experimental/Mutable/Source/MutableRuntime/Public/MuR/MutableString.h b/Engine/Plugins/Experimental/Mutable/Source/MutableRuntime/Public/MuR/MutableString.h
--- a/Engine/Plugins/Experimental/Mutable/Source/MutableRuntime/Public/MuR/MutableString.h
+++ b/Engine/Plugins/Experimental/Mutable/Source/MutableRuntime/Public/MuR/MutableString.h
@@ -14,7 +14,7 @@ namespace mu
 
     //! \brief Container for text string data.
 	//! \ingroup runtime
-    class MUTABLERUNTIME_API String : public RefCounted
+    class MUTABLERUNTIME_API String : public Resource
 	{
 	public:
 
@@ -31,6 +31,9 @@ namespace mu
 		static void Serialise( const String* p, OutputArchive& arch );
 		static Ptr<String> StaticUnserialise( InputArchive& arch );
 
+		// Resource interface
+		int32 GetDataSize() const override;
+
 		//-----------------------------------------------------------------------------------------
 		// Own interface
 		//-----------------------------------------------------------------------------------------
diff --git a/Engine/Plugins/Experimental/Mutable/Source/MutableRuntime/Public/MuR/RefCounted.h b/Engine/Plugins/Experimental/Mutable/Source/MutableRuntime/Public/MuR/RefCounted.h
--- a/Engine/Plugins/Experimental/Mutable/Source/MutableRuntime/Public/MuR/RefCounted.h
+++ b/Engine/Plugins/Experimental/Mutable/Source/MutableRuntime/Public/MuR/RefCounted.h
@@ -16,12 +16,10 @@ namespace mu
 {
 
 
-	//! \brief %Base class for all reference counted objects.
-	//!
-	//! Any subclass of this class can be managed using smart pointers through the Ptr<T> template.
-	//! \warning This base allow multi-threaded manipulation of smart pointers, since the count
-	//! increments and decrements are atomic.
-	//! \ingroup runtime
+	/** Base class for all reference counted objects.
+	* Any subclass of this class can be managed using smart pointers through the Ptr<T> template.
+	* \warning This base allow multi-threaded manipulation of smart pointers, since the count increments and decrements are atomic.
+	*/
 	class MUTABLERUNTIME_API RefCounted : public Base
 	{
 	public:
@@ -82,5 +80,13 @@ namespace mu
 		}
 	}
 
+	class MUTABLERUNTIME_API Resource : public RefCounted
+	{
+	public:
+
+		/** Return the size in bytes of all the LODs of the image. */
+		virtual int32 GetDataSize() const = 0;
+
+	};
 }
 
diff --git a/Engine/Plugins/Experimental/Mutable/Source/MutableRuntime/Public/MuR/Settings.h b/Engine/Plugins/Experimental/Mutable/Source/MutableRuntime/Public/MuR/Settings.h
--- a/Engine/Plugins/Experimental/Mutable/Source/MutableRuntime/Public/MuR/Settings.h
+++ b/Engine/Plugins/Experimental/Mutable/Source/MutableRuntime/Public/MuR/Settings.h
@@ -8,55 +8,39 @@
 
 namespace mu
 {
-	// Forward references
-    class Settings;
-
-    using SettingsPtr=Ptr<Settings>;
-    using SettingsPtrConst=Ptr<const Settings>;
-
-
-
-    //! \brief Settings class that can be used to create a customised System instance.
-    //! \ingroup runtime
+    /** Settings class that can be used to create a customised System instance. */
     class MUTABLERUNTIME_API Settings : public RefCounted
     {
     public:
 
-        //! Create new settings with the default values as explained in each method below.
+        /** Create new settings with the default values as explained in each method below. */
         Settings();
 
-        //! Record internal profiling data.
-        //! Disabled by default.
+        /** Record internal profiling data. Disabled by default. */
         void SetProfile( bool bEnabled );
 
-        //! Limit the maximum memory in bytes used to store streaming data. A low value will force
-        //! more streaming and higher instance construction times, but will use less memory while
-        //! building objects.
-        //! Defaults to 0, which disables the limit.
-        void SetStreamingCache( uint64 bytes );
-
-        //! Set the quality for the image compression algorithms.
-        //! \param quality Quality level of the compression. Higher quality increases the instance
-        //! generation time. The value is in the range 0 to 4 with 0 being the fastest. It defaults
-        //! to 0, which is the fastes for runtime. Tools may want ot set it higher, and low
-        //! performance profiles to lower.
-        //! Genreal rules are:
-        //! 0 - Fastest for runtime
-        //! 1 - Best for runtime
-        //! 2 - Fast for tools
-        //! 3 - Best for tools
-        //! 4 - Maximum, with no time limits.
-        //! Visit https://work.anticto.com/w/mutable/core/guide/ for
-        //! more information.
-        void SetImageCompressionQuality( int quality );
-
-        //-----------------------------------------------------------------------------------------
-        // Interface pattern
-        //-----------------------------------------------------------------------------------------
-        class Private;
-
-        Private* GetPrivate() const;
+        /** Limit the maximum memory in bytes used by the mutable core.A low value will force
+        * more streaming and higher instance construction times, but will use less memory while
+        * building objects.
+        * Defaults to 0, which disables the limit.
+		*/
+        void SetWorkingMemoryBytes( uint64 Bytes );
+
+        /** Set the quality for the image compression algorithms.
+        * \param quality Quality level of the compression. Higher quality increases the instance
+        * generation time. The value is in the range 0 to 4 with 0 being the fastest. It defaults
+        * to 0, which is the fastes for runtime. Tools may want ot set it higher, and low
+        * performance profiles to lower.
+        * Genreal rules are:
+        * 0 - Fastest for runtime
+        * 1 - Best for runtime
+        * 2 - Fast for tools
+        * 3 - Best for tools
+        * 4 - Maximum, with no time limits.
+        */
+        void SetImageCompressionQuality( int32 Quality );
 
+        // Interface pattern
         Settings( const Settings& ) = delete;
         Settings( Settings&& ) = delete;
         Settings& operator=(const Settings&) = delete;
@@ -64,15 +48,16 @@ namespace mu
 
     protected:
 
-        //! Forbidden. Manage with the Ptr<> template.
+        /** Forbidden.Manage with the Ptr<> template. */
         ~Settings() override;
 
-    private:
+    public:
 
-        Private* m_pD;
+		uint64 WorkingMemoryBytes = 0;
+		int32 ImageCompressionQuality = 0;
+		bool bProfile = false;
 
     };
 
-
 }
 
diff --git a/Engine/Plugins/Experimental/Mutable/Source/MutableRuntime/Public/MuR/System.h b/Engine/Plugins/Experimental/Mutable/Source/MutableRuntime/Public/MuR/System.h
--- a/Engine/Plugins/Experimental/Mutable/Source/MutableRuntime/Public/MuR/System.h
+++ b/Engine/Plugins/Experimental/Mutable/Source/MutableRuntime/Public/MuR/System.h
@@ -107,7 +107,7 @@ namespace mu
     public:
 
         //! This constant can be used in place of the lodMask in methods like BeginUpdate
-        static constexpr unsigned AllLODs = 0xffffffff;
+        static constexpr uint32 AllLODs = 0xffffffff;
 
     public:
 
@@ -115,29 +115,24 @@ namespace mu
         //! \param Settings Optional class with the settings to use in this system. The default
         //! value configures a production-ready system.
 		//! \param DataStreamer Optional interface to allow the Model to stream in ExtensionData from disk
-        System( const Ptr<Settings>& Settings = nullptr, ExtensionDataStreamer* DataStreamer = nullptr );
+        System( const Ptr<Settings>& Settings = nullptr, const TSharedPtr<ExtensionDataStreamer>& DataStreamer = nullptr );
 
-        //! Set a new provider for model data. The provider will become owned by this instance and
-        //! destroyed when necessary.
-        void SetStreamingInterface( ModelStreamer* );
+        //! Set a new provider for model data. 
+        void SetStreamingInterface(const TSharedPtr<ModelStreamer>& );
 
-        //! Overwrite the streaming memory limit set in the settings when the system was created.
-        //! Refer to Settings::SetStreamingCache for more information.
-        void SetStreamingCache( uint64 bytes );
+        /** Set the working memory limit, overrding any set in the settings when the system was created.
+         * Refer to Settings::SetWorkingMemoryBudget for more information.
+		 */
+        void SetWorkingMemoryBytes( uint64 Bytes );
 
-        //! Clears the streaming cache content. It does not modify the memory limit.
-		void ClearStreamingCache();
+        /** Removes all the possible working memory regardless of the budget set. This may make following 
+		* operations take longer.
+		*/
+		void ClearWorkingMemory();
 		
         //! Set a new provider for external image data. This is only necessary if image parameters
         //! are used in the models.
-		//! The provided pointer becomes owned by the system.
-        void SetImageParameterGenerator( ImageParameterGenerator* );
-
-        //! Set the maximum memory that this system can use. This memory is used for built data,
-        //! cached data, and streaming. If set to 0, the system will use as much memory as required.
-        //! \warning This limit is ignored if the system was built with SF_REFERENCE.
-        //! \param mem Maximum memory, in bytes
-        void SetMemoryLimit( uint32 mem );
+        void SetImageParameterGenerator(const TSharedPtr<ImageParameterGenerator>& );
 
         //! Create a new instance from the given model. The instance can then be configured through
         //! calls to BeginUpdate/EndUpdate.
@@ -200,11 +195,6 @@ namespace mu
 									const Ptr<const Parameters>& Parameters,
 									bool* Flags );
 
-        //! Free memory used in internal runtime caches. This is useful for long-running processes
-        //! that keep models loaded. This could be called when a game finishes, or a change of
-        //! level.
-        void ClearCaches();
-
 		//-----------------------------------------------------------------------------------------
 		// Interface pattern
 		//-----------------------------------------------------------------------------------------
diff --git a/Engine/Plugins/Experimental/Mutable/Source/MutableTools/Private/MuT/AST.cpp b/Engine/Plugins/Experimental/Mutable/Source/MutableTools/Private/MuT/AST.cpp
--- a/Engine/Plugins/Experimental/Mutable/Source/MutableTools/Private/MuT/AST.cpp
+++ b/Engine/Plugins/Experimental/Mutable/Source/MutableTools/Private/MuT/AST.cpp
@@ -1140,11 +1140,6 @@ FImageDesc ASTOpFixed::GetImageDesc( bool returnBestOption, FGetImageDescContext
         res = GetImageDesc( op.args.ImageInterpolate3.target0, returnBestOption, context );
         break;
 
-    case OP_TYPE::IM_DIFFERENCE:
-        res = GetImageDesc( op.args.ImageDifference.a, returnBestOption, context );
-        res.m_format = EImageFormat::IF_L_UBYTE;
-        break;
-
     case OP_TYPE::IM_PLAINCOLOUR:
         res.m_format = EImageFormat(op.args.ImagePlainColour.format);
         res.m_size[0] = op.args.ImagePlainColour.size[0];
@@ -1190,11 +1185,6 @@ FImageDesc ASTOpFixed::GetImageDesc( bool returnBestOption, FGetImageDescContext
         }
         break;
 
-    case OP_TYPE::IM_SELECTCOLOUR:
-        res = GetImageDesc( op.args.ImageSelectColour.base, returnBestOption, context );
-        res.m_format = EImageFormat::IF_L_UBYTE;
-        break;
-
     case OP_TYPE::IM_GRADIENT:
         res.m_size[0] = op.args.ImageGradient.size[0];
         res.m_size[1] = op.args.ImageGradient.size[1];
@@ -1645,13 +1635,6 @@ mu::Ptr<ImageSizeExpression> ASTOpFixed::GetImageSizeExpression() const
         pRes->factor[1] = op.args.ImageBlankLayout.blockSize[1];
         break;
 
-    case OP_TYPE::IM_DIFFERENCE:
-        if ( children[op.args.ImageDifference.a] )
-        {
-            pRes = children[op.args.ImageDifference.a].child()->GetImageSizeExpression();
-        }
-        break;
-
     case OP_TYPE::IM_INTERPOLATE:
         if ( children[op.args.ImageInterpolate.targets[0]] )
         {
@@ -1694,13 +1677,6 @@ mu::Ptr<ImageSizeExpression> ASTOpFixed::GetImageSizeExpression() const
         }
         break;
 
-    case OP_TYPE::IM_SELECTCOLOUR:
-        if ( children[op.args.ImageSelectColour.colour] )
-        {
-            pRes = children[op.args.ImageSelectColour.colour].child()->GetImageSizeExpression();
-        }
-        break;
-
     case OP_TYPE::IM_GRADIENT:
         pRes->type = ImageSizeExpression::ISET_CONSTANT;
         pRes->size[0] = op.args.ImageGradient.size[0];
diff --git a/Engine/Plugins/Experimental/Mutable/Source/MutableTools/Private/MuT/ASTOpImageLayerColor.cpp b/Engine/Plugins/Experimental/Mutable/Source/MutableTools/Private/MuT/ASTOpImageLayerColor.cpp
--- a/Engine/Plugins/Experimental/Mutable/Source/MutableTools/Private/MuT/ASTOpImageLayerColor.cpp
+++ b/Engine/Plugins/Experimental/Mutable/Source/MutableTools/Private/MuT/ASTOpImageLayerColor.cpp
@@ -98,6 +98,7 @@ namespace mu
 			args.BlendAlphaSourceChannel = BlendAlphaSourceChannel;
 			args.flags = Flags;
 
+			check(base);
 			if (base) args.base = base->linkedAddress;
 			if (color) args.colour = color->linkedAddress;
 			if (mask) args.mask = mask->linkedAddress;
diff --git a/Engine/Plugins/Experimental/Mutable/Source/MutableTools/Private/MuT/ASTOpReferenceResource.cpp b/Engine/Plugins/Experimental/Mutable/Source/MutableTools/Private/MuT/ASTOpReferenceResource.cpp
--- a/Engine/Plugins/Experimental/Mutable/Source/MutableTools/Private/MuT/ASTOpReferenceResource.cpp
+++ b/Engine/Plugins/Experimental/Mutable/Source/MutableTools/Private/MuT/ASTOpReferenceResource.cpp
@@ -6,13 +6,6 @@
 #include "HAL/PlatformMath.h"
 #include "Hash/CityHash.h"
 #include "Misc/AssertionMacros.h"
-#include "MuR/Layout.h"
-#include "MuR/Mesh.h"
-#include "MuR/ModelPrivate.h"
-#include "MuR/MutableMath.h"
-#include "MuR/PhysicsBody.h"
-#include "MuR/Serialisation.h"
-#include "MuR/Skeleton.h"
 #include "MuR/Types.h"
 #include "MuT/StreamsPrivate.h"
 
diff --git a/Engine/Plugins/Experimental/Mutable/Source/MutableTools/Private/MuT/CodeGenerator.h b/Engine/Plugins/Experimental/Mutable/Source/MutableTools/Private/MuT/CodeGenerator.h
--- a/Engine/Plugins/Experimental/Mutable/Source/MutableTools/Private/MuT/CodeGenerator.h
+++ b/Engine/Plugins/Experimental/Mutable/Source/MutableTools/Private/MuT/CodeGenerator.h
@@ -60,7 +60,6 @@ namespace mu
 	class NodeImageConditional;
 	class NodeImageConstant;
 	class NodeImageReference;
-	class NodeImageDifference;
 	class NodeImageFormat;
 	class NodeImageGradient;
 	class NodeImageInterpolate3;
@@ -76,7 +75,6 @@ namespace mu
 	class NodeImagePlainColour;
 	class NodeImageResize;
 	class NodeImageSaturate;
-	class NodeImageSelectColour;
 	class NodeImageSwitch;
 	class NodeImageSwizzle;
 	class NodeImageTable;
@@ -388,12 +386,10 @@ namespace mu
 		void GenerateImage(FImageGenerationResult& result, const NodeImagePtrConst& node);
 		void GenerateImage_Constant(FImageGenerationResult&, const NodeImageConstant*);
 		void GenerateImage_Reference(FImageGenerationResult&, const NodeImageReference*);
-		void GenerateImage_Difference(FImageGenerationResult&, const NodeImageDifference*);
 		void GenerateImage_Interpolate(FImageGenerationResult&, const NodeImageInterpolate*);
 		void GenerateImage_Saturate(FImageGenerationResult&, const NodeImageSaturate*);
 		void GenerateImage_Table(FImageGenerationResult&, const NodeImageTable*);
 		void GenerateImage_Swizzle(FImageGenerationResult&, const NodeImageSwizzle*);
-		void GenerateImage_SelectColour(FImageGenerationResult&, const NodeImageSelectColour*);
 		void GenerateImage_ColourMap(FImageGenerationResult&, const NodeImageColourMap*);
 		void GenerateImage_Gradient(FImageGenerationResult&, const NodeImageGradient*);
 		void GenerateImage_Binarise(FImageGenerationResult&, const NodeImageBinarise*);
diff --git a/Engine/Plugins/Experimental/Mutable/Source/MutableTools/Private/MuT/CodeGenerator_Image.cpp b/Engine/Plugins/Experimental/Mutable/Source/MutableTools/Private/MuT/CodeGenerator_Image.cpp
--- a/Engine/Plugins/Experimental/Mutable/Source/MutableTools/Private/MuT/CodeGenerator_Image.cpp
+++ b/Engine/Plugins/Experimental/Mutable/Source/MutableTools/Private/MuT/CodeGenerator_Image.cpp
@@ -57,8 +57,6 @@
 #include "MuT/NodeImageConstant.h"
 #include "MuT/NodeImageConstantPrivate.h"
 #include "MuT/NodeImageReferencePrivate.h"
-#include "MuT/NodeImageDifference.h"
-#include "MuT/NodeImageDifferencePrivate.h"
 #include "MuT/NodeImageFormat.h"
 #include "MuT/NodeImageFormatPrivate.h"
 #include "MuT/NodeImageGradient.h"
@@ -91,8 +89,6 @@
 #include "MuT/NodeImageResizePrivate.h"
 #include "MuT/NodeImageSaturate.h"
 #include "MuT/NodeImageSaturatePrivate.h"
-#include "MuT/NodeImageSelectColour.h"
-#include "MuT/NodeImageSelectColourPrivate.h"
 #include "MuT/NodeImageSwitch.h"
 #include "MuT/NodeImageSwitchPrivate.h"
 #include "MuT/NodeImageSwizzle.h"
@@ -158,12 +154,11 @@ namespace mu
 			switch (Untyped->GetImageNodeType())
 			{
 			case NodeImage::EType::Constant: GenerateImage_Constant(result, static_cast<const NodeImageConstant*>(Node)); break;
-			case NodeImage::EType::Difference: GenerateImage_Difference(result, static_cast<const NodeImageDifference*>(Node)); break;
+			case NodeImage::EType::Difference_Deprecated: check(false); break;
 			case NodeImage::EType::Interpolate: GenerateImage_Interpolate(result, static_cast<const NodeImageInterpolate*>(Node)); break;
 			case NodeImage::EType::Saturate: GenerateImage_Saturate(result, static_cast<const NodeImageSaturate*>(Node)); break;
 			case NodeImage::EType::Table: GenerateImage_Table(result, static_cast<const NodeImageTable*>(Node)); break;
 			case NodeImage::EType::Swizzle: GenerateImage_Swizzle(result, static_cast<const NodeImageSwizzle*>(Node)); break;
-			case NodeImage::EType::SelectColour: GenerateImage_SelectColour(result, static_cast<const NodeImageSelectColour*>(Node)); break;
 			case NodeImage::EType::ColourMap: GenerateImage_ColourMap(result, static_cast<const NodeImageColourMap*>(Node)); break;
 			case NodeImage::EType::Gradient: GenerateImage_Gradient(result, static_cast<const NodeImageGradient*>(Node)); break;
 			case NodeImage::EType::Binarise: GenerateImage_Binarise(result, static_cast<const NodeImageBinarise*>(Node)); break;
@@ -247,7 +242,8 @@ namespace mu
         //
         if ( pImage->GetSizeX()!=cropRect.size[0] || pImage->GetSizeY()!=cropRect.size[1] )
         {
-            ImagePtrConst pCropped = ImageCrop( m_compilerOptions->m_imageCompressionQuality, pImage.get(), cropRect );
+            Ptr<Image> pCropped = new Image(cropRect.size[0], cropRect.size[1], 1, pImage->GetFormat());
+			ImageCrop(pCropped.get(), m_compilerOptions->m_imageCompressionQuality, pImage.get(), cropRect );
             op->SetValue( pCropped, m_compilerOptions->m_optimisationOptions.m_useDiskCache );
         }
         else
@@ -972,93 +968,6 @@ namespace mu
     }
 
 
-    //---------------------------------------------------------------------------------------------
-	void CodeGenerator::GenerateImage_SelectColour(FImageGenerationResult& result, const NodeImageSelectColour* InNode)
-	{
-		const NodeImageSelectColour::Private& node = *InNode->GetPrivate();
-
-		Ptr<ASTOpFixed> op = new ASTOpFixed();
-        op->op.type = OP_TYPE::IM_SELECTCOLOUR;
-
-
-        // Source image
-        Ptr<ASTOp> base = 0;
-        if ( Node* pSource = node.m_pSource.get() )
-        {
-            base = Generate( pSource );
-        }
-        else
-        {
-            // This argument is required
-            base = GenerateMissingImageCode( "Image select colour base", EImageFormat::IF_RGB_UBYTE,
-                                             node.m_errorContext );
-        }
-        base = GenerateImageSize
-                ( base, FImageSize((uint16)m_imageState.Last().m_imageRect.size[0],(uint16)m_imageState.Last().m_imageRect.size[1]) );
-        op->SetChild( op->op.args.ImageSelectColour.base, base);
-
-
-        // Factor
-        if ( Node* pColour = node.m_pColour.get() )
-        {
-            op->SetChild( op->op.args.ImageSelectColour.colour, Generate( pColour ));
-        }
-        else
-        {
-            // This argument is required
-            op->SetChild( op->op.args.ImageSelectColour.colour, GenerateMissingColourCode( "Selected colour",
-                                                                          node.m_errorContext ));
-        }
-
-        result.op = op;
-    }
-
-
-    //---------------------------------------------------------------------------------------------
-	void CodeGenerator::GenerateImage_Difference(FImageGenerationResult& result, const NodeImageDifference* InNode)
-	{
-		const NodeImageDifference::Private& node = *InNode->GetPrivate();
-
-		Ptr<ASTOpFixed> op = new ASTOpFixed();
-        op->op.type = OP_TYPE::IM_DIFFERENCE;
-
-        // A image
-        Ptr<ASTOp> a = 0;
-        if ( Node* pA = node.m_pA.get() )
-        {
-            a = Generate( pA );
-        }
-        else
-        {
-            // This argument is required
-            a = GenerateMissingImageCode( "Image Difference A", EImageFormat::IF_RGB_UBYTE, node.m_errorContext );
-        }
-        a = GenerateImageFormat( a, EImageFormat::IF_RGB_UBYTE );
-        a = GenerateImageSize( a,
-                               FImageSize( (uint16)m_imageState.Last().m_imageRect.size[0],
-                                           (uint16)m_imageState.Last().m_imageRect.size[1]) );
-        op->SetChild( op->op.args.ImageDifference.a, a);
-
-        // B image
-        Ptr<ASTOp> b = 0;
-        if ( Node* pB = node.m_pB.get() )
-        {
-            b = Generate( pB );
-        }
-        else
-        {
-            // This argument is required
-            b = GenerateMissingImageCode( "Image Difference B", EImageFormat::IF_RGB_UBYTE, node.m_errorContext );
-        }
-        b = GenerateImageFormat( b, EImageFormat::IF_RGB_UBYTE );
-        b = GenerateImageSize
-                ( b, FImageSize((uint16)m_imageState.Last().m_imageRect.size[0],(uint16)m_imageState.Last().m_imageRect.size[1]) );
-        op->SetChild( op->op.args.ImageDifference.b, b);
-
-        result.op = op;
-    }
-
-
     //---------------------------------------------------------------------------------------------
 	void CodeGenerator::GenerateImage_ColourMap(FImageGenerationResult& result, const NodeImageColourMap* InNode)
 	{
diff --git a/Engine/Plugins/Experimental/Mutable/Source/MutableTools/Private/MuT/CodeOptimiser.cpp b/Engine/Plugins/Experimental/Mutable/Source/MutableTools/Private/MuT/CodeOptimiser.cpp
--- a/Engine/Plugins/Experimental/Mutable/Source/MutableTools/Private/MuT/CodeOptimiser.cpp
+++ b/Engine/Plugins/Experimental/Mutable/Source/MutableTools/Private/MuT/CodeOptimiser.cpp
@@ -453,7 +453,7 @@ namespace mu
 			OP_TYPE type = cloned->GetOpType();
 			DATATYPE dtype = GetOpDataType(type);
 
-			SettingsPtr pSettings = new Settings;
+			Ptr<Settings> pSettings = new Settings;
 			pSettings->SetProfile( false );
 			pSettings->SetImageCompressionQuality( m_imageCompressionQuality );
 			SystemPtr pSystem = new System( pSettings );
@@ -505,6 +505,8 @@ namespace mu
 					constantOp->SetValue( pImage, m_useDiskCache );
 					m_result = constantOp;
 				}
+
+				check(m_result);
 				break;
 			}
 
@@ -542,16 +544,16 @@ namespace mu
 			{
 				MUTABLE_CPUPROFILER_SCOPE(ConstantBool);
 
-				float r=0.0f, g=0.0f, b = 0.0f, a = 0.0f;
-				pSystem->GetPrivate()->BuildColour( model, localParams.get(), at, &r,&g,&b,&a );
+				FVector4f Result(0, 0, 0, 0);
+				Result = pSystem->GetPrivate()->BuildColour( model, localParams.get(), at );
 
 				{
 					mu::Ptr<ASTOpFixed> constantOp = new ASTOpFixed();
 					constantOp->op.type = OP_TYPE::CO_CONSTANT;
-					constantOp->op.args.ColourConstant.value[0] = r;
-					constantOp->op.args.ColourConstant.value[1] = g;
-					constantOp->op.args.ColourConstant.value[2] = b;
-					constantOp->op.args.ColourConstant.value[3] = a;
+					constantOp->op.args.ColourConstant.value[0] = Result[0];
+					constantOp->op.args.ColourConstant.value[1] = Result[1];
+					constantOp->op.args.ColourConstant.value[2] = Result[2];
+					constantOp->op.args.ColourConstant.value[3] = Result[3];
 					m_result = constantOp;
 				}
 				break;
diff --git a/Engine/Plugins/Experimental/Mutable/Source/MutableTools/Private/MuT/ErrorLog.cpp b/Engine/Plugins/Experimental/Mutable/Source/MutableTools/Private/MuT/ErrorLog.cpp
--- a/Engine/Plugins/Experimental/Mutable/Source/MutableTools/Private/MuT/ErrorLog.cpp
+++ b/Engine/Plugins/Experimental/Mutable/Source/MutableTools/Private/MuT/ErrorLog.cpp
@@ -258,7 +258,6 @@ namespace mu
 		TEXT("IM_RESIZEREL     "),
 		TEXT("IM_BLANKLAYOUT   "),
 		TEXT("IM_COMPOSE       "),
-		TEXT("IM_DIFFERENCE    "),
 		TEXT("IM_INTERPOLATE   "),
 		TEXT("IM_INTERPOLATE3  "),
 		TEXT("IM_SATURATE      "),
diff --git a/Engine/Plugins/Experimental/Mutable/Source/MutableTools/Private/MuT/ImageDescGenerator.cpp b/Engine/Plugins/Experimental/Mutable/Source/MutableTools/Private/MuT/ImageDescGenerator.cpp
--- a/Engine/Plugins/Experimental/Mutable/Source/MutableTools/Private/MuT/ImageDescGenerator.cpp
+++ b/Engine/Plugins/Experimental/Mutable/Source/MutableTools/Private/MuT/ImageDescGenerator.cpp
@@ -16,7 +16,6 @@
 #include "MuT/NodeImageColourMapPrivate.h"
 #include "MuT/NodeImageConditionalPrivate.h"
 #include "MuT/NodeImageConstantPrivate.h"
-#include "MuT/NodeImageDifferencePrivate.h"
 #include "MuT/NodeImageFormatPrivate.h"
 #include "MuT/NodeImageGradientPrivate.h"
 #include "MuT/NodeImageInterpolate3Private.h"
@@ -32,7 +31,6 @@
 #include "MuT/NodeImageProjectPrivate.h"
 #include "MuT/NodeImageResizePrivate.h"
 #include "MuT/NodeImageSaturatePrivate.h"
-#include "MuT/NodeImageSelectColourPrivate.h"
 #include "MuT/NodeImageSwitchPrivate.h"
 #include "MuT/NodeImageSwizzlePrivate.h"
 #include "MuT/NodeImageTablePrivate.h"
@@ -334,27 +332,6 @@ namespace mu
 	}
 
 
-    //---------------------------------------------------------------------------------------------
-    Ptr<ASTOp> ImageDescGenerator::Visit( const NodeImageDifference::Private& node )
-    {
-        m_desc = MUTABLE_MISSING_IMAGE_DESC;
-
-        // The first image size has higher priority
-        if ( NodeImage* pA = node.m_pA.get() )
-        {
-            Generate( *pA->GetBasePrivate() );
-        }
-        else if ( NodeImage* pB = node.m_pB.get() )
-        {
-            Generate( *pB->GetBasePrivate() );
-        }
-
-        m_desc.m_format = EImageFormat::IF_L_UBYTE;
-
-        return 0;
-    }
-
-
     //---------------------------------------------------------------------------------------------
     Ptr<ASTOp> ImageDescGenerator::Visit( const NodeImageBinarise::Private& node )
     {
@@ -371,22 +348,6 @@ namespace mu
     }
 
 
-    //---------------------------------------------------------------------------------------------
-    Ptr<ASTOp> ImageDescGenerator::Visit( const NodeImageSelectColour::Private& node )
-    {
-        m_desc = MUTABLE_MISSING_IMAGE_DESC;
-
-        if ( NodeImage* pBase = node.m_pSource.get() )
-        {
-            Generate( *pBase->GetBasePrivate() );
-        }
-
-        m_desc.m_format = EImageFormat::IF_L_UBYTE;
-
-        return 0;
-    }
-
-
     //---------------------------------------------------------------------------------------------
     Ptr<ASTOp> ImageDescGenerator::Visit( const NodeImageColourMap::Private& node )
     {
diff --git a/Engine/Plugins/Experimental/Mutable/Source/MutableTools/Private/MuT/ImageDescGenerator.h b/Engine/Plugins/Experimental/Mutable/Source/MutableTools/Private/MuT/ImageDescGenerator.h
--- a/Engine/Plugins/Experimental/Mutable/Source/MutableTools/Private/MuT/ImageDescGenerator.h
+++ b/Engine/Plugins/Experimental/Mutable/Source/MutableTools/Private/MuT/ImageDescGenerator.h
@@ -10,7 +10,6 @@
 #include "MuT/NodeImageColourMap.h"
 #include "MuT/NodeImageConditional.h"
 #include "MuT/NodeImageConstant.h"
-#include "MuT/NodeImageDifference.h"
 #include "MuT/NodeImageFormat.h"
 #include "MuT/NodeImageGradient.h"
 #include "MuT/NodeImageInterpolate.h"
@@ -27,7 +26,6 @@
 #include "MuT/NodeImageProject.h"
 #include "MuT/NodeImageResize.h"
 #include "MuT/NodeImageSaturate.h"
-#include "MuT/NodeImageSelectColour.h"
 #include "MuT/NodeImageSwitch.h"
 #include "MuT/NodeImageSwizzle.h"
 #include "MuT/NodeImageTable.h"
@@ -53,7 +51,6 @@ namespace mu
 
                                // Standard operations from models
                                public Visitor<NodeImageConstant::Private, Ptr<ASTOp>, true>,
-                               public Visitor<NodeImageDifference::Private, Ptr<ASTOp>, true>,
                                public Visitor<NodeImageInterpolate::Private, Ptr<ASTOp>, true>,
                                public Visitor<NodeImageInterpolate3::Private, Ptr<ASTOp>, true>,
                                public Visitor<NodeImageSaturate::Private, Ptr<ASTOp>, true>,
@@ -63,7 +60,6 @@ namespace mu
                                public Visitor<NodeImageTable::Private, Ptr<ASTOp>, true>,
                                public Visitor<NodeImageSwizzle::Private, Ptr<ASTOp>, true>,
                                public Visitor<NodeImageFormat::Private, Ptr<ASTOp>, true>,
-                               public Visitor<NodeImageSelectColour::Private, Ptr<ASTOp>, true>,
                                public Visitor<NodeImageColourMap::Private, Ptr<ASTOp>, true>,
                                public Visitor<NodeImageGradient::Private, Ptr<ASTOp>, true>,
                                public Visitor<NodeImageBinarise::Private, Ptr<ASTOp>, true>,
@@ -87,7 +83,6 @@ namespace mu
         void Generate( const Node::Private& );
 
         Ptr<ASTOp> Visit( const NodeImageConstant::Private& );
-        Ptr<ASTOp> Visit( const NodeImageDifference::Private& );
         Ptr<ASTOp> Visit( const NodeImageInterpolate::Private& );
         Ptr<ASTOp> Visit( const NodeImageInterpolate3::Private& );
         Ptr<ASTOp> Visit( const NodeImageSaturate::Private& );
@@ -97,7 +92,6 @@ namespace mu
         Ptr<ASTOp> Visit( const NodeImageTable::Private& );
         Ptr<ASTOp> Visit( const NodeImageSwizzle::Private&);
         Ptr<ASTOp> Visit( const NodeImageFormat::Private&);
-        Ptr<ASTOp> Visit( const NodeImageSelectColour::Private& );
         Ptr<ASTOp> Visit( const NodeImageColourMap::Private& );
         Ptr<ASTOp> Visit( const NodeImageGradient::Private& );
         Ptr<ASTOp> Visit( const NodeImageBinarise::Private& );
diff --git a/Engine/Plugins/Experimental/Mutable/Source/MutableTools/Private/MuT/NodeImage.cpp b/Engine/Plugins/Experimental/Mutable/Source/MutableTools/Private/MuT/NodeImage.cpp
--- a/Engine/Plugins/Experimental/Mutable/Source/MutableTools/Private/MuT/NodeImage.cpp
+++ b/Engine/Plugins/Experimental/Mutable/Source/MutableTools/Private/MuT/NodeImage.cpp
@@ -10,7 +10,6 @@
 #include "MuT/NodeImageColourMap.h"
 #include "MuT/NodeImageConditional.h"
 #include "MuT/NodeImageConstant.h"
-#include "MuT/NodeImageDifference.h"
 #include "MuT/NodeImageFormat.h"
 #include "MuT/NodeImageGradient.h"
 #include "MuT/NodeImageInterpolate.h"
@@ -27,7 +26,6 @@
 #include "MuT/NodeImageProject.h"
 #include "MuT/NodeImageResize.h"
 #include "MuT/NodeImageSaturate.h"
-#include "MuT/NodeImageSelectColour.h"
 #include "MuT/NodeImageSwitch.h"
 #include "MuT/NodeImageSwizzle.h"
 #include "MuT/NodeImageTable.h"
@@ -69,14 +67,14 @@ namespace mu
 		switch (id)
 		{
 		case 1 :  return NodeImageConstant::StaticUnserialise( arch ); break;
-		case 2 :  return NodeImageDifference::StaticUnserialise( arch ); break;
+		//case 2 :  return NodeImageDifference::StaticUnserialise( arch ); break;
 		//case 3 :  return NodeImageIdentity::StaticUnserialise( arch ); break;
 		case 4 :  return NodeImageInterpolate::StaticUnserialise( arch ); break;
 		case 6 :  return NodeImageSaturate::StaticUnserialise( arch ); break;
 		//case 8 :  return NodeSelectImage::StaticUnserialise( arch ); break;
 		case 9 :  return NodeImageTable::StaticUnserialise( arch ); break;
 		case 10 :  return NodeImageSwizzle::StaticUnserialise( arch ); break;
-		case 11 :  return NodeImageSelectColour::StaticUnserialise( arch ); break;
+		//case 11 :  return NodeImageSelectColour::StaticUnserialise( arch ); break;
 		case 12 :  return NodeImageColourMap::StaticUnserialise( arch ); break;
 		case 13 :  return NodeImageGradient::StaticUnserialise( arch ); break;
 		//case 14 :  return NodeImageVolumeLayer::StaticUnserialise( arch ); break;
diff --git a/Engine/Plugins/Experimental/Mutable/Source/MutableTools/Public/MuT/NodeImage.h b/Engine/Plugins/Experimental/Mutable/Source/MutableTools/Public/MuT/NodeImage.h
--- a/Engine/Plugins/Experimental/Mutable/Source/MutableTools/Public/MuT/NodeImage.h
+++ b/Engine/Plugins/Experimental/Mutable/Source/MutableTools/Public/MuT/NodeImage.h
@@ -26,12 +26,12 @@ namespace mu
 		enum class EType : uint8
 		{
 			Constant = 1,
-			Difference = 2,
+			Difference_Deprecated = 2,
 			Interpolate = 4,
 			Saturate = 6,
 			Table = 9,
 			Swizzle = 10,
-			SelectColour = 11,
+			SelectColour_Deprecated = 11,
 			ColourMap = 12,
 			Gradient = 13,
 			VolumeLayer_Deprecated = 14,
