[mutable] Fixed method that sets all the parameter values to their default values.
+ Added button to set a specific value to its default value (texture, int, float, color and enum parameters)

#rnx

diff --git a/Engine/Plugins/Experimental/Mutable/Source/CustomizableObject/Classes/MuCO/CustomizableObjectInstance.h b/Engine/Plugins/Experimental/Mutable/Source/CustomizableObject/Classes/MuCO/CustomizableObjectInstance.h
--- a/Engine/Plugins/Experimental/Mutable/Source/CustomizableObject/Classes/MuCO/CustomizableObjectInstance.h
+++ b/Engine/Plugins/Experimental/Mutable/Source/CustomizableObject/Classes/MuCO/CustomizableObjectInstance.h
@@ -322,6 +322,15 @@ public:
 	
 	UFUNCTION(BlueprintCallable, Category = CustomizableObjectInstance)
 	int32 GetProjectorValueRange(const FString& ParamName) const;
+
+	UFUNCTION(BlueprintCallable, Category = CustomizableObjectInstance)
+	int32 GetIntValueRange(const FString& ParamName) const;
+
+	UFUNCTION(BlueprintCallable, Category = CustomizableObjectInstance)
+	int32 GetFloatValueRange(const FString& ParamName) const;
+
+	UFUNCTION(BlueprintCallable, Category = CustomizableObjectInstance)
+	int32 GetTextureValueRange(const FString& ParamName) const;
 	
 	bool bShowOnlyRuntimeParameters = true;
 	bool bShowOnlyRelevantParameters = true;
diff --git a/Engine/Plugins/Experimental/Mutable/Source/CustomizableObject/Private/MuCO/CustomizableObjectInstance.cpp b/Engine/Plugins/Experimental/Mutable/Source/CustomizableObject/Private/MuCO/CustomizableObjectInstance.cpp
--- a/Engine/Plugins/Experimental/Mutable/Source/CustomizableObject/Private/MuCO/CustomizableObjectInstance.cpp
+++ b/Engine/Plugins/Experimental/Mutable/Source/CustomizableObject/Private/MuCO/CustomizableObjectInstance.cpp
@@ -476,12 +476,31 @@ FString UCustomizableObjectInstance::GetDesc()
 	return FString::Printf(TEXT("Instance of [%s]"), *ObjectName);
 }
 
+
 int32 UCustomizableObjectInstance::GetProjectorValueRange(const FString& ParamName) const
 {
 	return Descriptor.GetProjectorValueRange(ParamName);
 }
 
 
+int32 UCustomizableObjectInstance::GetIntValueRange(const FString& ParamName) const
+{
+	return Descriptor.GetIntValueRange(ParamName);
+}
+
+
+int32 UCustomizableObjectInstance::GetFloatValueRange(const FString& ParamName) const
+{
+	return Descriptor.GetFloatValueRange(ParamName);
+}
+
+
+int32 UCustomizableObjectInstance::GetTextureValueRange(const FString& ParamName) const
+{
+	return Descriptor.GetTextureValueRange(ParamName);
+}
+
+
 bool UCustomizableObjectInstance::IsParamMultidimensional(const FString& ParamName) const
 {
 	return GetCustomizableObject()->IsParameterMultidimensional(ParamName);
diff --git a/Engine/Plugins/Experimental/Mutable/Source/CustomizableObject/Private/MuCO/CustomizableObjectInstanceDescriptor.cpp b/Engine/Plugins/Experimental/Mutable/Source/CustomizableObject/Private/MuCO/CustomizableObjectInstanceDescriptor.cpp
--- a/Engine/Plugins/Experimental/Mutable/Source/CustomizableObject/Private/MuCO/CustomizableObjectInstanceDescriptor.cpp
+++ b/Engine/Plugins/Experimental/Mutable/Source/CustomizableObject/Private/MuCO/CustomizableObjectInstanceDescriptor.cpp
@@ -1912,6 +1912,47 @@ int32 FCustomizableObjectInstanceDescriptor::GetProjectorValueRange(const FStrin
 	return ProjectorParameters[ProjectorParamIndex].RangeValues.Num();
 }
 
+int32 FCustomizableObjectInstanceDescriptor::GetIntValueRange(const FString& ParamName) const
+{
+	check(CustomizableObject);
+
+	const int32 IntParamIndex = FindIntParameterNameIndex(ParamName);
+	if (IntParamIndex < 0)
+	{
+		return -1;
+	}
+
+	return IntParameters[IntParamIndex].ParameterRangeValueNames.Num();
+}
+
+
+int32 FCustomizableObjectInstanceDescriptor::GetFloatValueRange(const FString& ParamName) const
+{
+	check(CustomizableObject);
+
+	const int32 FloatParamIndex = FindFloatParameterNameIndex(ParamName);
+	if (FloatParamIndex < 0)
+	{
+		return -1;
+	}
+
+	return FloatParameters[FloatParamIndex].ParameterRangeValues.Num();
+}
+
+
+int32 FCustomizableObjectInstanceDescriptor::GetTextureValueRange(const FString& ParamName) const
+{
+	check(CustomizableObject);
+
+	const int32 TextureParamIndex = FindTextureParameterNameIndex(ParamName);
+	if (TextureParamIndex < 0)
+	{
+		return -1;
+	}
+
+	return TextureParameters[TextureParamIndex].ParameterRangeValues.Num();
+}
+
 
 int32 FCustomizableObjectInstanceDescriptor::AddValueToIntRange(const FString& ParamName)
 {
diff --git a/Engine/Plugins/Experimental/Mutable/Source/CustomizableObject/Public/MuCO/CustomizableObjectInstanceDescriptor.h b/Engine/Plugins/Experimental/Mutable/Source/CustomizableObject/Public/MuCO/CustomizableObjectInstanceDescriptor.h
--- a/Engine/Plugins/Experimental/Mutable/Source/CustomizableObject/Public/MuCO/CustomizableObjectInstanceDescriptor.h
+++ b/Engine/Plugins/Experimental/Mutable/Source/CustomizableObject/Public/MuCO/CustomizableObjectInstanceDescriptor.h
@@ -203,6 +203,15 @@ struct CUSTOMIZABLEOBJECT_API FCustomizableObjectInstanceDescriptor
 	/** Gets the range of values of the projector with ParamName, returns -1 if the parameter does not exist. */
 	int32 GetProjectorValueRange(const FString& ParamName) const;
 
+	/** Gets the range of values of the int with ParamName, returns -1 if the parameter does not exist. */
+	int32 GetIntValueRange(const FString& ParamName) const;
+
+	/** Gets the range of values of the float with ParamName, returns -1 if the parameter does not exist. */
+	int32 GetFloatValueRange(const FString& ParamName) const;
+
+	/** Gets the range of values of the texture with ParamName, returns -1 if the parameter does not exist. */
+	int32 GetTextureValueRange(const FString& ParamName) const;
+
 	/** Increases the range of values of the integer with ParamName, returns the index of the new integer value, -1 otherwise.
 	 * The added value is initialized with the first integer option and is the last one of the range. */
 	int32 AddValueToIntRange(const FString& ParamName);
diff --git a/Engine/Plugins/Experimental/Mutable/Source/CustomizableObjectEditor/Private/MuCOE/SCustomizableInstanceProperties.cpp b/Engine/Plugins/Experimental/Mutable/Source/CustomizableObjectEditor/Private/MuCOE/SCustomizableInstanceProperties.cpp
--- a/Engine/Plugins/Experimental/Mutable/Source/CustomizableObjectEditor/Private/MuCOE/SCustomizableInstanceProperties.cpp
+++ b/Engine/Plugins/Experimental/Mutable/Source/CustomizableObjectEditor/Private/MuCOE/SCustomizableInstanceProperties.cpp
@@ -665,13 +665,29 @@ void SCustomizableInstanceProperties::AddParameter(int32 ParamIndexInObject)
 		.AutoHeight()
 		.HAlign(HAlign_Fill)
 		[
-			SAssignNew(Slider, SSpinBox<float>)
-			.Value(this, &SCustomizableInstanceProperties::GetSliderValue, SliderIndex)
-			.MinValue(CustomizableObject->GetParameterUIMetadataFromIndex(ParamIndexInObject).ParamUIMetadata.MinimumValue)
-			.MaxValue(CustomizableObject->GetParameterUIMetadataFromIndex(ParamIndexInObject).ParamUIMetadata.MaximumValue)
-			.OnValueChanged(this, &SCustomizableInstanceProperties::OnFloatParameterChanged, SliderIndex)
-			.OnBeginSliderMovement(this, &SCustomizableInstanceProperties::OnFloatParameterSliderBegin)
-			.OnEndSliderMovement(this, &SCustomizableInstanceProperties::OnFloatParameterSliderEnd)
+			SNew(SHorizontalBox)
+			+ SHorizontalBox::Slot().FillWidth(0.9f)
+			[
+				SAssignNew(Slider, SSpinBox<float>)
+				.Value(this, &SCustomizableInstanceProperties::GetSliderValue, SliderIndex)
+				.MinValue(CustomizableObject->GetParameterUIMetadataFromIndex(ParamIndexInObject).ParamUIMetadata.MinimumValue)
+				.MaxValue(CustomizableObject->GetParameterUIMetadataFromIndex(ParamIndexInObject).ParamUIMetadata.MaximumValue)
+				.OnValueChanged(this, &SCustomizableInstanceProperties::OnFloatParameterChanged, SliderIndex)
+				.OnBeginSliderMovement(this, &SCustomizableInstanceProperties::OnFloatParameterSliderBegin)
+				.OnEndSliderMovement(this, &SCustomizableInstanceProperties::OnFloatParameterSliderEnd)
+			]
+			
+			+ SHorizontalBox::Slot().AutoWidth().Padding(2.5f, 0.0f, 0.0f, 0.0f)
+			[
+				SNew(SButton)
+				.OnClicked(this, &SCustomizableInstanceProperties::OnResetParameterButtonClicked, ParamIndexInObject)
+				.ButtonStyle(FAppStyle::Get(), "SimpleButton")
+				.Content()
+				[
+					SNew(SImage)
+					.Image(FAppStyle::GetBrush("PropertyWindow.DiffersFromDefault"))
+				]
+			]
 
 		];
 
@@ -722,15 +738,31 @@ void SCustomizableInstanceProperties::AddParameter(int32 ParamIndexInObject)
 			.FillWidth(8.0f)
 			.Padding(0.2f, 2.0f)
 			.HAlign(HAlign_Fill)
-			.VAlign(VAlign_Fill)
+			//.VAlign(VAlign_Fill)
 			[
-				SNew(SMutableTextSearchBox)
-				.ToolTipText(FText::FromString(ToolTipText))
-				.PossibleSuggestions(OptionNamesAttribute)
-				.InitialText(FText::FromString(OptionNamesAttribute[ValueIndex]))
-				.MustMatchPossibleSuggestions(TAttribute<bool>(true))
-				.SuggestionListPlacement(EMenuPlacement::MenuPlacement_ComboBox)
-				.OnTextCommitted(this, &SCustomizableInstanceProperties::OnIntParameterComboBoxChanged, ParamName)
+				SNew(SHorizontalBox)
+				+ SHorizontalBox::Slot().FillWidth(0.9f)
+				[
+					SNew(SMutableTextSearchBox)
+					.ToolTipText(FText::FromString(ToolTipText))
+					.PossibleSuggestions(OptionNamesAttribute)
+					.InitialText(FText::FromString(OptionNamesAttribute[ValueIndex]))
+					.MustMatchPossibleSuggestions(TAttribute<bool>(true))
+					.SuggestionListPlacement(EMenuPlacement::MenuPlacement_ComboBox)
+					.OnTextCommitted(this, &SCustomizableInstanceProperties::OnIntParameterComboBoxChanged, ParamName)
+				]
+
+				+ SHorizontalBox::Slot().AutoWidth().Padding(2.5f, 0.0f, 0.0f, 0.0f)
+				[
+					SNew(SButton)
+					.OnClicked(this, &SCustomizableInstanceProperties::OnResetParameterButtonClicked, ParamIndexInObject)
+					.ButtonStyle(FAppStyle::Get(), "SimpleButton")
+					.Content()
+					[
+						SNew(SImage)
+						.Image(FAppStyle::GetBrush("PropertyWindow.DiffersFromDefault"))
+					]
+				]
 			];
 		}
 		else
@@ -756,13 +788,30 @@ void SCustomizableInstanceProperties::AddParameter(int32 ParamIndexInObject)
 		.FillWidth(10.f)
 		.Padding(2.0f)
 		[
-			SNew(SColorBlock)
-			.Color(this, &SCustomizableInstanceProperties::GetColorParameterValue, ParamName)
-			.ShowBackgroundForAlpha(false)
-			.AlphaDisplayMode(EColorBlockAlphaDisplayMode::Ignore)
-			.UseSRGB(false)
-			.OnMouseButtonDown(this, &SCustomizableInstanceProperties::OnColorBlockMouseButtonDown, ParamName)
-			.Size(FVector2D(10.0f, 10.0f))
+			SNew(SHorizontalBox)
+			+SHorizontalBox::Slot().FillWidth(0.9f)
+			[
+				SNew(SColorBlock)
+				.Color(this, &SCustomizableInstanceProperties::GetColorParameterValue, ParamName)
+				.ShowBackgroundForAlpha(false)
+				.AlphaDisplayMode(EColorBlockAlphaDisplayMode::Ignore)
+				.UseSRGB(false)
+				.OnMouseButtonDown(this, &SCustomizableInstanceProperties::OnColorBlockMouseButtonDown, ParamName)
+				.Size(FVector2D(10.0f, 10.0f))
+			]
+			
+			+ SHorizontalBox::Slot().AutoWidth().Padding(2.5f,0.0f,0.0f,0.0f)
+			[
+				SNew(SButton)
+				.OnClicked(this, &SCustomizableInstanceProperties::OnResetParameterButtonClicked, ParamIndexInObject)
+				.ButtonStyle(FAppStyle::Get(), "SimpleButton")
+				.Content()
+				[
+					SNew(SImage)
+					.Image(FAppStyle::GetBrush("PropertyWindow.DiffersFromDefault"))
+				]
+			]
+
 		];
 
 		break;
@@ -790,10 +839,26 @@ void SCustomizableInstanceProperties::AddParameter(int32 ParamIndexInObject)
 		.FillWidth(1.0f)
 		.Padding(2.0f)
 		[
-			SNew(STextComboBox)
-			.OptionsSource(&TextureParameterValueNames)
-			.InitiallySelectedItem(InitiallySelected)
-			.OnSelectionChanged(this, &SCustomizableInstanceProperties::OnTextureParameterComboBoxSelectionChanged, ParamName)
+			SNew(SHorizontalBox)
+			+ SHorizontalBox::Slot().FillWidth(0.9f)
+			[
+				SNew(STextComboBox)
+				.OptionsSource(&TextureParameterValueNames)
+				.InitiallySelectedItem(InitiallySelected)
+				.OnSelectionChanged(this, &SCustomizableInstanceProperties::OnTextureParameterComboBoxSelectionChanged, ParamName)
+			]
+
+			+ SHorizontalBox::Slot().AutoWidth().Padding(2.5f, 0.0f, 0.0f, 0.0f)
+			[
+				SNew(SButton)
+				.OnClicked(this, &SCustomizableInstanceProperties::OnResetParameterButtonClicked, ParamIndexInObject)
+				.ButtonStyle(FAppStyle::Get(), "SimpleButton")
+				.Content()
+				[
+					SNew(SImage)
+					.Image(FAppStyle::GetBrush("PropertyWindow.DiffersFromDefault"))
+				]
+			]
 		];
 		
 		break;
@@ -1384,26 +1449,11 @@ FReply SCustomizableInstanceProperties::OnResetAllParameters()
 	TArray<FCustomizableObjectIntParameterValue>& IntParameters = CustomInstance->GetIntParameters();
 	
 	const UCustomizableObject* CustomObject = CustomInstance->GetCustomizableObject();
-	for (int32 i = 0; i < IntParameters.Num(); ++i)
+	const int32 NumObjectParameter = CustomObject->GetParameterCount();
+
+	for (int32 ParameterIndex = 0; ParameterIndex < NumObjectParameter; ++ParameterIndex)
 	{
-		const int32 ParamIndexInObject = CustomObject->FindParameter(IntParameters[i].ParameterName);
-		if (ParamIndexInObject != INDEX_NONE)
-		{
-			const int32 NumOptions = CustomObject->GetIntParameterNumOptions(ParamIndexInObject);
-			for (int32 v = 0; v < NumOptions; ++v)
-			{
-				FString ValueName = CustomObject->GetIntParameterAvailableOption(ParamIndexInObject, v);
-				if (ValueName == "None")
-				{
-					IntParameters[i].ParameterValueName = "None";
-				}
-				break;
-			}
-		}
-		else
-		{
-			UE_LOG(LogMutable, Error, TEXT("Failed to find parameter."));
-		}
+		SetParameterValueToDefault(ParameterIndex);
 	}
 
 	CustomInstance->UpdateSkeletalMeshAsync(true, true);
@@ -2266,6 +2316,141 @@ bool SCustomizableInstanceProperties::HasAnyParameters() const
 }
 
 
+bool SCustomizableInstanceProperties::SetParameterValueToDefault(int32 ParameterIndex)
+{
+	UCustomizableObject* CustomObject = CustomInstance->GetCustomizableObject();
+	FString ParameterName = CustomObject->GetParameterName(ParameterIndex);
+	EMutableParameterType ParameterType = CustomObject->GetParameterType(ParameterIndex);
+
+	// Checking if there are parameteres with the same name
+	if (ParameterType != CustomObject->GetParameterTypeByName(ParameterName))
+	{
+		return false;
+	}
+
+	switch (ParameterType)
+	{
+	case EMutableParameterType::None:
+		break;
+	case EMutableParameterType::Bool:
+	{
+		CustomInstance->SetBoolParameterSelectedOption(ParameterName, CustomObject->GetBoolParameterDefaultValue(ParameterName));
+		break;
+	}
+	case EMutableParameterType::Int:
+	{
+		int32 IndexValue = CustomObject->GetIntParameterDefaultValue(ParameterName);
+		FString DefaultValue = CustomObject->GetIntParameterAvailableOption(ParameterIndex, IndexValue);
+
+		if (!DefaultValue.IsEmpty())
+		{
+			if (CustomObject->IsParameterMultidimensional(ParameterName))
+			{
+				int32 NumRanges = CustomInstance->GetIntValueRange(ParameterName);
+
+				for (int32 RangeIndex = 0; RangeIndex < NumRanges; ++RangeIndex)
+				{
+					CustomInstance->SetIntParameterSelectedOption(ParameterName, DefaultValue, RangeIndex);
+				}
+			}
+			else
+			{
+				CustomInstance->SetIntParameterSelectedOption(ParameterName, DefaultValue);
+			}
+		}
+		break;
+	}
+	case EMutableParameterType::Float:
+	{
+		float DefaultValue = CustomObject->GetFloatParameterDefaultValue(ParameterName);
+
+		if (CustomObject->IsParameterMultidimensional(ParameterName))
+		{
+			int32 NumRanges = CustomInstance->GetFloatValueRange(ParameterName);
+
+			for (int32 RangeIndex = 0; RangeIndex < NumRanges; ++RangeIndex)
+			{
+				CustomInstance->SetFloatParameterSelectedOption(ParameterName, DefaultValue, RangeIndex);
+			}
+		}
+		else
+		{
+			CustomInstance->SetFloatParameterSelectedOption(ParameterName, DefaultValue);
+		}
+		break;
+	}
+	case EMutableParameterType::Color:
+	{
+		CustomInstance->SetColorParameterSelectedOption(ParameterName, CustomObject->GetColorParameterDefaultValue(ParameterName));
+		break;
+	}
+	case EMutableParameterType::Projector:
+	{
+		FCustomizableObjectProjector DefaultValue = CustomObject->GetProjectorParameterDefaultValue(ParameterName);
+
+		if (CustomObject->IsParameterMultidimensional(ParameterName))
+		{
+			int32 NumRanges = CustomInstance->GetProjectorValueRange(ParameterName);
+
+			for (int32 RangeIndex = 0; RangeIndex < NumRanges; ++RangeIndex)
+			{
+				CustomInstance->SetProjectorValue(ParameterName, FVector(DefaultValue.Position), FVector(DefaultValue.Direction), FVector(DefaultValue.Up), FVector(DefaultValue.Scale), DefaultValue.Angle, RangeIndex);
+			}
+		}
+		else
+		{
+			CustomInstance->SetProjectorValue(ParameterName, FVector(DefaultValue.Position), FVector(DefaultValue.Direction), FVector(DefaultValue.Up), FVector(DefaultValue.Scale), DefaultValue.Angle);
+		}
+
+		break;
+	}
+	case EMutableParameterType::Texture:
+	{
+		FString DefaultValue = CustomObject->GetTextureParameterDefaultValue(ParameterName);
+
+		if (CustomObject->IsParameterMultidimensional(ParameterName))
+		{
+			int32 NumRanges = CustomInstance->GetTextureValueRange(ParameterName);
+
+			for (int32 RangeIndex = 0; RangeIndex < NumRanges; ++RangeIndex)
+			{
+				CustomInstance->SetTextureParameterSelectedOption(ParameterName, DefaultValue, RangeIndex);
+			}
+		}
+		else
+		{
+			CustomInstance->SetTextureParameterSelectedOption(ParameterName, DefaultValue);
+		}
+		break;
+	}
+	default:
+	{
+		return false;
+		break;
+	}
+	}
+
+	return true;
+}
+
+
+FReply SCustomizableInstanceProperties::OnResetParameterButtonClicked(int32 ParameterIndex)
+{
+	if (UCustomizableObject* CustomObject = CustomInstance->GetCustomizableObject())
+	{
+		CustomInstance->PreEditChange(nullptr);
+		SetParameterValueToDefault(ParameterIndex);
+		CustomInstance->SetSelectedParameterProfileDirty();
+		CustomInstance->UpdateSkeletalMeshAsync(true, true);
+		CustomInstance->PostEditChange();
+
+		return FReply::Handled();
+	}
+
+	return FReply::Unhandled();
+}
+
+
 void SCreateProfileParameters::Construct(const FArguments& InArgs)
 {
 	if (AssetPath.IsEmpty())
diff --git a/Engine/Plugins/Experimental/Mutable/Source/CustomizableObjectEditor/Public/MuCOE/SCustomizableInstanceProperties.h b/Engine/Plugins/Experimental/Mutable/Source/CustomizableObjectEditor/Public/MuCOE/SCustomizableInstanceProperties.h
--- a/Engine/Plugins/Experimental/Mutable/Source/CustomizableObjectEditor/Public/MuCOE/SCustomizableInstanceProperties.h
+++ b/Engine/Plugins/Experimental/Mutable/Source/CustomizableObjectEditor/Public/MuCOE/SCustomizableInstanceProperties.h
@@ -231,6 +231,12 @@ private:
 	/*Search box that filters the Customizable Instance Properties*/
 	TSharedPtr<SSearchBox> SB_Properties;
 
+	/** Sets a parameter to its default value */
+	bool SetParameterValueToDefault(int32 ParameterIndex);
+
+	/** Callback for the OnButtonPressed of the parameter's reset button */
+	FReply OnResetParameterButtonClicked(int32 ParameterIndex);
+
 /*Profiles*/
 
 	/*Open the window to create a new profile*/
