[Mutable] Updated Grammar doc.

#rnx

diff --git a/Engine/Plugins/Experimental/Mutable/Source/CustomizableObjectEditor/Private/MuCOE/Compiler/GrammarImplicits.txt b/Engine/Plugins/Experimental/Mutable/Source/CustomizableObjectEditor/Private/MuCOE/Compiler/GrammarImplicits.txt
--- a/Engine/Plugins/Experimental/Mutable/Source/CustomizableObjectEditor/Private/MuCOE/Compiler/GrammarImplicits.txt
+++ b/Engine/Plugins/Experimental/Mutable/Source/CustomizableObjectEditor/Private/MuCOE/Compiler/GrammarImplicits.txt
@@ -59,6 +59,9 @@ TEXTURE -> NodeTextureVariation TEXTURE TEXTURE*
 TEXTURE -> NodeTextureTransform TEXTURE FLOAT:OffsetX FLOAT:OffsetY FLOAT:ScaleX FLOAT:ScaleY FLOAT:Rotation
 TEXTURE -> NodeTable TEXTUREASSET*
 
+TEXTUREPASSTHROUGH -> NodeTextureSwitch ENUM TEXTUREPASSTHROUGH*
+TEXTUREPASSTHROUGH -> NodePassThroughTexture
+
 COLOR -> NodeColorConstant
 COLOR -> NodeColorParameter
 COLOR -> NodeColorArithmeticOp COLOR:A COLOR:B
@@ -72,7 +75,7 @@ MATERIALASSET -> UMaterial
 
 MATERIALINSTANCE -> NodeTableMaterial MATERIALASSET MATERIALASSET*
 
-SURFACENODE -> NodeMaterial MATERIALASSET MESH:Mesh MATERIALINSTANCE TEXTURE* COLOR* FLOAT*
+SURFACENODE -> NodeMaterial MATERIALASSET MESH:Mesh MATERIALINSTANCE TEXTURE* TEXTUREPASSTHROUGH* COLOR* FLOAT*
 
 SURFACE -> NodeEditMaterial SURFACENODE TEXTURE*
 SURFACE -> NodeExtendMaterial SURFACENODE TEXTURE*
