[Mutable] Fix Edit Material copying old PinData.

- Reverted the Edit Material hotfix.
- Fixed the root issue where PinData was copying old data.
- Added new PinData Copy function which allows to decide which parameters to be copied from old data.

#rnx

diff --git a/Engine/Plugins/Experimental/Mutable/Source/CustomizableObject/Public/MuCO/CustomizableObjectCustomVersion.h b/Engine/Plugins/Experimental/Mutable/Source/CustomizableObject/Public/MuCO/CustomizableObjectCustomVersion.h
--- a/Engine/Plugins/Experimental/Mutable/Source/CustomizableObject/Public/MuCO/CustomizableObjectCustomVersion.h
+++ b/Engine/Plugins/Experimental/Mutable/Source/CustomizableObject/Public/MuCO/CustomizableObjectCustomVersion.h
@@ -79,6 +79,8 @@ struct CUSTOMIZABLEOBJECT_API FCustomizableObjectCustomVersion
 		AutomaticNodeSkeletalMeshPinDataOuter,
 
 		AutomaticNodeSkeletalMeshPinDataUProperty,
+
+		EditMaterialMaskPinDesync, // MTBL-979
 		
 		IgnoreDisabledSections,
 
diff --git a/Engine/Plugins/Experimental/Mutable/Source/CustomizableObjectEditor/Classes/MuCOE/Nodes/CustomizableObjectNode.h b/Engine/Plugins/Experimental/Mutable/Source/CustomizableObjectEditor/Classes/MuCOE/Nodes/CustomizableObjectNode.h
--- a/Engine/Plugins/Experimental/Mutable/Source/CustomizableObjectEditor/Classes/MuCOE/Nodes/CustomizableObjectNode.h
+++ b/Engine/Plugins/Experimental/Mutable/Source/CustomizableObjectEditor/Classes/MuCOE/Nodes/CustomizableObjectNode.h
@@ -30,6 +30,9 @@ public:
 	
 	virtual bool operator!=(const UCustomizableObjectNodePinData& Other) const;
 
+	/** Virtual function used to copy pin data when remapping pins. */
+	virtual void Copy(const UCustomizableObjectNodePinData& Other) {}
+
 protected:
 	/** Virtual function used to perform comparision between different UCustomizableObjectNodePinData inherited types. */
 	virtual bool Equals(const UCustomizableObjectNodePinData& Other) const;
diff --git a/Engine/Plugins/Experimental/Mutable/Source/CustomizableObjectEditor/Classes/MuCOE/Nodes/CustomizableObjectNodeEditMaterial.h b/Engine/Plugins/Experimental/Mutable/Source/CustomizableObjectEditor/Classes/MuCOE/Nodes/CustomizableObjectNodeEditMaterial.h
--- a/Engine/Plugins/Experimental/Mutable/Source/CustomizableObjectEditor/Classes/MuCOE/Nodes/CustomizableObjectNodeEditMaterial.h
+++ b/Engine/Plugins/Experimental/Mutable/Source/CustomizableObjectEditor/Classes/MuCOE/Nodes/CustomizableObjectNodeEditMaterial.h
@@ -53,7 +53,6 @@ public:
 	virtual void BackwardsCompatibleFixup() override;
 	virtual void PostBackwardsCompatibleFixup() override;
 	virtual void AllocateDefaultPins(UCustomizableObjectNodeRemapPins* RemapPins) override;
-	virtual void RemapPins(const TMap<UEdGraphPin*, UEdGraphPin*>& PinsToRemap) override;
 	virtual bool IsNodeOutDatedAndNeedsRefresh() override;
 	virtual FString GetRefreshMessage() const override;
 	virtual bool IsSingleOutputNode() const override;
diff --git a/Engine/Plugins/Experimental/Mutable/Source/CustomizableObjectEditor/Classes/MuCOE/Nodes/CustomizableObjectNodeSkeletalMesh.h b/Engine/Plugins/Experimental/Mutable/Source/CustomizableObjectEditor/Classes/MuCOE/Nodes/CustomizableObjectNodeSkeletalMesh.h
--- a/Engine/Plugins/Experimental/Mutable/Source/CustomizableObjectEditor/Classes/MuCOE/Nodes/CustomizableObjectNodeSkeletalMesh.h
+++ b/Engine/Plugins/Experimental/Mutable/Source/CustomizableObjectEditor/Classes/MuCOE/Nodes/CustomizableObjectNodeSkeletalMesh.h
@@ -213,9 +213,7 @@ public:
 	virtual void BackwardsCompatibleFixup() override;
 	virtual void AllocateDefaultPins(UCustomizableObjectNodeRemapPins* RemapPins) override;
 	virtual UCustomizableObjectNodeRemapPins* CreateRemapPinsDefault() const override;
-	virtual void RemapPinsData(const TMap<UEdGraphPin*, UEdGraphPin*>& PinsToRemap) override;
 
-public:
 	// UCustomizableObjectNodeMesh interface
 	virtual UTexture2D* FindTextureForPin(const UEdGraphPin* Pin) const override;
 	virtual void GetUVChannelForPin(const UEdGraphPin* Pin, TArray<FVector2f>& OutSegments, int32 UVIndex) const override;
diff --git a/Engine/Plugins/Experimental/Mutable/Source/CustomizableObjectEditor/Private/MuCOE/GenerateMutableSource/GenerateMutableSourceSurface.cpp b/Engine/Plugins/Experimental/Mutable/Source/CustomizableObjectEditor/Private/MuCOE/GenerateMutableSource/GenerateMutableSourceSurface.cpp
--- a/Engine/Plugins/Experimental/Mutable/Source/CustomizableObjectEditor/Private/MuCOE/GenerateMutableSource/GenerateMutableSourceSurface.cpp
+++ b/Engine/Plugins/Experimental/Mutable/Source/CustomizableObjectEditor/Private/MuCOE/GenerateMutableSource/GenerateMutableSourceSurface.cpp
@@ -1369,14 +1369,12 @@ mu::NodeSurfacePtr GenerateMutableSourceSurface(const UEdGraphPin * Pin, FMutabl
 					ImagePatchNode->SetImage(ImageNode);
 
 					const UEdGraphPin* ImageMaskPin = TypedNodeEdit->GetUsedImageMaskPin(ImageId);
-
-					if (ImageMaskPin)
+					check(ImageMaskPin); // Ensured when reconstructing EditMaterial nodes. If it fails, something is wrong.
+					
+					if (const UEdGraphPin* ConnectedMaskPin = FollowInputPin(*ImageMaskPin))
 					{
-						if (const UEdGraphPin* ConnectedMaskPin = FollowInputPin(*ImageMaskPin))
-						{
-							mu::NodeImagePtr MaskNode = GenerateMutableSourceImage(ConnectedMaskPin, GenerationContext, 0.f);
-							ImagePatchNode->SetMask(MaskNode);
-						}
+						mu::NodeImagePtr MaskNode = GenerateMutableSourceImage(ConnectedMaskPin, GenerationContext, 0.f);
+						ImagePatchNode->SetMask(MaskNode);
 					}
 
 					// Add the block indices
diff --git a/Engine/Plugins/Experimental/Mutable/Source/CustomizableObjectEditor/Private/MuCOE/Nodes/CustomizableObjectNode.cpp b/Engine/Plugins/Experimental/Mutable/Source/CustomizableObjectEditor/Private/MuCOE/Nodes/CustomizableObjectNode.cpp
--- a/Engine/Plugins/Experimental/Mutable/Source/CustomizableObjectEditor/Private/MuCOE/Nodes/CustomizableObjectNode.cpp
+++ b/Engine/Plugins/Experimental/Mutable/Source/CustomizableObjectEditor/Private/MuCOE/Nodes/CustomizableObjectNode.cpp
@@ -448,9 +448,9 @@ void UCustomizableObjectNode::RemapPinsData(const TMap<UEdGraphPin*, UEdGraphPin
 	for (const TTuple<UEdGraphPin*, UEdGraphPin*>& Pair : PinsToRemap)
 	{
 		// Move pin data.
-		if (TObjectPtr<UCustomizableObjectNodePinData>* PinDataOldPin = PinsDataId.Find(Pair.Key->PinId))
+		if (const TObjectPtr<UCustomizableObjectNodePinData>* PinDataOldPin = PinsDataId.Find(Pair.Key->PinId))
 		{
-			PinsDataId[Pair.Value->PinId] = *PinDataOldPin;		
+			PinsDataId[Pair.Value->PinId]->Copy(**PinDataOldPin);
 		}
 	}
 }
diff --git a/Engine/Plugins/Experimental/Mutable/Source/CustomizableObjectEditor/Private/MuCOE/Nodes/CustomizableObjectNodeEditMaterial.cpp b/Engine/Plugins/Experimental/Mutable/Source/CustomizableObjectEditor/Private/MuCOE/Nodes/CustomizableObjectNodeEditMaterial.cpp
--- a/Engine/Plugins/Experimental/Mutable/Source/CustomizableObjectEditor/Private/MuCOE/Nodes/CustomizableObjectNodeEditMaterial.cpp
+++ b/Engine/Plugins/Experimental/Mutable/Source/CustomizableObjectEditor/Private/MuCOE/Nodes/CustomizableObjectNodeEditMaterial.cpp
@@ -49,50 +49,6 @@ void UCustomizableObjectNodeEditMaterial::AllocateDefaultPins(UCustomizableObjec
 }
 
 
-void UCustomizableObjectNodeEditMaterial::RemapPins(const TMap<UEdGraphPin*, UEdGraphPin*>& PinsToRemap)
-{
-	Super::RemapPins(PinsToRemap);
-
-	// Update mask pin references.
-	for (const UEdGraphPin* Pin : GetAllPins())
-	{
-		if (UCustomizableObjectNodeEditMaterialPinEditImageData* PinData = Cast<UCustomizableObjectNodeEditMaterialPinEditImageData>(GetPinData(*Pin)))
-		{
-			if (UEdGraphPin* const* Result = PinsToRemap.Find(PinData->PinMask.Get()))
-			{
-				PinData->PinMask = FEdGraphPinReference(*Result);
-			}
-
-			FName PinMaskName = FName(Pin->PinName.ToString() + FString(" Mask"));
-
-			if (!PinData->PinMask.Get())
-			{
-				// Something went wrong, an Edit Material Image pin should have an associated mask pin, try to find it
-				for (const UEdGraphPin* PotentialMaskPin : GetAllPins())
-				{
-					if (PotentialMaskPin->PinName == PinMaskName)
-					{
-						PinData->PinMask = PotentialMaskPin;
-						break;
-					}
-				}
-			}
-
-			if (!PinData->PinMask.Get())
-			{
-				// The mask pin didn't actually exist, so create it
-				const UEdGraphSchema_CustomizableObject* Schema = GetDefault<UEdGraphSchema_CustomizableObject>();
-				UEdGraphPin* PinMask = CustomCreatePin(EGPD_Input, Schema->PC_Image, PinMaskName);
-				PinMask->bHidden = Pin->bHidden;
-				PinMask->bDefaultValueIsIgnored = true;
-
-				PinData->PinMask = FEdGraphPinReference(PinMask);
-			}
-		}
-	}
-}
-
-
 const UEdGraphPin* UCustomizableObjectNodeEditMaterial::GetUsedImageMaskPin(const FGuid& ImageId) const
 {
 	if (const UEdGraphPin* Pin = GetUsedImagePin(ImageId))
@@ -240,6 +196,11 @@ void UCustomizableObjectNodeEditMaterial::BackwardsCompatibleFixup()
 	{
 		ReconstructNode();
 	}
+
+	if (CustomizableObjectCustomVersion < FCustomizableObjectCustomVersion::EditMaterialMaskPinDesync)
+	{
+		ReconstructNode();
+	}
 }
 
 
diff --git a/Engine/Plugins/Experimental/Mutable/Source/CustomizableObjectEditor/Private/MuCOE/Nodes/CustomizableObjectNodeSkeletalMesh.cpp b/Engine/Plugins/Experimental/Mutable/Source/CustomizableObjectEditor/Private/MuCOE/Nodes/CustomizableObjectNodeSkeletalMesh.cpp
--- a/Engine/Plugins/Experimental/Mutable/Source/CustomizableObjectEditor/Private/MuCOE/Nodes/CustomizableObjectNodeSkeletalMesh.cpp
+++ b/Engine/Plugins/Experimental/Mutable/Source/CustomizableObjectEditor/Private/MuCOE/Nodes/CustomizableObjectNodeSkeletalMesh.cpp
@@ -149,12 +149,6 @@ UCustomizableObjectNodeRemapPins* UCustomizableObjectNodeSkeletalMesh::CreateRem
 }
 
 
-void UCustomizableObjectNodeSkeletalMesh::RemapPinsData(const TMap<UEdGraphPin*, UEdGraphPin*>& PinsToRemap)
-{
-	// Do not copy pin data from old pins
-}
-
-
 FText UCustomizableObjectNodeSkeletalMesh::GetNodeTitle(ENodeTitleType::Type TitleType) const
 {
 	if (SkeletalMesh)
