#rnx

diff --git a/Engine/Plugins/Experimental/Mutable/Config/DefaultMutable.ini b/Engine/Plugins/Experimental/Mutable/Config/DefaultMutable.ini
--- a/Engine/Plugins/Experimental/Mutable/Config/DefaultMutable.ini
+++ b/Engine/Plugins/Experimental/Mutable/Config/DefaultMutable.ini
@@ -6,8 +6,16 @@
 
 [Features]
 
-; If this is enabled, CustomizableObjects will be allowed to use 8 bone influences per vertex. 
-; Otherwise only the default 4 will be used. This is used at object compile time.
+; If CustomizableObjectNumBoneInfluences is set to N, CustomizableObjects will be allowed to use a max of N bone influences 
+; per vertex. The only valid values are 4, 8, 12, Four, Eight, Twelve. 
+; If r.GPUSkin.UnlimitedBoneInfluencesThreshold is set to a lower number, then only that number of influences
+; will always be used and allocated thus saving memory per vertex when the real number is higher for a vertex.
+; If CustomizableObjectNumBoneInfluences is not defined, only the default 4 will be used. This is used at object compile time.
+;CustomizableObjectNumBoneInfluences=Eight
+
+; For backwards compatibility bExtraBoneInfluencesEnabled can be used instead, but will be overridden by 
+; CustomizableObjectNumBoneInfluences if both are present. If enabled CustomizableObjects will be allowed to 
+; use 8 bone influences per vertex instead of the default 4. This is used at object compile time.
 bExtraBoneInfluencesEnabled=False
 
 ;--------------------------------------------------------------------------------------------------
diff --git a/Engine/Plugins/Experimental/Mutable/Source/CustomizableObject/Classes/MuCO/CustomizableObject.h b/Engine/Plugins/Experimental/Mutable/Source/CustomizableObject/Classes/MuCO/CustomizableObject.h
--- a/Engine/Plugins/Experimental/Mutable/Source/CustomizableObject/Classes/MuCO/CustomizableObject.h
+++ b/Engine/Plugins/Experimental/Mutable/Source/CustomizableObject/Classes/MuCO/CustomizableObject.h
@@ -122,6 +122,17 @@ struct FProfileParameterDat
 	TArray<FCustomizableObjectProjectorParameterValue> ProjectorParameters;
 };
 
+UENUM()
+enum class ECustomizableObjectNumBoneInfluences : uint8
+{
+	// The enum values can be used as the real numeric value of number of bone influences
+	Four = 4,
+	// 
+	Eight = 8,
+	//
+	Twelve = 12 // This is essentially the same as "Unlimited", but UE ultimately limits to 12
+};
+
 USTRUCT()
 struct FCompilationOptions
 {
@@ -147,7 +158,7 @@ struct FCompilationOptions
 	bool bUseDiskCompilation = false;
 
 	// Did we have the extra bones enabled when we compiled?
-	bool bExtraBoneInfluencesEnabled = false;
+	ECustomizableObjectNumBoneInfluences CustomizableObjectNumBoneInfluences = ECustomizableObjectNumBoneInfluences::Four;
 
 	// Compiling for cook
 	bool bIsCooking = false;
diff --git a/Engine/Plugins/Experimental/Mutable/Source/CustomizableObject/Private/MuCO/CustomizableObject.cpp b/Engine/Plugins/Experimental/Mutable/Source/CustomizableObject/Private/MuCO/CustomizableObject.cpp
--- a/Engine/Plugins/Experimental/Mutable/Source/CustomizableObject/Private/MuCO/CustomizableObject.cpp
+++ b/Engine/Plugins/Experimental/Mutable/Source/CustomizableObject/Private/MuCO/CustomizableObject.cpp
@@ -718,7 +718,7 @@ void UCustomizableObject::CompileForTargetPlatform(const ITargetPlatform* Target
 		Options.bIsCooking = true;
 		Options.bSaveCookedDataToDisk = bUsesOnCookStart;
 		Options.TargetPlatform = TargetPlatform;
-		Options.bExtraBoneInfluencesEnabled = ICustomizableObjectModule::Get().AreExtraBoneInfluencesEnabled();
+		Options.CustomizableObjectNumBoneInfluences = ICustomizableObjectModule::Get().GetNumBoneInfluences();
 
 		// If this is enabled, there are determinism problems. Disable it for packaging.
 		Options.bUseParallelCompilation = false;
diff --git a/Engine/Plugins/Experimental/Mutable/Source/CustomizableObject/Private/MuCO/CustomizableObjectModule.cpp b/Engine/Plugins/Experimental/Mutable/Source/CustomizableObject/Private/MuCO/CustomizableObjectModule.cpp
--- a/Engine/Plugins/Experimental/Mutable/Source/CustomizableObject/Private/MuCO/CustomizableObjectModule.cpp
+++ b/Engine/Plugins/Experimental/Mutable/Source/CustomizableObject/Private/MuCO/CustomizableObjectModule.cpp
@@ -12,6 +12,7 @@
 #include "MuCO/CustomizableSkeletalComponent.h"
 #include "MuCO/ICustomizableObjectModule.h"
 #include "UObject/StrongObjectPtr.h"
+#include "GPUSkinPublicDefs.h"
 
 /**
  * Customizable Object module implementation (private)
@@ -26,7 +27,7 @@ public:
 
 	// ICustomizableObjectModule interface
 	FString GetPluginVersion() const override;
-	bool AreExtraBoneInfluencesEnabled() const override;
+	ECustomizableObjectNumBoneInfluences GetNumBoneInfluences() const override;
 	void RegisterExtension(TObjectPtr<const UCustomizableObjectExtension> Extension) override;
 	void UnregisterExtension(TObjectPtr<const UCustomizableObjectExtension> Extension) override;
 	TArrayView<const TObjectPtr<const UCustomizableObjectExtension>> GetRegisteredExtensions() const override;
@@ -77,23 +78,57 @@ FString FCustomizableObjectModule::GetPluginVersion() const
 	return Version;
 }
 
-bool FCustomizableObjectModule::AreExtraBoneInfluencesEnabled() const
+
+ECustomizableObjectNumBoneInfluences FCustomizableObjectModule::GetNumBoneInfluences() const
 {
 	bool bAreExtraBoneInfluencesEnabled = false;
 
+#if WITH_EDITOR
+	ensure((int32)ECustomizableObjectNumBoneInfluences::Eight == EXTRA_BONE_INFLUENCES);
+	ensure((int32)ECustomizableObjectNumBoneInfluences::Twelve == MAX_TOTAL_INFLUENCES);
+#endif
+
 	FConfigFile* PluginConfig = GConfig->FindConfigFileWithBaseName("Mutable");
 	if (PluginConfig)
 	{
+		FString Value;
+
+		if (PluginConfig->GetString(TEXT("Features"), TEXT("CustomizableObjectNumBoneInfluences"), Value))
+		{
+			int32 NumInfluences = Value.IsNumeric() ? FCString::Atoi(*Value) : -1;
+
+			if (NumInfluences == 4 || Value.Equals(FString("Four"), ESearchCase::IgnoreCase))
+			{
+				return ECustomizableObjectNumBoneInfluences::Four;
+			}
+			else if (NumInfluences == 8 || Value.Equals(FString("Eight"), ESearchCase::IgnoreCase))
+			{
+				return ECustomizableObjectNumBoneInfluences::Eight;
+			}
+			else if (NumInfluences == 12 || Value.Equals(FString("Twelve"), ESearchCase::IgnoreCase))
+			{
+				return ECustomizableObjectNumBoneInfluences::Twelve;
+			}
+
+			UE_LOG(LogMutable, Warning, TEXT("The Mutable Plugin config. variable CustomizableObjectNumBoneInfluences has the invalid value [%s]."
+				"Only 4, 8, 12, Four, Eight, Twelve are valid values."
+				), *Value);
+		}
+
 		bool bValue = false;
 		if (PluginConfig->GetBool(TEXT("Features"), TEXT("bExtraBoneInfluencesEnabled"), bValue))
 		{
-			bAreExtraBoneInfluencesEnabled = bValue;
+			if (bValue)
+			{
+				return ECustomizableObjectNumBoneInfluences::Eight;
+			}
 		}
 	}
 
-	return bAreExtraBoneInfluencesEnabled;
+	return ECustomizableObjectNumBoneInfluences::Four;
 }
 
+
 void FCustomizableObjectModule::RegisterExtension(TObjectPtr<const UCustomizableObjectExtension> Extension)
 {
 	check(IsInGameThread());
diff --git a/Engine/Plugins/Experimental/Mutable/Source/CustomizableObject/Private/MuCO/UnrealConversionUtils.cpp b/Engine/Plugins/Experimental/Mutable/Source/CustomizableObject/Private/MuCO/UnrealConversionUtils.cpp
--- a/Engine/Plugins/Experimental/Mutable/Source/CustomizableObject/Private/MuCO/UnrealConversionUtils.cpp
+++ b/Engine/Plugins/Experimental/Mutable/Source/CustomizableObject/Private/MuCO/UnrealConversionUtils.cpp
@@ -5,6 +5,7 @@
 #include "MuCO/CustomizableObjectInstance.h"
 #include "MuCO/UnrealPortabilityHelpers.h"
 #include "Engine/SkeletalMesh.h"
+#include "GPUSkinVertexFactory.h"
 
 class USkeleton;
 
@@ -80,19 +81,42 @@ namespace UnrealConversionUtils
 			const int32 NumBones,
 			const int32 NumBoneInfluences,
 			const bool bNeedCPUAccess,
-			const void* InMutableData)
+			const void* InMutableData,
+			const uint32 MutableDataSize)
 		{
-			// \todo Ugly cast.
-			FSkinWeightDataVertexBuffer* VertexBuffer = const_cast<FSkinWeightDataVertexBuffer*>(OutVertexWeightBuffer.GetDataVertexBuffer());
+			FSkinWeightDataVertexBuffer* VertexBuffer = OutVertexWeightBuffer.GetDataVertexBuffer();
 			VertexBuffer->SetMaxBoneInfluences(NumBoneInfluences);
 			VertexBuffer->Init(NumBones, NumVertices);
 
+			bool bIsVariableBonesPerVertex = OutVertexWeightBuffer.GetDataVertexBuffer()->GetVariableBonesPerVertex();
+			check(!FGPUBaseSkinVertexFactory::UseUnlimitedBoneInfluences(NumBoneInfluences) || bIsVariableBonesPerVertex);
+
+			OutVertexWeightBuffer.SetNeedsCPUAccess(bNeedCPUAccess);
+
 			if (NumVertices)
 			{
-				OutVertexWeightBuffer.SetNeedsCPUAccess(bNeedCPUAccess);
-
 				uint8* Data = VertexBuffer->GetWeightData();
-				FMemory::Memcpy(Data, InMutableData, OutVertexWeightBuffer.GetVertexDataSize());
+				uint32 OutVertexWeightBufferSize = OutVertexWeightBuffer.GetDataVertexBuffer()->GetVertexDataSize();
+				ensure(MutableDataSize == OutVertexWeightBufferSize);
+
+				FMemory::Memcpy(Data, InMutableData, OutVertexWeightBufferSize);
+
+				if (bIsVariableBonesPerVertex)
+				{
+					OutVertexWeightBuffer.RebuildLookupVertexBuffer();
+
+					{
+						MUTABLE_CPUPROFILER_SCOPE(OptimizeVertexAndLookupBuffers);
+
+						// Everything in this scope is optional and makes extra copies, but will optimize the variable bone
+						// influences buffers. Without it, the vertices are assumed to have a constant NumBoneInfluences per vertex.
+						TArray<FSkinWeightInfo> TempVertices;
+						OutVertexWeightBuffer.GetSkinWeights(TempVertices);
+
+						// The assignment operator actually optimizes the DataVertexBuffer
+						OutVertexWeightBuffer = TempVertices;
+					}
+				}
 			}
 		}
 	}
@@ -231,7 +255,8 @@ namespace UnrealConversionUtils
 			NumBoneInfluences * NumVertices,
 			NumBoneInfluences,
 			bNeedsCPUAccess,
-			MutableMeshVertexBuffers.GetBufferData(BoneIndexBuffer)
+			MutableMeshVertexBuffers.GetBufferData(BoneIndexBuffer),
+			MutableMeshVertexBuffers.GetBufferDataSize(BoneIndexBuffer)
 		);
 
 		// Optional buffers
@@ -251,14 +276,16 @@ namespace UnrealConversionUtils
 					const int32 BonesPerVertex = ComponentCount;
 					const int32 NumBones = BonesPerVertex * NumVertices;
 
-					check(LODModel.SkinWeightVertexBuffer.GetVariableBonesPerVertex() == false);
+					check(FGPUBaseSkinVertexFactory::UseUnlimitedBoneInfluences(BonesPerVertex) ||
+						!LODModel.SkinWeightVertexBuffer.GetVariableBonesPerVertex());
 					FSkinWeightVertexBuffer_InitWithMutableData(
 						LODModel.SkinWeightVertexBuffer,
 						NumVertices,
 						NumBones,
 						NumBoneInfluences,
 						bNeedsCPUAccess,
-						MutableMeshVertexBuffers.GetBufferData(Buffer)
+						MutableMeshVertexBuffers.GetBufferData(Buffer),
+						MutableMeshVertexBuffers.GetBufferDataSize(Buffer)
 					);
 				}
 
diff --git a/Engine/Plugins/Experimental/Mutable/Source/CustomizableObject/Public/MuCO/ICustomizableObjectModule.h b/Engine/Plugins/Experimental/Mutable/Source/CustomizableObject/Public/MuCO/ICustomizableObjectModule.h
--- a/Engine/Plugins/Experimental/Mutable/Source/CustomizableObject/Public/MuCO/ICustomizableObjectModule.h
+++ b/Engine/Plugins/Experimental/Mutable/Source/CustomizableObject/Public/MuCO/ICustomizableObjectModule.h
@@ -5,6 +5,7 @@
 #include "Containers/ArrayView.h"
 #include "Modules/ModuleManager.h"
 #include "MuCO/CustomizableObjectExtension.h"
+#include "MuCO/CustomizableObject.h"
 
 class UCustomizableSkeletalComponent;
 
@@ -51,8 +52,8 @@ public:
 		return FModuleManager::Get().IsModuleLoaded( "CustomizableObject" );
 	}
 
-	// Return true if all engine patches are present
-	virtual bool AreExtraBoneInfluencesEnabled() const = 0;
+	// Returns the number of bone influences that Mutable will use according to the plugin conf
+	virtual ECustomizableObjectNumBoneInfluences GetNumBoneInfluences() const = 0;
 
 	// Return a string representing the plugin version.
 	virtual FString GetPluginVersion() const = 0;
diff --git a/Engine/Plugins/Experimental/Mutable/Source/CustomizableObjectEditor/Private/MuCOE/CustomizableObjectCompiler.cpp b/Engine/Plugins/Experimental/Mutable/Source/CustomizableObjectEditor/Private/MuCOE/CustomizableObjectCompiler.cpp
--- a/Engine/Plugins/Experimental/Mutable/Source/CustomizableObjectEditor/Private/MuCOE/CustomizableObjectCompiler.cpp
+++ b/Engine/Plugins/Experimental/Mutable/Source/CustomizableObjectEditor/Private/MuCOE/CustomizableObjectCompiler.cpp
@@ -39,7 +39,7 @@ FCustomizableObjectCompiler::FCustomizableObjectCompiler() : FCustomizableObject
 	, CompletedUnrealToMutableTask(0)
 	, MaxConvertToMutableTextureTime(0.2f)
 {
-	bAreExtraBoneInfluencesEnabled = ICustomizableObjectModule::Get().AreExtraBoneInfluencesEnabled();
+	CustomizableObjectNumBoneInfluences = ICustomizableObjectModule::Get().GetNumBoneInfluences();
 }
 
 
@@ -919,7 +919,7 @@ void FCustomizableObjectCompiler::CompileInternal(UCustomizableObject* Object, c
 	if (!Object || !CurrentObject) return;
 
 	Options = InOptions;
-	Options.bExtraBoneInfluencesEnabled = bAreExtraBoneInfluencesEnabled;
+	Options.CustomizableObjectNumBoneInfluences = CustomizableObjectNumBoneInfluences;
 	Options.bRealTimeMorphTargetsEnabled = Object->bEnableRealTimeMorphTargets;
 	Options.bClothingEnabled = Object->bEnableClothing;
 	Options.b16BitBoneWeightsEnabled = Object->bEnable16BitBoneWeights;
@@ -1366,7 +1366,7 @@ mu::NodePtr FCustomizableObjectCompiler::Export(UCustomizableObject* Object, con
 	FSlateNotificationManager::Get().AddNotification(Info);
 
 	FCompilationOptions CompilerOptions = InCompilerOptions;
-	CompilerOptions.bExtraBoneInfluencesEnabled = bAreExtraBoneInfluencesEnabled;
+	CompilerOptions.CustomizableObjectNumBoneInfluences = CustomizableObjectNumBoneInfluences;
 		
 	FMutableGraphGenerationContext GenerationContext(Object, this, CompilerOptions);
 	GenerationContext.ParamNamesToSelectedOptions = ParamNamesToSelectedOptions;
diff --git a/Engine/Plugins/Experimental/Mutable/Source/CustomizableObjectEditor/Private/MuCOE/CustomizableObjectCompiler.h b/Engine/Plugins/Experimental/Mutable/Source/CustomizableObjectEditor/Private/MuCOE/CustomizableObjectCompiler.h
--- a/Engine/Plugins/Experimental/Mutable/Source/CustomizableObjectEditor/Private/MuCOE/CustomizableObjectCompiler.h
+++ b/Engine/Plugins/Experimental/Mutable/Source/CustomizableObjectEditor/Private/MuCOE/CustomizableObjectCompiler.h
@@ -123,7 +123,7 @@ private:
 	TSharedPtr< FRunnableThread > SaveDDThread;
 
 	// Cache configuration settings from ini files 
-	bool bAreExtraBoneInfluencesEnabled = false;
+	ECustomizableObjectNumBoneInfluences CustomizableObjectNumBoneInfluences = ECustomizableObjectNumBoneInfluences::Four;
 
 	// TODO: add references
 	UCustomizableObject* CurrentObject = nullptr;
diff --git a/Engine/Plugins/Experimental/Mutable/Source/CustomizableObjectEditor/Private/MuCOE/GenerateMutableSource/GenerateMutableSourceMesh.cpp b/Engine/Plugins/Experimental/Mutable/Source/CustomizableObjectEditor/Private/MuCOE/GenerateMutableSource/GenerateMutableSourceMesh.cpp
--- a/Engine/Plugins/Experimental/Mutable/Source/CustomizableObjectEditor/Private/MuCOE/GenerateMutableSource/GenerateMutableSourceMesh.cpp
+++ b/Engine/Plugins/Experimental/Mutable/Source/CustomizableObjectEditor/Private/MuCOE/GenerateMutableSource/GenerateMutableSourceMesh.cpp
@@ -897,10 +897,22 @@ mu::MeshPtr ConvertSkeletalMeshToMutable(USkeletalMesh* InSkeletalMesh, const TS
 		const MESH_BUFFER_FORMAT Formats[ChannelCount] = { MBF_FLOAT32, MBF_FLOAT32, MBF_FLOAT32, MBF_FLOAT32, MBF_FLOAT32, MBF_FLOAT32, MBF_FLOAT32, MBF_FLOAT32, MBF_NUINT8, MBF_UINT16, BoneWeightFormat };
 
 		int Components[ChannelCount] = { 3, 3, 3, 4, 2, 2, 2, 2, 4, 4, 4 };
-		if (GenerationContext.Options.bExtraBoneInfluencesEnabled && MaxSectionInfluences > 4)
+		if (GenerationContext.Options.CustomizableObjectNumBoneInfluences != ECustomizableObjectNumBoneInfluences::Four && 
+			MaxSectionInfluences > 4)
 		{
-			Components[9] = EXTRA_BONE_INFLUENCES;
-			Components[10] = EXTRA_BONE_INFLUENCES;
+			int32 NewBoneInfluencesNum = (int32)GenerationContext.Options.CustomizableObjectNumBoneInfluences;
+
+			if (FGPUBaseSkinVertexFactory::UseUnlimitedBoneInfluences(MaxSectionInfluences, GenerationContext.Options.TargetPlatform) &&
+				MaxSectionInfluences < NewBoneInfluencesNum)
+			{
+				Components[9] = MaxSectionInfluences;
+				Components[10] = MaxSectionInfluences;
+			}
+			else
+			{
+				Components[9] = NewBoneInfluencesNum;
+				Components[10] = NewBoneInfluencesNum;
+			}
 		}
 
 		constexpr size_t SoftSkinVertexUVsElemSize = sizeof(TDecay<decltype(DeclVal<FSoftSkinVertex>().UVs[0])>::Type);
@@ -955,7 +967,7 @@ mu::MeshPtr ConvertSkeletalMeshToMutable(USkeletalMesh* InSkeletalMesh, const TS
 				TransferRemovedBonesInfluences(&Vertex.InfluenceBones[0], &Vertex.InfluenceWeights[0], MaxSectionInfluences, RemappedBoneMapIndices);
 			}
 
-			if (!GenerationContext.Options.bExtraBoneInfluencesEnabled)
+			if (GenerationContext.Options.CustomizableObjectNumBoneInfluences == ECustomizableObjectNumBoneInfluences::Four)
 			{
 				// Normalize weights
 				const int32 MaxMutableWeights = 4;
@@ -965,9 +977,21 @@ mu::MeshPtr ConvertSkeletalMeshToMutable(USkeletalMesh* InSkeletalMesh, const TS
 				NormalizeWeights(&Vertex.InfluenceBones[0], &Vertex.InfluenceWeights[0], MaxSectionInfluences, MaxMutableWeights,
 					&MaxOrderedWeighsIndices[0], MaxSectionBoneMapIndex, MaxBoneWeightValue);
 			}
-			else if (MaxSectionInfluences < EXTRA_BONE_INFLUENCES)
+			else if (GenerationContext.Options.CustomizableObjectNumBoneInfluences == ECustomizableObjectNumBoneInfluences::Eight)
+			{
+				if (!FGPUBaseSkinVertexFactory::UseUnlimitedBoneInfluences(MaxSectionInfluences, GenerationContext.Options.TargetPlatform) &&
+					MaxSectionInfluences < EXTRA_BONE_INFLUENCES) // EXTRA_BONE_INFLUENCES is ECustomizableObjectNumBoneInfluences::Eight
+				{
+					FMemory::Memzero(&Vertex.InfluenceWeights[MaxSectionInfluences], EXTRA_BONE_INFLUENCES - MaxSectionInfluences);
+				}
+			}
+			else if (GenerationContext.Options.CustomizableObjectNumBoneInfluences == ECustomizableObjectNumBoneInfluences::Twelve)
 			{
-				FMemory::Memzero(&Vertex.InfluenceWeights[MaxSectionInfluences], EXTRA_BONE_INFLUENCES - MaxSectionInfluences);
+				if (!FGPUBaseSkinVertexFactory::UseUnlimitedBoneInfluences(MaxSectionInfluences, GenerationContext.Options.TargetPlatform) &&
+					MaxSectionInfluences < MAX_TOTAL_INFLUENCES) // MAX_TOTAL_INFLUENCES is ECustomizableObjectNumBoneInfluences::Twelve
+				{
+					FMemory::Memzero(&Vertex.InfluenceWeights[MaxSectionInfluences], MAX_TOTAL_INFLUENCES - MaxSectionInfluences);
+				}
 			}
 		}
 
@@ -1348,13 +1372,15 @@ mu::MeshPtr ConvertSkeletalMeshToMutable(USkeletalMesh* InSkeletalMesh, const TS
 		MESH_BUFFER_FORMAT BoneWeightFormat = BoneWeightTypeSizeBytes == 1 ? MBF_NUINT8 : MBF_NUINT16;
 
 		// Limit skinning weights if necessary
-		// \todo: make it more flexible to support 3 or 5 or 1 weight, since there is support for this in 4.25
-		const int32 MutableBonesPerVertex = GenerationContext.Options.bExtraBoneInfluencesEnabled ? EXTRA_BONE_INFLUENCES : 4;
+		const int32 MaxSectionInfluences = ImportedModel->LODModels[LOD].Sections[MaterialIndex].GetMaxBoneInfluences();
+		const int32 MutableBonesPerVertex = FGPUBaseSkinVertexFactory::UseUnlimitedBoneInfluences(MaxSectionInfluences, 
+												GenerationContext.Options.TargetPlatform) ?
+											MaxSectionInfluences :
+											(int32)GenerationContext.Options.CustomizableObjectNumBoneInfluences;
 		const int32 BoneIndicesSize = MutableBonesPerVertex * sizeof(FBoneIndexType);
 		const int32 BoneWeightsSize = MutableBonesPerVertex * BoneWeightTypeSizeBytes;
 		const int32 SkinWeightProfileVertexSize = sizeof(int32) + BoneIndicesSize + BoneWeightsSize;
 
-		const int32 MaxSectionInfluences = ImportedModel->LODModels[LOD].Sections[MaterialIndex].GetMaxBoneInfluences();
 		const int32 MaxSectionBoneMapIndex = ImportedModel->LODModels[LOD].Sections[MaterialIndex].BoneMap.Num();
 
 		const TArray<FSkinWeightProfileInfo>& SkinWeightProfilesInfo = InSkeletalMesh->GetSkinWeightProfiles();
@@ -1381,7 +1407,7 @@ mu::MeshPtr ConvertSkeletalMeshToMutable(USkeletalMesh* InSkeletalMesh, const TS
 					TransferRemovedBonesInfluences(&SkinWeight.InfluenceBones[0], &SkinWeight.InfluenceWeights[0], MaxSectionInfluences, RemappedBoneMapIndices);
 				}
 
-				if (!GenerationContext.Options.bExtraBoneInfluencesEnabled)
+				if (GenerationContext.Options.CustomizableObjectNumBoneInfluences == ECustomizableObjectNumBoneInfluences::Four)
 				{
 					// Normalize weights
 					const int32 MaxMutableWeights = 4;
diff --git a/Engine/Plugins/Experimental/Mutable/Source/CustomizableObjectEditor/Private/MuCOE/GenerateMutableSource/GenerateMutableSourceSurface.cpp b/Engine/Plugins/Experimental/Mutable/Source/CustomizableObjectEditor/Private/MuCOE/GenerateMutableSource/GenerateMutableSourceSurface.cpp
--- a/Engine/Plugins/Experimental/Mutable/Source/CustomizableObjectEditor/Private/MuCOE/GenerateMutableSource/GenerateMutableSourceSurface.cpp
+++ b/Engine/Plugins/Experimental/Mutable/Source/CustomizableObjectEditor/Private/MuCOE/GenerateMutableSource/GenerateMutableSourceSurface.cpp
@@ -7,6 +7,8 @@
 #include "GPUSkinPublicDefs.h"
 #include "Interfaces/ITargetPlatform.h"
 #include "Materials/MaterialInstance.h"
+#include "GPUSkinVertexFactory.h"
+
 #include "MuCO/CustomizableObjectInstance.h"
 #include "MuCO/MutableMeshBufferUtils.h"
 #include "MuCOE/CustomizableObjectCompiler.h"
@@ -86,11 +88,17 @@ void LayoutOperationUVNormalizedWarning(FMutableGraphGenerationContext& Generati
 
 void SetSurfaceFormat( FMutableGraphGenerationContext& GenerationContext,
 					   mu::FMeshBufferSet& OutVertexBufferFormat, mu::FMeshBufferSet& OutIndexBufferFormat, const FMutableGraphMeshGenerationData& MeshData, 
-					   bool bWithExtraBoneInfluences, bool bWithRealTimeMorphs, bool bWithClothing, bool bWith16BitWeights )
+					   ECustomizableObjectNumBoneInfluences ECustomizableObjectNumBoneInfluences, bool bWithRealTimeMorphs, 
+					   bool bWithClothing, bool bWith16BitWeights )
 {
 	// Limit skinning weights if necessary
 	// \todo: make it more flexible to support 3 or 5 or 1 weight, since there is support for this in 4.25
-	const int32 MutableBonesPerVertex = bWithExtraBoneInfluences ? EXTRA_BONE_INFLUENCES : 4;
+	const int32 MutableBonesPerVertex = FGPUBaseSkinVertexFactory::UseUnlimitedBoneInfluences(MeshData.MaxNumBonesPerVertex, GenerationContext.Options.TargetPlatform) &&
+										MeshData.MaxNumBonesPerVertex < (int32)ECustomizableObjectNumBoneInfluences ?
+											MeshData.MaxNumBonesPerVertex :
+											(int32)ECustomizableObjectNumBoneInfluences;
+
+	ensure(MutableBonesPerVertex <= MAX_TOTAL_INFLUENCES);
 	
 	if (MutableBonesPerVertex != MeshData.MaxNumBonesPerVertex)
 	{
@@ -384,7 +392,7 @@ mu::NodeSurfacePtr GenerateMutableSourceSurface(const UEdGraphPin * Pin, FMutabl
 				MeshFormatNode->SetSource(MeshNode.get());
 				SetSurfaceFormat( GenerationContext,
 						MeshFormatNode->GetVertexBuffers(), MeshFormatNode->GetIndexBuffers(), MeshData,
-						GenerationContext.Options.bExtraBoneInfluencesEnabled,
+						GenerationContext.Options.CustomizableObjectNumBoneInfluences,
 						GenerationContext.Options.bRealTimeMorphTargetsEnabled, 
 						GenerationContext.Options.bClothingEnabled,
 						GenerationContext.Options.b16BitBoneWeightsEnabled);
@@ -1092,7 +1100,7 @@ mu::NodeSurfacePtr GenerateMutableSourceSurface(const UEdGraphPin * Pin, FMutabl
 				mu::NodeMeshFormatPtr MeshFormat = new mu::NodeMeshFormat();
 				SetSurfaceFormat( GenerationContext,
 						MeshFormat->GetVertexBuffers(), MeshFormat->GetIndexBuffers(), MeshData,
-						GenerationContext.Options.bExtraBoneInfluencesEnabled,
+						GenerationContext.Options.CustomizableObjectNumBoneInfluences,
 						GenerationContext.Options.bRealTimeMorphTargetsEnabled, 
 						GenerationContext.Options.bClothingEnabled,
 						GenerationContext.Options.b16BitBoneWeightsEnabled);
diff --git a/Engine/Plugins/Experimental/Mutable/Source/MutableRuntime/Private/MuR/MeshBufferSet.cpp b/Engine/Plugins/Experimental/Mutable/Source/MutableRuntime/Private/MuR/MeshBufferSet.cpp
--- a/Engine/Plugins/Experimental/Mutable/Source/MutableRuntime/Private/MuR/MeshBufferSet.cpp
+++ b/Engine/Plugins/Experimental/Mutable/Source/MutableRuntime/Private/MuR/MeshBufferSet.cpp
@@ -267,6 +267,21 @@ namespace mu
 	}
 
 
+	//---------------------------------------------------------------------------------------------
+	uint32 FMeshBufferSet::GetBufferDataSize(int32 buffer) const
+	{
+		check(buffer >= 0 && buffer < m_buffers.Num());
+		const uint32 Result = m_buffers[buffer].m_data.Num();
+
+#if WITH_EDITOR
+		int32 Size = m_buffers[buffer].m_elementSize * m_elementCount;
+		ensure(Size == Result);
+#endif
+
+		return Result;
+	}
+
+
 	//---------------------------------------------------------------------------------------------
 	void FMeshBufferSet::FindChannel
 	(
diff --git a/Engine/Plugins/Experimental/Mutable/Source/MutableRuntime/Private/MuR/MeshPrivate.h b/Engine/Plugins/Experimental/Mutable/Source/MutableRuntime/Private/MuR/MeshPrivate.h
--- a/Engine/Plugins/Experimental/Mutable/Source/MutableRuntime/Private/MuR/MeshPrivate.h
+++ b/Engine/Plugins/Experimental/Mutable/Source/MutableRuntime/Private/MuR/MeshPrivate.h
@@ -532,6 +532,16 @@ namespace mu
             return res;
         }
 
+		vec<int32_t, 12> GetAsVec12i() const
+		{
+			vec<int32_t, 12> res;
+			for (int c = 0; c < FMath::Min(m_components, 12); ++c)
+			{
+				ConvertData(c, &res[0], MBF_INT32, ptr(), m_format);
+			}
+			return res;
+		}
+
         inline UntypedMeshBufferIteratorConst operator+(int c) const
         {
             UntypedMeshBufferIteratorConst res = *this;
diff --git a/Engine/Plugins/Experimental/Mutable/Source/MutableRuntime/Private/MuR/OpMeshFormat.cpp b/Engine/Plugins/Experimental/Mutable/Source/MutableRuntime/Private/MuR/OpMeshFormat.cpp
--- a/Engine/Plugins/Experimental/Mutable/Source/MutableRuntime/Private/MuR/OpMeshFormat.cpp
+++ b/Engine/Plugins/Experimental/Mutable/Source/MutableRuntime/Private/MuR/OpMeshFormat.cpp
@@ -18,6 +18,8 @@
 #include "MuR/Ptr.h"
 #include "MuR/RefCounted.h"
 
+#include "GPUSkinPublicDefs.h"
+
 
 namespace mu
 {
@@ -580,7 +582,9 @@ namespace mu
 							int32_t maxBoneIndex = 0;
 							for (int v = 0; v < VertexBuffers.GetElementCount(); ++v)
 							{
-								vec<int32_t, 8> va = it.GetAsVec8i();
+								// If MAX_TOTAL_INFLUENCES ever changed, the next line would no longer work or compile and 
+								// GetAsVec12i would need to be changed accordingly
+								vec<int32_t, MAX_TOTAL_INFLUENCES> va = it.GetAsVec12i();
 								for (int c = 0; c < it.GetComponents(); ++c)
 								{
 									maxBoneIndex = FMath::Max(maxBoneIndex, va[c]);
diff --git a/Engine/Plugins/Experimental/Mutable/Source/MutableRuntime/Public/MuR/MeshBufferSet.h b/Engine/Plugins/Experimental/Mutable/Source/MutableRuntime/Public/MuR/MeshBufferSet.h
--- a/Engine/Plugins/Experimental/Mutable/Source/MutableRuntime/Public/MuR/MeshBufferSet.h
+++ b/Engine/Plugins/Experimental/Mutable/Source/MutableRuntime/Public/MuR/MeshBufferSet.h
@@ -306,7 +306,8 @@ namespace mu
 		//! \param buffer index of the buffer from 0 to GetBufferCount()-1
 		//! \todo Add padding support for better alignment of buffer elements.
         uint8* GetBufferData( int32 buffer );
-        const uint8* GetBufferData( int32 buffer ) const;
+		const uint8* GetBufferData(int32 buffer) const;
+		uint32 GetBufferDataSize( int32 buffer ) const;
 
 		//-----------------------------------------------------------------------------------------
 		// Utility methods
