[mutable] Fix reuse of non-black images. Fix debugger operation names in UI.
#rnx

diff --git a/Engine/Plugins/Experimental/Mutable/Source/MutableRuntime/Private/MuR/Image.cpp b/Engine/Plugins/Experimental/Mutable/Source/MutableRuntime/Private/MuR/Image.cpp
--- a/Engine/Plugins/Experimental/Mutable/Source/MutableRuntime/Private/MuR/Image.cpp
+++ b/Engine/Plugins/Experimental/Mutable/Source/MutableRuntime/Private/MuR/Image.cpp
@@ -127,7 +127,9 @@ namespace mu
 		int32 DataSize = CalculateDataSize();
 		if (bDisableCompressedImageBlackBlockInit)
 		{
-			m_data.SetNumZeroed(DataSize);
+			m_data.SetNumUninitialized(DataSize);
+			// Do it in separate steps in case we are reusing a buffer.
+			FMemory::Memzero(m_data.GetData(), DataSize);
 		}
 		else
 		{
@@ -141,7 +143,9 @@ namespace mu
 
 			if (bIsFormatBlackBlockZeroed)
 			{
-				m_data.SetNumZeroed(DataSize);
+				m_data.SetNumUninitialized(DataSize);
+				// Do it in separate steps in case we are reusing a buffer.
+				FMemory::Memzero(m_data.GetData(), DataSize);
 			}
 			else
 			{
diff --git a/Engine/Plugins/Experimental/Mutable/Source/MutableTools/Private/MuT/ErrorLog.cpp b/Engine/Plugins/Experimental/Mutable/Source/MutableTools/Private/MuT/ErrorLog.cpp
--- a/Engine/Plugins/Experimental/Mutable/Source/MutableTools/Private/MuT/ErrorLog.cpp
+++ b/Engine/Plugins/Experimental/Mutable/Source/MutableTools/Private/MuT/ErrorLog.cpp
@@ -263,7 +263,6 @@ namespace mu
 		TEXT("IM_SATURATE      "),
 		TEXT("IM_LUMINANCE     "),
 		TEXT("IM_SWIZZLE       "),
-		TEXT("IM_SELECTCOLOUR  "),
 		TEXT("IM_COLOURMAP     "),
 		TEXT("IM_GRADIENT      "),
 		TEXT("IM_BINARISE      "),
