#rnx

diff --git a/Engine/Plugins/Experimental/Mutable/Source/CustomizableObject/Private/MuCO/CustomizableObjectSystem.cpp b/Engine/Plugins/Experimental/Mutable/Source/CustomizableObject/Private/MuCO/CustomizableObjectSystem.cpp
--- a/Engine/Plugins/Experimental/Mutable/Source/CustomizableObject/Private/MuCO/CustomizableObjectSystem.cpp
+++ b/Engine/Plugins/Experimental/Mutable/Source/CustomizableObject/Private/MuCO/CustomizableObjectSystem.cpp
@@ -1595,9 +1595,13 @@ namespace impl
 				// If the image is null, it must be in the cache (or repeated in this instance), and we don't need to do anything here.
 				if (MutableImage)
 				{
-					FTexturePlatformData* PlatformData = UCustomizableInstancePrivateData::MutableCreateImagePlatformData(MutableImage, -1, Image.FullImageSizeX, Image.FullImageSizeY);
-					OperationData->ImageToPlatformDataMap.Add(Image.ImageID, PlatformData);
-					OperationData->PendingTextureCoverageQueries.Add({ KeyName, Surface.MaterialIndex, PlatformData });
+					// Image refences are just references to texture assets and require no work at all
+					if (!MutableImage->IsReference())
+					{
+						FTexturePlatformData* PlatformData = UCustomizableInstancePrivateData::MutableCreateImagePlatformData(MutableImage, -1, Image.FullImageSizeX, Image.FullImageSizeY);
+						OperationData->ImageToPlatformDataMap.Add(Image.ImageID, PlatformData);
+						OperationData->PendingTextureCoverageQueries.Add({ KeyName, Surface.MaterialIndex, PlatformData });
+					}
 				}
 			}
 		}
