[Mutable] Fix crash when closing the editor.

#rnx

diff --git a/Engine/Plugins/Experimental/Mutable/Source/CustomizableObject/Classes/MuCO/CustomizableObjectSystem.h b/Engine/Plugins/Experimental/Mutable/Source/CustomizableObject/Classes/MuCO/CustomizableObjectSystem.h
--- a/Engine/Plugins/Experimental/Mutable/Source/CustomizableObject/Classes/MuCO/CustomizableObjectSystem.h
+++ b/Engine/Plugins/Experimental/Mutable/Source/CustomizableObject/Classes/MuCO/CustomizableObjectSystem.h
@@ -317,9 +317,12 @@ public:
 
 
 	// Give access to the internal object data.
-	FCustomizableObjectSystemPrivate* GetPrivate() { return Private.Get(); }
-	const FCustomizableObjectSystemPrivate* GetPrivate() const { return Private.Get(); }
+	FCustomizableObjectSystemPrivate* GetPrivate();
+	const FCustomizableObjectSystemPrivate* GetPrivate() const;
 
+	FCustomizableObjectSystemPrivate* GetPrivateChecked();
+	const FCustomizableObjectSystemPrivate* GetPrivateChecked() const;
+	
 	FStreamableManager& GetStreamableManager() { return StreamableManager; }
 
 	UCustomizableInstanceLODManagementBase* GetInstanceLODManagement() { return CurrentInstanceLODManagement.Get(); }
diff --git a/Engine/Plugins/Experimental/Mutable/Source/CustomizableObject/Private/MuCO/CustomizableObjectInstance.cpp b/Engine/Plugins/Experimental/Mutable/Source/CustomizableObject/Private/MuCO/CustomizableObjectInstance.cpp
--- a/Engine/Plugins/Experimental/Mutable/Source/CustomizableObject/Private/MuCO/CustomizableObjectInstance.cpp
+++ b/Engine/Plugins/Experimental/Mutable/Source/CustomizableObject/Private/MuCO/CustomizableObjectInstance.cpp
@@ -362,7 +362,7 @@ void UCustomizableObjectInstance::ReleaseMutableResources(bool bCalledFromBeginD
 		}
 
 		// Remove all references to cached Texture Parameters
-		FUnrealMutableImageProvider* ImageProvider = UCustomizableObjectSystem::GetInstance()->GetPrivate()->ImageProvider.Get();
+		FUnrealMutableImageProvider* ImageProvider = UCustomizableObjectSystem::GetInstance()->GetPrivateChecked()->GetImageProviderChecked();
 
 		for (const FString& TextureParameter : PrivateData->UpdateTextureParameters)
 		{
diff --git a/Engine/Plugins/Experimental/Mutable/Source/CustomizableObject/Private/MuCO/CustomizableObjectMipDataProvider.cpp b/Engine/Plugins/Experimental/Mutable/Source/CustomizableObject/Private/MuCO/CustomizableObjectMipDataProvider.cpp
--- a/Engine/Plugins/Experimental/Mutable/Source/CustomizableObject/Private/MuCO/CustomizableObjectMipDataProvider.cpp
+++ b/Engine/Plugins/Experimental/Mutable/Source/CustomizableObject/Private/MuCO/CustomizableObjectMipDataProvider.cpp
@@ -25,16 +25,19 @@ FMutableUpdateContext::FMutableUpdateContext(mu::Ptr<mu::System> InSystem,
 {
 	if (Parameters)
 	{
-		UCustomizableObjectSystem::GetInstance()->GetPrivate()->ImageProvider->CacheImages(*Parameters);	
+		const FCustomizableObjectSystemPrivate* Private = UCustomizableObjectSystem::GetInstance()->GetPrivate();
+		Private->GetImageProviderChecked()->CacheImages(*Parameters);
 	}
 }
 
 
 FMutableUpdateContext::~FMutableUpdateContext()
 {
-	if (Parameters && UCustomizableObjectSystem::IsCreated() && UCustomizableObjectSystem::GetInstance()->GetPrivate()->ImageProvider)
+	if (Parameters &&
+		UCustomizableObjectSystem::IsCreated())
 	{
-		UCustomizableObjectSystem::GetInstance()->GetPrivate()->ImageProvider->UnCacheImages(*Parameters);
+		const FCustomizableObjectSystemPrivate* Private = UCustomizableObjectSystem::GetInstance()->GetPrivate();
+		Private->GetImageProviderChecked()->UnCacheImages(*Parameters);
 	}
 }
 
diff --git a/Engine/Plugins/Experimental/Mutable/Source/CustomizableObject/Private/MuCO/CustomizableObjectSystem.cpp b/Engine/Plugins/Experimental/Mutable/Source/CustomizableObject/Private/MuCO/CustomizableObjectSystem.cpp
--- a/Engine/Plugins/Experimental/Mutable/Source/CustomizableObject/Private/MuCO/CustomizableObjectSystem.cpp
+++ b/Engine/Plugins/Experimental/Mutable/Source/CustomizableObject/Private/MuCO/CustomizableObjectSystem.cpp
@@ -289,6 +289,32 @@ void UCustomizableObjectSystem::LogShowData(bool bFullInfo, bool ShowMaterialInf
 }
 
 
+FCustomizableObjectSystemPrivate* UCustomizableObjectSystem::GetPrivate()
+{
+	return Private.Get();
+}
+
+
+const FCustomizableObjectSystemPrivate* UCustomizableObjectSystem::GetPrivate() const
+{
+	return Private.Get();
+}
+
+
+FCustomizableObjectSystemPrivate* UCustomizableObjectSystem::GetPrivateChecked()
+{
+	check(Private)
+	return Private.Get();
+}
+
+
+const FCustomizableObjectSystemPrivate* UCustomizableObjectSystem::GetPrivateChecked() const
+{
+	check(Private)
+	return Private.Get();
+}
+
+
 bool UCustomizableObjectSystem::IsCreated()
 {
 	return FCustomizableObjectSystemPrivate::SSystem != 0;
@@ -768,7 +794,7 @@ void FinishUpdateGlobal(UCustomizableObjectInstance* Instance, EUpdateResult Upd
 
 	if (Parameters)
 	{
-		UCustomizableObjectSystem::GetInstance()->GetPrivate()->ImageProvider->UnCacheImages(*Parameters);
+		UCustomizableObjectSystem::GetInstance()->GetPrivateChecked()->GetImageProviderChecked()->UnCacheImages(*Parameters);
 	}
 
 	UCustomizableObjectSystem::GetInstance()->GetPrivate()->MutableTaskGraph.AllowLaunchingMutableTaskLowPriority(true, false);
@@ -908,7 +934,7 @@ void FCustomizableObjectSystemPrivate::InitUpdateSkeletalMesh(UCustomizableObjec
 		MutableOperationQueue.Enqueue(FMutableQueueElem::Create(Operation, Priority, Instance.GetPrivate()->MinSquareDistFromComponentToPlayer));
 
 		// Cache Texture Parameters being used during the update:
-		ImageProvider->CacheImages(*Operation->GetParameters());
+		GetImageProviderChecked()->CacheImages(*Operation->GetParameters());
 	}
 }
 
@@ -2072,7 +2098,7 @@ namespace impl
 		}
 
 		// Cache Texture Parameters
-		FUnrealMutableImageProvider* ImageProvider = UCustomizableObjectSystem::GetInstance()->GetPrivate()->ImageProvider.Get();
+		FUnrealMutableImageProvider* ImageProvider = UCustomizableObjectSystem::GetInstance()->GetPrivateChecked()->GetImageProviderChecked();
 		
 		// Uncache old Texture Parameters
 		for (const FString& TextureParameter : ObjectInstancePrivateData->UpdateTextureParameters)
@@ -2688,9 +2714,7 @@ TArray<FCustomizableObjectExternalTexture> UCustomizableObjectSystem::GetTexture
 {
 	TArray<FCustomizableObjectExternalTexture> Result;
 
-	check(Private != nullptr);
-	check(Private->ImageProvider != nullptr);
-	for (const TWeakObjectPtr<UCustomizableSystemImageProvider> Provider : Private->ImageProvider->ImageProviders)
+	for (const TWeakObjectPtr<UCustomizableSystemImageProvider> Provider : GetPrivateChecked()->GetImageProviderChecked()->ImageProviders)
 	{
 		if (Provider.IsValid())
 		{
@@ -2704,19 +2728,13 @@ TArray<FCustomizableObjectExternalTexture> UCustomizableObjectSystem::GetTexture
 
 void UCustomizableObjectSystem::RegisterImageProvider(UCustomizableSystemImageProvider* Provider)
 {
-	check(Private != nullptr);
-	check(Private->ImageProvider != nullptr);
-
-	Private->ImageProvider->ImageProviders.Add(Provider);
+	GetPrivateChecked()->GetImageProviderChecked()->ImageProviders.Add(Provider);
 }
 
 
 void UCustomizableObjectSystem::UnregisterImageProvider(UCustomizableSystemImageProvider* Provider)
 {
-	check(Private != nullptr);
-	check(Private->ImageProvider != nullptr);
-
-	Private->ImageProvider->ImageProviders.Remove(Provider);
+	GetPrivateChecked()->GetImageProviderChecked()->ImageProviders.Remove(Provider);
 }
 
 
@@ -3143,25 +3161,19 @@ void UCustomizableObjectSystem::OnEnableImageCacheChanged(IConsoleVariable* CVar
 
 void UCustomizableObjectSystem::CacheImage(FString ImageId)
 {
-	check(GetPrivate() != nullptr);
-	check(GetPrivate()->ImageProvider != nullptr);
-	GetPrivate()->ImageProvider->CacheImage(ImageId, true);
+	GetPrivateChecked()->GetImageProviderChecked()->CacheImage(ImageId, true);
 }
 
 
 void UCustomizableObjectSystem::UnCacheImage(FString ImageId)
 {
-	check(GetPrivate() != nullptr); 
-	check(GetPrivate()->ImageProvider != nullptr);
-	GetPrivate()->ImageProvider->UnCacheImage(ImageId, true);
+	GetPrivateChecked()->GetImageProviderChecked()->UnCacheImage(ImageId, true);
 }
 
 
 void UCustomizableObjectSystem::ClearImageCache()
 {
-	check(GetPrivate() != nullptr);
-	check(GetPrivate()->ImageProvider != nullptr);
-	GetPrivate()->ImageProvider->ClearCache(true);
+	GetPrivateChecked()->GetImageProviderChecked()->ClearCache(true);
 }
 
 
@@ -3175,10 +3187,17 @@ bool FCustomizableObjectSystemPrivate::IsMutableAnimInfoDebuggingEnabled() const
 }
 
 
+FUnrealMutableImageProvider* FCustomizableObjectSystemPrivate::GetImageProviderChecked() const
+{
+	check(ImageProvider)
+	return ImageProvider.Get();
+}
+
+
 bool UCustomizableObjectSystem::IsMutableAnimInfoDebuggingEnabled() const
 {
 #if WITH_EDITOR
-	return GetPrivate()->IsMutableAnimInfoDebuggingEnabled();
+	return GetPrivateChecked()->IsMutableAnimInfoDebuggingEnabled();
 #else
 	return false;
 #endif
diff --git a/Engine/Plugins/Experimental/Mutable/Source/CustomizableObject/Private/MuCO/CustomizableObjectSystemPrivate.h b/Engine/Plugins/Experimental/Mutable/Source/CustomizableObject/Private/MuCO/CustomizableObjectSystemPrivate.h
--- a/Engine/Plugins/Experimental/Mutable/Source/CustomizableObject/Private/MuCO/CustomizableObjectSystemPrivate.h
+++ b/Engine/Plugins/Experimental/Mutable/Source/CustomizableObject/Private/MuCO/CustomizableObjectSystemPrivate.h
@@ -741,6 +741,8 @@ public:
 
 	bool IsMutableAnimInfoDebuggingEnabled() const;
 
+	FUnrealMutableImageProvider* GetImageProviderChecked() const;
+	
 	/** Mutable TaskGraph system (Mutable Thread). */
 	FMutableTaskGraph MutableTaskGraph;
 };
