[Mutable] Fix Extend Material Node breaking when the duplicating a CO

#rnx

diff --git a/Engine/Plugins/Experimental/Mutable/Source/CustomizableObjectEditor/Classes/MuCOE/Nodes/CustomizableObjectNodeExtendMaterial.h b/Engine/Plugins/Experimental/Mutable/Source/CustomizableObjectEditor/Classes/MuCOE/Nodes/CustomizableObjectNodeExtendMaterial.h
--- a/Engine/Plugins/Experimental/Mutable/Source/CustomizableObjectEditor/Classes/MuCOE/Nodes/CustomizableObjectNodeExtendMaterial.h
+++ b/Engine/Plugins/Experimental/Mutable/Source/CustomizableObjectEditor/Classes/MuCOE/Nodes/CustomizableObjectNodeExtendMaterial.h
@@ -34,11 +34,13 @@ public:
 
 	// UObject interface.
 	virtual void Serialize(FArchive& Ar) override;
-
+	virtual void BeginPostDuplicate(bool bDuplicateForPIE) override;
+	
 	// UCustomizableObjectNode
 	virtual void PostBackwardsCompatibleFixup() override;
 	virtual void BackwardsCompatibleFixup() override;
-	
+	virtual void UpdateReferencedNodeId(const FGuid& NewGuid) override;
+
 	// Begin EdGraphNode interface
 	virtual FText GetNodeTitle(ENodeTitleType::Type TitleType) const override;
 	virtual FLinearColor GetNodeTitleColor() const override;
diff --git a/Engine/Plugins/Experimental/Mutable/Source/CustomizableObjectEditor/Private/MuCOE/Nodes/CustomizableObjectNodeExtendMaterial.cpp b/Engine/Plugins/Experimental/Mutable/Source/CustomizableObjectEditor/Private/MuCOE/Nodes/CustomizableObjectNodeExtendMaterial.cpp
--- a/Engine/Plugins/Experimental/Mutable/Source/CustomizableObjectEditor/Private/MuCOE/Nodes/CustomizableObjectNodeExtendMaterial.cpp
+++ b/Engine/Plugins/Experimental/Mutable/Source/CustomizableObjectEditor/Private/MuCOE/Nodes/CustomizableObjectNodeExtendMaterial.cpp
@@ -4,6 +4,7 @@
 
 #include "MaterialTypes.h"
 #include "MuCOE/CustomizableObjectEditor_Deprecated.h"
+#include "MuCOE/CustomizableObjectGraph.h"
 #include "MuCOE/EdGraphSchema_CustomizableObject.h"
 #include "MuCOE/GraphTraversal.h"
 #include "MuCOE/Nodes/CustomizableObjectNodeCopyMaterial.h"
@@ -24,6 +25,20 @@ void UCustomizableObjectNodeExtendMaterial::Serialize(FArchive& Ar)
 }
 
 
+void UCustomizableObjectNodeExtendMaterial::BeginPostDuplicate(bool bDuplicateForPIE)
+{
+	Super::BeginPostDuplicate(bDuplicateForPIE);
+
+	if (ParentMaterialObject)
+	{
+		if (UCustomizableObjectGraph* CEdGraph = Cast<UCustomizableObjectGraph>(GetGraph()))
+		{
+			ParentMaterialNodeId = CEdGraph->RequestNotificationForNodeIdChange(ParentMaterialNodeId, NodeGuid);
+		}
+	}
+}
+
+
 void UCustomizableObjectNodeExtendMaterial::BackwardsCompatibleFixup()
 {
 	Super::PostBackwardsCompatibleFixup();
@@ -64,6 +79,16 @@ void UCustomizableObjectNodeExtendMaterial::BackwardsCompatibleFixup()
 	}
 }
 
+
+void UCustomizableObjectNodeExtendMaterial::UpdateReferencedNodeId(const FGuid& NewGuid)
+{
+	if (ParentMaterialObject)
+	{
+		ParentMaterialNodeId = NewGuid;
+	}
+}
+
+
 void UCustomizableObjectNodeExtendMaterial::PostBackwardsCompatibleFixup()
 {
 	Super::PostBackwardsCompatibleFixup();
