Also changed the logging of unexpected texture formats to display the full format instead of just the unprefixed format for better debugging of build issues.
#rnx

diff --git a/Engine/Plugins/Experimental/Mutable/Source/CustomizableObjectEditor/Private/MuCOE/GenerateMutableSource/GenerateMutableSourceSurface.cpp b/Engine/Plugins/Experimental/Mutable/Source/CustomizableObjectEditor/Private/MuCOE/GenerateMutableSource/GenerateMutableSourceSurface.cpp
--- a/Engine/Plugins/Experimental/Mutable/Source/CustomizableObjectEditor/Private/MuCOE/GenerateMutableSource/GenerateMutableSourceSurface.cpp
+++ b/Engine/Plugins/Experimental/Mutable/Source/CustomizableObjectEditor/Private/MuCOE/GenerateMutableSource/GenerateMutableSourceSurface.cpp
@@ -772,6 +772,18 @@ mu::NodeSurfacePtr GenerateMutableSourceSurface(const UEdGraphPin * Pin, FMutabl
 										else if (FormatWithoutPrefix == TEXT("RGB")) mutableFormat = mu::EImageFormat::IF_ASTC_4x4_RGB_LDR;
 										else if (FormatWithoutPrefix == TEXT("RGBA")) mutableFormat = mu::EImageFormat::IF_ASTC_4x4_RGBA_LDR;
 										else if (FormatWithoutPrefix == TEXT("NormalRG")) mutableFormat = mu::EImageFormat::IF_ASTC_4x4_RG_LDR;
+										else if (PlatformFormat.Contains(TEXT("ASTC_NormalRG_Precise")))
+										{
+											// TODO: This is just a workaround to prevent the "Unexpected image format" warning below. ASTC_NormalRG_Precise is
+											// not supported yet by Mutable so it forces IF_ASTC_4x4_RG_LDR as a replacement. It should be changed with a more
+											// appropriate format or directly implement it
+											mutableFormat = mu::EImageFormat::IF_ASTC_4x4_RG_LDR;
+
+											const FString ReplacedImageFormatMsg = FString::Printf(TEXT("In object [%s] the unsupported ASTC_NormalRG_Precise image format is used, ASTC_4x4_RG_LDR will be used instead."), *GenerationContext.Object->GetName());
+											const FText ReplacedImageFormatText = FText::FromString(ReplacedImageFormatMsg);
+											GenerationContext.Compiler->CompilerLog(ReplacedImageFormatText, Node, EMessageSeverity::Info);
+											UE_LOG(LogMutable, Log, TEXT("%s"), *ReplacedImageFormatMsg);
+										}
 									}
 
 									if (mutableFormat == mu::EImageFormat::IF_NONE)
@@ -779,7 +791,7 @@ mu::NodeSurfacePtr GenerateMutableSourceSurface(const UEdGraphPin * Pin, FMutabl
 										// Format not supported by Mutable, use RBGA_UBYTE as default.
 										mutableFormat = mu::EImageFormat::IF_RGBA_UBYTE;
 
-										const FString UnexpectedImageFormatMsg = FString::Printf(TEXT("In object [%s] Unexpected image format [%s], RGBA_UBYTE will be used instead."), *GenerationContext.Object->GetName(), *FormatWithoutPrefix);
+										const FString UnexpectedImageFormatMsg = FString::Printf(TEXT("In object [%s] Unexpected image format [%s], RGBA_UBYTE will be used instead."), *GenerationContext.Object->GetName(), *PlatformFormat);
 										const FText UnexpectedImageFormatText = FText::FromString(UnexpectedImageFormatMsg);
 										GenerationContext.Compiler->CompilerLog(UnexpectedImageFormatText, Node);
 										UE_LOG(LogMutable, Warning, TEXT("%s"), *UnexpectedImageFormatMsg);
