[mutable] MutableTools: Added missing string for mu:OP_Type.
- The issue manifested itself in the mutable code debugger as tree elements without an operation type name.

diff --git a/Engine/Plugins/Experimental/Mutable/Source/MutableTools/Private/MuT/ErrorLog.cpp b/Engine/Plugins/Experimental/Mutable/Source/MutableTools/Private/MuT/ErrorLog.cpp
--- a/Engine/Plugins/Experimental/Mutable/Source/MutableTools/Private/MuT/ErrorLog.cpp
+++ b/Engine/Plugins/Experimental/Mutable/Source/MutableTools/Private/MuT/ErrorLog.cpp
@@ -310,7 +310,8 @@ namespace mu
         TEXT("IN_ADDSURFACE    "),
         TEXT("IN_ADDCOMPONENT  "),
         TEXT("IN_ADDLOD        "),
-
+		TEXT("IN_ADDEXTENSIDATA"),
+		
 		TEXT("LA_PACK          "),
 		TEXT("LA_MERGE         "),
 		TEXT("LA_REMOVEBLOCKS  "),
