[mutable] Fixed duplicated LOCTEXT keys.
- Removed stray pragma once in .cpp file

diff --git a/Engine/Plugins/Experimental/Mutable/Source/MutableValidation/Private/MuV/AssetValidator_CustomizableObjects.cpp b/Engine/Plugins/Experimental/Mutable/Source/MutableValidation/Private/MuV/AssetValidator_CustomizableObjects.cpp
--- a/Engine/Plugins/Experimental/Mutable/Source/MutableValidation/Private/MuV/AssetValidator_CustomizableObjects.cpp
+++ b/Engine/Plugins/Experimental/Mutable/Source/MutableValidation/Private/MuV/AssetValidator_CustomizableObjects.cpp
@@ -64,7 +64,7 @@ EDataValidationResult UAssetValidator_CustomizableObjects::ValidateLoadedAsset_I
 			AssetFails(InAsset,ErrorMessage,ValidationErrors);
 		}
 
-		const FText ErrorMessage = FText::Format(LOCTEXT("CustomizableObjectsValidator", "Validation compilation of {0} CO failed."),  FText::FromString( CustomizableObjectToValidate->GetName()));
+		const FText ErrorMessage = FText::Format(LOCTEXT("CO_Validation_Failed", "Validation compilation of {0} CO failed."),  FText::FromString( CustomizableObjectToValidate->GetName()));
 		AssetFails(InAsset,ErrorMessage,ValidationErrors);
 	}
 	else
@@ -183,12 +183,12 @@ EDataValidationResult UAssetValidator_CustomizableObjects::IsCustomizableObjectV
 	const FString ReferencedCOName = FString(TEXT("\"" +  InCustomizableObject->GetName() + "\""));
 	for (const FText& ValidationError : CachedValidationErrors)
 	{
-		FText ComposedMessage = FText::Format(LOCTEXT("CustomizableObjectsValidator", "Customizable Object : {0} {1}"),  FText::FromString(ReferencedCOName),ValidationError );
+		FText ComposedMessage = FText::Format(LOCTEXT("CO_Compilation_Error", "Customizable Object : {0} {1}"),  FText::FromString(ReferencedCOName),ValidationError );
 		OutValidationErrors.Add(ComposedMessage);
 	}
 	for (const FText& ValidationWarning : CachedValidationWarnings)
 	{
-		FText ComposedMessage = FText::Format(LOCTEXT("CustomizableObjectsValidator", "Customizable Object : {0} {1}"),  FText::FromString(ReferencedCOName),ValidationWarning );
+		FText ComposedMessage = FText::Format(LOCTEXT("CO_Compilation_nWarning", "Customizable Object : {0} {1}"),  FText::FromString(ReferencedCOName),ValidationWarning );
 		OutValidationWarnings.Add(ComposedMessage);
 	}
 	
diff --git a/Engine/Plugins/Experimental/Mutable/Source/MutableValidation/Public/MuV/MutableValidationModule.cpp b/Engine/Plugins/Experimental/Mutable/Source/MutableValidation/Public/MuV/MutableValidationModule.cpp
--- a/Engine/Plugins/Experimental/Mutable/Source/MutableValidation/Public/MuV/MutableValidationModule.cpp
+++ b/Engine/Plugins/Experimental/Mutable/Source/MutableValidation/Public/MuV/MutableValidationModule.cpp
@@ -1,14 +1,11 @@
 // Copyright Epic Games, Inc. All Rights Reserved.
 
-#pragma once
-
 #include "ISettingsModule.h"
 #include "ISettingsSection.h"
 #include "Modules/ModuleManager.h"
 #include "MuV/MutableValidationSettings.h"
 
-
-#define LOCTEXT_NAMESPACE "MutableSettings"
+#define LOCTEXT_NAMESPACE "MutableValidationSettings"
 
 /**
  * StaticMesh editor module
@@ -47,8 +44,8 @@ void FMutableValidationModule::StartupModule()
 	if (SettingsModule != nullptr)
 	{ 		
 		 SettingsSectionPtr = SettingsModule->RegisterSettings("Project", "Plugins", "MutableValidationSettings",
-			LOCTEXT("MutableSettings_Setting", "Mutable Validation"),
-			LOCTEXT("MutableSettings_Setting_Desc", "Mutable resources validation settings"),
+			LOCTEXT("MutableValidationSettings_Setting", "Mutable Validation"),
+			LOCTEXT("MutableValidatioinSettings_Setting_Desc", "Mutable resources validation settings"),
 			GetMutableDefault<UMutableValidationSettings>()
 		);
 		
