[mutable] Fix bug with image-size-relative operation setting incorrect size. Better message and checks for mutable images with unexpected sizes.
Some additional portability and indentation cleanups.
#rnx

diff --git a/Engine/Plugins/Experimental/Mutable/Source/CustomizableObject/Private/MuCO/CustomizableObjectSystem.cpp b/Engine/Plugins/Experimental/Mutable/Source/CustomizableObject/Private/MuCO/CustomizableObjectSystem.cpp
--- a/Engine/Plugins/Experimental/Mutable/Source/CustomizableObject/Private/MuCO/CustomizableObjectSystem.cpp
+++ b/Engine/Plugins/Experimental/Mutable/Source/CustomizableObject/Private/MuCO/CustomizableObjectSystem.cpp
@@ -91,7 +91,7 @@ static TAutoConsoleVariable<int32> CVarWorkingMemory(
 TAutoConsoleVariable<bool> CVarClearWorkingMemoryOnUpdateEnd(
 	TEXT("mutable.CVarClearWorkingMemoryOnUpdateEnd"),
 	false,
-	TEXT("Clear the wokring memory and cache after every Mutable operation."),
+	TEXT("Clear the working memory and cache after every Mutable operation."),
 	ECVF_Scalability);
 
 
@@ -1558,10 +1558,10 @@ namespace impl
 					MipsToSkip = 0;
 				}
 
+				const int32 MipSizeX = FMath::Max(Image.FullImageSizeX >> MipsToSkip, 1);
+				const int32 MipSizeY = FMath::Max(Image.FullImageSizeY >> MipsToSkip, 1);
 				if (MipsToSkip > 0 && CustomizableObjectSystemPrivateData->EnableSkipGenerateResidentMips != 0 && OperationData->LowPriorityTextures.Find(Image.Name) != INDEX_NONE)
 				{
-					const int32 MipSizeX = FMath::Max(Image.FullImageSizeX >> MipsToSkip, 1);
-					const int32 MipSizeY = FMath::Max(Image.FullImageSizeY >> MipsToSkip, 1);
 					Image.Image = new mu::Image(MipSizeX, MipSizeY, FullLODCount - MipsToSkip, ImageDesc.m_format);
 				}
 				else
@@ -1577,10 +1577,25 @@ namespace impl
 					continue;
 				}
 
+
+				// We should have genrated exactly this size.
+				bool bSizeMissmatch = Image.Image->GetSizeX() != MipSizeX || Image.Image->GetSizeY() != MipSizeY;
+				if (bSizeMissmatch)
+				{
+					// Generate a correctly-sized but empty image instead, to avoid crashes.
+					UE_LOG(LogMutable, Warning, TEXT("Mutable generated a wrongly-sized image %d."), Image.ImageID);
+					Image.Image = new mu::Image(MipSizeX, MipSizeY, FullLODCount - MipsToSkip, Image.Image->GetFormat());
+				}
+
 				// We need one mip or the complete chain. Otherwise there was a bug.
 				int32 FullMipCount = Image.Image->GetMipmapCount(Image.Image->GetSizeX(), Image.Image->GetSizeY());
 				int32 RealMipCount = Image.Image->GetLODCount();
-				if ((RealMipCount != 1) && (RealMipCount != FullMipCount))
+
+				bool bForceMipchain = 
+					// Did we fail to generate the entire mipchain (if we have mips at all)?
+					(RealMipCount != 1) && (RealMipCount != FullMipCount);
+
+				if (bForceMipchain)
 				{
 					MUTABLE_CPUPROFILER_SCOPE(GetImage_MipFix);
 
diff --git a/Engine/Plugins/Experimental/Mutable/Source/CustomizableObject/Private/MuCO/UnrealConversionUtils.cpp b/Engine/Plugins/Experimental/Mutable/Source/CustomizableObject/Private/MuCO/UnrealConversionUtils.cpp
--- a/Engine/Plugins/Experimental/Mutable/Source/CustomizableObject/Private/MuCO/UnrealConversionUtils.cpp
+++ b/Engine/Plugins/Experimental/Mutable/Source/CustomizableObject/Private/MuCO/UnrealConversionUtils.cpp
@@ -95,7 +95,8 @@ namespace UnrealConversionUtils
 
 			if (NumVertices)
 			{
-				uint8* Data = VertexBuffer->GetWeightData();
+				void* Data = VertexBuffer->GetWeightData();
+
 				uint32 OutVertexWeightBufferSize = OutVertexWeightBuffer.GetDataVertexBuffer()->GetVertexDataSize();
 				ensure(MutableDataSize == OutVertexWeightBufferSize);
 
@@ -525,8 +526,8 @@ namespace UnrealConversionUtils
 			 {
 				 DestSkinWeightBuffer.SetNeedsCPUAccess(bAllowCPUAccess);
 
-				 const uint8* SrcData = SrcSkinWeightBuffer.GetDataVertexBuffer()->GetWeightData();
-				 uint8* Data = SkinWeightDataVertexBuffer->GetWeightData();
+				 const void* SrcData = SrcSkinWeightBuffer.GetDataVertexBuffer()->GetWeightData();
+				 void* Data = SkinWeightDataVertexBuffer->GetWeightData();
 				 check(SrcData);
 				 check(Data);
 
diff --git a/Engine/Plugins/Experimental/Mutable/Source/CustomizableObjectEditor/Private/MuCOE/GenerateMutableSource/GenerateMutableSourceMesh.cpp b/Engine/Plugins/Experimental/Mutable/Source/CustomizableObjectEditor/Private/MuCOE/GenerateMutableSource/GenerateMutableSourceMesh.cpp
--- a/Engine/Plugins/Experimental/Mutable/Source/CustomizableObjectEditor/Private/MuCOE/GenerateMutableSource/GenerateMutableSourceMesh.cpp
+++ b/Engine/Plugins/Experimental/Mutable/Source/CustomizableObjectEditor/Private/MuCOE/GenerateMutableSource/GenerateMutableSourceMesh.cpp
@@ -855,6 +855,9 @@ mu::MeshPtr ConvertSkeletalMeshToMutable(USkeletalMesh* InSkeletalMesh, const TS
 
 	MutableMesh->GetVertexBuffers().SetBufferCount(VertexBuffersCount);
 
+	const int32 MaxSectionInfluences = ImportedModel->LODModels[LOD].Sections[MaterialIndex].MaxBoneInfluences;
+	const bool bUseUnlimitedInfluences = FGPUBaseSkinVertexFactory::UseUnlimitedBoneInfluences(MaxSectionInfluences, GenerationContext.Options.TargetPlatform);
+
 	if (bIgnoreSkeleton)
 	{
 		// Create the mesh with the same data, but skinning is considered padding.
@@ -885,7 +888,6 @@ mu::MeshPtr ConvertSkeletalMeshToMutable(USkeletalMesh* InSkeletalMesh, const TS
 	}
 	else
 	{
-		int32 MaxSectionInfluences = ImportedModel->LODModels[LOD].Sections[MaterialIndex].MaxBoneInfluences;
 		using namespace mu;
 		const int ElementSize = sizeof(FSoftSkinVertex);
 		const int ChannelCount = 11;
@@ -902,7 +904,7 @@ mu::MeshPtr ConvertSkeletalMeshToMutable(USkeletalMesh* InSkeletalMesh, const TS
 		{
 			int32 NewBoneInfluencesNum = (int32)GenerationContext.Options.CustomizableObjectNumBoneInfluences;
 
-			if (FGPUBaseSkinVertexFactory::UseUnlimitedBoneInfluences(MaxSectionInfluences, GenerationContext.Options.TargetPlatform) &&
+			if (bUseUnlimitedInfluences &&
 				MaxSectionInfluences < NewBoneInfluencesNum)
 			{
 				Components[9] = MaxSectionInfluences;
@@ -979,7 +981,7 @@ mu::MeshPtr ConvertSkeletalMeshToMutable(USkeletalMesh* InSkeletalMesh, const TS
 			}
 			else if (GenerationContext.Options.CustomizableObjectNumBoneInfluences == ECustomizableObjectNumBoneInfluences::Eight)
 			{
-				if (!FGPUBaseSkinVertexFactory::UseUnlimitedBoneInfluences(MaxSectionInfluences, GenerationContext.Options.TargetPlatform) &&
+				if (!bUseUnlimitedInfluences &&
 					MaxSectionInfluences < EXTRA_BONE_INFLUENCES) // EXTRA_BONE_INFLUENCES is ECustomizableObjectNumBoneInfluences::Eight
 				{
 					FMemory::Memzero(&Vertex.InfluenceWeights[MaxSectionInfluences], EXTRA_BONE_INFLUENCES - MaxSectionInfluences);
@@ -987,7 +989,7 @@ mu::MeshPtr ConvertSkeletalMeshToMutable(USkeletalMesh* InSkeletalMesh, const TS
 			}
 			else if (GenerationContext.Options.CustomizableObjectNumBoneInfluences == ECustomizableObjectNumBoneInfluences::Twelve)
 			{
-				if (!FGPUBaseSkinVertexFactory::UseUnlimitedBoneInfluences(MaxSectionInfluences, GenerationContext.Options.TargetPlatform) &&
+				if (!bUseUnlimitedInfluences &&
 					MaxSectionInfluences < MAX_TOTAL_INFLUENCES) // MAX_TOTAL_INFLUENCES is ECustomizableObjectNumBoneInfluences::Twelve
 				{
 					FMemory::Memzero(&Vertex.InfluenceWeights[MaxSectionInfluences], MAX_TOTAL_INFLUENCES - MaxSectionInfluences);
@@ -1372,9 +1374,7 @@ mu::MeshPtr ConvertSkeletalMeshToMutable(USkeletalMesh* InSkeletalMesh, const TS
 		MESH_BUFFER_FORMAT BoneWeightFormat = BoneWeightTypeSizeBytes == 1 ? MBF_NUINT8 : MBF_NUINT16;
 
 		// Limit skinning weights if necessary
-		const int32 MaxSectionInfluences = ImportedModel->LODModels[LOD].Sections[MaterialIndex].GetMaxBoneInfluences();
-		const int32 MutableBonesPerVertex = FGPUBaseSkinVertexFactory::UseUnlimitedBoneInfluences(MaxSectionInfluences, 
-												GenerationContext.Options.TargetPlatform) ?
+		const int32 MutableBonesPerVertex = bUseUnlimitedInfluences ?
 											MaxSectionInfluences :
 											(int32)GenerationContext.Options.CustomizableObjectNumBoneInfluences;
 		const int32 BoneIndicesSize = MutableBonesPerVertex * sizeof(FBoneIndexType);
