[Mutable] Cache merged skeletons to avoid loading and merging skeletons with every update.
Add all skeletons involved in the merge as compatible.

#rnx

diff --git a/Engine/Plugins/Experimental/Mutable/Source/CustomizableObject/Classes/MuCO/CustomizableObject.h b/Engine/Plugins/Experimental/Mutable/Source/CustomizableObject/Classes/MuCO/CustomizableObject.h
--- a/Engine/Plugins/Experimental/Mutable/Source/CustomizableObject/Classes/MuCO/CustomizableObject.h
+++ b/Engine/Plugins/Experimental/Mutable/Source/CustomizableObject/Classes/MuCO/CustomizableObject.h
@@ -49,6 +49,25 @@ namespace mu
 
 CUSTOMIZABLEOBJECT_API DECLARE_LOG_CATEGORY_EXTERN(LogMutable, Log, All);
 
+struct FMergedSkeleton
+{
+	FMergedSkeleton(TObjectPtr<USkeleton> InSkeleton, const int32 InComponentIndex, const TArray<uint16>& InSkeletonIds);
+
+	// Merged skeleton
+	TWeakObjectPtr<USkeleton> Skeleton;
+	
+	// Component Index and Ids of the Skeletons involved in the merge. Used to identify the skeleton
+	uint16 ComponentIndex;
+	TArray<uint16> SkeletonIds;
+
+	bool operator==(const FMergedSkeleton& Other) const
+	{
+		return ComponentIndex == Other.ComponentIndex
+			&& SkeletonIds == Other.SkeletonIds;
+	}
+};
+
+
 USTRUCT()
 struct FFParameterOptionsTags
 {
@@ -1299,6 +1318,15 @@ public:
 	
 	TSoftObjectPtr<UMaterialInterface> GetReferencedMaterialAssetPtr(uint32 Index);
 
+	// Return a valid Skeletons if cached. ComponentIndex and SkeletonIds are used as a key to find the skeleton.
+	TObjectPtr<USkeleton> GetCachedMergedSkeleton(int32 ComponentIndex, const TArray<uint16>& SkeletonIds) const;
+
+	// Add merged skeleton to the cache. It'll be cached as a TWeakObjPtr.
+	void CacheMergedSkeleton(const int32 ComponentIndex, const TArray<uint16>& SkeletonIds, TObjectPtr<USkeleton> Skeleton);
+
+	// Remove skeletons that have been destroyed by the garbage collector from the cache.
+	void UnCacheInvalidSkeletons();
+
 	// Call before using Mutable's Projector testing with mask out features. It should only be loaded when needed because it can spend quite a lot of memory
 	// Can cause a loading hitch
 	UFUNCTION(BlueprintCallable, Category = CustomizableObject)
@@ -1653,6 +1681,9 @@ public:
 	UPROPERTY()
 	FGuid CompilationGuid;
 
+	/** Cache of merged skeletons */
+	TArray<FMergedSkeleton> MergedSkeletons;
+
 public:
 	FGuid GetCompilationGuid() const;
 
diff --git a/Engine/Plugins/Experimental/Mutable/Source/CustomizableObject/Private/MuCO/CustomizableInstancePrivateData.h b/Engine/Plugins/Experimental/Mutable/Source/CustomizableObject/Private/MuCO/CustomizableInstancePrivateData.h
--- a/Engine/Plugins/Experimental/Mutable/Source/CustomizableObject/Private/MuCO/CustomizableInstancePrivateData.h
+++ b/Engine/Plugins/Experimental/Mutable/Source/CustomizableObject/Private/MuCO/CustomizableInstancePrivateData.h
@@ -97,8 +97,11 @@ struct FReferencedSkeletons
 {
 	GENERATED_USTRUCT_BODY();
 
+	// Merged skeleton if found in the cache
+	TObjectPtr<USkeleton> Skeleton;
+
 	UPROPERTY()
-	TArray<int32> SkeletonsToLoad;
+	TArray<uint16> SkeletonIds;
 
 	UPROPERTY()
 	TArray< TObjectPtr<USkeleton> > SkeletonsToMerge;
diff --git a/Engine/Plugins/Experimental/Mutable/Source/CustomizableObject/Private/MuCO/CustomizableObject.cpp b/Engine/Plugins/Experimental/Mutable/Source/CustomizableObject/Private/MuCO/CustomizableObject.cpp
--- a/Engine/Plugins/Experimental/Mutable/Source/CustomizableObject/Private/MuCO/CustomizableObject.cpp
+++ b/Engine/Plugins/Experimental/Mutable/Source/CustomizableObject/Private/MuCO/CustomizableObject.cpp
@@ -7,6 +7,7 @@
 #include "Async/AsyncFileHandle.h"
 #include "EdGraph/EdGraph.h"
 #include "Engine/SkeletalMesh.h"
+#include "Animation/Skeleton.h"
 #include "Engine/AssetUserData.h"
 #include "HAL/FileManager.h"
 #include "HAL/PlatformFileManager.h"
@@ -38,11 +39,6 @@
 
 #include UE_INLINE_GENERATED_CPP_BY_NAME(CustomizableObject)
 
-class UMaterialInterface;
-class UPhysicsAsset;
-class USkeletalMesh;
-class USkeleton;
-
 #define LOCTEXT_NAMESPACE "CustomizableObject"
 
 DEFINE_LOG_CATEGORY(LogMutable);
@@ -1053,6 +1049,43 @@ void UCustomizableObject::OnReferenceSkeletalMeshesAsyncLoaded()
 }
 
 
+TObjectPtr<USkeleton> UCustomizableObject::GetCachedMergedSkeleton(const int32 ComponentIndex, const TArray<uint16>& SkeletonIds) const
+{
+	const FMergedSkeleton* MergedSkeleton = MergedSkeletons.FindByPredicate([&ComponentIndex, &SkeletonIds](const FMergedSkeleton& MergedSkeleton) 
+		{ return MergedSkeleton.ComponentIndex == ComponentIndex && MergedSkeleton.SkeletonIds == SkeletonIds; });
+
+	if (MergedSkeleton && MergedSkeleton->Skeleton.IsValid())
+	{
+		return MergedSkeleton->Skeleton.Get();
+	}
+
+	return nullptr;
+}
+
+
+void UCustomizableObject::CacheMergedSkeleton(const int32 ComponentIndex, const TArray<uint16>& SkeletonIds, TObjectPtr<USkeleton> Skeleton)
+{
+	if (Skeleton)
+	{
+		MergedSkeletons.AddUnique({ Skeleton, ComponentIndex, SkeletonIds });
+	}
+}
+
+
+void UCustomizableObject::UnCacheInvalidSkeletons()
+{
+	for (int32 SkeletonIndex = MergedSkeletons.Num() - 1; SkeletonIndex >= 0; --SkeletonIndex)
+	{
+		FMergedSkeleton& MergedSkeleton = MergedSkeletons[SkeletonIndex];
+		if (!MergedSkeleton.Skeleton.IsValid())
+		{
+			MergedSkeletons.RemoveSingleSwap(MergedSkeleton);
+		}
+	}
+}
+
+
+
 int32 UCustomizableObject::FindState( const FString& Name ) const
 {
 	int32 Result = -1;
@@ -1756,6 +1789,19 @@ TSharedPtr<const mu::Model, ESPMode::ThreadSafe> FCustomizableObjectPrivateData:
 	return MutableModel;
 }
 
+
+//-------------------------------------------------------------------------------------------------
+//-------------------------------------------------------------------------------------------------
+//-------------------------------------------------------------------------------------------------
+FMergedSkeleton::FMergedSkeleton(TObjectPtr<USkeleton> InSkeleton, const int32 InComponentIndex, const TArray<uint16>& InSkeletonIds)
+{
+	check(InSkeleton);
+	Skeleton = InSkeleton;
+	ComponentIndex = InComponentIndex;
+	SkeletonIds = InSkeletonIds;
+}
+
+
 //-------------------------------------------------------------------------------------------------
 //-------------------------------------------------------------------------------------------------
 //-------------------------------------------------------------------------------------------------
diff --git a/Engine/Plugins/Experimental/Mutable/Source/CustomizableObject/Private/MuCO/CustomizableObjectInstance.cpp b/Engine/Plugins/Experimental/Mutable/Source/CustomizableObject/Private/MuCO/CustomizableObjectInstance.cpp
--- a/Engine/Plugins/Experimental/Mutable/Source/CustomizableObject/Private/MuCO/CustomizableObjectInstance.cpp
+++ b/Engine/Plugins/Experimental/Mutable/Source/CustomizableObject/Private/MuCO/CustomizableObjectInstance.cpp
@@ -190,8 +190,8 @@ void UCustomizableInstancePrivateData::PrepareForUpdate(const TSharedPtr<FMutabl
 #if WITH_EDITORONLY_DATA
 				ComponentData->MeshPartPaths.Empty();
 #endif
-
-				ComponentData->Skeletons.SkeletonsToLoad.Empty();
+				ComponentData->Skeletons.Skeleton = nullptr;
+				ComponentData->Skeletons.SkeletonIds.Empty();
 				ComponentData->Skeletons.SkeletonsToMerge.Empty();
 				ComponentData->PhysicsAssets.PhysicsAssetToLoad.Empty();
 				ComponentData->PhysicsAssets.PhysicsAssetsToMerge.Empty();
@@ -1042,7 +1042,7 @@ USkeleton* UCustomizableInstancePrivateData::MergeSkeletons(UCustomizableObjectI
 {
 	MUTABLE_CPUPROFILER_SCOPE(BuildSkeletonData_MergeSkeletons);
 
-	const UCustomizableObject* CustomizableObject = Public->GetCustomizableObject();
+	UCustomizableObject* CustomizableObject = Public->GetCustomizableObject();
 
 	check(CustomizableObject);
 	check(RefSkeletalMeshData);
@@ -1052,12 +1052,20 @@ USkeleton* UCustomizableInstancePrivateData::MergeSkeletons(UCustomizableObjectI
 
 	FReferencedSkeletons& ReferencedSkeletons = ComponentData->Skeletons;
 	
+	// Merged skeleton found in the cache
+	if (ReferencedSkeletons.Skeleton)
+	{
+		USkeleton* MergedSkeleton = ReferencedSkeletons.Skeleton;
+		ReferencedSkeletons.Skeleton = nullptr;
+		return MergedSkeleton;
+	}
+
 	// No need to merge skeletons
 	if(ReferencedSkeletons.SkeletonsToMerge.Num() == 1)
 	{
 		const TObjectPtr<USkeleton> RefSkeleton = ReferencedSkeletons.SkeletonsToMerge[0];
+		ReferencedSkeletons.SkeletonIds.Empty();
 		ReferencedSkeletons.SkeletonsToMerge.Empty();
-
 		return RefSkeleton;
 	}
 
@@ -1072,10 +1080,6 @@ USkeleton* UCustomizableInstancePrivateData::MergeSkeletons(UCustomizableObjectI
 	FSkeletonMergeParams Params;
 	Exchange(Params.SkeletonsToMerge, ReferencedSkeletons.SkeletonsToMerge);
 
-#if !UE_BUILD_SHIPPING
-	//Params.bCheckSkeletonsCompatibility = true;
-#endif
-
 	USkeleton* FinalSkeleton = USkeletalMergingLibrary::MergeSkeletons(Params);
 	if (!FinalSkeleton)
 	{
@@ -1094,7 +1098,17 @@ USkeleton* UCustomizableInstancePrivateData::MergeSkeletons(UCustomizableObjectI
 	else
 	{
 		// The reference skeleton may be used by AnimBp as a target skeleton, add it as a compatible skeleton.
-		FinalSkeleton->AddCompatibleSkeleton(RefSkeletalMeshData->Skeleton.Get());
+		for (USkeleton* Skeleton : Params.SkeletonsToMerge)
+		{
+			if (Skeleton)
+			{
+				FinalSkeleton->AddCompatibleSkeleton(Skeleton);
+			}
+		}
+
+		// Add Skeleton to the cache
+		CustomizableObject->CacheMergedSkeleton(ComponentIndex, ReferencedSkeletons.SkeletonIds, FinalSkeleton);
+		ReferencedSkeletons.SkeletonIds.Empty();
 	}
 	
 	return FinalSkeleton;
@@ -5125,6 +5139,9 @@ FGraphEventRef UCustomizableInstancePrivateData::LoadAdditionalAssetsAsync(const
 		}
 	}
 
+	// Clear invalid skeletons from the MergedSkeletons cache
+	CustomizableObject->UnCacheInvalidSkeletons();
+
 	const int32 ComponentCount = OperationData->InstanceUpdateData.Skeletons.Num();
 
 	// Load Skeletons required by the SubMeshes of the newly generated Mesh, will be merged later
@@ -5136,6 +5153,14 @@ FGraphEventRef UCustomizableInstancePrivateData::LoadAdditionalAssetsAsync(const
 			[&ComponentIndex](FInstanceUpdateData::FSkeletonData& S) { return S.ComponentIndex == ComponentIndex; });
 		check(SkeletonData);
 
+		ComponentData.Skeletons.Skeleton = CustomizableObject->GetCachedMergedSkeleton(ComponentIndex, SkeletonData->SkeletonIds);
+		if (ComponentData.Skeletons.Skeleton)
+		{
+			ComponentData.Skeletons.SkeletonIds.Empty();
+			ComponentData.Skeletons.SkeletonsToMerge.Empty();
+			continue;
+		}
+
 		for (const uint32 SkeletonId : SkeletonData->SkeletonIds)
 		{
 			TSoftObjectPtr<USkeleton> AssetPtr;
@@ -5156,14 +5181,15 @@ FGraphEventRef UCustomizableInstancePrivateData::LoadAdditionalAssetsAsync(const
 				continue;
 			}
 
+			// Add referenced skeletons to the assets to stream
+			ComponentData.Skeletons.SkeletonIds.Add(SkeletonId);
+
 			if (USkeleton* Skeleton = AssetPtr.Get())
 			{
 				ComponentData.Skeletons.SkeletonsToMerge.Add(Skeleton);
 			}
 			else
 			{
-				// Add referenced skeletons to the assets to stream
-				ComponentData.Skeletons.SkeletonsToLoad.Add(SkeletonId);
 				AssetsToStream.Add(AssetPtr.ToSoftObjectPath());
 			}
 
@@ -5372,11 +5398,10 @@ void UCustomizableInstancePrivateData::AdditionalAssetsAsyncLoaded(UCustomizable
 	{
 		// Loaded Skeletons
 		FReferencedSkeletons& Skeletons = ComponentData.Skeletons;
-		for (const int32& SkeletonIndex : Skeletons.SkeletonsToLoad)
+		for (const int32& SkeletonIndex : Skeletons.SkeletonIds)
 		{
-			Skeletons.SkeletonsToMerge.Add(CustomizableObject->GetReferencedSkeletonAssetPtr(SkeletonIndex).Get());
+			Skeletons.SkeletonsToMerge.AddUnique(CustomizableObject->GetReferencedSkeletonAssetPtr(SkeletonIndex).Get());
 		}
-		Skeletons.SkeletonsToLoad.Empty();
 
 		// Loaded PhysicsAssets
 		FReferencedPhysicsAssets& PhysicsAssets = ComponentData.PhysicsAssets;
diff --git a/Engine/Plugins/Experimental/Mutable/Source/CustomizableObject/Private/MuCO/CustomizableObjectSystem.cpp b/Engine/Plugins/Experimental/Mutable/Source/CustomizableObject/Private/MuCO/CustomizableObjectSystem.cpp
--- a/Engine/Plugins/Experimental/Mutable/Source/CustomizableObject/Private/MuCO/CustomizableObjectSystem.cpp
+++ b/Engine/Plugins/Experimental/Mutable/Source/CustomizableObject/Private/MuCO/CustomizableObjectSystem.cpp
@@ -1670,7 +1670,7 @@ namespace impl
 			OperationData->InstanceUpdateData.Skeletons[ComponentIndex].ComponentIndex = MinLODComponent.Id;
 
 			// Fill the data used to generate the RefSkeletalMesh
-			TArray<uint32>& SkeletonIds = OperationData->InstanceUpdateData.Skeletons[ComponentIndex].SkeletonIds;
+			TArray<uint16>& SkeletonIds = OperationData->InstanceUpdateData.Skeletons[ComponentIndex].SkeletonIds;
 			TArray<FName>& BoneNames = OperationData->InstanceUpdateData.Skeletons[ComponentIndex].BoneNames;
 			TMap<FName, FMatrix44f>& BoneMatricesWithScale = OperationData->InstanceUpdateData.Skeletons[ComponentIndex].BoneMatricesWithScale;
 
diff --git a/Engine/Plugins/Experimental/Mutable/Source/CustomizableObject/Private/MuCO/CustomizableObjectSystemPrivate.h b/Engine/Plugins/Experimental/Mutable/Source/CustomizableObject/Private/MuCO/CustomizableObjectSystemPrivate.h
--- a/Engine/Plugins/Experimental/Mutable/Source/CustomizableObject/Private/MuCO/CustomizableObjectSystemPrivate.h
+++ b/Engine/Plugins/Experimental/Mutable/Source/CustomizableObject/Private/MuCO/CustomizableObjectSystemPrivate.h
@@ -412,7 +412,7 @@ struct FInstanceUpdateData
 	{
 		int16 ComponentIndex = INDEX_NONE;
 
-		TArray<uint32> SkeletonIds;
+		TArray<uint16> SkeletonIds;
 
 		TArray<FName> BoneNames;
 		TMap<FName, FMatrix44f> BoneMatricesWithScale;
