[Mutable] Fix compilation error on Linux

#rnx

diff --git a/Engine/Plugins/Experimental/Mutable/Source/CustomizableObject/Private/MuCO/CustomizableObjectInstance.cpp b/Engine/Plugins/Experimental/Mutable/Source/CustomizableObject/Private/MuCO/CustomizableObjectInstance.cpp
--- a/Engine/Plugins/Experimental/Mutable/Source/CustomizableObject/Private/MuCO/CustomizableObjectInstance.cpp
+++ b/Engine/Plugins/Experimental/Mutable/Source/CustomizableObject/Private/MuCO/CustomizableObjectInstance.cpp
@@ -5398,7 +5398,7 @@ void UCustomizableInstancePrivateData::AdditionalAssetsAsyncLoaded(UCustomizable
 	{
 		// Loaded Skeletons
 		FReferencedSkeletons& Skeletons = ComponentData.Skeletons;
-		for (const int32& SkeletonIndex : Skeletons.SkeletonIds)
+		for (int32 SkeletonIndex : Skeletons.SkeletonIds)
 		{
 			Skeletons.SkeletonsToMerge.AddUnique(CustomizableObject->GetReferencedSkeletonAssetPtr(SkeletonIndex).Get());
 		}
