[mutable] Fix bug with reuse of images without clearing their flags.
#rnx

diff --git a/Engine/Plugins/Experimental/Mutable/Source/MutableRuntime/Private/MuR/CodeRunnerTasked.cpp b/Engine/Plugins/Experimental/Mutable/Source/MutableRuntime/Private/MuR/CodeRunnerTasked.cpp
--- a/Engine/Plugins/Experimental/Mutable/Source/MutableRuntime/Private/MuR/CodeRunnerTasked.cpp
+++ b/Engine/Plugins/Experimental/Mutable/Source/MutableRuntime/Private/MuR/CodeRunnerTasked.cpp
@@ -1762,11 +1762,6 @@ namespace mu
 			uint16(FMath::Max(1.0, Base->GetSizeX() * Args.factor[0] + 0.5f)),
 			uint16(FMath::Max(1.0, Base->GetSizeY() * Args.factor[1] + 0.5f)));
 
-		// Apply the mips-to-skip to the dest size
-		int32 MipsToSkip = Op.ExecutionOptions;
-		DestSize[0] = FMath::Max(DestSize[0] >> MipsToSkip, 1);
-		DestSize[1] = FMath::Max(DestSize[1] >> MipsToSkip, 1);
-
 		if (Base->GetSizeX() == DestSize[0] && Base->GetSizeY() == DestSize[1])
 		{
 			Runner->StoreImage(Op, Base);
diff --git a/Engine/Plugins/Experimental/Mutable/Source/MutableRuntime/Private/MuR/SystemPrivate.h b/Engine/Plugins/Experimental/Mutable/Source/MutableRuntime/Private/MuR/SystemPrivate.h
--- a/Engine/Plugins/Experimental/Mutable/Source/MutableRuntime/Private/MuR/SystemPrivate.h
+++ b/Engine/Plugins/Experimental/Mutable/Source/MutableRuntime/Private/MuR/SystemPrivate.h
@@ -1054,6 +1054,12 @@ namespace mu
 					{
 						Result->InitToBlack();
 					}
+					else
+					{
+						Result->m_flags = 0;
+						Result->RelevancyMinY = 0;
+						Result->RelevancyMaxY = 0;
+					}
 					return Result;
 				}
 			}
diff --git a/Engine/Plugins/Experimental/Mutable/Source/MutableRuntime/Public/MuR/Image.h b/Engine/Plugins/Experimental/Mutable/Source/MutableRuntime/Public/MuR/Image.h
--- a/Engine/Plugins/Experimental/Mutable/Source/MutableRuntime/Public/MuR/Image.h
+++ b/Engine/Plugins/Experimental/Mutable/Source/MutableRuntime/Public/MuR/Image.h
@@ -326,16 +326,16 @@ namespace mu
 		// This used to be the data in the private implementation of the image interface
 		//-----------------------------------------------------------------------------------------
 
-		//!
+		/** */
 		FImageSize m_size = FImageSize(0, 0);
 
-		//!
+		/** */
 		EImageFormat m_format = EImageFormat::IF_NONE;
 
-		//! Levels of detail (mipmaps)
+		/** Levels of detail (mipmaps) */
 		uint8 m_lods = 0;
 
-		//! These set of flags are used to cache information for images at runtime.
+		/** These set of flags are used to cache information for images at runtime. */
 		typedef enum
 		{
 			// Set if the next flag has been calculated and its value is valid.
@@ -354,20 +354,19 @@ namespace mu
 			IF_IS_REFERENCE = 1 << 4,
 		} EImageFlags;
 
-		//! Persistent flags with some image properties. The meaning will depend of every
-		//! context.
+		/** Persistent flags with some image properties. The meaning will depend of every context. */
 		mutable uint8 m_flags = 0;
 
 		/** Non-persistent relevancy map. */
 		uint16 RelevancyMinY = 0;
 		uint16 RelevancyMaxY = 0;
 
-		//! Pixel data
-		TArray<uint8> m_data;
-
 		/** Only valid if the right flags are set, this identifies a referenced image. */
 		uint32 ReferenceID = 0;
 
+		/** Pixel data for all lods. */
+		TArray<uint8> m_data;
+
 
 		// This used to be the methods in the private implementation of the image interface
 		//-----------------------------------------------------------------------------------------
