[Mutable] Ensure UV islands remain in the same block across different LODs when reusing textures between LODs.

#rnx

diff --git a/Engine/Plugins/Experimental/Mutable/Source/CustomizableObject/Classes/MuCO/CustomizableObject.h b/Engine/Plugins/Experimental/Mutable/Source/CustomizableObject/Classes/MuCO/CustomizableObject.h
--- a/Engine/Plugins/Experimental/Mutable/Source/CustomizableObject/Classes/MuCO/CustomizableObject.h
+++ b/Engine/Plugins/Experimental/Mutable/Source/CustomizableObject/Classes/MuCO/CustomizableObject.h
@@ -212,9 +212,6 @@ struct FCompilationOptions
 	// Used to enable AnimBp override physics mainipualtion.  
 	bool bAnimBpPhysicsManipulationEnabled = false;
 
-	// Used to enable Reuse textures between lods. Requires automatic LODs to be enabled.  
-	bool bReuseTexturesEnabled = false;
-
 	// Used to reduce the number of notifications when compiling objects
 	bool bSilentCompilation = true;
 
diff --git a/Engine/Plugins/Experimental/Mutable/Source/CustomizableObjectEditor/Classes/MuCOE/Nodes/CustomizableObjectNodeMaterial.h b/Engine/Plugins/Experimental/Mutable/Source/CustomizableObjectEditor/Classes/MuCOE/Nodes/CustomizableObjectNodeMaterial.h
--- a/Engine/Plugins/Experimental/Mutable/Source/CustomizableObjectEditor/Classes/MuCOE/Nodes/CustomizableObjectNodeMaterial.h
+++ b/Engine/Plugins/Experimental/Mutable/Source/CustomizableObjectEditor/Classes/MuCOE/Nodes/CustomizableObjectNodeMaterial.h
@@ -104,6 +104,11 @@ public:
 	/** Selects which Mesh component of the Instance this material belongs to */
 	UPROPERTY(EditAnywhere, BlueprintReadWrite, Category = CustomizableObject, meta = (ClampMin = "0"))
 	int32 MeshComponentIndex = 0;
+
+	/** Reuse material between LODs when using the LODStrategy - Automatic from mesh. If the material is extended or edited by a CO using
+	* custom LODs (LODStrategy), on that given LOD, reuse materials will not be applied and the material and its textures will be unique. */
+	UPROPERTY(EditAnywhere, Category = CustomizableObject)
+	bool bReuseMaterialBetweenLODs = true;
 	
 	// UObject interface.
 	virtual void Serialize(FArchive& Ar) override;
diff --git a/Engine/Plugins/Experimental/Mutable/Source/CustomizableObjectEditor/Private/MuCOE/CustomizableObjectCompiler.cpp b/Engine/Plugins/Experimental/Mutable/Source/CustomizableObjectEditor/Private/MuCOE/CustomizableObjectCompiler.cpp
--- a/Engine/Plugins/Experimental/Mutable/Source/CustomizableObjectEditor/Private/MuCOE/CustomizableObjectCompiler.cpp
+++ b/Engine/Plugins/Experimental/Mutable/Source/CustomizableObjectEditor/Private/MuCOE/CustomizableObjectCompiler.cpp
@@ -620,10 +620,6 @@ mu::NodeObjectPtr FCustomizableObjectCompiler::GenerateMutableRoot(
 
 	}
 
-	// Reuse textures when using automatic LODs from mesh if the LayoutManagement is disabled.
-	GenerationContext.Options.bReuseTexturesEnabled = ActualRoot->AutoLODStrategy == ECustomizableObjectAutomaticLODStrategy::AutomaticFromMesh &&
-		GenerationContext.bDisableTextureLayoutManagementFlag;
-
 	// Ensure that the CO has a valid AutoLODStrategy on the ActualRoot.
 	if (ActualRoot->AutoLODStrategy == ECustomizableObjectAutomaticLODStrategy::Inherited)
 	{
diff --git a/Engine/Plugins/Experimental/Mutable/Source/CustomizableObjectEditor/Private/MuCOE/GenerateMutableSource/GenerateMutableSourceSurface.cpp b/Engine/Plugins/Experimental/Mutable/Source/CustomizableObjectEditor/Private/MuCOE/GenerateMutableSource/GenerateMutableSourceSurface.cpp
--- a/Engine/Plugins/Experimental/Mutable/Source/CustomizableObjectEditor/Private/MuCOE/GenerateMutableSource/GenerateMutableSourceSurface.cpp
+++ b/Engine/Plugins/Experimental/Mutable/Source/CustomizableObjectEditor/Private/MuCOE/GenerateMutableSource/GenerateMutableSourceSurface.cpp
@@ -220,7 +220,9 @@ void SetSurfaceFormat( FMutableGraphGenerationContext& GenerationContext,
 
 void UpdateSharedSurfaceId(FMutableGraphGenerationContext& GenerationContext, UCustomizableObjectNodeMaterial* NodeMaterial, mu::NodeSurfaceNewPtr InNodeSurface)
 {
-	const bool bCanShareSurface = GenerationContext.Options.bReuseTexturesEnabled && GenerationContext.CurrentAutoLODStrategy == ECustomizableObjectAutomaticLODStrategy::AutomaticFromMesh;
+	// Reuse materials between LODs when using automatics LODs, if texture layout management is disabled or if bReuseMaterials is enabled in the node material.
+	const bool bCanShareSurface = GenerationContext.CurrentAutoLODStrategy == ECustomizableObjectAutomaticLODStrategy::AutomaticFromMesh
+		&& (GenerationContext.bDisableTextureLayoutManagementFlag || NodeMaterial->bReuseMaterialBetweenLODs);
 
 	// Set shared surface Id to reuse materials and textures between LODs if automatic LODs from mesh is being used
 	if (bCanShareSurface && InNodeSurface)
diff --git a/Engine/Plugins/Experimental/Mutable/Source/MutableTools/MutableTools.Build.cs b/Engine/Plugins/Experimental/Mutable/Source/MutableTools/MutableTools.Build.cs
--- a/Engine/Plugins/Experimental/Mutable/Source/MutableTools/MutableTools.Build.cs
+++ b/Engine/Plugins/Experimental/Mutable/Source/MutableTools/MutableTools.Build.cs
@@ -26,6 +26,7 @@ namespace UnrealBuildTool.Rules
                 new string[] {
 					"MutableRuntime", 
 					"Core",
+					"GeometryCore",
 				}
             );
 		}
diff --git a/Engine/Plugins/Experimental/Mutable/Source/MutableTools/Private/MuT/CodeGenerator.cpp b/Engine/Plugins/Experimental/Mutable/Source/MutableTools/Private/MuT/CodeGenerator.cpp
--- a/Engine/Plugins/Experimental/Mutable/Source/MutableTools/Private/MuT/CodeGenerator.cpp
+++ b/Engine/Plugins/Experimental/Mutable/Source/MutableTools/Private/MuT/CodeGenerator.cpp
@@ -844,6 +844,12 @@ namespace mu
 			MeshOptions.State = m_currentStateIndex;
 			MeshOptions.ActiveTags = m_activeTags.Last();
 
+			// Ensure UV islands remain within their main layout block on lower LODs to avoid unexpected reordering 
+			// of the layout blocks when reusing a surface between LODs. Used to fix small displacements on vertices
+			// that may cause them to fall on a different block.
+			const bool bClampUVIslands = node.SharedSurfaceId != INDEX_NONE && m_currentParents.Last().m_lod > 0;
+			MeshOptions.bClampUVIslands = bClampUVIslands;
+
             GenerateMesh(MeshOptions, meshResults, pMesh );
             lastMeshOp = meshResults.meshOp;
             meshResults.extraMeshLayouts.SetNum( edits.Num() );
@@ -870,6 +876,7 @@ namespace mu
 						FMeshGenerationOptions MergedMeshOptions;
 						MergedMeshOptions.bUniqueVertexIDs = true;
 						MergedMeshOptions.bLayouts = true;
+						MergedMeshOptions.bClampUVIslands = bClampUVIslands;
 						MergedMeshOptions.State = m_currentStateIndex;
 						MergedMeshOptions.ActiveTags = m_activeTags.Last();
 						FMeshGenerationResult addResults;
diff --git a/Engine/Plugins/Experimental/Mutable/Source/MutableTools/Private/MuT/CodeGenerator.h b/Engine/Plugins/Experimental/Mutable/Source/MutableTools/Private/MuT/CodeGenerator.h
--- a/Engine/Plugins/Experimental/Mutable/Source/MutableTools/Private/MuT/CodeGenerator.h
+++ b/Engine/Plugins/Experimental/Mutable/Source/MutableTools/Private/MuT/CodeGenerator.h
@@ -332,7 +332,8 @@ namespace mu
         void PrepareForLayout( LayoutPtrConst GeneratedLayout,
                                   MeshPtr currentLayoutMesh,
                                   size_t currentLayoutChannel,
-                                  const void* errorContext );
+                                  const void* errorContext,
+								  const bool bClampUVIslands);
 
         //-----------------------------------------------------------------------------------------
         //!
@@ -455,6 +456,11 @@ namespace mu
 			/** The meshes at the leaves will need their own layout block data. */
 			bool bLayouts = false;
 
+			/** If true, Ensure UV Islands are not split between two or more blocks. UVs shared between multiple 
+			* layout blocks will be clamped to fit the one with more vertices belonging to the UV island.
+			* Mainly used to keep consistent layouts when reusing textures between LODs. */
+			bool bClampUVIslands = false;
+
 			/** If this has something the layouts in constant meshes will be ignored, because
 			* they are supposed to match some other set of layouts. If the vector is empty, layouts
 			* are generated normally.
diff --git a/Engine/Plugins/Experimental/Mutable/Source/MutableTools/Private/MuT/CodeGenerator_Mesh.cpp b/Engine/Plugins/Experimental/Mutable/Source/MutableTools/Private/MuT/CodeGenerator_Mesh.cpp
--- a/Engine/Plugins/Experimental/Mutable/Source/MutableTools/Private/MuT/CodeGenerator_Mesh.cpp
+++ b/Engine/Plugins/Experimental/Mutable/Source/MutableTools/Private/MuT/CodeGenerator_Mesh.cpp
@@ -75,6 +75,8 @@
 #include "MuT/Table.h"
 #include "MuT/TablePrivate.h"
 
+#include "Spatial/PointHashGrid3.h"
+
 #include <memory>
 #include <utility>
 
@@ -83,204 +85,618 @@ namespace mu
 {
 class Node;
 
+	
+	//---------------------------------------------------------------------------------------------
+    //! Create a map from vertices into vertices, collapsing vertices that have the same position. 
+	//! This version uses UE Containers to return.	
     //---------------------------------------------------------------------------------------------
-    void CodeGenerator::PrepareForLayout(Ptr<const Layout> GeneratedLayout,
+	void MeshCreateCollapsedVertexMap(const mu::Mesh* Mesh,
+		TArray<int32>& CollapsedVertices)
+	{
+		MUTABLE_CPUPROFILER_SCOPE(LayoutUV_CreateCollapsedVertexMap);
+	
+		const int32 NumVertices = Mesh->GetVertexCount();
+		CollapsedVertices.Reserve(NumVertices);
+	
+		UE::Geometry::TPointHashGrid3f<int32> VertHash(0.01f, INDEX_NONE);
+		VertHash.Reserve(NumVertices);
+	
+		TArray<FVector3f> Vertices;
+		Vertices.SetNumUninitialized(NumVertices);
+	
+		mu::UntypedMeshBufferIteratorConst ItPosition = mu::UntypedMeshBufferIteratorConst(Mesh->GetVertexBuffers(), mu::MBS_POSITION);
+	
+		FVector3f* VertexData = Vertices.GetData();
+	
+		for (int32 VertexIndex = 0; VertexIndex < NumVertices; ++VertexIndex)
+		{
+			*VertexData = ItPosition.GetAsVec3f();
+			VertHash.InsertPointUnsafe(VertexIndex, *VertexData);
+	
+			++ItPosition;
+			++VertexData;
+		}
+	
+		// Find unique vertices
+		CollapsedVertices.Init(INDEX_NONE, NumVertices);
+	
+		TArray<int32> NearbyVertices;
+		for (int32 VertexIndex = 0; VertexIndex < NumVertices; ++VertexIndex)
+		{
+			if (CollapsedVertices[VertexIndex] != INDEX_NONE)
+			{
+				continue;
+			}
+	
+			const FVector3f& Vertex = Vertices[VertexIndex];
+	
+			NearbyVertices.Reset();
+			VertHash.FindPointsInBall(Vertex, 0.00001,
+				[&Vertex, &Vertices](const int32& Other) -> float {return FVector3f::DistSquared(Vertices[Other], Vertex); },
+				NearbyVertices);
+	
+			// Find equals
+			for (int32 NearbyVertexIndex : NearbyVertices)
+			{
+				CollapsedVertices[NearbyVertexIndex] = VertexIndex;
+			}
+		}
+	}
+
+	// Helper used to connect triangles of a UV island
+	struct FTriangle
+	{
+		uint32 Indices[3];
+		uint32 CollapsedIndices[3];
+
+		uint16 BlockIndices[3];
+
+		bool bUVsFixed = false;
+	};
+
+	//---------------------------------------------------------------------------------------------
+	//! Fill an array with the indices of all triangles belonging to the same UV island as InFirstTriangle. 
+	//---------------------------------------------------------------------------------------------
+	void GetUVIsland(TArray<FTriangle>& InTriangles,
+		const uint32 InFirstTriangle,
+		TArray<uint32>& OutTriangleIndices,
+		const TArray<vec2<float>>& InUVs,
+		const TMultiMap<int32, uint32>& InVertexToTriangleMap)
+	{
+		MUTABLE_CPUPROFILER_SCOPE(LayoutUV_GetUVIsland);
+
+		const uint32 NumTriangles = (uint32)InTriangles.Num();
+		
+		OutTriangleIndices.Reserve(NumTriangles);
+		OutTriangleIndices.Add(InFirstTriangle);
+
+		TArray<bool> SkipTrinalges;
+		SkipTrinalges.Init(false, NumTriangles);
+
+		TArray<uint32> PendingTriangles;
+		PendingTriangles.Reserve(NumTriangles / 64);
+		PendingTriangles.Add(InFirstTriangle);
+
+		while (!PendingTriangles.IsEmpty())
+		{
+			const uint32 TriangleIndex = PendingTriangles.Pop();
+
+			// Triangle about to be proccessed, mark as skip;
+			SkipTrinalges[TriangleIndex] = true;
+
+			bool ConnectedEdges[3] = { false, false, false };
+
+			const FTriangle& Triangle = InTriangles[TriangleIndex];
+
+			// Find Triangles connected to edges 0 and 2
+			int32 CollapsedVertex1 = Triangle.CollapsedIndices[1];
+			int32 CollapsedVertex2 = Triangle.CollapsedIndices[2];
+
+			TArray<uint32> FoundTriangleIndices;
+			InVertexToTriangleMap.MultiFind(Triangle.CollapsedIndices[0], FoundTriangleIndices);
+
+			for (uint32 OtherTriangleIndex : FoundTriangleIndices)
+			{
+				const FTriangle& OtherTriangle = InTriangles[OtherTriangleIndex];
+
+				for (int32 OtherIndex = 0; OtherIndex < 3; ++OtherIndex)
+				{
+					const int32 OtherCollapsedIndex = OtherTriangle.CollapsedIndices[OtherIndex];
+					if (OtherCollapsedIndex == CollapsedVertex1)
+					{
+						// Check if the vertex is in the same UV Island 
+						if (!SkipTrinalges[OtherTriangleIndex]
+							&& InUVs[Triangle.Indices[1]].AlmostEqual(InUVs[OtherTriangle.Indices[OtherIndex]], 0.00001))
+						{
+							OutTriangleIndices.Add(OtherTriangleIndex);
+							PendingTriangles.Add(OtherTriangleIndex);
+							SkipTrinalges[OtherTriangleIndex] = true;
+						}
+
+						// Connected but already processed or in another island
+						break;
+					}
+
+					if (OtherCollapsedIndex  == CollapsedVertex2)
+					{
+						// Check if the vertex is in the same UV Island 
+						if (!SkipTrinalges[OtherTriangleIndex]
+							&& InUVs[Triangle.Indices[2]].AlmostEqual(InUVs[OtherTriangle.Indices[OtherIndex]], 0.00001))
+						{
+							OutTriangleIndices.Add(OtherTriangleIndex);
+							PendingTriangles.Add(OtherTriangleIndex);
+							SkipTrinalges[OtherTriangleIndex] = true;
+						}
+
+						// Connected but already processed or in another UV Island
+						break;
+					}
+				}
+
+			}
+
+			// Find the triangle connected to edge 1
+			FoundTriangleIndices.Reset();
+			InVertexToTriangleMap.MultiFind(CollapsedVertex1, FoundTriangleIndices);
+
+			for (uint32 OtherTriangleIndex : FoundTriangleIndices)
+			{
+				const FTriangle& OtherTriangle = InTriangles[OtherTriangleIndex];
+
+				for (int32 OtherIndex = 0; OtherIndex < 3; ++OtherIndex)
+				{
+					const int32 OtherCollapsedIndex = OtherTriangle.CollapsedIndices[OtherIndex];
+					if (OtherCollapsedIndex == CollapsedVertex2)
+					{
+						// Check if the vertex belong to the same UV island
+						if (!SkipTrinalges[OtherTriangleIndex]
+							&& InUVs[Triangle.Indices[2]].AlmostEqual(InUVs[OtherTriangle.Indices[OtherIndex]], 0.00001))
+						{
+							OutTriangleIndices.Add(OtherTriangleIndex);
+							PendingTriangles.Add(OtherTriangleIndex);
+							SkipTrinalges[OtherTriangleIndex] = true;
+						}
+
+						// Connected but already processed or in another island
+						break;
+					}
+				}
+			}
+		}
+	}
+
+    //---------------------------------------------------------------------------------------------
+	void CodeGenerator::PrepareForLayout(Ptr<const Layout> GeneratedLayout,
 		MeshPtr currentLayoutMesh,
 		size_t currentLayoutChannel,
-		const void* errorContext )
-    {
-        if (currentLayoutMesh->GetVertexCount()==0)
-        {
-            return;
-        }
+		const void* errorContext,
+		const bool bClampUVIslands)
+	{
+		MUTABLE_CPUPROFILER_SCOPE(LayoutUV_PrepareForLayout);
+
+		if (currentLayoutMesh->GetVertexCount() == 0)
+		{
+			return;
+		}
 
 		// The layout we are adding must have block ids.
 		check(GeneratedLayout->m_blocks.IsEmpty() || GeneratedLayout->m_blocks[0].m_id != -1);
 
-        // 
-		Ptr<const Layout> pLayout = GeneratedLayout;
-        currentLayoutMesh->AddLayout( pLayout );
-
-        int buffer = -1;
-        int channel = -1;
-        currentLayoutMesh->GetVertexBuffers().FindChannel( MBS_TEXCOORDS,
-                                                              (int)currentLayoutChannel,
-                                                              &buffer,
-                                                              &channel );
-        check( buffer>=0 );
-        check( channel>=0 );
-
-        // Create the layout block vertex buffer
-        uint16* pLayoutData = 0;
-        {
-            int layoutBuf = currentLayoutMesh->GetVertexBuffers().GetBufferCount();
-            currentLayoutMesh->GetVertexBuffers().SetBufferCount( layoutBuf+1 );
 
-            // TODO
-            check( pLayout->GetBlockCount()<65535 );
-            MESH_BUFFER_SEMANTIC semantic = MBS_LAYOUTBLOCK;
-            int semanticIndex = int(currentLayoutChannel);
-            MESH_BUFFER_FORMAT format = MBF_UINT16;
-            int components = 1;
-            int offset = 0;
-            currentLayoutMesh->GetVertexBuffers().SetBuffer
-                    (
-                        layoutBuf,
-                        sizeof(uint16),
-                        1,
-                        &semantic, &semanticIndex,
-                        &format, &components,
-                        &offset
-                    );
-            pLayoutData = (uint16*)currentLayoutMesh->GetVertexBuffers().GetBufferData( layoutBuf );
-        }
+		// 
+		Ptr<const Layout> Layout = GeneratedLayout;
+		currentLayoutMesh->AddLayout(Layout);
 
-        // Get the information about the texture coordinates channel
-        MESH_BUFFER_SEMANTIC semantic;
-        int semanticIndex;
-        MESH_BUFFER_FORMAT format;
-        int components;
-        int offset;
-        currentLayoutMesh->GetVertexBuffers().GetChannel
-            ( buffer, channel, &semantic, &semanticIndex, &format, &components, &offset );
-        check( semantic == MBS_TEXCOORDS );
-
-        const uint8_t* pData = currentLayoutMesh->GetVertexBuffers().GetBufferData( buffer );
-        int elemSize = currentLayoutMesh->GetVertexBuffers().GetElementSize( buffer );
-        int channelOffset = currentLayoutMesh->GetVertexBuffers().GetChannelOffset( buffer, channel );
-        pData += channelOffset;
-
-        // Clear block data
-        for (int i=0; i< currentLayoutMesh->GetVertexBuffers().GetElementCount(); ++i)
-        {
-            pLayoutData[i] = 65535;
-        }
+		// Create the layout block vertex buffer
+		uint16* LayoutData = 0;
+		{
+			const int32 LayoutBufferIndex = currentLayoutMesh->GetVertexBuffers().GetBufferCount();
+			currentLayoutMesh->GetVertexBuffers().SetBufferCount(LayoutBufferIndex + 1);
+
+			// TODO
+			check(Layout->GetBlockCount() < MAX_uint16);
+			const MESH_BUFFER_SEMANTIC semantic = MBS_LAYOUTBLOCK;
+			const int32 semanticIndex = int32(currentLayoutChannel);
+			const MESH_BUFFER_FORMAT format = MBF_UINT16;
+			const int32 components = 1;
+			const int32 offset = 0;
+			currentLayoutMesh->GetVertexBuffers().SetBuffer
+			(
+				LayoutBufferIndex,
+				sizeof(uint16),
+				1,
+				&semantic, &semanticIndex,
+				&format, &components,
+				&offset
+			);
+			LayoutData = (uint16*)currentLayoutMesh->GetVertexBuffers().GetBufferData(LayoutBufferIndex);
+		}
 
-        // TODO: Check overlapping layout blocks
-        // TODO: Check triangles crossing blocks
-        int inside = 0;
-        int temp = 0;
-        for ( int b=0; b<pLayout->GetBlockCount(); ++b )
-        {
-            FIntPoint grid = pLayout->GetGridSize();
+		const int32 NumVertices = currentLayoutMesh->GetVertexCount();
+		const int32 NumBlocks = Layout->GetBlockCount();
 
-            box< UE::Math::TIntVector2<uint16> > block;
-            pLayout->GetBlock( b, &block.min[0], &block.min[1], &block.size[0], &block.size[1] );
+		// Clear block data
+		const uint16 NullBlockId = MAX_uint16 - 1;
+		for (int32 VertexIndex = 0; VertexIndex < NumVertices; ++VertexIndex)
+		{
+			*(LayoutData + VertexIndex) = NullBlockId;
+		}
 
-            box< vec2<float> > rect;
-            rect.min[0] = ( (float)block.min[0] ) / (float) grid[0];
-            rect.min[1] = ( (float)block.min[1] ) / (float) grid[1];
-            rect.size[0] = ( (float)block.size[0] ) / (float) grid[0];
-            rect.size[1] = ( (float)block.size[1] ) / (float) grid[1];
+		// Find block ids for each block in the grid
+		const FIntPoint Grid = Layout->GetGridSize();
+		TArray<int32> GridBlockBlockId;
+		GridBlockBlockId.Init(MAX_uint16, Grid.X * Grid.Y);
 
-            const uint8_t* pVertices = pData;
+		// 
+		TArray<uint16> BlockIds;
+		TArray<box<vec2<float>>> BlockRects;
 
-            if ( format == MBF_FLOAT32 )
-            {
-                for ( int v=0; v<currentLayoutMesh->GetVertexBuffers().GetElementCount(); ++v )
-                {
-                    vec2<float>* pUV = (vec2<float>*)pVertices;
-                    pVertices += elemSize;
+		BlockIds.SetNumUninitialized(NumBlocks);
+		BlockRects.SetNumUninitialized(NumBlocks);
 
-                    if ( pLayoutData[v]==65535 && rect.ContainsInclusive(*pUV) )
-                    {
-                        *pUV = rect.Homogenize( *pUV );
+		// Create an array of block index per cell
+		TArray<int32> OverlappingBlocks;
+		for (int32 BlockIndex = 0; BlockIndex < NumBlocks; ++BlockIndex)
+		{
+			// Get the block id
+			BlockIds[BlockIndex] = Layout->m_blocks[BlockIndex].m_id;
 
-                        // Set the value to the unique block id
-                        check( pLayout->m_blocks[b].m_id < 65535 );
-                        pLayoutData[v] =(uint16) pLayout->m_blocks[b].m_id;
+			// Get the block rect
+			uint16 MinX, MinY, SizeX, SizeY;
+			Layout->GetBlock(BlockIndex, &MinX, &MinY, &SizeX, &SizeY);
 
-                        inside++;
-                    }
-                    else
-                    {
-                        temp++;
-                    }
-                }
-            }
-            else if ( format == MBF_FLOAT16 )
-            {
-                // TODO: Very slow?
-                for ( int v=0; v<currentLayoutMesh->GetVertexBuffers().GetElementCount(); ++v )
-                {
-                    float16* pUV = (float16*)pVertices;
-                    pVertices += elemSize;
+			box<vec2<float>>& BlockRect = BlockRects[BlockIndex];
+			BlockRect.min[0] = ((float)MinX) / (float)Grid.X;
+			BlockRect.min[1] = ((float)MinY) / (float)Grid.Y;
+			BlockRect.size[0] = ((float)SizeX) / (float)Grid.X;
+			BlockRect.size[1] = ((float)SizeY) / (float)Grid.Y;
 
-                    vec2<float> UV( halfToFloat(pUV[0]), halfToFloat(pUV[1]) );
+			// Create block index per cell array
+			for (uint16 Y = MinY; Y < MinY + SizeY; ++Y)
+			{
+				const uint16 PositionY = Y * Grid.X;
 
-                    if ( pLayoutData[v]==65535 && rect.ContainsInclusive(UV) )
-                    {
-                        UV = rect.Homogenize( UV );
+				for (uint16 X = MinX; X < MinX + SizeX; ++X)
+				{
+					if (GridBlockBlockId[PositionY + X] == MAX_uint16)
+					{
+						GridBlockBlockId[PositionY + X] = BlockIndex;
+					}
+					else
+					{
+						OverlappingBlocks.AddUnique(BlockIndex);
+					}
+				}
+			}
+		}
 
-                        // Set the value to the unique block id
-                        check( pLayout->m_blocks[b].m_id < 65535 );
-                        pLayoutData[v] = (uint16)pLayout->m_blocks[b].m_id;
 
-                        inside++;
+		// Notify Overlapping layout blocks
+		if (!OverlappingBlocks.IsEmpty())
+		{
+			char buf[256];
+			mutable_snprintf
+			(
+				buf, 256,
+				"Source mesh has %d layout block overlapping in LOD %d",
+				OverlappingBlocks.Num() + 1, m_currentParents.Last().m_lod
+			);
+			m_pErrorLog->GetPrivate()->Add(buf, ELMT_WARNING, errorContext);
+		}
 
-                        pUV[0] = floatToHalf( UV[0] );
-                        pUV[1] = floatToHalf( UV[1] );
-                    }
-                }
-            }
-        }
+		// Get the information about the texture coordinates channel
+		int buffer = -1;
+		int channel = -1;
+		currentLayoutMesh->GetVertexBuffers().FindChannel(MBS_TEXCOORDS,
+			(int)currentLayoutChannel,
+			&buffer,
+			&channel);
+		check(buffer >= 0);
+		check(channel >= 0);
+
+		MESH_BUFFER_SEMANTIC semantic;
+		int semanticIndex;
+		MESH_BUFFER_FORMAT format;
+		int components;
+		int offset;
+		currentLayoutMesh->GetVertexBuffers().GetChannel
+		(buffer, channel, &semantic, &semanticIndex, &format, &components, &offset);
+		check(semantic == MBS_TEXCOORDS);
+
+		const uint8_t* pData = currentLayoutMesh->GetVertexBuffers().GetBufferData(buffer);
+		int elemSize = currentLayoutMesh->GetVertexBuffers().GetElementSize(buffer);
+		int channelOffset = currentLayoutMesh->GetVertexBuffers().GetChannelOffset(buffer, channel);
+		pData += channelOffset;
+
+
+		// Temp copy of the UVs
+		TArray<vec2<float>> TempUVs;
+		TempUVs.SetNumUninitialized(NumVertices);
+
+		// Get a copy of the UVs as vec2<float> to work with them. 
+		{
+			bool bNonNormalizedUVs = false;
 
-        int outside = currentLayoutMesh->GetVertexBuffers().GetElementCount() - inside;
-        if (outside>0 && (pLayout->FirstLODToIgnoreWarnings == -1 || m_currentParents.Last().m_lod < pLayout->FirstLODToIgnoreWarnings))
-        {
-            char buf[256];
-            mutable_snprintf
-                (
-                    buf, 256,
-                    "Source mesh has %d vertices not assigned to any layout block in LOD %d",
-					outside, m_currentParents.Last().m_lod
-                );
-			int blockCount = pLayout->GetBlockCount();
-        
-            //m_pErrorLog->GetPrivate()->Add( buf, blockCount==1?ELMT_INFO:ELMT_WARNING, errorContext );
-            TArray< float > unassignedUVs;
-            unassignedUVs.Reserve(64);           
-            
-            const uint8_t* pVertices = pData;
-            for (int i=0; i<currentLayoutMesh->GetVertexBuffers().GetElementCount(); ++i)
-            {
-                vec2<float> UV;
-                if ( format == MBF_FLOAT32 )
-                {
-                    UV = *((vec2<float>*)pVertices);
-                }
-                else if ( format == MBF_FLOAT16 )
-                {
-                    float16* pUV = (float16*)pVertices;
-                    UV = vec2<float>( halfToFloat(pUV[0]), halfToFloat(pUV[1]) );
-                }
+			const uint8* pVertices = pData;
+			for (int32 VertexIndex = 0; VertexIndex < NumVertices; ++VertexIndex)
+			{
+				vec2<float>& UV = TempUVs[VertexIndex];
+				if (format == MBF_FLOAT32)
+				{
+					UV = *((vec2<float>*)pVertices);
+				}
+				else if (format == MBF_FLOAT16)
+				{
+					float16* pUV = (float16*)pVertices;
+					UV = vec2<float>(halfToFloat(pUV[0]), halfToFloat(pUV[1]));
+				}
 
-                pVertices += elemSize;
+				// Check that UVs are normalized. If not, clamp the values and throw a warning.
+				if (UV[0] < 0.f || UV[0] > 1.f || UV[1] < 0.f || UV[1] > 1.f)
+				{
+					UV[0] = FMath::Clamp(UV[0], 0.f, 1.f);
+					UV[1] = FMath::Clamp(UV[1], 0.f, 1.f);
+					bNonNormalizedUVs = true;
+				}
 
-                if (pLayoutData[i]==65535)
-                {
-                    unassignedUVs.Add(UV[0]);
-                    unassignedUVs.Add(UV[1]);
-                }
-            }
+				pVertices += elemSize;
+			}
+
+			// Mutable does not support non-normalized UVs
+			if (bNonNormalizedUVs)
+			{
+				char buf[256];
+				mutable_snprintf
+				(
+					buf, 256,
+					"Source mesh has non-normalized UVs in LOD %d",
+					m_currentParents.Last().m_lod
+				);
+				m_pErrorLog->GetPrivate()->Add(buf, ELMT_WARNING, errorContext);
+			}
+		}
 
-            ErrorLogMessageAttachedDataView attachedDataView;
-            attachedDataView.m_unassignedUVs = unassignedUVs.GetData();
-            attachedDataView.m_unassignedUVsSize = (size_t)unassignedUVs.Num();
 
-            m_pErrorLog->GetPrivate()->Add( buf, attachedDataView, ELMT_WARNING, errorContext );
-        }
-        
-        // Assign broken vertices to the first block
-        for (int i=0; i< currentLayoutMesh->GetVertexBuffers().GetElementCount(); ++i)
-        {
-            if ( pLayoutData[i]==65535 )
-            {
-                pLayoutData[i] = 0;
-            }
-        }
-    }
+		const int32 NumTriangles = currentLayoutMesh->GetIndexCount() / 3;
+		TArray<FTriangle> Triangles;
+
+		// Vertices mapped to unique vertex index
+		TArray<int32> CollapsedVertices;
+
+		// Vertex to face map used to speed up connectivity building
+		TMultiMap<int32, uint32> VertexToFaceMap;
+
+		// Find Unique Vertices
+		if (bClampUVIslands)
+		{
+			VertexToFaceMap.Reserve(NumVertices);
+			Triangles.SetNumUninitialized(NumTriangles);
+
+			MeshCreateCollapsedVertexMap(currentLayoutMesh.get(), CollapsedVertices);
+		}
+
+		UntypedMeshBufferIteratorConst ItIndices(currentLayoutMesh->GetIndexBuffers(), MBS_VERTEXINDEX);
+		TArray<int32> ConflictiveTriangles;
 
+		// TODO: We could skip this if there's only one block and it fits the entire grid
+		for (int32 TriangleIndex = 0; TriangleIndex < NumTriangles; ++TriangleIndex)
+		{
+			uint32 Index0 = int(ItIndices.GetAsUINT32());
+			++ItIndices;
+			uint32 Index1 = int(ItIndices.GetAsUINT32());
+			++ItIndices;
+			uint32 Index2 = int(ItIndices.GetAsUINT32());
+			++ItIndices;
+
+			uint16 X, Y;
+
+			uint16& BlockIndexV0 = *(LayoutData + Index0);
+			if (BlockIndexV0 == NullBlockId)
+			{
+				X = (uint16)FMath::Min<uint32>(Grid.X - 1, FMath::Max<uint32>(0, (uint32)(Grid.X * TempUVs[Index0][0])));
+				Y = (uint16)FMath::Min<uint32>(Grid.Y - 1, FMath::Max<uint32>(0, (uint32)(Grid.Y * TempUVs[Index0][1])));
+				BlockIndexV0 = GridBlockBlockId[Y * Grid.X + X];
+			}
+
+			uint16& BlockIndexV1 = *(LayoutData + Index1);
+			if (BlockIndexV1 == NullBlockId)
+			{
+				X = (uint16)FMath::Min<uint32>(Grid.X - 1, FMath::Max<uint32>(0, (uint32)(Grid.X * TempUVs[Index1][0])));
+				Y = (uint16)FMath::Min<uint32>(Grid.Y - 1, FMath::Max<uint32>(0, (uint32)(Grid.Y * TempUVs[Index1][1])));
+				BlockIndexV1 = GridBlockBlockId[Y * Grid.X + X];
+			}
+
+			uint16& BlockIndexV2 = *(LayoutData + Index2);
+			if (BlockIndexV2 == NullBlockId)
+			{
+				X = (uint16)FMath::Min<uint32>(Grid.X - 1, FMath::Max<uint32>(0, (uint32)(Grid.X * TempUVs[Index2][0])));
+				Y = (uint16)FMath::Min<uint32>(Grid.Y - 1, FMath::Max<uint32>(0, (uint32)(Grid.Y * TempUVs[Index2][1])));
+				BlockIndexV2 = GridBlockBlockId[Y * Grid.X + X];
+			}
+
+			if (bClampUVIslands)
+			{
+				if (BlockIndexV0 != BlockIndexV1 || BlockIndexV0 != BlockIndexV2)
+				{
+					ConflictiveTriangles.Add(TriangleIndex);
+				}
+
+				FTriangle& Triangle = Triangles[TriangleIndex];
+
+				Triangle.Indices[0] = Index0;
+				Triangle.Indices[1] = Index1;
+				Triangle.Indices[2] = Index2;
+				Triangle.CollapsedIndices[0] = CollapsedVertices[Index0];
+				Triangle.CollapsedIndices[1] = CollapsedVertices[Index1];
+				Triangle.CollapsedIndices[2] = CollapsedVertices[Index2];
+
+				Triangle.BlockIndices[0] = BlockIndexV0;
+				Triangle.BlockIndices[1] = BlockIndexV1;
+				Triangle.BlockIndices[2] = BlockIndexV2;
+				Triangle.bUVsFixed = false;
+
+				VertexToFaceMap.Add(Triangle.CollapsedIndices[0], TriangleIndex);
+				VertexToFaceMap.Add(Triangle.CollapsedIndices[1], TriangleIndex);
+				VertexToFaceMap.Add(Triangle.CollapsedIndices[2], TriangleIndex);
+			}
+		}
+
+		// Clamp UV islands to the predominant block of each island. Will only happen if bClampUVIslands is true.
+		for (const int32& ConflictiveTriangleIndex : ConflictiveTriangles)
+		{
+			FTriangle& Triangle = Triangles[ConflictiveTriangleIndex];
+
+			// Skip the ones that have been fixed already
+			if (Triangle.bUVsFixed)
+			{
+				continue;
+			}
+
+			// Find triangles from the same UV Island
+			TArray<uint32> TriangleIndices;
+			GetUVIsland(Triangles, ConflictiveTriangleIndex, TriangleIndices, TempUVs, VertexToFaceMap);
+
+			// Get predominant BlockId != MAX_uint16
+			TArray<uint32> NumVerticesPerBlock;
+			NumVerticesPerBlock.SetNumZeroed(NumBlocks);
+
+			for (const int32& TriangleIndex : TriangleIndices)
+			{
+				FTriangle& OtherTriangle = Triangles[TriangleIndex];
+				for (int32 VertexIndex = 0; VertexIndex < 3; ++VertexIndex)
+				{
+					const uint16& BlockIndex = OtherTriangle.BlockIndices[VertexIndex];
+					if (BlockIndex != MAX_uint16)
+					{
+						NumVerticesPerBlock[BlockIndex]++;
+					}
+				}
+			}
+
+			uint16 BlockIndex = 0;
+			uint32 CurrentMaxVertices = 0;
+			for (int32 Index = 0; Index < NumBlocks; ++Index)
+			{
+				if (NumVerticesPerBlock[Index] > CurrentMaxVertices)
+				{
+					BlockIndex = BlockIndex;
+					CurrentMaxVertices = NumVerticesPerBlock[Index];
+				}
+			}
+
+			// Get the limits of the predominant block rect
+			const Layout::FBlock& LayoutBlock = Layout->m_blocks[BlockIndex];
+
+			const float SmallNumber = 0.000001;
+			const float MinX = ((float)LayoutBlock.m_min.X) / (float)Grid.X + SmallNumber;
+			const float MinY = ((float)LayoutBlock.m_min.Y) / (float)Grid.Y + SmallNumber;
+			const float MaxX = (((float)LayoutBlock.m_size.X + LayoutBlock.m_min.X) / (float)Grid.X) - 2 * SmallNumber;
+			const float MaxY = (((float)LayoutBlock.m_size.Y + LayoutBlock.m_min.Y) / (float)Grid.Y) - 2 * SmallNumber;
+
+			// Iterate triangles and clamp the UVs
+			vec2<float>* TempUVsData = TempUVs.GetData();
+			for (const int32& TriangleIndex : TriangleIndices)
+			{
+				FTriangle& OtherTriangle = Triangles[TriangleIndex];
+
+				for (int8 VertexIndex = 0; VertexIndex < 3; ++VertexIndex)
+				{
+					if (OtherTriangle.BlockIndices[VertexIndex] == BlockIndex)
+					{
+						continue;
+					}
+
+					OtherTriangle.BlockIndices[VertexIndex] = BlockIndex;
+
+					// Clamp UVs
+					const int32 UVIndex = OtherTriangle.Indices[VertexIndex];
+					vec2<float>& UV = TempUVs[UVIndex];
+					UV[0] = FMath::Clamp(UV[0], MinX, MaxX);
+					UV[1] = FMath::Clamp(UV[1], MinY, MaxY);
+					*(LayoutData + UVIndex) = BlockIndex;
+				}
+
+				OtherTriangle.bUVsFixed = true;
+			}
+		}
+
+		// Warn about vertices without a block id
+		if (Layout->FirstLODToIgnoreWarnings == -1 || m_currentParents.Last().m_lod < Layout->FirstLODToIgnoreWarnings)
+		{
+			TArray<float> UnassignedUVs;
+			UnassignedUVs.Reserve(NumVertices / 100);
+
+			const vec2<float>* UVs = TempUVs.GetData();
+			for (int32 VertexIndex = 0; VertexIndex < NumVertices; ++VertexIndex)
+			{
+				if (LayoutData[VertexIndex] == MAX_uint16)
+				{
+					UnassignedUVs.Add((*(UVs + VertexIndex))[0]);
+					UnassignedUVs.Add((*(UVs + VertexIndex))[1]);
+				}
+			}
+
+			if (!UnassignedUVs.IsEmpty())
+			{
+				char buf[256];
+				mutable_snprintf
+				(
+					buf, 256,
+					"Source mesh has %d vertices not assigned to any layout block in LOD %d",
+					UnassignedUVs.Num(), m_currentParents.Last().m_lod
+				);
+
+				ErrorLogMessageAttachedDataView attachedDataView;
+				attachedDataView.m_unassignedUVs = UnassignedUVs.GetData();
+				attachedDataView.m_unassignedUVsSize = (size_t)UnassignedUVs.Num();
+
+				m_pErrorLog->GetPrivate()->Add(buf, attachedDataView, ELMT_WARNING, errorContext);
+			}
+		}
+
+		// Format and copy UVs
+		{
+			uint8* pVertices = const_cast<uint8*>(pData);
+			vec2<float>* UVs = TempUVs.GetData();
+
+			for (int32 VertexIndex = 0; VertexIndex < NumVertices; ++VertexIndex)
+			{
+				vec2<float>* UV = &TempUVs[VertexIndex];
+
+				uint16& LayoutBlockIndex = LayoutData[VertexIndex];
+				if (BlockIds.IsValidIndex(LayoutBlockIndex))
+				{
+					// Homogenize UVs
+					*UV = BlockRects[LayoutBlockIndex].Homogenize(*UV);
+
+					// Replace block index by the actual id of the block
+					LayoutBlockIndex = BlockIds[LayoutBlockIndex];
+				}
+				else
+				{
+					// Map vertices without block to the first block.
+					LayoutBlockIndex = 0;
+				}
+
+				// Copy UVs
+				if (format == MBF_FLOAT32)
+				{
+					vec2<float>* pUV = (vec2<float>*)pVertices;
+					*pUV = *UV;
+					check((*pUV)[0] <= 1 && (*pUV)[1] <= 1);
+				}
+				else if (format == MBF_FLOAT16)
+				{
+					float16* pUV = (float16*)pVertices;
+					pUV[0] = floatToHalf((*UV)[0]);
+					pUV[1] = floatToHalf((*UV)[1]);
+				}
+
+				pVertices += elemSize;
+			}
+		}
+	}
+	
 
     //---------------------------------------------------------------------------------------------
     void CodeGenerator::GenerateMesh( const FMeshGenerationOptions& InOptions, FMeshGenerationResult& OutResult, const NodeMeshPtrConst& InUntypedNode)
@@ -968,7 +1384,7 @@ class Node;
 						{
 							Ptr<const Layout> SourceLayout = TypedNode->GetPrivate()->m_pLayout;
 							Ptr<const Layout> GeneratedLayout = AddLayout( SourceLayout );
-							PrepareForLayout(GeneratedLayout, pCloned, LayoutIndex, TypedNode->GetPrivate()->m_errorContext);
+							PrepareForLayout(GeneratedLayout, pCloned, LayoutIndex, TypedNode->GetPrivate()->m_errorContext, InOptions.bClampUVIslands);
 
 							OutResult.GeneratedLayouts.Add(GeneratedLayout);
 						}
@@ -980,7 +1396,7 @@ class Node;
 					for (int32 LayoutIndex = 0; LayoutIndex < InOptions.OverrideLayouts.Num(); ++LayoutIndex)
 					{
 						Ptr<const Layout> GeneratedLayout = InOptions.OverrideLayouts[LayoutIndex];
-						PrepareForLayout(GeneratedLayout, pCloned, LayoutIndex, node.m_errorContext);
+						PrepareForLayout(GeneratedLayout, pCloned, LayoutIndex, node.m_errorContext, InOptions.bClampUVIslands);
 
 						OutResult.GeneratedLayouts.Add(GeneratedLayout);
 					}
