[Mutable] Fix compilation error on Linux

#rnx

diff --git a/Engine/Plugins/Experimental/Mutable/Source/MutableTools/Private/MuT/CodeGenerator_Mesh.cpp b/Engine/Plugins/Experimental/Mutable/Source/MutableTools/Private/MuT/CodeGenerator_Mesh.cpp
--- a/Engine/Plugins/Experimental/Mutable/Source/MutableTools/Private/MuT/CodeGenerator_Mesh.cpp
+++ b/Engine/Plugins/Experimental/Mutable/Source/MutableTools/Private/MuT/CodeGenerator_Mesh.cpp
@@ -542,7 +542,7 @@ class Node;
 		}
 
 		// Clamp UV islands to the predominant block of each island. Will only happen if bClampUVIslands is true.
-		for (const int32& ConflictiveTriangleIndex : ConflictiveTriangles)
+		for (int32 ConflictiveTriangleIndex : ConflictiveTriangles)
 		{
 			FTriangle& Triangle = Triangles[ConflictiveTriangleIndex];
 
@@ -560,7 +560,7 @@ class Node;
 			TArray<uint32> NumVerticesPerBlock;
 			NumVerticesPerBlock.SetNumZeroed(NumBlocks);
 
-			for (const int32& TriangleIndex : TriangleIndices)
+			for (int32 TriangleIndex : TriangleIndices)
 			{
 				FTriangle& OtherTriangle = Triangles[TriangleIndex];
 				for (int32 VertexIndex = 0; VertexIndex < 3; ++VertexIndex)
@@ -579,7 +579,7 @@ class Node;
 			{
 				if (NumVerticesPerBlock[Index] > CurrentMaxVertices)
 				{
-					BlockIndex = BlockIndex;
+					BlockIndex = Index;
 					CurrentMaxVertices = NumVerticesPerBlock[Index];
 				}
 			}
@@ -595,7 +595,7 @@ class Node;
 
 			// Iterate triangles and clamp the UVs
 			vec2<float>* TempUVsData = TempUVs.GetData();
-			for (const int32& TriangleIndex : TriangleIndices)
+			for (int32 TriangleIndex : TriangleIndices)
 			{
 				FTriangle& OtherTriangle = Triangles[TriangleIndex];
 
