[Mutable] Use PointHashGrid3f to optimize MeshCreateCollapsedVertexMap.
Remove Mutable's implementation of an Octree

#rnx

diff --git a/Engine/Plugins/Experimental/Mutable/Source/MutableRuntime/Private/MuR/Octree.h b/Engine/Plugins/Experimental/Mutable/Source/MutableRuntime/Private/MuR/Octree.h
deleted file mode 100644
diff --git a/Engine/Plugins/Experimental/Mutable/Source/MutableRuntime/Private/MuR/OpImageProject.cpp b/Engine/Plugins/Experimental/Mutable/Source/MutableRuntime/Private/MuR/OpImageProject.cpp
--- a/Engine/Plugins/Experimental/Mutable/Source/MutableRuntime/Private/MuR/OpImageProject.cpp
+++ b/Engine/Plugins/Experimental/Mutable/Source/MutableRuntime/Private/MuR/OpImageProject.cpp
@@ -10,7 +10,6 @@
 #include "MuR/MeshPrivate.h"
 #include "MuR/MutableMath.h"
 #include "MuR/MutableTrace.h"
-#include "MuR/Octree.h"
 #include "MuR/OpMeshClipWithMesh.h"
 #include "MuR/Parameters.h"
 #include "MuR/ParametersPrivate.h"
@@ -26,6 +25,8 @@
 
 #include "Math/VectorRegister.h"
 
+#include "Spatial/PointHashGrid3.h"
+
 namespace mu
 {
 
@@ -1444,88 +1445,59 @@ constexpr float vert_collapse_eps = 0.0001f;
 //! Create a map from vertices into vertices, collapsing vertices that have the same position
 //---------------------------------------------------------------------------------------------
 inline void MeshCreateCollapsedVertexMap( const Mesh* pMesh,
-	TArray<int>& collapsedVertexMap,
-	TArray<FVector3f>& vertices,
-	TMultiMap<int, int>& collapsedVertsMap
-)
+	TArray<int32>& CollapsedVertices,
+	TMultiMap<int32, int32>& CollapsedVerticesMap,
+	TArray<FVector3f>& OutVertices)
 {
 	MUTABLE_CPUPROFILER_SCOPE(CreateCollapseMap);
 
-    int vcount = pMesh->GetVertexCount();
-    collapsedVertexMap.SetNum(vcount);
-    vertices.SetNum(vcount);
-	collapsedVertsMap.Reserve(vcount);
-
-    const FMeshBufferSet& MBSPriv = pMesh->GetVertexBuffers();
-    for (int32 b = 0; b < MBSPriv.m_buffers.Num(); ++b)
-    {
-        for (int32 c = 0; c < MBSPriv.m_buffers[b].m_channels.Num(); ++c)
-        {
-            MESH_BUFFER_SEMANTIC sem = MBSPriv.m_buffers[b].m_channels[c].m_semantic;
-            int semIndex = MBSPriv.m_buffers[b].m_channels[c].m_semanticIndex;
-
-			UntypedMeshBufferIteratorConst it = UntypedMeshBufferIteratorConst(pMesh->GetVertexBuffers(), sem, semIndex);
+	const int32 NumVertices = pMesh->GetVertexCount();
 
-			box<FVector3f> meshBoundingBox;
-			Octree* pOctree = nullptr;
+	// Used to speed up vertex comparison
+	UE::Geometry::TPointHashGrid3f<int32> VertHash(0.01f, INDEX_NONE);
+	VertHash.Reserve(NumVertices);
 
-            switch (sem)
-            {
-            case MBS_POSITION:
-
-                // First create a cache of the vertices in the vertices array
-                for (int v = 0; v < vcount; ++v)
-                {
-                    FVector3f vertex(0.0f, 0.0f, 0.0f);
-                    for (int i = 0; i < 3; ++i)
-                    {
-                        ConvertData(i, &vertex[0], MBF_FLOAT32, it.ptr(), it.GetFormat());
-                    }
+	// Info to collect. Vertices, collapsed vertices and unique vertices to nearby vertices map
+	OutVertices.SetNumUninitialized(NumVertices);
+	CollapsedVertices.Init(INDEX_NONE, NumVertices);
+	CollapsedVerticesMap.Reserve(NumVertices);
 
-                    vertices[v] = vertex;
+	// Get Vertices
+	mu::UntypedMeshBufferIteratorConst ItPosition = mu::UntypedMeshBufferIteratorConst(pMesh->GetVertexBuffers(), mu::MBS_POSITION);
+	FVector3f* VertexData = OutVertices.GetData();
+	
+	for (int32 VertexIndex = 0; VertexIndex < NumVertices; ++VertexIndex)
+	{
+		*VertexData = ItPosition.GetAsVec3f();
+		VertHash.InsertPointUnsafe(VertexIndex, *VertexData);
 
-                    ++it;
-                }
-				
-				//Initializing Octree boundingbox
-				meshBoundingBox.min = vertices[0];
-				meshBoundingBox.size = FVector3f(0.0f, 0.0f, 0.0f);
-				for (int32 i = 1; i < vertices.Num(); ++i)
-				{
-					meshBoundingBox.Bound3(vertices[i]);
-				}
+		++ItPosition;
+		++VertexData;
+	}
 
-				//Initializing and filling Octree
-				pOctree = new Octree(meshBoundingBox.min, meshBoundingBox.min + meshBoundingBox.size,0.001f);
-				for (int32 itr = 0; itr < vertices.Num(); ++itr)
-				{
-					pOctree->InsertElement(vertices[itr], int(itr));
-				}
 
-				// Create map to store which vertices are the same (collapse nearby vertices)
-				for (int32 itr = 0; itr < vertices.Num(); ++itr)
-				{
-					int collapsed_candidate_v_idx = pOctree->GetNearests(vertices[itr], vert_collapse_eps, int(itr));
-					collapsedVertexMap[itr] = collapsed_candidate_v_idx;
+	// Find unique vertices
+	TArray<int32> NearbyVertices;
+	for (int32 VertexIndex = 0; VertexIndex < NumVertices; ++VertexIndex)
+	{
+		if (CollapsedVertices[VertexIndex] != INDEX_NONE)
+		{
+			continue;
+		}
 
-					if (collapsed_candidate_v_idx != int(itr))
-					{
-						collapsedVertsMap.Add(collapsed_candidate_v_idx, itr);
-					}
-				}
-                break;
+		const FVector3f& Vertex = OutVertices[VertexIndex];
 
-            default:
-                break;
-            }
+		NearbyVertices.Reset();
+		VertHash.FindPointsInBall(Vertex, vert_collapse_eps,
+			[&Vertex, &OutVertices](const int32& Other) -> float {return FVector3f::DistSquared(OutVertices[Other], Vertex); },
+			NearbyVertices);
 
-			if (pOctree)
-			{
-				delete pOctree;
-				pOctree = nullptr;
-			}
-        }
-    }
+		for (int32 NearbyVertexIndex : NearbyVertices)
+		{
+			CollapsedVertices[NearbyVertexIndex] = VertexIndex;
+			CollapsedVerticesMap.Add(VertexIndex, NearbyVertexIndex);
+		}
+	}
 }
 
 
@@ -1972,10 +1944,10 @@ void MeshProject_Optimised_Wrapping( const Mesh* pMesh,
 	faceStep.SetNum(faceCount);
 
     // Map vertices to the one they are collapsed to because they are very similar, if they aren't collapsed then they are mapped to themselves
-    TArray<int> collapsedVertexMap;
+    TArray<int32> collapsedVertexMap;
 	TArray<FVector3f> vertices;
-    TMultiMap<int, int> collapsedVertsMap; // Maps a collapsed vertex to all the vertices that collapse to it
-    MeshCreateCollapsedVertexMap(pMesh, collapsedVertexMap, vertices, collapsedVertsMap);
+    TMultiMap<int32, int32> collapsedVertsMap; // Maps a collapsed vertex to all the vertices that collapse to it
+    MeshCreateCollapsedVertexMap(pMesh, collapsedVertexMap, collapsedVertsMap, vertices);
 
     // Used to speed up connectivity building
 	TMultiMap<int, int> vertToFacesMap;
@@ -1990,9 +1962,6 @@ void MeshProject_Optimised_Wrapping( const Mesh* pMesh,
         vertToFacesMap.Add(collapsedVertexMap[i0], f);
         vertToFacesMap.Add(collapsedVertexMap[i1], f);
         vertToFacesMap.Add(collapsedVertexMap[i2], f);
-        //vertToFacesMap.insert(std::pair<int, int>(i0, f));
-        //vertToFacesMap.insert(std::pair<int, int>(i1, f));
-        //vertToFacesMap.insert(std::pair<int, int>(i2, f));
     }
 
     // Trace a ray in the projection direction to find the face that will be projected planarly and be the root of the unfolding
diff --git a/Engine/Plugins/Experimental/Mutable/Source/MutableRuntime/Private/MuR/OpMeshClipWithMesh.h b/Engine/Plugins/Experimental/Mutable/Source/MutableRuntime/Private/MuR/OpMeshClipWithMesh.h
--- a/Engine/Plugins/Experimental/Mutable/Source/MutableRuntime/Private/MuR/OpMeshClipWithMesh.h
+++ b/Engine/Plugins/Experimental/Mutable/Source/MutableRuntime/Private/MuR/OpMeshClipWithMesh.h
@@ -16,6 +16,7 @@
 
 #include "Math/UnrealMathUtility.h"
 
+#include "Spatial/PointHashGrid3.h"
 
 namespace mu
 {
@@ -24,72 +25,56 @@ namespace mu
     //! Create a map from vertices into vertices, collapsing vertices that have the same position, 
 	//! This version uses UE Containers to return.	
     //---------------------------------------------------------------------------------------------
-    inline void MeshCreateCollapsedVertexMap( const Mesh* pMesh, TArray<int32>& OutCollapsedVertexMap, TArray<FVector3f>& Vertices )
-    {
-        MUTABLE_CPUPROFILER_SCOPE(MeshCreateCollapsedVertexMap);
-    	
-        int32 VCount = pMesh->GetVertexCount();
+	inline void MeshCreateCollapsedVertexMap(const Mesh* pMesh,	TArray<int32>& OutCollapsedVertices, TArray<FVector3f>& OutVertices)
+	{
+		MUTABLE_CPUPROFILER_SCOPE(MeshCreateCollapsedVertexMap);
 
-        OutCollapsedVertexMap.SetNum(VCount);
-        Vertices.SetNum(VCount);
+		const int32 NumVertices = pMesh->GetVertexCount();
 
-        const FMeshBufferSet& MBSPriv = pMesh->GetVertexBuffers();
-        for (int32 b = 0; b < MBSPriv.m_buffers.Num(); ++b)
-        {
-            for (int32 c = 0; c < MBSPriv.m_buffers[b].m_channels.Num(); ++c)
-            {
-                MESH_BUFFER_SEMANTIC Sem = MBSPriv.m_buffers[b].m_channels[c].m_semantic;
-                int32 SemIndex = MBSPriv.m_buffers[b].m_channels[c].m_semanticIndex;
+		// Used to speed up vertex comparison
+		UE::Geometry::TPointHashGrid3f<int32> VertHash(0.01f, INDEX_NONE);
+		VertHash.Reserve(NumVertices);
 
-                UntypedMeshBufferIteratorConst It(pMesh->GetVertexBuffers(), Sem, SemIndex);
+		// Info to collect. Vertices and collapsed vertices
+		OutVertices.SetNumUninitialized(NumVertices);
+		OutCollapsedVertices.Init(INDEX_NONE, NumVertices);
 
-                switch (Sem)
-                {
-                case MBS_POSITION:
-				{
-                    // First create a cache of the vertices in the vertices array
-                    for (int32 V = 0; V < VCount; ++V)
-                    {
-                        FVector3f Vertex(0.0f, 0.0f, 0.0f);
-                        for (int32 I = 0; I < 3; ++I)
-                        {
-                            ConvertData(I, &Vertex[0], MBF_FLOAT32, It.ptr(), It.GetFormat());
-                        }
+		// Get Vertices
+		mu::UntypedMeshBufferIteratorConst ItPosition = mu::UntypedMeshBufferIteratorConst(pMesh->GetVertexBuffers(), mu::MBS_POSITION);
+		FVector3f* VertexData = OutVertices.GetData();
 
-                        Vertices[V] = Vertex;
+		for (int32 VertexIndex = 0; VertexIndex < NumVertices; ++VertexIndex)
+		{
+			*VertexData = ItPosition.GetAsVec3f();
+			VertHash.InsertPointUnsafe(VertexIndex, *VertexData);
 
-                        ++It;
-                    }
-					
-                    // Create map to store which vertices are the same (collapse nearby vertices)
-					const int32 VerticesNum = Vertices.Num();
-                    for (int32 V = 0; V < VerticesNum; ++V)
-                    {
-                        OutCollapsedVertexMap[V] = V;
+			++ItPosition;
+			++VertexData;
+		}
 
-                        for (int32 CandidateV = 0; CandidateV < V; ++CandidateV)
-                        {
-                            const int32 CollapsedCandidateVIdx = OutCollapsedVertexMap[CandidateV];
+		// Find unique vertices
+		TArray<int32> NearbyVertices;
+		for (int32 VertexIndex = 0; VertexIndex < NumVertices; ++VertexIndex)
+		{
+			if (OutCollapsedVertices[VertexIndex] != INDEX_NONE)
+			{
+				continue;
+			}
 
-                            // Test whether v and the collapsed candidate are close enough to be collapsed
-                            const FVector3f R = Vertices[CollapsedCandidateVIdx] - Vertices[V];
+			const FVector3f& Vertex = OutVertices[VertexIndex];
 
-                            if (R.Dot(R) <= TMathUtilConstants<float>::ZeroTolerance)
-                            {
-                                OutCollapsedVertexMap[V] = CollapsedCandidateVIdx;
-                                break;
-                            }
-                        }
-                    }
+			NearbyVertices.Reset();
+			VertHash.FindPointsInBall(Vertex, TMathUtilConstants<float>::ZeroTolerance,
+				[&Vertex, &OutVertices](const int32& Other) -> float {return FVector3f::DistSquared(OutVertices[Other], Vertex); },
+				NearbyVertices);
+
+			for (int32 NearbyVertexIndex : NearbyVertices)
+			{
+				OutCollapsedVertices[NearbyVertexIndex] = VertexIndex;
+			}
+		}
+	}
 
-                    break;
-				}
-                default:
-                    break;
-                }
-            }
-        }
-    }
 
     //---------------------------------------------------------------------------------------------
     //! Return true if a mesh is closed
