[mutable] Fixed morphs. Now a null target will leave the base unchanged.
+ Also fixed case where a morph factor of 1.0 returned the base mesh.

#rnx

diff --git a/Engine/Plugins/Experimental/Mutable/Source/MutableRuntime/Private/MuR/CodeRunner.cpp b/Engine/Plugins/Experimental/Mutable/Source/MutableRuntime/Private/MuR/CodeRunner.cpp
--- a/Engine/Plugins/Experimental/Mutable/Source/MutableRuntime/Private/MuR/CodeRunner.cpp
+++ b/Engine/Plugins/Experimental/Mutable/Source/MutableRuntime/Private/MuR/CodeRunner.cpp
@@ -1138,8 +1138,8 @@ namespace mu
 
 				float absFactor = FMath::Abs(factor);
                 float delta = 1.0f/(NumTargets -1);
-				int min = (int)FMath::FloorToFloat(absFactor/delta);
 				int max = (int)FMath::CeilToFloat(absFactor/delta);
+				int min = max > 0 ? max - 1 : 0;
 
                 // Factor from 0 to 1 between the two targets
 				float bifactor = absFactor/delta - min;
@@ -1218,42 +1218,21 @@ namespace mu
                     // Just the first of the targets
 					if ( bifactor < UE_SMALL_NUMBER && bifactor > -UE_SMALL_NUMBER )
                     {
-                        // Base with one full morph
-                        Ptr<const Mesh> pMorph = LoadMesh( FCacheAddress(Targets[min],item) );
-                        if (pMorph)
-                        {
-							pResult = MeshMorph(pBase.get(), pMorph.get());
-                        }
+						pResult = pBase->Clone();
                     }
-                    // Just the second of the targets
-                    else if ( bifactor > 1.0f - UE_SMALL_NUMBER )
+                    // Just the second of the targets (positive or negative)
+                    else if ( bifactor >= 1.0f - UE_SMALL_NUMBER || bifactor <= -1.0f + UE_SMALL_NUMBER)
                     {
                         check( max>0 );
                         Ptr<const Mesh> pMorph = LoadMesh( FCacheAddress(Targets[max],item) );
-                        if (pMorph)
-                        {
-							pResult = MeshMorph(pBase.get(), pMorph.get());
-                        }
+						pResult = MeshMorph(pBase.get(), pMorph.get(), bifactor);
                     }
-					// Negative target
-					else if (bifactor <= -1.0f + UE_SMALL_NUMBER)
-					{
-						check( max > 0 );
-						Ptr<const Mesh> pMorph = LoadMesh(FCacheAddress(Targets[max], item));
-						if (pMorph)
-						{
-							pResult = MeshMorph(pBase.get(), pMorph.get(), bifactor);
-						}
-					}
                     // Mix two targets on the base
                     else
                     {
                         Ptr<const Mesh> pMin = LoadMesh( FCacheAddress(Targets[min],item) );
                         Ptr<const Mesh> pMax = LoadMesh( FCacheAddress(Targets[max],item) );
-                        if (pMin && pMax)
-                        {
-                            pResult = MeshMorph2( pBase.get(), pMin.get(), pMax.get(), bifactor );
-                        }
+                        pResult = MeshMorph2( pBase.get(), pMin.get(), pMax.get(), bifactor );
                     }
 					// Missing branch. With the current system we can never get apply a negative pMin morph, only a negative pMax morph.
 					// This is due to not having a bifactor value that represents a negative pMin morph.
