[Mutable] Set reuse materials between LODs to false by default.

#rnx

diff --git a/Engine/Plugins/Experimental/Mutable/Source/CustomizableObjectEditor/Classes/MuCOE/Nodes/CustomizableObjectNodeMaterial.h b/Engine/Plugins/Experimental/Mutable/Source/CustomizableObjectEditor/Classes/MuCOE/Nodes/CustomizableObjectNodeMaterial.h
--- a/Engine/Plugins/Experimental/Mutable/Source/CustomizableObjectEditor/Classes/MuCOE/Nodes/CustomizableObjectNodeMaterial.h
+++ b/Engine/Plugins/Experimental/Mutable/Source/CustomizableObjectEditor/Classes/MuCOE/Nodes/CustomizableObjectNodeMaterial.h
@@ -108,7 +108,7 @@ public:
 	/** Reuse material between LODs when using the LODStrategy - Automatic from mesh. If the material is extended or edited by a CO using
 	* custom LODs (LODStrategy), on that given LOD, reuse materials will not be applied and the material and its textures will be unique. */
 	UPROPERTY(EditAnywhere, Category = CustomizableObject)
-	bool bReuseMaterialBetweenLODs = true;
+	bool bReuseMaterialBetweenLODs = false;
 	
 	// UObject interface.
 	virtual void Serialize(FArchive& Ar) override;
