#rnx

diff --git a/Engine/Plugins/Experimental/Mutable/Source/CustomizableObjectEditor/Private/MuCOE/GenerateMutableSource/GenerateMutableSourceSurface.cpp b/Engine/Plugins/Experimental/Mutable/Source/CustomizableObjectEditor/Private/MuCOE/GenerateMutableSource/GenerateMutableSourceSurface.cpp
--- a/Engine/Plugins/Experimental/Mutable/Source/CustomizableObjectEditor/Private/MuCOE/GenerateMutableSource/GenerateMutableSourceSurface.cpp
+++ b/Engine/Plugins/Experimental/Mutable/Source/CustomizableObjectEditor/Private/MuCOE/GenerateMutableSource/GenerateMutableSourceSurface.cpp
@@ -786,6 +786,18 @@ mu::NodeSurfacePtr GenerateMutableSourceSurface(const UEdGraphPin * Pin, FMutabl
 											GenerationContext.Compiler->CompilerLog(ReplacedImageFormatText, Node, EMessageSeverity::Info);
 											UE_LOG(LogMutable, Log, TEXT("%s"), *ReplacedImageFormatMsg);
 										}
+										else if (PlatformFormat.Contains(TEXT("ASTC_NormalLA")))
+										{
+											// TODO: This is just a workaround to prevent the "Unexpected image format" warning below. ASTC_NormalLA is
+											// not supported yet by Mutable so it forces IF_ASTC_4x4_RG_LDR as a replacement. It should be changed with a more
+											// appropriate format or directly implement it
+											mutableFormat = mu::EImageFormat::IF_ASTC_4x4_RG_LDR;
+
+											const FString ReplacedImageFormatMsg2 = FString::Printf(TEXT("In object [%s] the unsupported ASTC_NormalLA image format is used, ASTC_4x4_RG_LDR will be used instead."), *GenerationContext.Object->GetName());
+											const FText ReplacedImageFormatText2 = FText::FromString(ReplacedImageFormatMsg2);
+											GenerationContext.Compiler->CompilerLog(ReplacedImageFormatText2, Node, EMessageSeverity::Info);
+											UE_LOG(LogMutable, Log, TEXT("%s"), *ReplacedImageFormatMsg2);
+										}
 									}
 
 									if (mutableFormat == mu::EImageFormat::IF_NONE)
