[mutable] Ensure full mip chain after patch operation, if necessary.
Honor "disable layouts" flag in debugger.
Fallback to ensure full mipmap chain reaches Unreal in case Mutable fails.
Bypass mipmap operation if all mips are already there.
#rnx

diff --git a/Engine/Plugins/Experimental/Mutable/Source/CustomizableObject/Private/MuCO/CustomizableObjectMipDataProvider.cpp b/Engine/Plugins/Experimental/Mutable/Source/CustomizableObject/Private/MuCO/CustomizableObjectMipDataProvider.cpp
--- a/Engine/Plugins/Experimental/Mutable/Source/CustomizableObject/Private/MuCO/CustomizableObjectMipDataProvider.cpp
+++ b/Engine/Plugins/Experimental/Mutable/Source/CustomizableObject/Private/MuCO/CustomizableObjectMipDataProvider.cpp
@@ -182,6 +182,29 @@ namespace impl
 
 				check(Image);
 
+				int32 FullMipCount = Image->GetMipmapCount(Image->GetSizeX(), Image->GetSizeY());
+				int32 RealMipCount = Image->GetLODCount();
+
+				bool bForceMipchain =
+					// Did we fail to generate the entire mipchain (if we have mips at all)?
+					(RealMipCount != 1) && (RealMipCount != FullMipCount);
+
+				if (bForceMipchain)
+				{
+					MUTABLE_CPUPROFILER_SCOPE(GetImage_MipFix);
+
+					UE_LOG(LogMutable, Warning, TEXT("Mutable generated an incomplete mip chain for image."));
+
+					// Force the right number of mips. The missing data will be black.
+					mu::Ptr<mu::Image> NewImage = new mu::Image(Image->GetSizeX(), Image->GetSizeY(), FullMipCount, Image->GetFormat());
+					check(NewImage);
+					if (NewImage->GetDataSize() >= Image->GetDataSize())
+					{
+						FMemory::Memcpy(NewImage->GetData(), Image->GetData(), Image->GetDataSize());
+					}
+					Image = NewImage;
+				}
+
 				OperationData->Result = Image;
 			}
 
@@ -344,17 +367,29 @@ bool FMutableTextureMipDataProvider::PollMips(const FTextureUpdateSyncOptions& S
 		// The counter must be zero meaning the Mutable image operation has finished
 		check(SyncOptions.Counter->GetValue() == 0);
 
-		mu::ImagePtrConst Mip = OperationData->Result;
+		mu::Ptr<const mu::Image> Mip = OperationData->Result;
 		int32 MipIndex = 0;
 		check(Mip->GetSizeX() == OperationData->Levels[0].SizeX);
 		check(Mip->GetSizeY() == OperationData->Levels[0].SizeY);
 
 		for (FMutableMipUpdateLevel& Level : OperationData->Levels)
 		{
-			// Check Mip DataSize for consistency, but skip if 0 because it's optional and might be zero in cooked mips
-			check(Level.DataSize == 0 || Mip->GetLODDataSize(MipIndex) == Level.DataSize);
 			void* Dest = Level.Dest;
-			FMemory::Memcpy(Dest, Mip->GetMipData(MipIndex), Mip->GetLODDataSize(MipIndex));
+
+			if (MipIndex >= Mip->GetLODCount())
+			{
+				// Mutable didn't generate all the expected mips
+				UE_LOG(LogMutable, Warning, TEXT("Mutable image is missing mips."));
+				FMemory::Memzero(Dest, Level.DataSize);
+			}
+			else
+			{
+				int32 MipDataSize = Mip->GetLODDataSize(MipIndex);
+
+				// Check Mip DataSize for consistency, but skip if 0 because it's optional and might be zero in cooked mips
+				check(Level.DataSize == 0 || MipDataSize == Level.DataSize);
+				FMemory::Memcpy(Dest, Mip->GetMipData(MipIndex), Level.DataSize);
+			}
 			++MipIndex;
 		}
 	}
diff --git a/Engine/Plugins/Experimental/Mutable/Source/CustomizableObject/Private/MuCO/CustomizableObjectSystem.cpp b/Engine/Plugins/Experimental/Mutable/Source/CustomizableObject/Private/MuCO/CustomizableObjectSystem.cpp
--- a/Engine/Plugins/Experimental/Mutable/Source/CustomizableObject/Private/MuCO/CustomizableObjectSystem.cpp
+++ b/Engine/Plugins/Experimental/Mutable/Source/CustomizableObject/Private/MuCO/CustomizableObjectSystem.cpp
@@ -1602,7 +1602,7 @@ namespace impl
 					UE_LOG(LogMutable, Warning, TEXT("Mutable generated an incomplete mip chain for image %d."), Image.ImageID);
 
 					// Force the right number of mips. The missing data will be black.
-					mu::ImagePtr NewImage = new mu::Image(Image.Image->GetSizeX(), Image.Image->GetSizeY(), FullMipCount, Image.Image->GetFormat());
+					mu::Ptr<mu::Image> NewImage = new mu::Image(Image.Image->GetSizeX(), Image.Image->GetSizeY(), FullMipCount, Image.Image->GetFormat());
 					check(NewImage);
 					if (NewImage->GetDataSize() >= Image.Image->GetDataSize())
 					{
diff --git a/Engine/Plugins/Experimental/Mutable/Source/CustomizableObjectEditor/Private/MuCOE/CustomizableObjectCompiler.cpp b/Engine/Plugins/Experimental/Mutable/Source/CustomizableObjectEditor/Private/MuCOE/CustomizableObjectCompiler.cpp
--- a/Engine/Plugins/Experimental/Mutable/Source/CustomizableObjectEditor/Private/MuCOE/CustomizableObjectCompiler.cpp
+++ b/Engine/Plugins/Experimental/Mutable/Source/CustomizableObjectEditor/Private/MuCOE/CustomizableObjectCompiler.cpp
@@ -1350,7 +1350,7 @@ void FCustomizableObjectCompiler::CompileInternal(UCustomizableObject* Object, c
 }
 
 
-mu::NodePtr FCustomizableObjectCompiler::Export(UCustomizableObject* Object, const FCompilationOptions& InCompilerOptions)
+mu::NodePtr FCustomizableObjectCompiler::Export(UCustomizableObject* Object, const FCompilationOptions& InCompilerOptions, bool& bOutDisableLayouts)
 {
 	UE_LOG(LogMutable, Log, TEXT("Started Customizable Object Export %s."), *Object->GetName());
 
@@ -1393,6 +1393,8 @@ mu::NodePtr FCustomizableObjectCompiler::Export(UCustomizableObject* Object, con
 		UpdatePendingTextureConversion(false);
 	}
 
+	bOutDisableLayouts = GenerationContext.bDisableTextureLayoutManagementFlag;
+
 	return MutableRoot;
 }
 
diff --git a/Engine/Plugins/Experimental/Mutable/Source/CustomizableObjectEditor/Private/MuCOE/CustomizableObjectCompiler.h b/Engine/Plugins/Experimental/Mutable/Source/CustomizableObjectEditor/Private/MuCOE/CustomizableObjectCompiler.h
--- a/Engine/Plugins/Experimental/Mutable/Source/CustomizableObjectEditor/Private/MuCOE/CustomizableObjectCompiler.h
+++ b/Engine/Plugins/Experimental/Mutable/Source/CustomizableObjectEditor/Private/MuCOE/CustomizableObjectCompiler.h
@@ -42,7 +42,7 @@ public:
 	virtual UCustomizableObject* GetRootObject(UCustomizableObject* InObject) override;
 
 	/** Generate the Mutable Graph from the Unreal Graph. */
-	mu::NodePtr Export(UCustomizableObject* Object, const FCompilationOptions& Options);
+	mu::NodePtr Export(UCustomizableObject* Object, const FCompilationOptions& Options, bool& bOutDisableLayouts);
 
 	void CompilerLog(const FText& Message, const TArray<const UCustomizableObjectNode*>& ArrayNode, EMessageSeverity::Type MessageSeverity = EMessageSeverity::Warning, bool bAddBaseObjectInfo = true);
 	void CompilerLog(const FText& Message, const UCustomizableObjectNode* Node = nullptr, EMessageSeverity::Type MessageSeverity = EMessageSeverity::Warning, bool bAddBaseObjectInfo = true);
diff --git a/Engine/Plugins/Experimental/Mutable/Source/CustomizableObjectEditor/Private/MuCOE/SMutableObjectViewer.cpp b/Engine/Plugins/Experimental/Mutable/Source/CustomizableObjectEditor/Private/MuCOE/SMutableObjectViewer.cpp
--- a/Engine/Plugins/Experimental/Mutable/Source/CustomizableObjectEditor/Private/MuCOE/SMutableObjectViewer.cpp
+++ b/Engine/Plugins/Experimental/Mutable/Source/CustomizableObjectEditor/Private/MuCOE/SMutableObjectViewer.cpp
@@ -215,7 +215,9 @@ void SMutableObjectViewer::Construct(const FArguments& InArgs, UCustomizableObje
 void SMutableObjectViewer::GenerateMutableGraphPressed()
 {
 	// Convert from Unreal graph to Mutable graph.
-	mu::NodePtr RootNode = Compiler.Export(CustomizableObject, CompileOptions);
+	bool bDisableLayouts = false;
+
+	mu::NodePtr RootNode = Compiler.Export(CustomizableObject, CompileOptions, bDisableLayouts);
 	if (!RootNode)
 	{
 		// TODO: Show errors
@@ -239,6 +241,8 @@ void SMutableObjectViewer::GenerateMutableGraphPressed()
 
 void SMutableObjectViewer::CompileMutableCodePressed()
 {
+	bool bDisableLayouts = false;
+
 	if (CompileOptions.bForceLargeLODBias)
 	{
 		// Debug compile with many different biasses
@@ -247,7 +251,7 @@ void SMutableObjectViewer::CompileMutableCodePressed()
 		{
 			CompileOptions.DebugBias = Bias;
 				
-			mu::NodePtr RootNode = Compiler.Export(CustomizableObject, CompileOptions);
+			mu::NodePtr RootNode = Compiler.Export(CustomizableObject, CompileOptions, bDisableLayouts);
 			if (!RootNode)
 			{
 				// TODO: Show errors
@@ -255,7 +259,7 @@ void SMutableObjectViewer::CompileMutableCodePressed()
 			}
 
 			// Do the compilation to Mutable Code synchronously.
-			TSharedPtr<FCustomizableObjectCompileRunnable> CompileTask = MakeShareable(new FCustomizableObjectCompileRunnable(RootNode, false));
+			TSharedPtr<FCustomizableObjectCompileRunnable> CompileTask = MakeShareable(new FCustomizableObjectCompileRunnable(RootNode, bDisableLayouts));
 			CompileTask->Options = CompileOptions;
 			CompileTask->Init();
 			CompileTask->Run();
@@ -264,7 +268,7 @@ void SMutableObjectViewer::CompileMutableCodePressed()
 	}
 
 	// Convert from Unreal graph to Mutable graph.
-	mu::NodePtr RootNode = Compiler.Export(CustomizableObject, CompileOptions);
+	mu::NodePtr RootNode = Compiler.Export(CustomizableObject, CompileOptions, bDisableLayouts);
 	if (!RootNode)
 	{
 		// TODO: Show errors
@@ -272,7 +276,7 @@ void SMutableObjectViewer::CompileMutableCodePressed()
 	}
 
 	// Do the compilation to Mutable Code synchronously.
-	TSharedPtr<FCustomizableObjectCompileRunnable> CompileTask = MakeShareable(new FCustomizableObjectCompileRunnable(RootNode, false));
+	TSharedPtr<FCustomizableObjectCompileRunnable> CompileTask = MakeShareable(new FCustomizableObjectCompileRunnable(RootNode, bDisableLayouts));
 	CompileTask->Options = CompileOptions;
 	CompileTask->Init();
 	CompileTask->Run();
diff --git a/Engine/Plugins/Experimental/Mutable/Source/MutableRuntime/Private/MuR/CodeRunnerTasked.cpp b/Engine/Plugins/Experimental/Mutable/Source/MutableRuntime/Private/MuR/CodeRunnerTasked.cpp
--- a/Engine/Plugins/Experimental/Mutable/Source/MutableRuntime/Private/MuR/CodeRunnerTasked.cpp
+++ b/Engine/Plugins/Experimental/Mutable/Source/MutableRuntime/Private/MuR/CodeRunnerTasked.cpp
@@ -1401,6 +1401,12 @@ namespace mu
 		// At least keep the levels we already have.
 		int StartLevel = Base->GetLODCount();
 		LevelCount = FMath::Max(StartLevel, LevelCount);
+		
+		if (LevelCount == Base->GetLODCount())
+		{
+			Runner->StoreImage(Op, Base);
+			return false;
+		}
 
 		// Create destination data
 		Result = Runner->CreateImage(Base->GetSizeX(), Base->GetSizeY(), LevelCount, Base->GetFormat(), EInitializationType::NotInitialized);
@@ -1431,8 +1437,6 @@ namespace mu
 		// This runs in the Runner thread
 		Runner->Release(Base);
 
-		check(Result);
-
 		// If no shortcut was taken
 		if (Result)
 		{
diff --git a/Engine/Plugins/Experimental/Mutable/Source/MutableTools/Private/MuT/ASTOpImageMipmap.cpp b/Engine/Plugins/Experimental/Mutable/Source/MutableTools/Private/MuT/ASTOpImageMipmap.cpp
--- a/Engine/Plugins/Experimental/Mutable/Source/MutableTools/Private/MuT/ASTOpImageMipmap.cpp
+++ b/Engine/Plugins/Experimental/Mutable/Source/MutableTools/Private/MuT/ASTOpImageMipmap.cpp
@@ -113,6 +113,27 @@ namespace mu
 	}
 
 
+	mu::Ptr<ASTOp> ASTOpImageMipmap::OptimiseSemantic(const FModelOptimizationOptions& options, int32 Pass) const
+	{
+		mu::Ptr<ASTOp> at;
+
+		// \TODO: This seems to fail with the bandit test model.
+		//if (Source.child())
+		//{
+		//	ASTOp::FGetImageDescContext context;
+		//	FImageDesc ChildDesc = Source.child()->GetImageDesc(false, &context);
+
+		//	if (ChildDesc.m_lods>0 && ChildDesc.m_lods>=Levels)
+		//	{
+		//		// We can skip the mipmaps, because the child will always contain all lods anyway.
+		//		at = Source.child();
+		//	}
+		//}
+
+		return at;
+	}
+
+
 	mu::Ptr<ASTOp> ASTOpImageMipmap::OptimiseSink(const FModelOptimizationOptions& options, FOptimizeSinkContext& context) const
 	{
 		mu::Ptr<ASTOp> at;
@@ -425,7 +446,11 @@ namespace mu
 					newOp->patch = Visit(newOp->patch.child(), newMip.get());
 					newAt = newOp;
 
-					if (currentMipmapOp->Levels != currentMipmapOp->BlockLevels)
+					if (currentMipmapOp->Levels != currentMipmapOp->BlockLevels						
+						// If the current levels are all of them, we will want to rebuild the mips after patch.
+						// This happens if ignoring layouts, in which case there is no top-most mipmap to ensure the tail already.
+						|| currentMipmapOp->BlockLevels == 0 
+						)
 					{
 						// We need to add a mipmap on top to finish the mipmapping
 						mu::Ptr<ASTOpImageMipmap> topMipOp = mu::Clone<ASTOpImageMipmap>(currentMipmapOp);
diff --git a/Engine/Plugins/Experimental/Mutable/Source/MutableTools/Private/MuT/ASTOpImageMipmap.h b/Engine/Plugins/Experimental/Mutable/Source/MutableTools/Private/MuT/ASTOpImageMipmap.h
--- a/Engine/Plugins/Experimental/Mutable/Source/MutableTools/Private/MuT/ASTOpImageMipmap.h
+++ b/Engine/Plugins/Experimental/Mutable/Source/MutableTools/Private/MuT/ASTOpImageMipmap.h
@@ -49,6 +49,7 @@ namespace mu
 		Ptr<ASTOp> Clone(MapChildFuncRef mapChild) const override;
 		void ForEachChild(const TFunctionRef<void(ASTChild&)>) override;
 		void Link(FProgram& program, const FLinkerOptions* Options) override;
+		mu::Ptr<ASTOp> OptimiseSemantic(const FModelOptimizationOptions& options, int32 Pass) const override;
 		Ptr<ASTOp> OptimiseSink(const FModelOptimizationOptions& options, FOptimizeSinkContext& context) const override;
 		FImageDesc GetImageDesc(bool returnBestOption, FGetImageDescContext* context) const override;
 		void GetLayoutBlockSize(int* pBlockX, int* pBlockY) override;
diff --git a/Engine/Plugins/Experimental/Mutable/Source/MutableTools/Private/MuT/ASTOpImagePixelFormat.cpp b/Engine/Plugins/Experimental/Mutable/Source/MutableTools/Private/MuT/ASTOpImagePixelFormat.cpp
--- a/Engine/Plugins/Experimental/Mutable/Source/MutableTools/Private/MuT/ASTOpImagePixelFormat.cpp
+++ b/Engine/Plugins/Experimental/Mutable/Source/MutableTools/Private/MuT/ASTOpImagePixelFormat.cpp
@@ -405,7 +405,7 @@ namespace mu
 				if (acceptable)
 				{
 					// We move the format down the two paths
-					auto newOp = mu::Clone<ASTOpImageCompose>(at);
+					Ptr<ASTOpImageCompose> newOp = mu::Clone<ASTOpImageCompose>(at);
 
 					// TODO: We have to make sure we don't end up with two different formats if
 					// there is an formatIfAlpha
@@ -428,7 +428,7 @@ namespace mu
 			if (isBlockFormat)
 			{
 				// We move the format down the two paths
-				auto newOp = mu::Clone<ASTOpImagePatch>(at);
+				Ptr<ASTOpImagePatch> newOp = mu::Clone<ASTOpImagePatch>(at);
 
 				newOp->base = Visit(newOp->base.child(), currentFormatOp);
 				newOp->patch = Visit(newOp->patch.child(), currentFormatOp);
@@ -441,7 +441,7 @@ namespace mu
 
 		case OP_TYPE::IM_MIPMAP:
 		{
-			auto typedSource = dynamic_cast<const ASTOpImageMipmap*>(at.get());
+			const ASTOpImageMipmap* typedSource = dynamic_cast<const ASTOpImageMipmap*>(at.get());
 
 			// If its a compressed format, only sink formats on mipmap operations that
 			// generate the tail. To avoid optimization loop.
diff --git a/Engine/Plugins/Experimental/Mutable/Source/MutableTools/Private/MuT/CodeOptimiser_Semantic.cpp b/Engine/Plugins/Experimental/Mutable/Source/MutableTools/Private/MuT/CodeOptimiser_Semantic.cpp
--- a/Engine/Plugins/Experimental/Mutable/Source/MutableTools/Private/MuT/CodeOptimiser_Semantic.cpp
+++ b/Engine/Plugins/Experimental/Mutable/Source/MutableTools/Private/MuT/CodeOptimiser_Semantic.cpp
@@ -1090,7 +1090,7 @@ namespace mu
 
 				if (acceptable)
 				{
-					auto newOp = mu::Clone<ASTOpImageCompose>(at);
+					Ptr<ASTOpImageCompose> newOp = mu::Clone<ASTOpImageCompose>(at);
 
 					Ptr<ASTOp> baseOp = newOp->Base.child();
 					newOp->Base = Visit(baseOp, currentSinkingOp);
@@ -1106,7 +1106,7 @@ namespace mu
 
 			case OP_TYPE::IM_PATCH:
 			{
-				auto newOp = mu::Clone<ASTOpImagePatch>(at);
+				Ptr<ASTOpImagePatch> newOp = mu::Clone<ASTOpImagePatch>(at);
 
 				newOp->base = Visit(newOp->base.child(), currentSinkingOp);
 				newOp->patch = Visit(newOp->patch.child(), currentSinkingOp);
@@ -1122,7 +1122,7 @@ namespace mu
 
 			case OP_TYPE::IM_MIPMAP:
 			{
-				auto newOp = mu::Clone<ASTOpImageMipmap>(at.get());
+				Ptr<ASTOpImageMipmap> newOp = mu::Clone<ASTOpImageMipmap>(at.get());
 				Ptr<ASTOp> baseOp = newOp->Source.child();
 				newOp->Source = Visit(baseOp, currentSinkingOp);
 				newAt = newOp;
