#rnx

diff --git a/Engine/Plugins/Experimental/Mutable/Source/CustomizableObjectEditor/Private/MuCOE/GenerateMutableSource/GenerateMutableSourceSurface.cpp b/Engine/Plugins/Experimental/Mutable/Source/CustomizableObjectEditor/Private/MuCOE/GenerateMutableSource/GenerateMutableSourceSurface.cpp
--- a/Engine/Plugins/Experimental/Mutable/Source/CustomizableObjectEditor/Private/MuCOE/GenerateMutableSource/GenerateMutableSourceSurface.cpp
+++ b/Engine/Plugins/Experimental/Mutable/Source/CustomizableObjectEditor/Private/MuCOE/GenerateMutableSource/GenerateMutableSourceSurface.cpp
@@ -484,8 +484,9 @@ mu::NodeSurfacePtr GenerateMutableSourceSurface(const UEdGraphPin * Pin, FMutabl
 		for (int32 ImageIndex = 0; ImageIndex < NumImages; ++ImageIndex)
 		{
 			const UEdGraphPin* ImagePin = TypedNodeMat->GetParameterPin(EMaterialParameterType::Texture, ImageIndex);
+			const bool bIsImagePinLinked = ImagePin && !ImagePin->LinkedTo.IsEmpty();
 
-			if (ImagePin && !TypedNodeMaterial->IsImageMutableMode(ImageIndex))
+			if (bIsImagePinLinked && !TypedNodeMaterial->IsImageMutableMode(ImageIndex))
 			{
 				// This is a connected pass-through texture that simply has to be passed to the core
 
