[mutable] Fix mip data size that could fail during cook.
Indentation and back-port macros cleanup.
#rnx

diff --git a/Engine/Plugins/Experimental/Mutable/Source/CustomizableObject/Private/MuCO/CustomizableObjectMipDataProvider.cpp b/Engine/Plugins/Experimental/Mutable/Source/CustomizableObject/Private/MuCO/CustomizableObjectMipDataProvider.cpp
--- a/Engine/Plugins/Experimental/Mutable/Source/CustomizableObject/Private/MuCO/CustomizableObjectMipDataProvider.cpp
+++ b/Engine/Plugins/Experimental/Mutable/Source/CustomizableObject/Private/MuCO/CustomizableObjectMipDataProvider.cpp
@@ -388,7 +388,7 @@ bool FMutableTextureMipDataProvider::PollMips(const FTextureUpdateSyncOptions& S
 
 				// Check Mip DataSize for consistency, but skip if 0 because it's optional and might be zero in cooked mips
 				check(Level.DataSize == 0 || MipDataSize == Level.DataSize);
-				FMemory::Memcpy(Dest, Mip->GetMipData(MipIndex), Level.DataSize);
+				FMemory::Memcpy(Dest, Mip->GetMipData(MipIndex), MipDataSize);
 			}
 			++MipIndex;
 		}
diff --git a/Engine/Plugins/Experimental/Mutable/Source/CustomizableObjectEditor/Private/MuCOE/SCustomizableInstanceProperties.cpp b/Engine/Plugins/Experimental/Mutable/Source/CustomizableObjectEditor/Private/MuCOE/SCustomizableInstanceProperties.cpp
--- a/Engine/Plugins/Experimental/Mutable/Source/CustomizableObjectEditor/Private/MuCOE/SCustomizableInstanceProperties.cpp
+++ b/Engine/Plugins/Experimental/Mutable/Source/CustomizableObjectEditor/Private/MuCOE/SCustomizableInstanceProperties.cpp
@@ -685,7 +685,7 @@ void SCustomizableInstanceProperties::AddParameter(int32 ParamIndexInObject)
 				.Content()
 				[
 					SNew(SImage)
-					.Image(FAppStyle::GetBrush("PropertyWindow.DiffersFromDefault"))
+					.Image(UE_MUTABLE_GET_BRUSH("PropertyWindow.DiffersFromDefault"))
 				]
 			]
 
@@ -760,7 +760,7 @@ void SCustomizableInstanceProperties::AddParameter(int32 ParamIndexInObject)
 					.Content()
 					[
 						SNew(SImage)
-						.Image(FAppStyle::GetBrush("PropertyWindow.DiffersFromDefault"))
+						.Image(UE_MUTABLE_GET_BRUSH("PropertyWindow.DiffersFromDefault"))
 					]
 				]
 			];
@@ -808,7 +808,7 @@ void SCustomizableInstanceProperties::AddParameter(int32 ParamIndexInObject)
 				.Content()
 				[
 					SNew(SImage)
-					.Image(FAppStyle::GetBrush("PropertyWindow.DiffersFromDefault"))
+					.Image(UE_MUTABLE_GET_BRUSH("PropertyWindow.DiffersFromDefault"))
 				]
 			]
 
@@ -856,7 +856,7 @@ void SCustomizableInstanceProperties::AddParameter(int32 ParamIndexInObject)
 				.Content()
 				[
 					SNew(SImage)
-					.Image(FAppStyle::GetBrush("PropertyWindow.DiffersFromDefault"))
+					.Image(UE_MUTABLE_GET_BRUSH("PropertyWindow.DiffersFromDefault"))
 				]
 			]
 		];
