[Mutable] Matching Sections between Generated LODs

- Sections between LODs now will be matched using Material Index.
- Fixed lower LODs with a section removed using LOD lod than expected.
- Fixed crash when a lower LOD section with morph were not present in higher LODs.

#rnx

diff --git a/Engine/Plugins/Experimental/Mutable/Source/CustomizableObjectEditor/Classes/MuCOE/Nodes/CustomizableObjectNodeTable.h b/Engine/Plugins/Experimental/Mutable/Source/CustomizableObjectEditor/Classes/MuCOE/Nodes/CustomizableObjectNodeTable.h
--- a/Engine/Plugins/Experimental/Mutable/Source/CustomizableObjectEditor/Classes/MuCOE/Nodes/CustomizableObjectNodeTable.h
+++ b/Engine/Plugins/Experimental/Mutable/Source/CustomizableObjectEditor/Classes/MuCOE/Nodes/CustomizableObjectNodeTable.h
@@ -94,7 +94,7 @@ public:
 	UPROPERTY()
 	int32 LOD = 0;
 
-	/** Material Index (Surface Index) of the mesh related to this Mesh pin */
+	/** Section Index (Surface Index) of the mesh related to this Mesh pin */
 	UPROPERTY()
 	int32 Material = 0;
 
@@ -196,7 +196,7 @@ public:
 	FString GetMutableColumnName(const UEdGraphPin* Pin, const int32& LOD) const;
 
 	// Returns the LOD of the mesh associated to the input pin
-	void GetPinLODAndMaterial(const UEdGraphPin* Pin, int32& LOD, int32& Material) const;
+	void GetPinLODAndSection(const UEdGraphPin* Pin, int32& LODIndex, int32& SectionIndex) const;
 
 	// Get the anim blueprint and anim slot columns related to a mesh
 	void GetAnimationColumns(const FString& ColumnName, FString& AnimBPColumnName, FString& AnimSlotColumnName, FString& AnimTagColumnName) const;
diff --git a/Engine/Plugins/Experimental/Mutable/Source/CustomizableObjectEditor/Private/MuCOE/GenerateMutableSource/GenerateMutableSource.h b/Engine/Plugins/Experimental/Mutable/Source/CustomizableObjectEditor/Private/MuCOE/GenerateMutableSource/GenerateMutableSource.h
--- a/Engine/Plugins/Experimental/Mutable/Source/CustomizableObjectEditor/Private/MuCOE/GenerateMutableSource/GenerateMutableSource.h
+++ b/Engine/Plugins/Experimental/Mutable/Source/CustomizableObjectEditor/Private/MuCOE/GenerateMutableSource/GenerateMutableSource.h
@@ -533,7 +533,8 @@ struct FMutableGraphGenerationContext
 		{
 			/** Source mesh data. */
 			const UObject* Mesh = nullptr;
-			int32 LOD = 0;
+			int32 LOD = 0; // Mesh Data LOD (i.e., LOD where we are getting the vertices from)
+			int32 CurrentLOD = 0; // Derived data LOD (i.e., LOD where we are generating the non-Core Data like morphs)
 			int32 MaterialIndex = 0;
 
 			/** Flag used to generate this mesh. Bit mask of EMutableMeshConversionFlags */
@@ -547,7 +548,7 @@ struct FMutableGraphGenerationContext
 
 			bool operator==( const FKey& OtherKey ) const
 			{
-				return Mesh == OtherKey.Mesh && LOD == OtherKey.LOD && MaterialIndex == OtherKey.MaterialIndex
+				return Mesh == OtherKey.Mesh && LOD == OtherKey.LOD && CurrentLOD == OtherKey.CurrentLOD && MaterialIndex == OtherKey.MaterialIndex
 					&& Flags == OtherKey.Flags && Tags == OtherKey.Tags;
 			}
 		};
diff --git a/Engine/Plugins/Experimental/Mutable/Source/CustomizableObjectEditor/Private/MuCOE/GenerateMutableSource/GenerateMutableSourceMesh.cpp b/Engine/Plugins/Experimental/Mutable/Source/CustomizableObjectEditor/Private/MuCOE/GenerateMutableSource/GenerateMutableSourceMesh.cpp
--- a/Engine/Plugins/Experimental/Mutable/Source/CustomizableObjectEditor/Private/MuCOE/GenerateMutableSource/GenerateMutableSourceMesh.cpp
+++ b/Engine/Plugins/Experimental/Mutable/Source/CustomizableObjectEditor/Private/MuCOE/GenerateMutableSource/GenerateMutableSourceMesh.cpp
@@ -48,81 +48,86 @@
 
 #define LOCTEXT_NAMESPACE "CustomizableObjectEditor"
 
-int32 FindLODSectionIndexSkippingDisabled(const FSkeletalMeshLODModel& LODModel, int32 TargetSectionIndex)
+
+/** Returns the corrected LOD and Section Index when using Automatic LOD From Mesh strategy.
+ *
+ * Do not confuse Section Index and Material Index, they are not the same.
+ * 
+ * @param Context 
+ * @param Node 
+ * @param SkeletalMesh 
+ * @param InOutLODIndex In connected pin LOD Index. Out corrected Skeletal Mesh LOD Index.
+ * @param InOutPinSectionIndex In connected pin Section Index. Out corrected Skeletal Mesh Section Index. */
+void GetEffectiveLODAndSection(const FMutableGraphGenerationContext& Context, const UCustomizableObjectNode& Node, const USkeletalMesh& SkeletalMesh, int32& InOutLODIndex, int32& InOutPinSectionIndex)
 {
-	if (LODModel.Sections.IsValidIndex(TargetSectionIndex))
+	if (Context.CurrentAutoLODStrategy != ECustomizableObjectAutomaticLODStrategy::AutomaticFromMesh)
 	{
-		return INDEX_NONE;
+		return;
 	}
-
-	if (LODModel.Sections[TargetSectionIndex].bDisabled)
+	
+	const int32 CurrentLOD = InOutLODIndex + Context.CurrentLOD;
+	
+	if (CurrentLOD == InOutLODIndex)
 	{
-		return INDEX_NONE;
+		return;
 	}
 
-	int32 SkipDisabledSectionIndex = 0;
-	const int32 NumSections = LODModel.Sections.Num();
-	for (int32 SectionIndex = 0; SectionIndex < NumSections; ++SectionIndex)
+	const FSkeletalMeshModel* ImportedModel = SkeletalMesh.GetImportedModel();
+	if (!ImportedModel)
 	{
-		if (TargetSectionIndex == SectionIndex)
-		{
-			return SkipDisabledSectionIndex;
-		}
-
-		const FSkelMeshSection& Section = LODModel.Sections[SectionIndex];
-
-		SkipDisabledSectionIndex += static_cast<int32>(!Section.bDisabled);
+		return;
 	}
 
-	return INDEX_NONE;
-}
-
-void GetLODAndSectionForAutomaticLODs(const FSkeletalMeshModel* ImportedModel, int32 InPinLODIndex, int32 InCurrentLOD, int32 InPinSectionIndex, int32& OutLODIndex, int32& OutSectionIndex)
-{
-	check(ImportedModel);
-
-	// The disabled sections are not replicated to unreal generated lods and a mismatch can happen when 
-	// selecting the section with automatic from mesh strategy. The fix for that is only parcial and the casuistics 
-	// of section matching between lods are more extense, for instance, when lower resolution lods have 
-	// sections removed in an arbitrary order. To minimize discrepacines with the old behvaior the disabled sections
-	// will only be taken into account if the PinLODIndex is the base. 
-
-	int32 CurrentLODSectionIndex = (InCurrentLOD != InPinLODIndex && InPinLODIndex == 0) 
-		? FindLODSectionIndexSkippingDisabled(ImportedModel->LODModels[InPinLODIndex], InPinSectionIndex)
-		: InPinSectionIndex;
-
-	// if we couldn't find any matching section, fallback with the old behaviour of using 
-	// the pin SectionIndex
-	if (CurrentLODSectionIndex == INDEX_NONE)
+	if (!ImportedModel->LODModels.IsValidIndex(InOutLODIndex) ||
+		!ImportedModel->LODModels[InOutLODIndex].Sections.IsValidIndex(InOutPinSectionIndex))
 	{
-		CurrentLODSectionIndex = InPinSectionIndex;
+		return;
 	}
 	
-	// If the mesh has additional LODs and they use the same material, use them.
-	if (ImportedModel->LODModels.Num() > InCurrentLOD
-		&&
-		ImportedModel->LODModels[InCurrentLOD].Sections.Num() > CurrentLODSectionIndex)
-	{
-		OutLODIndex = InCurrentLOD;
-		OutSectionIndex = CurrentLODSectionIndex;
-	}
-	else
+	const int32 SearchLODMaterialIndex = ImportedModel->LODModels[InOutLODIndex].Sections[InOutPinSectionIndex].MaterialIndex; // Material Index of the connected pin
+
+	bool bRepeatedSections = false;
+	
+	const int32 StartLOD = FMath::Min(CurrentLOD, ImportedModel->LODModels.Num() - 1); // Not all meshes have the number of LODs specified in the Mutable graph
+	for (int32 LODIndex = StartLOD; LODIndex >= 0; --LODIndex)
 	{
-		// Find the closest valid LOD.
-		// Don't try to account for disabled sections for this case.
-		for (int32 LODIndex = ImportedModel->LODModels.Num() - 1; LODIndex >= 0; --LODIndex)
+		const FSkeletalMeshLODModel& LODModel = ImportedModel->LODModels[LODIndex];
+		const TArray<int32>& MaterialMap = SkeletalMesh.GetLODInfoArray()[LODIndex].LODMaterialMap;
+
+		bool bFound = false;
+		for (int32 SectionIndex = 0; SectionIndex < LODModel.Sections.Num(); ++SectionIndex)
 		{
-			if (ImportedModel->LODModels[LODIndex].Sections.Num() > InPinSectionIndex)
+			const int32 MaterialIndex = MaterialMap.IsValidIndex(SectionIndex) ? MaterialMap[SectionIndex] : LODModel.Sections[SectionIndex].MaterialIndex; // MaterialMap overrides the MaterialIndex in the section
+				
+			if (MaterialIndex == SearchLODMaterialIndex &&
+				!LODModel.Sections[SectionIndex].bDisabled)
 			{
-				OutLODIndex = LODIndex;
-				OutSectionIndex = InPinSectionIndex;
-				break;
+				if (!bFound)
+				{
+					InOutLODIndex = LODIndex;
+					InOutPinSectionIndex = SectionIndex;
+					bFound = true;
+				}
+				else
+				{
+					bRepeatedSections = true;
+				}
 			}
 		}
+
+		if (bFound)
+		{
+			break;
+		}
 	}
 
+	if (bRepeatedSections)
+	{
+		Context.Compiler->CompilerLog(FText::Format(LOCTEXT("MeshMultipleMaterialIndex", "Mesh {0} contains multiple sections with the same Material Index"), FText::FromString(SkeletalMesh.GetName())), &Node);
+	}
 }
 
+
 void BuildRemappedBonesArray(const FMutableComponentInfo& InComponentInfo, TObjectPtr<USkeletalMesh> InSkeletalMesh, int32 InLODIndex, const TArray<FBoneIndexType>& InRequiredBones, TArray<FBoneIndexType>& OutRemappedBones)
 {
 	if (!InSkeletalMesh)
@@ -1200,10 +1205,10 @@ mu::MeshPtr ConvertSkeletalMeshToMutable(USkeletalMesh* InSkeletalMesh, const TS
 
 					DestMorphTargetIdx = DestMorphTargetIdx != INDEX_NONE
 						? DestMorphTargetIdx
-						: ContributingMorphTargetsInfo.Emplace(FMorphTargetInfo{ MorphTargetName, LOD + 1 });
+						: ContributingMorphTargetsInfo.Emplace(FMorphTargetInfo{ MorphTargetName, GenerationContext.CurrentLOD + 1 });
 
 					FMorphTargetInfo& MorphTargetInfo = ContributingMorphTargetsInfo[DestMorphTargetIdx];
-					MorphTargetInfo.LodNum = FMath::Max(MorphTargetInfo.LodNum, LOD + 1);
+					MorphTargetInfo.LodNum = FMath::Max(MorphTargetInfo.LodNum, GenerationContext.CurrentLOD + 1);
 
 					MorphsUsed.Emplace(FMorphTargetVertexData{ VertexFound.PositionDelta, VertexFound.TangentZDelta, DestMorphTargetIdx });
 				}
@@ -1989,7 +1994,7 @@ mu::MeshPtr GenerateMutableMesh(UObject * Mesh, const TSoftClassPtr<UAnimInstanc
 	// Get the mesh generation flags to use
 	EMutableMeshConversionFlags CurrentFlags = GenerationContext.MeshGenerationFlags.Last();
 
-	FMutableGraphGenerationContext::FGeneratedMeshData::FKey Key = { Mesh, LOD, MaterialIndex, CurrentFlags, UniqueTags };
+	FMutableGraphGenerationContext::FGeneratedMeshData::FKey Key = { Mesh, LOD, GenerationContext.CurrentLOD, MaterialIndex, CurrentFlags, UniqueTags };
 	mu::MeshPtr MutableMesh = GenerationContext.FindGeneratedMesh(Key);
 	
 	if (!MutableMesh)
@@ -2055,34 +2060,13 @@ mu::MeshPtr BuildMorphedMutableMesh(const UEdGraphPin* BaseSourcePin, const FStr
 
 	else if (const UCustomizableObjectNodeTable* TypedNodeTable = Cast<UCustomizableObjectNodeTable>(Node))
 	{
-		TypedNodeTable->GetPinLODAndMaterial(BaseSourcePin, LODIndex, SectionIndex);
+		TypedNodeTable->GetPinLODAndSection(BaseSourcePin, LODIndex, SectionIndex);
 		SkeletalMesh = TypedNodeTable->GetSkeletalMeshAt(BaseSourcePin, RowName);
 	}
 
 	if (SkeletalMesh)
 	{
-		// See if we need to select another LOD because of automatic LOD generation
-		if (GenerationContext.CurrentAutoLODStrategy == ECustomizableObjectAutomaticLODStrategy::AutomaticFromMesh)
-		{
-			if (SkeletalMesh && SkeletalMesh->GetImportedModel())
-			{
-				FSkeletalMeshModel* ImportedModel = SkeletalMesh->GetImportedModel();
-				
-				const int32 PinLODIndex = LODIndex;
-				// Add the CurrentLOD being generated to the first LOD used from this SkeletalMesh
-				const int32 CurrentLOD = PinLODIndex + GenerationContext.CurrentLOD; 
-				const int32 PinSectionIndex = SectionIndex;
-
-				int32 OutLOD = -1;
-				int32 OutSectionIndex = -1;
-
-				GetLODAndSectionForAutomaticLODs(ImportedModel, PinLODIndex, CurrentLOD, PinSectionIndex, OutLOD, OutSectionIndex);
-
-				LODIndex = OutLOD;
-				SectionIndex = OutSectionIndex;
-			}
-		}
-
+		GetEffectiveLODAndSection(GenerationContext, *Node, *SkeletalMesh, LODIndex, SectionIndex);
 
 		// Get the base mesh
 		mu::MeshPtr BaseSourceMesh = GenerateMutableMesh(SkeletalMesh, TSoftClassPtr<UAnimInstance>(), LODIndex, SectionIndex, FString(), GenerationContext, Node);
@@ -2817,6 +2801,26 @@ void GenerateMorphTarget(const UCustomizableObjectNode* Node, const UEdGraphPin*
 		GenerationContext.Compiler->CompilerLog(LOCTEXT("MorphGenerationFailed", "Failed to generate morph target."), Node);
 	}
 }
+
+
+/** Create a default layout. Used when no layout is found. */
+mu::NodeLayoutBlocksPtr CreateDefaultLayout()
+{
+	constexpr int32 GridSize = 4;
+	
+	mu::NodeLayoutBlocksPtr LayoutNode = new mu::NodeLayoutBlocks();
+	LayoutNode->SetGridSize(GridSize, GridSize);
+	LayoutNode->SetMaxGridSize(GridSize, GridSize);
+	LayoutNode->SetLayoutPackingStrategy(mu::EPackStrategy::RESIZABLE_LAYOUT);
+	LayoutNode->SetBlockReductionMethod(mu::EReductionMethod::HALVE_REDUCTION);
+	LayoutNode->SetBlockCount(1);
+	LayoutNode->SetBlock(0, 0, 0, GridSize, GridSize);
+	LayoutNode->SetBlockOptions(0, 0, false);
+
+	return LayoutNode;
+}
+
+
 /** Convert a CustomizableObject Source Graph into a mutable source graph  */
 mu::NodeMeshPtr GenerateMutableSourceMesh(const UEdGraphPin * Pin,
 	FMutableGraphGenerationContext & GenerationContext,
@@ -2856,35 +2860,18 @@ mu::NodeMeshPtr GenerateMutableSourceMesh(const UEdGraphPin * Pin,
 
 		if (TypedNodeSkel->SkeletalMesh)
 		{
-			int32 LOD;
+			int32 LODIndexConnected; // LOD which the pin is connected to
+			int32 LODIndex;
 			int32 SectionIndex;
 
 			{
 				int32 LayoutIndex;
-				TypedNodeSkel->GetPinSection(*Pin, LOD, SectionIndex, LayoutIndex);
-			}
-				
-			// See if we need to select another LOD because of automatic LOD generation
-			if (GenerationContext.CurrentAutoLODStrategy == ECustomizableObjectAutomaticLODStrategy::AutomaticFromMesh)
-			{
-				if (TypedNodeSkel->SkeletalMesh && TypedNodeSkel->SkeletalMesh->GetImportedModel())
-				{
-					const FSkeletalMeshModel* ImportedModel = TypedNodeSkel->SkeletalMesh->GetImportedModel();
-
-					const int32 PinLODIndex = LOD;
-					const int32 CurrentLOD = PinLODIndex + GenerationContext.CurrentLOD; 
-					const int32 PinSectionIndex = SectionIndex;
-
-					int32 OutLOD = -1;
-					int32 OutSectionIndex = -1;
-
-					GetLODAndSectionForAutomaticLODs(ImportedModel, PinLODIndex, CurrentLOD, PinSectionIndex, OutLOD, OutSectionIndex);
-
-					LOD = OutLOD;
-					SectionIndex = OutSectionIndex;
-				}
+				TypedNodeSkel->GetPinSection(*Pin, LODIndexConnected, SectionIndex, LayoutIndex);
 
+				LODIndex = LODIndexConnected;
 			}
+				
+			GetEffectiveLODAndSection(GenerationContext, *Node, *TypedNodeSkel->SkeletalMesh, LODIndex, SectionIndex);
 
 			// First process the mesh tags that are going to make the mesh unique and affect whether it's repeated in 
 			// the mesh cache or not
@@ -2909,7 +2896,9 @@ mu::NodeMeshPtr GenerateMutableSourceMesh(const UEdGraphPin * Pin,
 				MeshUniqueTags += AnimBPTag;
 			}
 
-			mu::MeshPtr MutableMesh = GenerateMutableMesh(TypedNodeSkel->SkeletalMesh, TypedNodeSkel->AnimInstance, LOD, SectionIndex, MeshUniqueTags, GenerationContext, TypedNodeSkel);
+			FSkeletalMeshModel* ImportedModel = TypedNodeSkel->SkeletalMesh->GetImportedModel();
+			
+			mu::MeshPtr MutableMesh = GenerateMutableMesh(TypedNodeSkel->SkeletalMesh, TypedNodeSkel->AnimInstance, LODIndex, SectionIndex, MeshUniqueTags, GenerationContext, TypedNodeSkel);
 			if (MutableMesh)
 			{
 				MeshNode->SetValue(MutableMesh);
@@ -2927,7 +2916,7 @@ mu::NodeMeshPtr GenerateMutableSourceMesh(const UEdGraphPin * Pin,
 
 				if (GenerationContext.Options.bClothingEnabled)
 				{
-					UClothingAssetBase* ClothingAssetBase = TypedNodeSkel->SkeletalMesh->GetSectionClothingAsset(LOD, SectionIndex);	
+					UClothingAssetBase* ClothingAssetBase = TypedNodeSkel->SkeletalMesh->GetSectionClothingAsset(LODIndex, SectionIndex);	
 					UClothingAssetCommon* ClothingAssetCommon = Cast<UClothingAssetCommon>(ClothingAssetBase);
 
 					if (ClothingAssetCommon && ClothingAssetCommon->PhysicsAsset)
@@ -2957,7 +2946,7 @@ mu::NodeMeshPtr GenerateMutableSourceMesh(const UEdGraphPin * Pin,
 					const int32 SkinWeightProfilesCount = GenerationContext.SkinWeightProfilesInfo.Num();
 					for (int32 ProfileIndex = 0; ProfileIndex < SkinWeightProfilesCount; ++ProfileIndex)
 					{
-						if (ImportModel->LODModels[LOD].SkinWeightProfiles.Find(GenerationContext.SkinWeightProfilesInfo[ProfileIndex].Name))
+						if (ImportModel->LODModels[LODIndex].SkinWeightProfiles.Find(GenerationContext.SkinWeightProfilesInfo[ProfileIndex].Name))
 						{
 							const int32 ProfileScemanticIndex = ProfileIndex + 10;
 							MeshData.SkinWeightProfilesSemanticIndices.AddUnique(ProfileScemanticIndex);
@@ -2986,67 +2975,48 @@ mu::NodeMeshPtr GenerateMutableSourceMesh(const UEdGraphPin * Pin,
 				}
 
 				MeshData.bHasVertexColors = TypedNodeSkel->SkeletalMesh->GetHasVertexColors();
-				FSkeletalMeshModel* importModel = TypedNodeSkel->SkeletalMesh->GetImportedModel();
-				MeshData.NumTexCoordChannels = importModel->LODModels[LOD].NumTexCoords;
+				MeshData.NumTexCoordChannels = ImportedModel->LODModels[LODIndex].NumTexCoords;
 				MeshData.MaxBoneIndexTypeSizeBytes = MutableMesh->GetBoneMap().Num() > 256 ? 2 : 1;
-				MeshData.MaxNumBonesPerVertex = importModel->LODModels[LOD].GetMaxBoneInfluences();
+				MeshData.MaxNumBonesPerVertex = ImportedModel->LODModels[LODIndex].GetMaxBoneInfluences();
 				
 				// When mesh data is combined we will get an upper and lower bound of the number of triangles.
-				MeshData.MaxNumTriangles = importModel->LODModels[LOD].Sections[SectionIndex].NumTriangles;
-				MeshData.MinNumTriangles = importModel->LODModels[LOD].Sections[SectionIndex].NumTriangles;
+				MeshData.MaxNumTriangles = ImportedModel->LODModels[LODIndex].Sections[SectionIndex].NumTriangles;
+				MeshData.MinNumTriangles = ImportedModel->LODModels[LODIndex].Sections[SectionIndex].NumTriangles;
 			}
 
 			// Layouts
-			// If we didn't find a layout, but we are generating LODs and this LOD is automatic, reuse the first valid LOD layout
-			bool LayoutFound = false;
-
-			FSkeletalMeshModel* ImportedModel = TypedNodeSkel->SkeletalMesh->GetImportedModel();
-
-			int32 LayoutLOD = LOD;
-			while (!LayoutFound && LayoutLOD >= 0)
 			{
-				const int32 NumLayouts = ImportedModel->LODModels[LayoutLOD].NumTexCoords;
-				
+				// When using Automatic From Mesh all LODs share the same base layout, hence we use LODIndexConnected (as the base layout) instead of the LODIndex.
+				const int32 LODIndexLayout = GenerationContext.CurrentAutoLODStrategy == ECustomizableObjectAutomaticLODStrategy::AutomaticFromMesh ?
+					LODIndexConnected :
+					LODIndex;
+					
+				const int32 NumLayouts = ImportedModel->LODModels[LODIndexLayout].NumTexCoords;
 				MeshNode->SetLayoutCount(NumLayouts);
+				
+				bool bAtLeastOneLayout = false;
+
 				for (int32 LayoutIndex = 0; LayoutIndex < NumLayouts; ++LayoutIndex)
 				{
-					if (UEdGraphPin* LayoutPin = TypedNodeSkel->GetLayoutPin(LayoutLOD, SectionIndex, LayoutIndex))
+					if (UEdGraphPin* LayoutPin = TypedNodeSkel->GetLayoutPin(LODIndexLayout, SectionIndex, LayoutIndex))
 					{
 						if (const UEdGraphPin* ConnectedPin = FollowInputPin(*LayoutPin))
 						{
 							mu::NodeLayoutPtr LayoutNode = GenerateMutableSourceLayout(ConnectedPin, GenerationContext, bLinkedToExtendMaterial);
 							MeshNode->SetLayout(LayoutIndex, LayoutNode);
-							LayoutFound = true;
+							bAtLeastOneLayout = true;
 						}
 					}
 				}
 
-				if (GenerationContext.CurrentAutoLODStrategy == ECustomizableObjectAutomaticLODStrategy::AutomaticFromMesh)
-				{
-					--LayoutLOD;
-				}
-				else
+				if (!bAtLeastOneLayout)
 				{
-					break;
-				}
-			}
-
-			if (!LayoutFound)
-			{
-				MeshNode->SetLayoutCount(1);
-				mu::NodeLayoutBlocksPtr LayoutNode = new mu::NodeLayoutBlocks();
-				int GridSize = 4;
-				LayoutNode->SetGridSize(GridSize, GridSize);
-				LayoutNode->SetMaxGridSize(GridSize, GridSize);
-				LayoutNode->SetLayoutPackingStrategy(mu::EPackStrategy::RESIZABLE_LAYOUT);
-				LayoutNode->SetBlockReductionMethod(mu::EReductionMethod::HALVE_REDUCTION);
-				LayoutNode->SetBlockCount(1);
-				LayoutNode->SetBlock(0, 0, 0, GridSize, GridSize);
-				LayoutNode->SetBlockOptions(0, 0, false);
-				MeshNode->SetLayout(0, LayoutNode);
+					MeshNode->SetLayoutCount(1);
 
-				// We need it here because we create multiple nodes.
-				LayoutNode->SetMessageContext(Node);
+					mu::NodeLayoutBlocksPtr LayoutNode = CreateDefaultLayout();
+					MeshNode->SetLayout(0, LayoutNode);
+					LayoutNode->SetMessageContext(Node); // We need it here because we create multiple nodes.
+				}
 			}
 
 			// Applying Mesh Morph Nodes
@@ -3111,19 +3081,9 @@ mu::NodeMeshPtr GenerateMutableSourceMesh(const UEdGraphPin * Pin,
 				}
 				else
 				{
-					mu::NodeLayoutBlocksPtr LayoutNode = new mu::NodeLayoutBlocks();
-					int GridSize = 4;
-					LayoutNode->SetGridSize(GridSize, GridSize);
-					LayoutNode->SetMaxGridSize(GridSize, GridSize);
-					LayoutNode->SetLayoutPackingStrategy(mu::EPackStrategy::RESIZABLE_LAYOUT);
-					LayoutNode->SetBlockReductionMethod(mu::EReductionMethod::HALVE_REDUCTION);
-					LayoutNode->SetBlockCount(1);
-					LayoutNode->SetBlock(0, 0, 0, GridSize, GridSize);
-					LayoutNode->SetBlockOptions(0, 0, false);
+					mu::NodeLayoutBlocksPtr LayoutNode = CreateDefaultLayout();
 					MeshNode->SetLayout(0, LayoutNode);
-
-					// We need it here because we create multiple nodes.
-					LayoutNode->SetMessageContext(Node);
+					LayoutNode->SetMessageContext(Node);  // We need it here because we create multiple nodes.
 				}
 			}
 			else
@@ -3606,44 +3566,23 @@ mu::NodeMeshPtr GenerateMutableSourceMesh(const UEdGraphPin * Pin,
 					mu::NodeMeshTablePtr MeshTableNode = new mu::NodeMeshTable();
 					USkeletalMesh* SkeletalMesh = TypedNodeTable->GetColumnDefaultAssetByType<USkeletalMesh>(Pin);
 
-					int32 CurrentLOD = 0;
-					int32 MaterialIndex = 0;
+					int32 LODIndex = 0;
+					int32 SectionIndex = 0;
 
-					TypedNodeTable->GetPinLODAndMaterial(Pin, CurrentLOD, MaterialIndex);
+					TypedNodeTable->GetPinLODAndSection(Pin, LODIndex, SectionIndex);
 
-					if (GenerationContext.CurrentAutoLODStrategy == ECustomizableObjectAutomaticLODStrategy::AutomaticFromMesh)
+					if (SkeletalMesh)
 					{
-						if (SkeletalMesh && SkeletalMesh->GetImportedModel())
-						{
-							FSkeletalMeshModel* ImportedModel = SkeletalMesh->GetImportedModel();
-
-							CurrentLOD += GenerationContext.CurrentLOD;
-
-							// Checking if the current LOD is valid
-							if (ImportedModel->LODModels.Num() <= CurrentLOD
-								||
-								ImportedModel->LODModels[CurrentLOD].Sections.Num() <= MaterialIndex)
-							{
-								// Find the closest valid LOD 
-								for (int32 LODIndex = ImportedModel->LODModels.Num() - 1; LODIndex >= 0; --LODIndex)
-								{
-									if (ImportedModel->LODModels[LODIndex].Sections.Num() > MaterialIndex)
-									{
-										CurrentLOD = LODIndex;
-										break;
-									}
-								}
-							}
-						}
+						GetEffectiveLODAndSection(GenerationContext, *Node, *SkeletalMesh, LODIndex, SectionIndex);
 					}
 
 					// Getting the mutable table mesh column name
-					FString MutableColumnName = TypedNodeTable->GetMutableColumnName(Pin, CurrentLOD);
+					FString MutableColumnName = TypedNodeTable->GetMutableColumnName(Pin, LODIndex);
 
 					// Generating a new Mesh column if not exists
 					if (Table->FindColumn(StringCast<ANSICHAR>(*MutableColumnName).Get()) == INDEX_NONE)
 					{
-						bSuccess = GenerateTableColumn(TypedNodeTable, Pin, Table, DataTableColumnName, CurrentLOD, GenerationContext);
+						bSuccess = GenerateTableColumn(TypedNodeTable, Pin, Table, DataTableColumnName, LODIndex, GenerationContext);
 
 						if (!bSuccess)
 						{
@@ -3668,13 +3607,13 @@ mu::NodeMeshPtr GenerateMutableSourceMesh(const UEdGraphPin * Pin,
 							// Filling Mesh Data
 							FSkeletalMeshModel* importModel = SkeletalMesh->GetImportedModel();
 							MeshData.bHasVertexColors = SkeletalMesh->GetHasVertexColors();
-							MeshData.NumTexCoordChannels = importModel->LODModels[CurrentLOD].NumTexCoords;
-							MeshData.MaxBoneIndexTypeSizeBytes = importModel->LODModels[CurrentLOD].RequiredBones.Num() > 256 ? 2 : 1;
-							MeshData.MaxNumBonesPerVertex = importModel->LODModels[CurrentLOD].GetMaxBoneInfluences();
+							MeshData.NumTexCoordChannels = importModel->LODModels[LODIndex].NumTexCoords;
+							MeshData.MaxBoneIndexTypeSizeBytes = importModel->LODModels[LODIndex].RequiredBones.Num() > 256 ? 2 : 1;
+							MeshData.MaxNumBonesPerVertex = importModel->LODModels[LODIndex].GetMaxBoneInfluences();
 
 							// When mesh data is combined we will get an upper and lower bound of the number of triangles.
-							MeshData.MaxNumTriangles = importModel->LODModels[CurrentLOD].Sections[MaterialIndex].NumTriangles;
-							MeshData.MinNumTriangles = importModel->LODModels[CurrentLOD].Sections[MaterialIndex].NumTriangles;
+							MeshData.MaxNumTriangles = importModel->LODModels[LODIndex].Sections[SectionIndex].NumTriangles;
+							MeshData.MinNumTriangles = importModel->LODModels[LODIndex].Sections[SectionIndex].NumTriangles;
 						}
 
 						TArray<UCustomizableObjectLayout*> Layouts = TypedNodeTable->GetLayouts(Pin);
diff --git a/Engine/Plugins/Experimental/Mutable/Source/CustomizableObjectEditor/Private/MuCOE/Nodes/CustomizableObjectNodeTable.cpp b/Engine/Plugins/Experimental/Mutable/Source/CustomizableObjectEditor/Private/MuCOE/Nodes/CustomizableObjectNodeTable.cpp
--- a/Engine/Plugins/Experimental/Mutable/Source/CustomizableObjectEditor/Private/MuCOE/Nodes/CustomizableObjectNodeTable.cpp
+++ b/Engine/Plugins/Experimental/Mutable/Source/CustomizableObjectEditor/Private/MuCOE/Nodes/CustomizableObjectNodeTable.cpp
@@ -927,17 +927,17 @@ FString UCustomizableObjectNodeTable::GetMutableColumnName(const UEdGraphPin* Pi
 }
 
 
-void UCustomizableObjectNodeTable::GetPinLODAndMaterial(const UEdGraphPin* Pin, int32& LOD, int32& Material) const
+void UCustomizableObjectNodeTable::GetPinLODAndSection(const UEdGraphPin* Pin, int32& LODIndex, int32& SectionIndex) const
 {
 	if (const UCustomizableObjectNodeTableMeshPinData* PinData = Cast<UCustomizableObjectNodeTableMeshPinData >(GetPinData(*Pin)))
 	{
-		LOD = PinData->LOD;
-		Material = PinData->Material;
+		LODIndex = PinData->LOD;
+		SectionIndex = PinData->Material;
 	}
 	else
 	{
-		LOD = -1;
-		Material = -1;
+		LODIndex = -1;
+		SectionIndex = -1;
 	}
 }
 
