#rnx

diff --git a/Engine/Plugins/Experimental/Mutable/Source/CustomizableObjectEditor/Private/MuCOE/CustomizableObjectCompiler.cpp b/Engine/Plugins/Experimental/Mutable/Source/CustomizableObjectEditor/Private/MuCOE/CustomizableObjectCompiler.cpp
--- a/Engine/Plugins/Experimental/Mutable/Source/CustomizableObjectEditor/Private/MuCOE/CustomizableObjectCompiler.cpp
+++ b/Engine/Plugins/Experimental/Mutable/Source/CustomizableObjectEditor/Private/MuCOE/CustomizableObjectCompiler.cpp
@@ -929,7 +929,17 @@ void FCustomizableObjectCompiler::CompileInternal(UCustomizableObject* Object, c
 		return;
 	}
 
-	UE_LOG(LogMutable, Log, TEXT("Started Customizable Object Compile %s."), *Object->GetName());	
+	UE_LOG(LogMutable, Display, TEXT("Started Customizable Object Compile %s."), *Object->GetName());
+
+	if (InOptions.bIsCooking && InOptions.TargetPlatform)
+	{
+		UE_LOG(LogMutable, Display, TEXT("Compiling Customizable Object %s for platform %s."), *Object->GetName(), *InOptions.TargetPlatform->PlatformName());
+	}
+
+	if (InOptions.bIsCooking && InOptions.bForceLargeLODBias)
+	{
+		UE_LOG(LogMutable, Display, TEXT("Compiling Customizable Object with %d LODBias."), InOptions.DebugBias);
+	}
 
 	CompilationLogsContainer.ClearMessageCounters();
 
@@ -1433,7 +1443,7 @@ void FCustomizableObjectCompiler::FinishCompilation()
 		UCustomizableObjectSystem::GetInstance()->UnlockObject(CurrentObject);
 	}
 
-	UE_LOG(LogMutable, Log, TEXT("Finished Customizable Object Compile."));
+	UE_LOG(LogMutable, Display, TEXT("Finished Customizable Object Compile %s."), *CurrentObject->GetName());
 }
 
 void FCustomizableObjectCompiler::FinishSavingDerivedData()
