[Mutable] Change ClearWorkingMemoryOnUpdateEnd CVar name.

#rnx

diff --git a/Engine/Plugins/Experimental/Mutable/Source/CustomizableObject/Private/MuCO/CustomizableObjectSystem.cpp b/Engine/Plugins/Experimental/Mutable/Source/CustomizableObject/Private/MuCO/CustomizableObjectSystem.cpp
--- a/Engine/Plugins/Experimental/Mutable/Source/CustomizableObject/Private/MuCO/CustomizableObjectSystem.cpp
+++ b/Engine/Plugins/Experimental/Mutable/Source/CustomizableObject/Private/MuCO/CustomizableObjectSystem.cpp
@@ -89,7 +89,7 @@ static TAutoConsoleVariable<int32> CVarWorkingMemory(
 
 
 TAutoConsoleVariable<bool> CVarClearWorkingMemoryOnUpdateEnd(
-	TEXT("mutable.CVarClearWorkingMemoryOnUpdateEnd"),
+	TEXT("mutable.ClearWorkingMemoryOnUpdateEnd"),
 	false,
 	TEXT("Clear the working memory and cache after every Mutable operation."),
 	ECVF_Scalability);
