[Mutable] Fix realtime morph targets with different nodes with same skeletalmesh but different selections.

- It adds the node as key in the generated meshes cache.
- Make the skeletal mesh passed to convert to mutable const.

#rnx

diff --git a/Engine/Plugins/Experimental/Mutable/Source/CustomizableObjectEditor/Private/MuCOE/GenerateMutableSource/GenerateMutableSource.h b/Engine/Plugins/Experimental/Mutable/Source/CustomizableObjectEditor/Private/MuCOE/GenerateMutableSource/GenerateMutableSource.h
--- a/Engine/Plugins/Experimental/Mutable/Source/CustomizableObjectEditor/Private/MuCOE/GenerateMutableSource/GenerateMutableSource.h
+++ b/Engine/Plugins/Experimental/Mutable/Source/CustomizableObjectEditor/Private/MuCOE/GenerateMutableSource/GenerateMutableSource.h
@@ -376,7 +376,7 @@ struct FMutableComponentInfo
 	USkeleton* RefSkeleton = nullptr;
 
 	// Map to check skeleton compatibility
-	TMap<USkeleton*, bool> SkeletonCompatibility;
+	TMap<const USkeleton*, bool> SkeletonCompatibility;
 	
 	// Hierarchy hash from parent-bone to root bone, used to check if additional skeletons are compatible with
 	// the RefSkeleton
@@ -546,10 +546,16 @@ struct FMutableGraphGenerationContext
 			*/
 			FString Tags;
 
+			/**
+			* SkeletalMeshNode is needed to disambiguate realtime morph selection from diferent nodes.
+			* TODO: Consider using the actual selection.
+			*/
+			const UCustomizableObjectNode* SkeletalMeshNode = nullptr;
+
 			bool operator==( const FKey& OtherKey ) const
 			{
 				return Mesh == OtherKey.Mesh && LOD == OtherKey.LOD && CurrentLOD == OtherKey.CurrentLOD && MaterialIndex == OtherKey.MaterialIndex
-					&& Flags == OtherKey.Flags && Tags == OtherKey.Tags;
+					&& Flags == OtherKey.Flags && Tags == OtherKey.Tags && SkeletalMeshNode == OtherKey.SkeletalMeshNode;
 			}
 		};
 
diff --git a/Engine/Plugins/Experimental/Mutable/Source/CustomizableObjectEditor/Private/MuCOE/GenerateMutableSource/GenerateMutableSourceMesh.cpp b/Engine/Plugins/Experimental/Mutable/Source/CustomizableObjectEditor/Private/MuCOE/GenerateMutableSource/GenerateMutableSourceMesh.cpp
--- a/Engine/Plugins/Experimental/Mutable/Source/CustomizableObjectEditor/Private/MuCOE/GenerateMutableSource/GenerateMutableSourceMesh.cpp
+++ b/Engine/Plugins/Experimental/Mutable/Source/CustomizableObjectEditor/Private/MuCOE/GenerateMutableSource/GenerateMutableSourceMesh.cpp
@@ -128,7 +128,7 @@ void GetEffectiveLODAndSection(const FMutableGraphGenerationContext& Context, co
 }
 
 
-void BuildRemappedBonesArray(const FMutableComponentInfo& InComponentInfo, TObjectPtr<USkeletalMesh> InSkeletalMesh, int32 InLODIndex, const TArray<FBoneIndexType>& InRequiredBones, TArray<FBoneIndexType>& OutRemappedBones)
+void BuildRemappedBonesArray(const FMutableComponentInfo& InComponentInfo, TObjectPtr<const USkeletalMesh> InSkeletalMesh, int32 InLODIndex, const TArray<FBoneIndexType>& InRequiredBones, TArray<FBoneIndexType>& OutRemappedBones)
 {
 	if (!InSkeletalMesh)
 	{
@@ -324,9 +324,9 @@ void NormalizeWeights(FBoneIndexType* InfluenceBones, uint16* InfluenceWeights,
 }
 
 
-bool IsSkeletalMeshCompatibleWithRefSkeleton(FMutableComponentInfo& ComponentInfo, TObjectPtr<USkeletalMesh> InSkeletalMesh, FString& OutErrorMessage)
+bool IsSkeletalMeshCompatibleWithRefSkeleton(FMutableComponentInfo& ComponentInfo, TObjectPtr<const USkeletalMesh> InSkeletalMesh, FString& OutErrorMessage)
 {
-	TObjectPtr<USkeleton> Skeleton = InSkeletalMesh->GetSkeleton();
+	TObjectPtr<const USkeleton> Skeleton = InSkeletalMesh->GetSkeleton();
 
 	if (Skeleton == ComponentInfo.RefSkeleton)
 	{
@@ -627,7 +627,7 @@ mu::Ptr<mu::PhysicsBody> MakePhysicsBodyFromAsset(UPhysicsAsset* Asset, const TA
 	return PhysicsBody;
 }
 
-mu::MeshPtr ConvertSkeletalMeshToMutable(USkeletalMesh* InSkeletalMesh, const TSoftClassPtr<UAnimInstance>& AnimBp, int LOD, int MaterialIndex, FMutableGraphGenerationContext& GenerationContext, const UCustomizableObjectNode* CurrentNode)
+mu::MeshPtr ConvertSkeletalMeshToMutable(const USkeletalMesh* InSkeletalMesh, const TSoftClassPtr<UAnimInstance>& AnimBp, int LOD, int MaterialIndex, FMutableGraphGenerationContext& GenerationContext, const UCustomizableObjectNode* CurrentNode)
 {
 	if(!InSkeletalMesh)
 	{
@@ -689,7 +689,7 @@ mu::MeshPtr ConvertSkeletalMeshToMutable(USkeletalMesh* InSkeletalMesh, const TS
 	// Check if the Skeleton is valid and build the mu::Skeleton
 	if (!bIgnoreSkeleton)
 	{
-		USkeleton* InSkeleton = InSkeletalMesh->GetSkeleton();
+		const USkeleton* InSkeleton = InSkeletalMesh->GetSkeleton();
 		if (!InSkeleton)
 		{
 			FString Msg = FString::Printf(TEXT("No skeleton provided when converting SkeletalMesh [%s]."), *InSkeletalMesh->GetName());
@@ -1800,7 +1800,7 @@ mu::MeshPtr ConvertSkeletalMeshToMutable(USkeletalMesh* InSkeletalMesh, const TS
 }
 
 
-mu::MeshPtr ConvertStaticMeshToMutable(UStaticMesh* StaticMesh, int LOD, int MaterialIndex, FMutableGraphGenerationContext& GenerationContext, const UCustomizableObjectNode* CurrentNode)
+mu::MeshPtr ConvertStaticMeshToMutable(const UStaticMesh* StaticMesh, int LOD, int MaterialIndex, FMutableGraphGenerationContext& GenerationContext, const UCustomizableObjectNode* CurrentNode)
 {
 	if (!StaticMesh->GetRenderData() ||
 		!StaticMesh->GetRenderData()->LODResources.IsValidIndex(LOD) ||
@@ -1844,7 +1844,7 @@ mu::MeshPtr ConvertStaticMeshToMutable(UStaticMesh* StaticMesh, int LOD, int Mat
 
 		// Tangent buffer
 		{
-			FStaticMeshVertexBuffer& VertexBuffer = StaticMesh->GetRenderData()->LODResources[LOD].VertexBuffers.StaticMeshVertexBuffer;
+			const FStaticMeshVertexBuffer& VertexBuffer = StaticMesh->GetRenderData()->LODResources[LOD].VertexBuffers.StaticMeshVertexBuffer;
 
 			MESH_BUFFER_SEMANTIC Semantics[2];
 			int SemanticIndices[2];
@@ -1883,7 +1883,7 @@ mu::MeshPtr ConvertStaticMeshToMutable(UStaticMesh* StaticMesh, int LOD, int Mat
 
 		// Texture coordinates
 		{
-			FStaticMeshVertexBuffer& VertexBuffer = StaticMesh->GetRenderData()->LODResources[LOD].VertexBuffers.StaticMeshVertexBuffer;
+			const FStaticMeshVertexBuffer& VertexBuffer = StaticMesh->GetRenderData()->LODResources[LOD].VertexBuffers.StaticMeshVertexBuffer;
 
 			int texChannels = VertexBuffer.GetNumTexCoords();
 			int ChannelCount = texChannels;
@@ -1989,21 +1989,21 @@ mu::MeshPtr ConvertStaticMeshToMutable(UStaticMesh* StaticMesh, int LOD, int Mat
 
 // Convert a Mesh constant to a mutable format. UniqueTags are the tags that make this Mesh unique that cannot be merged in the cache 
 //  with the exact same Mesh with other tags
-mu::MeshPtr GenerateMutableMesh(UObject * Mesh, const TSoftClassPtr<UAnimInstance>& AnimInstance, int32 LOD, int32 MaterialIndex, const FString& UniqueTags, FMutableGraphGenerationContext & GenerationContext, const UCustomizableObjectNode* CurrentNode)
+mu::MeshPtr GenerateMutableMesh(const UObject* Mesh, const TSoftClassPtr<UAnimInstance>& AnimInstance, int32 LOD, int32 MaterialIndex, const FString& UniqueTags, FMutableGraphGenerationContext & GenerationContext, const UCustomizableObjectNode* CurrentNode)
 {
 	// Get the mesh generation flags to use
 	EMutableMeshConversionFlags CurrentFlags = GenerationContext.MeshGenerationFlags.Last();
 
-	FMutableGraphGenerationContext::FGeneratedMeshData::FKey Key = { Mesh, LOD, GenerationContext.CurrentLOD, MaterialIndex, CurrentFlags, UniqueTags };
+	FMutableGraphGenerationContext::FGeneratedMeshData::FKey Key = { Mesh, LOD, GenerationContext.CurrentLOD, MaterialIndex, CurrentFlags, UniqueTags, CurrentNode };
 	mu::MeshPtr MutableMesh = GenerationContext.FindGeneratedMesh(Key);
 	
 	if (!MutableMesh)
 	{
-		if (USkeletalMesh* SkeletalMesh = Cast<USkeletalMesh>(Mesh))
+		if (const USkeletalMesh* SkeletalMesh = Cast<USkeletalMesh>(Mesh))
 		{
 			MutableMesh = ConvertSkeletalMeshToMutable(SkeletalMesh, AnimInstance, LOD, MaterialIndex, GenerationContext, CurrentNode);
 		}
-		else if (UStaticMesh* StaticMesh = Cast<UStaticMesh>(Mesh))
+		else if (const UStaticMesh* StaticMesh = Cast<UStaticMesh>(Mesh))
 		{
 			MutableMesh = ConvertStaticMeshToMutable(StaticMesh, LOD, MaterialIndex, GenerationContext, CurrentNode);
 		}
diff --git a/Engine/Plugins/Experimental/Mutable/Source/CustomizableObjectEditor/Private/MuCOE/GenerateMutableSource/GenerateMutableSourceMesh.h b/Engine/Plugins/Experimental/Mutable/Source/CustomizableObjectEditor/Private/MuCOE/GenerateMutableSource/GenerateMutableSourceMesh.h
--- a/Engine/Plugins/Experimental/Mutable/Source/CustomizableObjectEditor/Private/MuCOE/GenerateMutableSource/GenerateMutableSourceMesh.h
+++ b/Engine/Plugins/Experimental/Mutable/Source/CustomizableObjectEditor/Private/MuCOE/GenerateMutableSource/GenerateMutableSourceMesh.h
@@ -18,13 +18,13 @@ class UStaticMesh;
 struct FMutableGraphGenerationContext;
 struct FMutableGraphMeshGenerationData;
 
-mu::MeshPtr ConvertSkeletalMeshToMutable(USkeletalMesh* InSkeletalMesh, const TSoftClassPtr<class UAnimInstance>& AnimBp, int LOD, int MaterialIndex, FMutableGraphGenerationContext& GenerationContext, const UCustomizableObjectNode* CurrentNode);
+mu::MeshPtr ConvertSkeletalMeshToMutable(const USkeletalMesh* InSkeletalMesh, const TSoftClassPtr<class UAnimInstance>& AnimBp, int LOD, int MaterialIndex, FMutableGraphGenerationContext& GenerationContext, const UCustomizableObjectNode* CurrentNode);
 
 
-mu::MeshPtr ConvertStaticMeshToMutable(UStaticMesh* StaticMesh, int LOD, int MaterialIndex, FMutableGraphGenerationContext& GenerationContext, const UCustomizableObjectNode* CurrentNode);
+mu::MeshPtr ConvertStaticMeshToMutable(const UStaticMesh* StaticMesh, int LOD, int MaterialIndex, FMutableGraphGenerationContext& GenerationContext, const UCustomizableObjectNode* CurrentNode);
 
 
-mu::MeshPtr GenerateMutableMesh(UObject * Mesh, const TSoftClassPtr<class UAnimInstance>& AnimBp, int32 LOD, int32 MaterialIndex, const FString& MeshUniqueTags, FMutableGraphGenerationContext & GenerationContext, const UCustomizableObjectNode* CurrentNode);
+mu::MeshPtr GenerateMutableMesh(const UObject * Mesh, const TSoftClassPtr<class UAnimInstance>& AnimBp, int32 LOD, int32 MaterialIndex, const FString& MeshUniqueTags, FMutableGraphGenerationContext & GenerationContext, const UCustomizableObjectNode* CurrentNode);
 
 
 mu::MeshPtr BuildMorphedMutableMesh(const UEdGraphPin * BaseSourcePin, const FString& MorphTargetName, FMutableGraphGenerationContext & GenerationContext, const FName& RowName = "");
