[mutable] Fixed Node tables support for static meshes.

#rnx

diff --git a/Engine/Plugins/Experimental/Mutable/Source/CustomizableObjectEditor/Classes/MuCOE/Nodes/CustomizableObjectNodeTable.h b/Engine/Plugins/Experimental/Mutable/Source/CustomizableObjectEditor/Classes/MuCOE/Nodes/CustomizableObjectNodeTable.h
--- a/Engine/Plugins/Experimental/Mutable/Source/CustomizableObjectEditor/Classes/MuCOE/Nodes/CustomizableObjectNodeTable.h
+++ b/Engine/Plugins/Experimental/Mutable/Source/CustomizableObjectEditor/Classes/MuCOE/Nodes/CustomizableObjectNodeTable.h
@@ -29,6 +29,15 @@ enum class ETableTextureType : uint8
 };
 
 
+/** Enum class for the different types of pin meshes */
+enum class ETableMeshPinType : uint8
+{
+	NONE = 0,
+	SKELETAL_MESH = 1,
+	STATIC_MESH = 2
+};
+
+
 /** Base class for all Table Pins. */
 UCLASS()
 class CUSTOMIZABLEOBJECTEDITOR_API UCustomizableObjectNodeTableObjectPinData : public UCustomizableObjectNodePinData
@@ -292,7 +301,14 @@ public:
 	bool IsImagePinDefault(UEdGraphPin* Pin);
 
 	// Returns the image mode of the column
-	ETableTextureType GetColumnImageMode(FString ColumnName) const;
+	ETableTextureType GetColumnImageMode(const FString& ColumnName) const;
+
+	// Returns the mesh type of the Pin
+	ETableMeshPinType GetPinMeshType(const UEdGraphPin* Pin) const;
+
+	// Functions to generate the names of a mutable table's column
+	FString GenerateSkeletalMeshMutableColumName(const FString& PinName, int32 LODIndex, int32 MaterialIndex) const;
+	FString GenerateStaticMeshMutableColumName(const FString& PinName, int32 MaterialIndex) const;
 
 private:
 
@@ -303,7 +319,7 @@ private:
 	FDelegateHandle OnTableChangedDelegateHandle;
 	
 	// Generates a mesh pin for each LOD and Material Surface of the reference SkeletalMesh
-	void GenerateMeshPins(UObject* Mesh, FString Name);
+	void GenerateMeshPins(UObject* Mesh, const FString& Name);
 
 	// Checks if a pin already exists and if it has the same type as before the node refresh
 	bool CheckPinUpdated(const FString& PinName, const FName& PinType) const;
diff --git a/Engine/Plugins/Experimental/Mutable/Source/CustomizableObjectEditor/Private/MuCOE/CustomizableObjectEditor.cpp b/Engine/Plugins/Experimental/Mutable/Source/CustomizableObjectEditor/Private/MuCOE/CustomizableObjectEditor.cpp
--- a/Engine/Plugins/Experimental/Mutable/Source/CustomizableObjectEditor/Private/MuCOE/CustomizableObjectEditor.cpp
+++ b/Engine/Plugins/Experimental/Mutable/Source/CustomizableObjectEditor/Private/MuCOE/CustomizableObjectEditor.cpp
@@ -43,6 +43,7 @@
 #include "MuCOE/Nodes/CustomizableObjectNodeProjectorConstant.h"
 #include "MuCOE/Nodes/CustomizableObjectNodeSkeletalMesh.h"
 #include "MuCOE/Nodes/CustomizableObjectNodeStaticMesh.h"
+#include "MuCOE/Nodes/CustomizableObjectNodeTable.h"
 #include "MuCOE/SCustomizableObjectEditorAdvancedPreviewSettings.h"
 #include "MuCOE/SCustomizableObjectEditorPerformanceReport.h"
 #include "MuCOE/SCustomizableObjectEditorTagExplorer.h"
@@ -1576,6 +1577,10 @@ void FCustomizableObjectEditor::Tick( float InDeltaTime )
 				{
 					ClipMesh = TypedNode->StaticMesh;
 				}
+				else if (const UCustomizableObjectNodeTable* TableNode = Cast<UCustomizableObjectNodeTable>(Node))
+				{
+					ClipMesh = TableNode->GetColumnDefaultAssetByType<UStaticMesh>(ConnectedPin);
+				}
 			}
 		}
 
diff --git a/Engine/Plugins/Experimental/Mutable/Source/CustomizableObjectEditor/Private/MuCOE/GenerateMutableSource/GenerateMutableSourceMesh.cpp b/Engine/Plugins/Experimental/Mutable/Source/CustomizableObjectEditor/Private/MuCOE/GenerateMutableSource/GenerateMutableSourceMesh.cpp
--- a/Engine/Plugins/Experimental/Mutable/Source/CustomizableObjectEditor/Private/MuCOE/GenerateMutableSource/GenerateMutableSourceMesh.cpp
+++ b/Engine/Plugins/Experimental/Mutable/Source/CustomizableObjectEditor/Private/MuCOE/GenerateMutableSource/GenerateMutableSourceMesh.cpp
@@ -3547,7 +3547,10 @@ mu::NodeMeshPtr GenerateMutableSourceMesh(const UEdGraphPin * Pin,
 				bSuccess = false;
 			}
 
-			if (bSuccess && !TypedNodeTable->GetColumnDefaultAssetByType<USkeletalMesh>(Pin))
+			USkeletalMesh* DefaultSkeletalMesh = TypedNodeTable->GetColumnDefaultAssetByType<USkeletalMesh>(Pin);
+			UStaticMesh* DefaultStaticMesh = TypedNodeTable->GetColumnDefaultAssetByType<UStaticMesh>(Pin);
+
+			if (bSuccess && !DefaultSkeletalMesh && !DefaultStaticMesh)
 			{
 				FString Msg = FString::Printf(TEXT("Couldn't find a default value in the data table's struct for the column [%s]."), *DataTableColumnName);
 				GenerationContext.Compiler->CompilerLog(FText::FromString(Msg), Node);
@@ -3564,16 +3567,15 @@ mu::NodeMeshPtr GenerateMutableSourceMesh(const UEdGraphPin * Pin,
 				if (Table)
 				{
 					mu::NodeMeshTablePtr MeshTableNode = new mu::NodeMeshTable();
-					USkeletalMesh* SkeletalMesh = TypedNodeTable->GetColumnDefaultAssetByType<USkeletalMesh>(Pin);
 
 					int32 LODIndex = 0;
 					int32 SectionIndex = 0;
 
 					TypedNodeTable->GetPinLODAndSection(Pin, LODIndex, SectionIndex);
 
-					if (SkeletalMesh)
+					if (DefaultSkeletalMesh)
 					{
-						GetEffectiveLODAndSection(GenerationContext, *Node, *SkeletalMesh, LODIndex, SectionIndex);
+						GetEffectiveLODAndSection(GenerationContext, *Node, *DefaultSkeletalMesh, LODIndex, SectionIndex);
 					}
 
 					// Getting the mutable table mesh column name
@@ -3601,12 +3603,12 @@ mu::NodeMeshPtr GenerateMutableSourceMesh(const UEdGraphPin * Pin,
 
 						GenerationContext.AddParameterNameUnique(Node, TypedNodeTable->ParameterName);
 
-						if (SkeletalMesh)
+						if (DefaultSkeletalMesh)
 						{
 							// TODO: this should be made for all the meshes of the Column to support meshes with different values
 							// Filling Mesh Data
-							FSkeletalMeshModel* importModel = SkeletalMesh->GetImportedModel();
-							MeshData.bHasVertexColors = SkeletalMesh->GetHasVertexColors();
+							FSkeletalMeshModel* importModel = DefaultSkeletalMesh->GetImportedModel();
+							MeshData.bHasVertexColors = DefaultSkeletalMesh->GetHasVertexColors();
 							MeshData.NumTexCoordChannels = importModel->LODModels[LODIndex].NumTexCoords;
 							MeshData.MaxBoneIndexTypeSizeBytes = importModel->LODModels[LODIndex].RequiredBones.Num() > 256 ? 2 : 1;
 							MeshData.MaxNumBonesPerVertex = importModel->LODModels[LODIndex].GetMaxBoneInfluences();
@@ -3665,7 +3667,7 @@ mu::NodeMeshPtr GenerateMutableSourceMesh(const UEdGraphPin * Pin,
 						}
 
 						// Applying Mesh Morph Nodes
-						if (GenerationContext.MeshMorphStack.Num())
+						if (DefaultSkeletalMesh && GenerationContext.MeshMorphStack.Num())
 						{
 							MorphResult = GenerateMorphMesh(Pin, GenerationContext.MeshMorphStack, 0, Result, GenerationContext, MeshData, MutableColumnName);
 						}
diff --git a/Engine/Plugins/Experimental/Mutable/Source/CustomizableObjectEditor/Private/MuCOE/GenerateMutableSource/GenerateMutableSourceTable.cpp b/Engine/Plugins/Experimental/Mutable/Source/CustomizableObjectEditor/Private/MuCOE/GenerateMutableSource/GenerateMutableSourceTable.cpp
--- a/Engine/Plugins/Experimental/Mutable/Source/CustomizableObjectEditor/Private/MuCOE/GenerateMutableSource/GenerateMutableSourceTable.cpp
+++ b/Engine/Plugins/Experimental/Mutable/Source/CustomizableObjectEditor/Private/MuCOE/GenerateMutableSource/GenerateMutableSourceTable.cpp
@@ -188,7 +188,7 @@ bool FillTableColumn(const UCustomizableObjectNodeTable* TableNode,	mu::TablePtr
 
 			for (int32 MatIndex = 0; MatIndex < Materials; ++MatIndex)
 			{
-				FString MutableColumnName = ColumnName + FString::Printf(TEXT(" LOD_%d "), LOD) + FString::Printf(TEXT("Mat_%d"), MatIndex);
+				FString MutableColumnName = TableNode->GenerateSkeletalMeshMutableColumName(ColumnName, LOD, MatIndex);
 
 				CurrentColumn = MutableTable.get()->FindColumn(StringCast<ANSICHAR>(*MutableColumnName).Get());
 
@@ -278,8 +278,9 @@ bool FillTableColumn(const UCustomizableObjectNodeTable* TableNode,	mu::TablePtr
 				return false;
 			}
 
-			// Parameter used for LOD differences
-			int32 CurrentLOD = LOD;
+			// Static Meshes only support LOD 0 at the moment (TODO)
+			//int32 CurrentLOD = LOD;
+			int32 CurrentLOD = 0;
 
 			int NumLODs = StaticMesh->GetRenderData()->LODResources.Num();
 
@@ -308,7 +309,7 @@ bool FillTableColumn(const UCustomizableObjectNodeTable* TableNode,	mu::TablePtr
 
 			for (int32 MatIndex = 0; MatIndex < Materials; ++MatIndex)
 			{
-				FString MutableColumnName = ColumnName + FString::Printf(TEXT(" LOD_%d "), LOD) + FString::Printf(TEXT("Mat_%d"), MatIndex);
+				FString MutableColumnName = TableNode->GenerateStaticMeshMutableColumName(ColumnName, MatIndex);
 
 				CurrentColumn = MutableTable.get()->FindColumn(StringCast<ANSICHAR>(*MutableColumnName).Get());
 
diff --git a/Engine/Plugins/Experimental/Mutable/Source/CustomizableObjectEditor/Private/MuCOE/Nodes/CustomizableObjectNodeTable.cpp b/Engine/Plugins/Experimental/Mutable/Source/CustomizableObjectEditor/Private/MuCOE/Nodes/CustomizableObjectNodeTable.cpp
--- a/Engine/Plugins/Experimental/Mutable/Source/CustomizableObjectEditor/Private/MuCOE/Nodes/CustomizableObjectNodeTable.cpp
+++ b/Engine/Plugins/Experimental/Mutable/Source/CustomizableObjectEditor/Private/MuCOE/Nodes/CustomizableObjectNodeTable.cpp
@@ -369,7 +369,7 @@ void UCustomizableObjectNodeTable::AllocateDefaultPins(UCustomizableObjectNodeRe
 }
 
 
-void UCustomizableObjectNodeTable::GenerateMeshPins(UObject* Mesh, FString Name)
+void UCustomizableObjectNodeTable::GenerateMeshPins(UObject* Mesh, const FString& Name)
 {
 	if (USkeletalMesh* SkeletalMesh = Cast<USkeletalMesh>(Mesh))
 	{
@@ -382,7 +382,7 @@ void UCustomizableObjectNodeTable::GenerateMeshPins(UObject* Mesh, FString Name)
 
 			for (int32 MatIndex = 0; MatIndex < NumMaterials; ++MatIndex)
 			{
-				FString TableMeshPinName = Name + FString::Printf(TEXT(" LOD_%d "), LODIndex) + FString::Printf(TEXT("Mat_%d"), MatIndex);
+				FString TableMeshPinName = GenerateSkeletalMeshMutableColumName(Name, LODIndex, MatIndex);
 
 				UCustomizableObjectNodeTableMeshPinData* PinData = NewObject<UCustomizableObjectNodeTableMeshPinData>(this);
 				PinData->ColumnName = Name;
@@ -424,36 +424,25 @@ void UCustomizableObjectNodeTable::GenerateMeshPins(UObject* Mesh, FString Name)
 	
 	else if (UStaticMesh* StaticMesh = Cast<UStaticMesh>(Mesh))
 	{
-		int NumLODs = StaticMesh->GetRenderData()->LODResources.Num();
 		const UEdGraphSchema_CustomizableObject* Schema = GetDefault<UEdGraphSchema_CustomizableObject>();
 
-		for (int32 LODIndex = 0; LODIndex < NumLODs; ++LODIndex)
+		if (StaticMesh->GetRenderData()->LODResources.Num())
 		{
-			int32 NumMaterials = StaticMesh->GetRenderData()->LODResources[LODIndex].Sections.Num();
+			int32 NumMaterials = StaticMesh->GetRenderData()->LODResources[0].Sections.Num();
 
 			for (int32 MatIndex = 0; MatIndex < NumMaterials; ++MatIndex)
 			{
-				FString TableMeshPinName = Name + FString::Printf(TEXT(" LOD_%d "), LODIndex) + FString::Printf(TEXT("Mat_%d"), MatIndex);
+				FString TableMeshPinName = GenerateStaticMeshMutableColumName(Name, MatIndex);
 
 				UCustomizableObjectNodeTableMeshPinData* PinData = NewObject<UCustomizableObjectNodeTableMeshPinData>(this);
 				PinData->ColumnName = Name;
 				PinData->MutableColumnName = TableMeshPinName;
-				PinData->LOD = LODIndex;
+				PinData->LOD = 0;
 				PinData->Material = MatIndex;
 
 				UEdGraphPin* MeshPin = CustomCreatePin(EGPD_Output, Schema->PC_Mesh, FName(*TableMeshPinName), PinData);
 				MeshPin->PinFriendlyName = FText::FromString(TableMeshPinName);
 				MeshPin->SafeSetHidden(false);
-
-				UCustomizableObjectLayout* Layout = NewObject<UCustomizableObjectLayout>(this);
-				FString LayoutName = TableMeshPinName;
-
-				if (Layout)
-				{
-					PinData->Layouts.Add(Layout);
-					Layout->SetLayout(StaticMesh, LODIndex, MatIndex, 0);
-					Layout->SetLayoutName(LayoutName);
-				}
 			}
 		}
 	}
@@ -533,8 +522,8 @@ bool UCustomizableObjectNodeTable::IsNodeOutDatedAndNeedsRefresh()
 
 							for (int32 MatIndex = 0; MatIndex < NumMaterials; ++MatIndex)
 							{
-								FString PinName = DataTableUtils::GetPropertyExportName(ColumnProperty);
-								PinName += FString::Printf(TEXT(" LOD_%d "), LODIndex) + FString::Printf(TEXT("Mat_%d"), MatIndex);
+								FString PropertyName = DataTableUtils::GetPropertyExportName(ColumnProperty);
+								FString PinName = GenerateSkeletalMeshMutableColumName(PropertyName, LODIndex, MatIndex);
 
 								if (CheckPinUpdated(PinName, Schema->PC_Mesh))
 								{
@@ -547,16 +536,14 @@ bool UCustomizableObjectNodeTable::IsNodeOutDatedAndNeedsRefresh()
 					}
 					else if (UStaticMesh* StaticMesh = Cast<UStaticMesh>(Object))
 					{
-						int NumLODs = StaticMesh->GetRenderData()->LODResources.Num();
-
-						for (int32 LODIndex = 0; LODIndex < NumLODs; ++LODIndex)
+						if (StaticMesh->GetRenderData()->LODResources.Num())
 						{
-							int32 NumMaterials = StaticMesh->GetRenderData()->LODResources[LODIndex].Sections.Num();
+							int32 NumMaterials = StaticMesh->GetRenderData()->LODResources[0].Sections.Num();
 
 							for (int32 MatIndex = 0; MatIndex < NumMaterials; ++MatIndex)
 							{
-								FString PinName = DataTableUtils::GetPropertyExportName(ColumnProperty);
-								PinName += FString::Printf(TEXT(" LOD_%d "), LODIndex) + FString::Printf(TEXT("Mat_%d"), MatIndex);
+								FString PropertyName = DataTableUtils::GetPropertyExportName(ColumnProperty);
+								FString PinName = GenerateStaticMeshMutableColumName(PropertyName, MatIndex);
 
 								if (CheckPinUpdated(PinName, Schema->PC_Mesh))
 								{
@@ -678,20 +665,23 @@ void UCustomizableObjectNodeTable::RemapPinsData(const TMap<UEdGraphPin*, UEdGra
 					PinDataNewPin->AnimTagColumnName = PinDataOldPin->AnimTagColumnName;
 				}
 
-				// Keeping information added in layout editor if the layout is the same
-				for (TObjectPtr<UCustomizableObjectLayout>& NewLayout : PinDataNewPin->Layouts)
+				if (GetPinMeshType(Pair.Value) == ETableMeshPinType::SKELETAL_MESH)
 				{
-					for (TObjectPtr<UCustomizableObjectLayout>& OldLayout : PinDataOldPin->Layouts)
+					// Keeping information added in layout editor if the layout is the same
+					for (TObjectPtr<UCustomizableObjectLayout>& NewLayout : PinDataNewPin->Layouts)
 					{
-						if (NewLayout->GetLayoutName() == OldLayout->GetLayoutName())
+						for (TObjectPtr<UCustomizableObjectLayout>& OldLayout : PinDataOldPin->Layouts)
 						{
-							NewLayout->Blocks = OldLayout->Blocks;
-							NewLayout->SetGridSize(OldLayout->GetGridSize());
-							NewLayout->SetMaxGridSize(OldLayout->GetMaxGridSize());
-							NewLayout->SetPackingStrategy(OldLayout->GetPackingStrategy());
-							NewLayout->SetBlockReductionMethod(OldLayout->GetBlockReductionMethod());
+							if (NewLayout->GetLayoutName() == OldLayout->GetLayoutName())
+							{
+								NewLayout->Blocks = OldLayout->Blocks;
+								NewLayout->SetGridSize(OldLayout->GetGridSize());
+								NewLayout->SetMaxGridSize(OldLayout->GetMaxGridSize());
+								NewLayout->SetPackingStrategy(OldLayout->GetPackingStrategy());
+								NewLayout->SetBlockReductionMethod(OldLayout->GetBlockReductionMethod());
 
-							break;
+								break;
+							}
 						}
 					}
 				}
@@ -1174,7 +1164,7 @@ bool UCustomizableObjectNodeTable::IsImagePinDefault(UEdGraphPin* Pin)
 }
 
 
-ETableTextureType UCustomizableObjectNodeTable::GetColumnImageMode(FString ColumnName) const
+ETableTextureType UCustomizableObjectNodeTable::GetColumnImageMode(const FString& ColumnName) const
 {
 	for (UEdGraphPin* Pin : Pins)
 	{
@@ -1190,4 +1180,40 @@ ETableTextureType UCustomizableObjectNodeTable::GetColumnImageMode(FString Colum
 	return DefaultImageMode;
 }
 
+
+ETableMeshPinType UCustomizableObjectNodeTable::GetPinMeshType(const UEdGraphPin* Pin) const
+{
+	if (Pin && Pin->PinType.PinCategory == UEdGraphSchema_CustomizableObject::PC_Mesh)
+	{
+		FString ColumnName = GetColumnNameByPin(Pin);
+		FProperty* Property = Table->FindTableProperty(FName(*ColumnName));
+
+		if (const FSoftObjectProperty* SoftObjectProperty = CastField<FSoftObjectProperty>(Property))
+		{
+			if (SoftObjectProperty->PropertyClass->IsChildOf(USkeletalMesh::StaticClass()))
+			{
+				return ETableMeshPinType::SKELETAL_MESH;
+			}
+			else if (SoftObjectProperty->PropertyClass->IsChildOf(UStaticMesh::StaticClass()))
+			{
+				return ETableMeshPinType::STATIC_MESH;
+			}
+		}
+	}
+
+	return ETableMeshPinType::NONE;
+}
+
+
+FString UCustomizableObjectNodeTable::GenerateSkeletalMeshMutableColumName(const FString& PinName, int32 LODIndex, int32 MaterialIndex) const
+{
+	return PinName + FString::Printf(TEXT(" LOD_%d "), LODIndex) + FString::Printf(TEXT("Mat_%d"), MaterialIndex);
+}
+
+
+FString UCustomizableObjectNodeTable::GenerateStaticMeshMutableColumName(const FString& PinName, int32 MaterialIndex) const
+{
+	return PinName + FString::Printf(TEXT(" Mat_%d"), MaterialIndex);
+}
+
 #undef LOCTEXT_NAMESPACE
diff --git a/Engine/Plugins/Experimental/Mutable/Source/CustomizableObjectEditor/Private/MuCOE/Nodes/CustomizableObjectNodeTableDetails.cpp b/Engine/Plugins/Experimental/Mutable/Source/CustomizableObjectEditor/Private/MuCOE/Nodes/CustomizableObjectNodeTableDetails.cpp
--- a/Engine/Plugins/Experimental/Mutable/Source/CustomizableObjectEditor/Private/MuCOE/Nodes/CustomizableObjectNodeTableDetails.cpp
+++ b/Engine/Plugins/Experimental/Mutable/Source/CustomizableObjectEditor/Private/MuCOE/Nodes/CustomizableObjectNodeTableDetails.cpp
@@ -298,7 +298,7 @@ void FCustomizableObjectNodeTableDetails::GenerateMeshColumnComboBoxOptions()
 					{
 						const UCustomizableObjectNodeTableMeshPinData* PinData = Cast<UCustomizableObjectNodeTableMeshPinData >(Node->GetPinData(*Pin));
 
-						if (!PinData || PinData->ColumnName != MeshColumnName)
+						if (!PinData || PinData->ColumnName != MeshColumnName || Node->GetPinMeshType(Pin) != ETableMeshPinType::SKELETAL_MESH)
 						{
 							continue;
 						}
