[Mutable] Fix crash when connecting a Morph Material Node to an emtpy Skeletal Mesh Node.

#rnx

diff --git a/Engine/Plugins/Experimental/Mutable/Source/CustomizableObject/Public/MuCO/CustomizableObjectCustomVersion.h b/Engine/Plugins/Experimental/Mutable/Source/CustomizableObject/Public/MuCO/CustomizableObjectCustomVersion.h
--- a/Engine/Plugins/Experimental/Mutable/Source/CustomizableObject/Public/MuCO/CustomizableObjectCustomVersion.h
+++ b/Engine/Plugins/Experimental/Mutable/Source/CustomizableObject/Public/MuCO/CustomizableObjectCustomVersion.h
@@ -84,6 +84,8 @@ struct CUSTOMIZABLEOBJECT_API FCustomizableObjectCustomVersion
 		
 		IgnoreDisabledSections,
 
+		SkeletalMeshNodeDefaultPinWithoutPinData,
+
 		// -----<new versions can be added above this line>-------------------------------------------------
 		VersionPlusOne,
 		LatestVersion = VersionPlusOne - 1
diff --git a/Engine/Plugins/Experimental/Mutable/Source/CustomizableObjectEditor/Private/MuCOE/Nodes/CustomizableObjectNodeMorphMaterialDetails.cpp b/Engine/Plugins/Experimental/Mutable/Source/CustomizableObjectEditor/Private/MuCOE/Nodes/CustomizableObjectNodeMorphMaterialDetails.cpp
--- a/Engine/Plugins/Experimental/Mutable/Source/CustomizableObjectEditor/Private/MuCOE/Nodes/CustomizableObjectNodeMorphMaterialDetails.cpp
+++ b/Engine/Plugins/Experimental/Mutable/Source/CustomizableObjectEditor/Private/MuCOE/Nodes/CustomizableObjectNodeMorphMaterialDetails.cpp
@@ -51,14 +51,17 @@ void FCustomizableObjectNodeMorphMaterialDetails::CustomizeDetails( IDetailLayou
 			{
 				if ( const UCustomizableObjectNodeSkeletalMesh* TypedSourceNode = Cast<UCustomizableObjectNodeSkeletalMesh>(BaseSourcePin->GetOwningNode()) )
 				{
-					for (int m = 0; m < TypedSourceNode->SkeletalMesh->GetMorphTargets().Num(); ++m )
+					if (TypedSourceNode->SkeletalMesh)
 					{
-						FString MorphName = *TypedSourceNode->SkeletalMesh->GetMorphTargets()[m]->GetName();
-						MorphTargetComboOptions.Add( MakeShareable( new FString( MorphName ) ) );
-
-						if (Node->MorphTargetName == MorphName)
+						for (int m = 0; m < TypedSourceNode->SkeletalMesh->GetMorphTargets().Num(); ++m )
 						{
-							ItemToSelect = MorphTargetComboOptions.Last();
+							FString MorphName = *TypedSourceNode->SkeletalMesh->GetMorphTargets()[m]->GetName();
+							MorphTargetComboOptions.Add( MakeShareable( new FString( MorphName ) ) );
+
+							if (Node->MorphTargetName == MorphName)
+							{
+								ItemToSelect = MorphTargetComboOptions.Last();
+							}
 						}
 					}
 				}
diff --git a/Engine/Plugins/Experimental/Mutable/Source/CustomizableObjectEditor/Private/MuCOE/Nodes/CustomizableObjectNodeSkeletalMesh.cpp b/Engine/Plugins/Experimental/Mutable/Source/CustomizableObjectEditor/Private/MuCOE/Nodes/CustomizableObjectNodeSkeletalMesh.cpp
--- a/Engine/Plugins/Experimental/Mutable/Source/CustomizableObjectEditor/Private/MuCOE/Nodes/CustomizableObjectNodeSkeletalMesh.cpp
+++ b/Engine/Plugins/Experimental/Mutable/Source/CustomizableObjectEditor/Private/MuCOE/Nodes/CustomizableObjectNodeSkeletalMesh.cpp
@@ -52,7 +52,10 @@ void UCustomizableObjectNodeSkeletalMesh::AllocateDefaultPins(UCustomizableObjec
 
 	if (!SkeletalMesh)
 	{
-		DefaultPin = CustomCreatePin(EGPD_Output, Schema->PC_Mesh, FName(SKELETAL_MESH_PIN_NAME));
+		UCustomizableObjectNodeSkeletalMeshPinDataMesh* PinData = NewObject<UCustomizableObjectNodeSkeletalMeshPinDataMesh>(this);
+		PinData->Init(-1, -1);
+		
+		DefaultPin = CustomCreatePin(EGPD_Output, Schema->PC_Mesh, FName(SKELETAL_MESH_PIN_NAME), PinData);
 		return;
 	}
 	else
@@ -619,6 +622,17 @@ void UCustomizableObjectNodeSkeletalMesh::BackwardsCompatibleFixup()
 	{
 		ReconstructNode(); // Pins representing disabled sections could be present. 
 	}
+
+	if (CustomizableObjectCustomVersion < FCustomizableObjectCustomVersion::SkeletalMeshNodeDefaultPinWithoutPinData)
+	{
+		if (const UEdGraphPin* Pin = DefaultPin.Get())
+		{
+			UCustomizableObjectNodeSkeletalMeshPinDataMesh* PinData = NewObject<UCustomizableObjectNodeSkeletalMeshPinDataMesh>(this);
+			PinData->Init(-1, -1);
+
+			AddPinData(*Pin, *PinData);
+		}
+	}
 }
 
 
