[mutable] Redesign runtime operation tracking data structures to reduce memory operations and fragmentation.
Add unloading of cached results if memory budget is not met.
Remove hack with op-hit-count to track state-cache results.
#rnx

diff --git a/Engine/Plugins/Experimental/Mutable/Source/CustomizableObject/Private/MuCO/CustomizableObject.cpp b/Engine/Plugins/Experimental/Mutable/Source/CustomizableObject/Private/MuCO/CustomizableObject.cpp
--- a/Engine/Plugins/Experimental/Mutable/Source/CustomizableObject/Private/MuCO/CustomizableObject.cpp
+++ b/Engine/Plugins/Experimental/Mutable/Source/CustomizableObject/Private/MuCO/CustomizableObject.cpp
@@ -35,8 +35,6 @@
 #endif
 
 
-#include "MuR/Model.h"
-
 #include UE_INLINE_GENERATED_CPP_BY_NAME(CustomizableObject)
 
 #define LOCTEXT_NAMESPACE "CustomizableObject"
diff --git a/Engine/Plugins/Experimental/Mutable/Source/CustomizableObjectEditor/Private/MuCOE/SMutableCodeViewer.cpp b/Engine/Plugins/Experimental/Mutable/Source/CustomizableObjectEditor/Private/MuCOE/SMutableCodeViewer.cpp
--- a/Engine/Plugins/Experimental/Mutable/Source/CustomizableObjectEditor/Private/MuCOE/SMutableCodeViewer.cpp
+++ b/Engine/Plugins/Experimental/Mutable/Source/CustomizableObjectEditor/Private/MuCOE/SMutableCodeViewer.cpp
@@ -2482,9 +2482,8 @@ void SMutableCodeViewer::Tick(const FGeometry& AllottedGeometry, const double In
 	case mu::DT_PROJECTOR:
 	{
 		check (PreviewProjectorViewer);
-		const mu::Ptr<const mu::Projector> ProjectorPtr = System->GetPrivate()->BuildProjector(
-				MutableModel, PreviewParameters.get(), SelectedOperationAddress);
-		PreviewProjectorViewer->SetProjector(ProjectorPtr->m_value);
+		mu::FProjector Projector = System->GetPrivate()->BuildProjector(MutableModel, PreviewParameters.get(), SelectedOperationAddress);
+		PreviewProjectorViewer->SetProjector(Projector);
 	}
 	
 	default:
diff --git a/Engine/Plugins/Experimental/Mutable/Source/MutableRuntime/Private/MuR/CodeRunner.cpp b/Engine/Plugins/Experimental/Mutable/Source/MutableRuntime/Private/MuR/CodeRunner.cpp
--- a/Engine/Plugins/Experimental/Mutable/Source/MutableRuntime/Private/MuR/CodeRunner.cpp
+++ b/Engine/Plugins/Experimental/Mutable/Source/MutableRuntime/Private/MuR/CodeRunner.cpp
@@ -307,7 +307,7 @@ namespace mu
 				case DT_SCALAR:		StoreScalar(item, 0.0f); break;
 				case DT_STRING:		StoreString(item, nullptr); break;
 				case DT_COLOUR:		StoreColor(item, FVector4f()); break;
-				case DT_PROJECTOR:  StoreProjector(item, nullptr); break;
+				case DT_PROJECTOR:  StoreProjector(item, FProjector()); break;
 				case DT_MESH:       StoreMesh(item, nullptr); break;
 				case DT_IMAGE:      StoreImage(item, nullptr); break;
 				case DT_LAYOUT:     StoreLayout(item, nullptr); break;
@@ -2499,7 +2499,7 @@ namespace mu
 				MUTABLE_CPUPROFILER_SCOPE(ME_PROJECT_1)
 
                 Ptr<const Mesh> pMesh = LoadMesh( FCacheAddress(args.mesh,item) );
-                Ptr<const Projector> pProjector = LoadProjector( FCacheAddress(args.projector,item) );
+                FProjector Projector = LoadProjector( FCacheAddress(args.projector,item) );
 
                 // Only if both are valid.
                 MeshPtr pResult;
@@ -2507,18 +2507,7 @@ namespace mu
                     && pMesh.get()
                     && pMesh.get()->GetVertexBuffers().GetBufferCount() > 0 )
                 {
-                    if (pProjector)
-                    {
-                        pResult = MeshProject(pMesh.get(), pProjector->m_value);
-//                        if (pResult)
-//                        {
-//                            pResult->GetPrivate()->CheckIntegrity();
-//                        }
-                    }
-                    else if (pMesh)
-                    {
-                        pResult = pMesh->Clone();
-                    }
+                    pResult = MeshProject(pMesh.get(), Projector);
                 }
                 StoreMesh( item, pResult );
                 break;
@@ -4019,37 +4008,35 @@ namespace mu
 				FScheduledOpData Data;
 				Data.RasterMesh.Mip = ProjectionMip;
 				Data.RasterMesh.MipValue = static_cast<float>(ProjectionMip);
-				Ptr<const Projector> pProjector = LoadProjector(FScheduledOp::FromOpAndOptions(args.projector, item, 0));
-				if (pProjector)
+				FProjector Projector = LoadProjector(FScheduledOp::FromOpAndOptions(args.projector, item, 0));
+
+				EMinFilterMethod MinFilterMethod = Invoke([&]() -> EMinFilterMethod
 				{
-					EMinFilterMethod MinFilterMethod = Invoke([&]() -> EMinFilterMethod
+					if (ForcedProjectionMode == 0)
 					{
-						if (ForcedProjectionMode == 0)
-						{
-							return EMinFilterMethod::None;
-						}
-						else if (ForcedProjectionMode == 1)
-						{
-							return EMinFilterMethod::TotalAreaHeuristic;
-						}
+						return EMinFilterMethod::None;
+					}
+					else if (ForcedProjectionMode == 1)
+					{
+						return EMinFilterMethod::TotalAreaHeuristic;
+					}
 						
-						return static_cast<EMinFilterMethod>(args.MinFilterMethod);
-					});
+					return static_cast<EMinFilterMethod>(args.MinFilterMethod);
+				});
 
-					if (MinFilterMethod == EMinFilterMethod::TotalAreaHeuristic)
-					{
-						FVector2f TargetImageSizeF = FVector2f(
-							FMath::Max(args.sizeX >> MipsToSkip, 1),
-							FMath::Max(args.sizeY >> MipsToSkip, 1));
-						FVector2f SourceImageSizeF = FVector2f(args.SourceSizeX, args.SourceSizeY);
+				if (MinFilterMethod == EMinFilterMethod::TotalAreaHeuristic)
+				{
+					FVector2f TargetImageSizeF = FVector2f(
+						FMath::Max(args.sizeX >> MipsToSkip, 1),
+						FMath::Max(args.sizeY >> MipsToSkip, 1));
+					FVector2f SourceImageSizeF = FVector2f(args.SourceSizeX, args.SourceSizeY);
 						
-						if (pMesh)
-						{ 
-							const float ComputedMip = ComputeProjectedFootprintBestMip(pMesh.get(), pProjector->m_value, TargetImageSizeF, SourceImageSizeF);
+					if (pMesh)
+					{ 
+						const float ComputedMip = ComputeProjectedFootprintBestMip(pMesh.get(), Projector, TargetImageSizeF, SourceImageSizeF);
 
-							Data.RasterMesh.MipValue = FMath::Max(0.0f, ComputedMip + GlobalProjectionLodBias);
-							Data.RasterMesh.Mip = static_cast<uint8>(FMath::FloorToInt32(Data.RasterMesh.MipValue));
-						}
+						Data.RasterMesh.MipValue = FMath::Max(0.0f, ComputedMip + GlobalProjectionLodBias);
+						Data.RasterMesh.Mip = static_cast<uint8>(FMath::FloorToInt32(Data.RasterMesh.MipValue));
 					}
 				}
 	
@@ -4057,7 +4044,7 @@ namespace mu
 				// this needs to be done before the operation is re-scheduled otherwise the execution may deadlock.	
 				// TODO: Review if this is a good way of reusing data between stages.
 				StoreMesh(FScheduledOp::FromOpAndOptions(args.mesh, item, 0), pMesh);
-				StoreProjector(FScheduledOp::FromOpAndOptions(args.projector, item, 0), pProjector);
+				StoreProjector(FScheduledOp::FromOpAndOptions(args.projector, item, 0), Projector);
 				
 
 				int32 DataHeapAddress = m_heapData.Add(Data);
@@ -4208,48 +4195,46 @@ namespace mu
 
 				if (args.projector && Source && Source->GetSizeX() > 0 && Source->GetSizeY() > 0)
 				{
-					Ptr<const Projector> pProjector = LoadProjector(FScheduledOp::FromOpAndOptions(args.projector, item, 0));
-					if (pProjector)
+					FProjector Projector = LoadProjector(FScheduledOp::FromOpAndOptions(args.projector, item, 0));
+
+					switch (Projector.type)
 					{
-						switch (pProjector->m_value.type)
-						{
-						case PROJECTOR_TYPE::PLANAR:
-							ImageRasterProjectedPlanar(pMesh.get(), pNew.get(),
-								Source.get(), Mask.get(),
-								args.bIsRGBFadingEnabled, args.bIsAlphaFadingEnabled,
-								SamplingMethod,
-								FadeStartRad, FadeEndRad, FMath::Frac(Data.RasterMesh.MipValue),
-								layout, args.blockIndex,
-								CropMin, UncroppedSize,
-								&scratch, bUseProjectionVectorImpl);
-							break;
-
-						case PROJECTOR_TYPE::WRAPPING:
-							ImageRasterProjectedWrapping(pMesh.get(), pNew.get(),
-								Source.get(), Mask.get(),
-								args.bIsRGBFadingEnabled, args.bIsAlphaFadingEnabled,
-								SamplingMethod,
-								FadeStartRad, FadeEndRad, FMath::Frac(Data.RasterMesh.MipValue),
-								layout, args.blockIndex,
-								CropMin, UncroppedSize,
-								&scratch);
-							break;
-
-						case PROJECTOR_TYPE::CYLINDRICAL:
-							ImageRasterProjectedCylindrical(pMesh.get(), pNew.get(),
-								Source.get(), Mask.get(),
-								args.bIsRGBFadingEnabled, args.bIsAlphaFadingEnabled,
-								FadeStartRad, FadeEndRad,
-								layout,
-								pProjector->m_value.projectionAngle,
-								CropMin, UncroppedSize,
-								&scratch);
-							break;
-
-						default:
-							check(false);
-							break;
-						}
+					case PROJECTOR_TYPE::PLANAR:
+						ImageRasterProjectedPlanar(pMesh.get(), pNew.get(),
+							Source.get(), Mask.get(),
+							args.bIsRGBFadingEnabled, args.bIsAlphaFadingEnabled,
+							SamplingMethod,
+							FadeStartRad, FadeEndRad, FMath::Frac(Data.RasterMesh.MipValue),
+							layout, args.blockIndex,
+							CropMin, UncroppedSize,
+							&scratch, bUseProjectionVectorImpl);
+						break;
+
+					case PROJECTOR_TYPE::WRAPPING:
+						ImageRasterProjectedWrapping(pMesh.get(), pNew.get(),
+							Source.get(), Mask.get(),
+							args.bIsRGBFadingEnabled, args.bIsAlphaFadingEnabled,
+							SamplingMethod,
+							FadeStartRad, FadeEndRad, FMath::Frac(Data.RasterMesh.MipValue),
+							layout, args.blockIndex,
+							CropMin, UncroppedSize,
+							&scratch);
+						break;
+
+					case PROJECTOR_TYPE::CYLINDRICAL:
+						ImageRasterProjectedCylindrical(pMesh.get(), pNew.get(),
+							Source.get(), Mask.get(),
+							args.bIsRGBFadingEnabled, args.bIsAlphaFadingEnabled,
+							FadeStartRad, FadeEndRad,
+							layout,
+							Projector.projectionAngle,
+							CropMin, UncroppedSize,
+							&scratch);
+						break;
+
+					default:
+						check(false);
+						break;
 					}
 				}
 
@@ -5247,9 +5232,8 @@ namespace mu
         case OP_TYPE::PR_CONSTANT:
         {
 			OP::ResourceConstantArgs args = pModel->GetPrivate()->m_program.GetOpArgs<OP::ResourceConstantArgs>(item.At);
-            ProjectorPtr pResult = new Projector();
-            pResult->m_value = program.m_constantProjectors[args.value];
-            StoreProjector( item, pResult );
+            FProjector Result = program.m_constantProjectors[args.value];
+            StoreProjector( item, Result );
             break;
         }
 
@@ -5257,14 +5241,13 @@ namespace mu
         {
 			OP::ParameterArgs args = pModel->GetPrivate()->m_program.GetOpArgs<OP::ParameterArgs>(item.At);
 			Ptr<RangeIndex> index = BuildCurrentOpRangeIndex( item, pParams, pModel, args.variable );
-            ProjectorPtr pResult = new Projector();
-            pResult->m_value = pParams->GetPrivate()->GetProjectorValue(args.variable,index);
+            FProjector Result = pParams->GetPrivate()->GetProjectorValue(args.variable,index);
 
             // The type cannot be changed, take it from the default value
             const FProjector& def = program.m_parameters[args.variable].m_defaultValue.Get<ParamProjectorType>();
-            pResult->m_value.type = def.type;
+            Result.type = def.type;
 
-            StoreProjector( item, pResult );
+            StoreProjector( item, Result );
             break;
         }
 
diff --git a/Engine/Plugins/Experimental/Mutable/Source/MutableRuntime/Private/MuR/CodeRunner.h b/Engine/Plugins/Experimental/Mutable/Source/MutableRuntime/Private/MuR/CodeRunner.h
--- a/Engine/Plugins/Experimental/Mutable/Source/MutableRuntime/Private/MuR/CodeRunner.h
+++ b/Engine/Plugins/Experimental/Mutable/Source/MutableRuntime/Private/MuR/CodeRunner.h
@@ -412,14 +412,14 @@ namespace  mu
 			return m_pSystem->WorkingMemoryManager.CurrentInstanceCache->GetString(From);
 		}
 
-		Ptr<const Projector> LoadProjector(const FCacheAddress& From)
+		FProjector LoadProjector(const FCacheAddress& From)
 		{
 			return m_pSystem->WorkingMemoryManager.CurrentInstanceCache->GetProjector(From);
 		}
 
 		Ptr<const Mesh> LoadMesh(const FCacheAddress& From)
 		{
-			return m_pSystem->WorkingMemoryManager.CurrentInstanceCache->GetMesh(From);
+			return m_pSystem->WorkingMemoryManager.LoadMesh(From);
 		}
 
 		Ptr<const Image> LoadImage(const FCacheAddress& From)
@@ -472,7 +472,7 @@ namespace  mu
 			m_pSystem->WorkingMemoryManager.CurrentInstanceCache->SetColour(To, Value);
 		}
 
-		void StoreProjector(const FCacheAddress& To, Ptr<const Projector> Value)
+		void StoreProjector(const FCacheAddress& To, const FProjector& Value)
 		{
 			m_pSystem->WorkingMemoryManager.CurrentInstanceCache->SetProjector(To, Value);
 		}
diff --git a/Engine/Plugins/Experimental/Mutable/Source/MutableRuntime/Private/MuR/System.cpp b/Engine/Plugins/Experimental/Mutable/Source/MutableRuntime/Private/MuR/System.cpp
--- a/Engine/Plugins/Experimental/Mutable/Source/MutableRuntime/Private/MuR/System.cpp
+++ b/Engine/Plugins/Experimental/Mutable/Source/MutableRuntime/Private/MuR/System.cpp
@@ -348,14 +348,13 @@ namespace mu
 				const mu::Model* Model = pLiveInstance->Model.Get();
 				const mu::FProgram& program = Model->GetPrivate()->m_program;
 
+				// TODO: It should be possible to reuse this data if cleared in the correct places only, together with m_heapImageDesc.
 				int32 VarValue = CVarClearImageDescCache.GetValueOnAnyThread();
 				if (VarValue != 0)
 				{
-					m_pD->WorkingMemoryManager.CurrentInstanceCache->m_descCache.Reset();
+					m_pD->WorkingMemoryManager.CurrentInstanceCache->ClearDescCache();
 				}
 
-				m_pD->WorkingMemoryManager.CurrentInstanceCache->m_descCache.SetNum(program.m_opAddress.Num());
-
 				OP::ADDRESS at = res.m_rootAddress;
 				mu::OP_TYPE opType = program.GetOpType(at);
 				if (GetOpDataType(opType) == DT_IMAGE)
@@ -441,6 +440,13 @@ namespace mu
 
 		// Reduce the cache until it fits the limit.
 		m_pD->WorkingMemoryManager.EnsureBudgetBelow(0);
+
+		// If we don't constrain the memory budget, free the pooled images or they may pile up.
+		if (m_pD->WorkingMemoryManager.BudgetBytes == 0)
+		{
+			m_pD->WorkingMemoryManager.PooledImages.Empty();
+		}
+
 	}
 
 
@@ -456,16 +462,19 @@ namespace mu
 			if (Instance.InstanceID == instanceID)
 			{
 				// Make sure all the resources cached in the instance are removed from the tracking list
-				CodeContainer<TPair<int32, Ptr<const Resource>>>::iterator it = Instance.Cache->m_resources.begin();
-				for (; it.IsValid(); ++it)
+				for (Ptr<const Image>& Data : Instance.Cache->ImageResults)
 				{
-					Ptr<const Resource> Value = (*it).Value.get();
-					if (!Value)
+					if (Data)
 					{
-						continue;
+						m_pD->WorkingMemoryManager.CacheResources.Remove(Data);
+					}
+				}
+				for (Ptr<const Mesh>& Data : Instance.Cache->MeshResults)
+				{
+					if (Data)
+					{
+						m_pD->WorkingMemoryManager.CacheResources.Remove(Data);
 					}
-
-					m_pD->WorkingMemoryManager.CacheResources.Remove(Value);
 				}
 
 				m_pD->WorkingMemoryManager.LiveInstances.RemoveAtSwap(Index);
@@ -729,13 +738,13 @@ namespace mu
 
 	
 	//---------------------------------------------------------------------------------------------
-	Ptr<const Projector> System::Private::BuildProjector(const TSharedPtr<const Model>& pModel, const Parameters* pParams, OP::ADDRESS at)
+	FProjector System::Private::BuildProjector(const TSharedPtr<const Model>& pModel, const Parameters* pParams, OP::ADDRESS at)
 	{
 		WorkingMemoryManager.BeginRunnerThread();
 
     	RunCode(pModel, pParams, at);
 
-		Ptr<const Projector> Result;
+		FProjector Result;
 		if (!bUnrecoverableError)
 		{
 			Result = WorkingMemoryManager.CurrentInstanceCache->GetProjector(FCacheAddress(at, 0, 0));
@@ -784,7 +793,7 @@ namespace mu
 			RunCode(pModel, pParams, at);
 			if (!bUnrecoverableError)
 			{
-				Result = WorkingMemoryManager.CurrentInstanceCache->GetMesh(FCacheAddress(at, 0, 0));
+				Result = WorkingMemoryManager.LoadMesh(FCacheAddress(at, 0, 0), true);
 			}
 
 			// Debug check to see if we managed the op-hit-counts correctly
@@ -854,18 +863,25 @@ namespace mu
 	{
 		MUTABLE_CPUPROFILER_SCOPE(PrepareCache);
 
-		FProgram& program = InModel->GetPrivate()->m_program;
-		int32 opCount = program.m_opAddress.Num();
-		WorkingMemoryManager.CurrentInstanceCache->m_opHitCount.clear();
-		WorkingMemoryManager.CurrentInstanceCache->Init(opCount);
+		FProgram& Program = InModel->GetPrivate()->m_program;
+		int32 OpCount = Program.m_opAddress.Num();
+		WorkingMemoryManager.CurrentInstanceCache->Init(OpCount);
+
+		// Clear cache flags of existing data
+		CodeContainer<FProgramCache::FOpExecutionData>::iterator It = WorkingMemoryManager.CurrentInstanceCache->OpExecutionData.begin();
+		for (; It.IsValid(); ++It)
+		{
+			(*It).OpHitCount = 0; // This should already be 0, but just in case.
+			(*It).IsCacheLocked = false;
+		}
 
 		// Mark the resources that have to be cached to update the instance in this state
-		if (InState >= 0)
+		if (InState >= 0 && InState< Program.m_states.Num())
 		{
-			const FProgram::FState& s = program.m_states[InState];
-			for (uint32 a : s.m_updateCache)
+			const FProgram::FState& State = Program.m_states[InState];
+			for (uint32 Address : State.m_updateCache)
 			{
-				WorkingMemoryManager.CurrentInstanceCache->SetForceCached(a);
+				WorkingMemoryManager.CurrentInstanceCache->SetForceCached(Address);
 			}
 		}
 	}
@@ -878,8 +894,7 @@ namespace mu
 
 		// For now, we calculate these for every log. We will later track on resource creation, destruction or state change.
 		const uint32 RomBytes = GetRomBytes();
-		const uint32 Cache0Bytes = GetCache0Bytes();
-		const uint32 Cache1Bytes = GetCache1Bytes();
+		const uint32 CacheBytes = GetCacheBytes();
 		const uint32 TrackedCacheBytes = GetTrackedCacheBytes();
 		const uint32 PoolBytes = GetPooledBytes();
 		const uint32 TempBytes = GetTempBytes();
@@ -908,9 +923,13 @@ namespace mu
 		{
 			Models.AddUnique(Instance.Model.Get());
 
-			InternalBytes += Instance.Cache->m_resources.GetAllocatedSize()
-				+ Instance.Cache->m_descCache.GetAllocatedSize()
-				+ Instance.Cache->m_opHitCount.GetAllocatedSize()
+			InternalBytes += Instance.Cache->OpExecutionData.GetAllocatedSize()
+				+ Instance.Cache->ImageResults.GetAllocatedSize()
+				+ Instance.Cache->MeshResults.GetAllocatedSize()
+				+ Instance.Cache->LayoutResults.GetAllocatedSize()
+				+ Instance.Cache->InstanceResults.GetAllocatedSize()
+				+ Instance.Cache->ColorResults.GetAllocatedSize()
+				+ Instance.Cache->ProjectorResults.GetAllocatedSize()
 				+ Instance.Cache->m_usedRangeIndices.GetAllocatedSize();
 
 			for (const FWorkingMemoryManager::FModelCacheEntry& c : CachePerModel)
@@ -940,7 +959,7 @@ namespace mu
 			InternalBytes / 1024,
 			TempBytes / 1024,
 			PoolBytes / 1024,
-			(Cache0Bytes + Cache1Bytes) / 1024,
+			CacheBytes / 1024,
 			TrackedCacheBytes / 1024,
 			RomBytes / 1024,
 			StreamingBufferBytes / 1024,
@@ -950,52 +969,52 @@ namespace mu
 
 		// Detailed log for debug.
 		//if (TotalBytes > 30 * 1024 * 1024)
-		if (false)
-		{
-			UE_LOG(LogMutableCore, Log, TEXT("Mem detailed report:"));
-
-			for (const FLiveInstance& Instance : LiveInstances)
-			{
-				UE_LOG(LogMutableCore, Log, TEXT("  instance:"));
-
-				TSet<const Resource*> InstanceCache0Unique;
-				TSet<const Resource*> InstanceCache1Unique;
-
-				CodeContainer<int>::iterator it = Instance.Cache->m_opHitCount.begin();
-				for (; it.IsValid(); ++it)
-				{
-					const Resource* Value = Instance.Cache->m_resources[it.get_address()].Value.get();
-					if (!Value)
-					{
-						continue;
-					}
-
-					int32 Count = *it;
-					if (Count < UE_MUTABLE_CACHE_COUNT_LIMIT)
-					{
-						InstanceCache0Unique.Add(Value);
-					}
-					else
-					{
-						InstanceCache1Unique.Add(Value);
-					}
-				}
-
-				UE_LOG(LogMutableCore, Log, TEXT("    cache level 0:"));
-
-				for (const Resource* Value : InstanceCache0Unique)
-				{
-					UE_LOG(LogMutableCore, Log, TEXT("      resource: %d"), Value->GetDataSize());
-				}
-
-				UE_LOG(LogMutableCore, Log, TEXT("    cache level 1:"));
-
-				for (const Resource* Value : InstanceCache1Unique)
-				{
-					UE_LOG(LogMutableCore, Log, TEXT("      resource: %d"), Value->GetDataSize());
-				}
-			}
-		}
+		//if (false)
+		//{
+		//	UE_LOG(LogMutableCore, Log, TEXT("Mem detailed report:"));
+
+		//	for (const FLiveInstance& Instance : LiveInstances)
+		//	{
+		//		UE_LOG(LogMutableCore, Log, TEXT("  instance:"));
+
+		//		TSet<const Resource*> InstanceCache0Unique;
+		//		TSet<const Resource*> InstanceCache1Unique;
+
+		//		CodeContainer<int>::iterator it = Instance.Cache->m_opHitCount.begin();
+		//		for (; it.IsValid(); ++it)
+		//		{
+		//			const Resource* Value = Instance.Cache->m_resources[it.get_address()].Value.get();
+		//			if (!Value)
+		//			{
+		//				continue;
+		//			}
+
+		//			int32 Count = *it;
+		//			if (Count < UE_MUTABLE_CACHE_COUNT_LIMIT)
+		//			{
+		//				InstanceCache0Unique.Add(Value);
+		//			}
+		//			else
+		//			{
+		//				InstanceCache1Unique.Add(Value);
+		//			}
+		//		}
+
+		//		UE_LOG(LogMutableCore, Log, TEXT("    cache level 0:"));
+
+		//		for (const Resource* Value : InstanceCache0Unique)
+		//		{
+		//			UE_LOG(LogMutableCore, Log, TEXT("      resource: %d"), Value->GetDataSize());
+		//		}
+
+		//		UE_LOG(LogMutableCore, Log, TEXT("    cache level 1:"));
+
+		//		for (const Resource* Value : InstanceCache1Unique)
+		//		{
+		//			UE_LOG(LogMutableCore, Log, TEXT("      resource: %d"), Value->GetDataSize());
+		//		}
+		//	}
+		//}
 
 		TRACE_COUNTER_SET(MutableRuntime_MemInternal, InternalBytes);
 		TRACE_COUNTER_SET(MutableRuntime_MemTemp, TempBytes);
@@ -1054,21 +1073,14 @@ namespace mu
 		uint32 RomBytes = 0;
 		uint32 PoolBytes = 0;
 		uint32 TempBytes = 0;
-		//uint32 Cache0Bytes = 0;
-		//uint32 Cache1Bytes = 0;
 		uint32 TrackedCacheBytes = 0;
 		{
 			MUTABLE_CPUPROFILER_SCOPE(EnsureBudgetBelow_Measure);
 			RomBytes = GetRomBytes();
 			PoolBytes = GetPooledBytes();
 			TempBytes = GetTempBytes();
-			//Cache0Bytes = GetCache0Bytes();
-			//Cache1Bytes = GetCache1Bytes();
 			TrackedCacheBytes = GetTrackedCacheBytes();
 		}
-		// These types of memory are not tracked here yet.
-		//uint32 StreamingBufferBytes = 0;
-		//uint32 InternalBytes = 0;
 
 		uint64 TotalBytes = 0;
 		TotalBytes += TempBytes;
@@ -1149,7 +1161,7 @@ namespace mu
 				{
 					MUTABLE_CPUPROFILER_SCOPE(EnsureBudgetBelow_UnloadRom);
 
-					//UE_LOG(LogMutableCore,Log, "Unloading rom because of memory budget: %d.", lowestPriorityRom);
+					// UE_LOG(LogMutableCore,Log, "Unloading rom because of memory budget: %d.", lowestPriorityRom);
 					LowestPriorityModel->GetPrivate()->m_program.UnloadRom(LowestPriorityRom);
 					TotalBytes -= LowestPriorityModel->GetPrivate()->m_program.m_roms[LowestPriorityRom].Size;
 					bFinished = TotalBytes <= BudgetBytes;
@@ -1157,13 +1169,229 @@ namespace mu
 			}
 		}
 
-		// Try to free cache 1 memory (from other models first?)
+		// Try to free cache 1 memory
 		if (!bFinished)
 		{
 			MUTABLE_CPUPROFILER_SCOPE(EnsureBudgetBelow_FreeCached);
 
-			// TODO.
-			// from other models first?
+			// From other live instances first
+			for (const FLiveInstance& Instance : LiveInstances)
+			{
+				if (Instance.Cache==CurrentInstanceCache)
+				{
+					// Ignore the current live instance.
+					continue;
+				}
+
+				// Gather all data in the cache for this instance
+				TArray<const Resource*> CacheUnique;
+				CodeContainer<FProgramCache::FOpExecutionData>::iterator It = Instance.Cache->OpExecutionData.begin();
+				for (; It.IsValid(); ++It)
+				{
+					FProgramCache::FOpExecutionData& Data = *It;
+
+					if (!Data.DataTypeIndex)
+					{
+						continue;
+					}
+
+					const Resource* Value = nullptr;
+					switch (Data.DataType)
+					{
+					case DATATYPE::DT_IMAGE:
+						Value = Instance.Cache->ImageResults[Data.DataTypeIndex].get();
+						if (Value)
+						{
+							CacheUnique.AddUnique(Value);
+						}
+						break;
+
+					case DATATYPE::DT_MESH:
+						Value = Instance.Cache->MeshResults[Data.DataTypeIndex].get();
+						if (Value)
+						{
+							CacheUnique.AddUnique(Value);
+						}
+						break;
+
+					default:
+						break;
+					}
+				}
+
+				while (!bFinished && CacheUnique.Num())
+				{
+					// Free one
+					const Resource* Removed = CacheUnique.Pop();
+
+					int32 RemovedDataSize = Removed->GetDataSize();
+
+					// Clear its cache references
+					CodeContainer<FProgramCache::FOpExecutionData>::iterator RemIt = Instance.Cache->OpExecutionData.begin();
+					for (; RemIt.IsValid(); ++RemIt)
+					{
+						FProgramCache::FOpExecutionData& Data = *RemIt;
+
+						if (!Data.DataTypeIndex)
+						{
+							continue;
+						}
+
+						const Resource* Value = nullptr;
+						switch (Data.DataType)
+						{
+						case DATATYPE::DT_IMAGE:
+							Value = Instance.Cache->ImageResults[Data.DataTypeIndex].get();
+							break;
+
+						case DATATYPE::DT_MESH:
+							Value = Instance.Cache->MeshResults[Data.DataTypeIndex].get();
+							break;
+
+						default:
+							break;
+						}
+
+						if (Value == Removed)
+						{
+							// \TODO: This is not very efficient, since it will search the data again.
+							Instance.Cache->SetUnused(RemIt.get_address());
+						}
+					}
+
+					TotalBytes -= RemovedDataSize;
+					bFinished = TotalBytes <= BudgetBytes;
+				}
+			}
+
+			// From the current live instances. It is more involved: we have to make sure any data we want to free is not also
+			// in any cache (0 or 1) position with hit-count > 0.
+			if (CurrentInstanceCache)
+			{
+				// Gather all data in the cache for this instance
+				TArray<const Resource*> CacheUnique;
+				CodeContainer<FProgramCache::FOpExecutionData>::iterator It = CurrentInstanceCache->OpExecutionData.begin();
+				for (; It.IsValid(); ++It)
+				{
+					FProgramCache::FOpExecutionData& Data = *It;
+
+					if (!Data.DataTypeIndex || Data.OpHitCount>0 || !Data.IsCacheLocked)
+					{
+						continue;
+					}
+
+					const Resource* Value = nullptr;
+					switch (Data.DataType)
+					{
+					case DATATYPE::DT_IMAGE:
+						Value = CurrentInstanceCache->ImageResults[Data.DataTypeIndex].get();
+						if (Value)
+						{
+							CacheUnique.AddUnique(Value);
+						}
+						break;
+
+					case DATATYPE::DT_MESH:
+						Value = CurrentInstanceCache->MeshResults[Data.DataTypeIndex].get();
+						if (Value)
+						{
+							CacheUnique.AddUnique(Value);
+						}
+						break;
+
+					default:
+						break;
+					}
+				}
+
+				while (!bFinished && CacheUnique.Num())
+				{
+					// Free one
+					const Resource* Removed = CacheUnique.Pop();
+
+					// Does this data have any other cache references with op-hit-count bigger than 0?
+					bool bStillUsed = false;
+
+					CodeContainer<FProgramCache::FOpExecutionData>::iterator CheckIt = CurrentInstanceCache->OpExecutionData.begin();
+					for (; CheckIt.IsValid(); ++CheckIt)
+					{
+						const FProgramCache::FOpExecutionData& Data = *CheckIt;
+
+						if (!Data.DataTypeIndex)
+						{
+							continue;
+						}
+
+						const Resource* Value = nullptr;
+						switch (Data.DataType)
+						{
+						case DATATYPE::DT_IMAGE:
+							Value = CurrentInstanceCache->ImageResults[Data.DataTypeIndex].get();
+							break;
+
+						case DATATYPE::DT_MESH:
+							Value = CurrentInstanceCache->MeshResults[Data.DataTypeIndex].get();
+							break;
+
+						default:
+							break;
+						}
+
+						if (Value == Removed && Data.OpHitCount>0)
+						{
+							bStillUsed = true;
+							break;
+						}
+					}
+
+					if (bStillUsed)
+					{
+						continue;
+					}
+
+
+					int32 RemovedDataSize = Removed->GetDataSize();
+
+					// Clear its cache references
+					CodeContainer<FProgramCache::FOpExecutionData>::iterator RemIt = CurrentInstanceCache->OpExecutionData.begin();
+					for (; RemIt.IsValid(); ++RemIt)
+					{
+						FProgramCache::FOpExecutionData& Data = *RemIt;
+
+						if (!Data.DataTypeIndex)
+						{
+							continue;
+						}
+
+						const Resource* Value = nullptr;
+						switch (Data.DataType)
+						{
+						case DATATYPE::DT_IMAGE:
+							Value = CurrentInstanceCache->ImageResults[Data.DataTypeIndex].get();
+							break;
+
+						case DATATYPE::DT_MESH:
+							Value = CurrentInstanceCache->MeshResults[Data.DataTypeIndex].get();
+							break;
+
+						default:
+							break;
+						}
+
+						if (Value == Removed)
+						{
+							CacheResources.Remove(Removed);
+							// \TODO: This is not very efficient, since it will search the data again.
+							CurrentInstanceCache->SetUnused(RemIt.get_address());
+						}
+					}
+
+					TotalBytes -= RemovedDataSize;
+					bFinished = TotalBytes <= BudgetBytes;
+				}
+
+			}
+
 		}
 
 		TRACE_COUNTER_SET(MutableRuntime_MemTemp, TempBytes);
diff --git a/Engine/Plugins/Experimental/Mutable/Source/MutableRuntime/Private/MuR/SystemPrivate.h b/Engine/Plugins/Experimental/Mutable/Source/MutableRuntime/Private/MuR/SystemPrivate.h
--- a/Engine/Plugins/Experimental/Mutable/Source/MutableRuntime/Private/MuR/SystemPrivate.h
+++ b/Engine/Plugins/Experimental/Mutable/Source/MutableRuntime/Private/MuR/SystemPrivate.h
@@ -19,7 +19,6 @@
 #include "HAL/PlatformTLS.h"
 #endif
 
-#define UE_MUTABLE_CACHE_COUNT_LIMIT	3000000
 
 namespace mu
 {
@@ -42,60 +41,6 @@ namespace mu
 	constexpr uint64 AllParametersMask = TNumericLimits<uint64> ::Max();
 
 
-	//! Reference-counted colour to be stored in the cache
-	class Colour : public Resource
-	{
-	public:
-		Colour(FVector4f v = FVector4f()) : m_colour(v) {}
-		FVector4f m_colour;
-		int32 GetDataSize() const override { return sizeof(Colour); }
-	};
-	typedef Ptr<Colour> ColourPtr;
-
-
-	//! Reference-counted bool to be stored in the cache
-	class Bool : public Resource
-	{
-	public:
-		Bool(bool v = false) : m_value(v) {}
-		bool m_value;
-		int32 GetDataSize() const override { return sizeof(Bool); }
-	};
-	typedef Ptr<Bool> BoolPtr;
-
-
-	//! Reference-counted scalar to be stored in the cache
-	class Scalar : public Resource
-	{
-	public:
-		Scalar(float v = 0.0f) : m_value(v) {}
-		float m_value;
-		int32 GetDataSize() const override { return sizeof(Scalar); }
-	};
-	typedef Ptr<Scalar> ScalarPtr;
-
-
-	//! Reference-counted scalar to be stored in the cache
-	class Int : public Resource
-	{
-	public:
-		Int(int32 v = 0) : m_value(v) {}
-		int32 m_value;
-		int32 GetDataSize() const override { return sizeof(Int); }
-	};
-	typedef Ptr<Int> IntPtr;
-
-
-	//! Reference-counted scalar to be stored in the cache
-	class Projector : public Resource
-	{
-	public:
-		FProjector m_value;
-		int32 GetDataSize() const override { return sizeof(Projector); }
-	};
-	typedef Ptr<Projector> ProjectorPtr;
-
-
 	/** ExecutinIndex stores the location inside all ranges of the execution of a specific
 	* operation. The first integer on each pair is the dimension/range index in the program
 	* array of ranges, and the second integer is the value inside that range.
@@ -264,7 +209,7 @@ namespace mu
 
 		void resize(size_t s)
 		{
-			m_index0.SetNum(s);
+			m_index0.SetNumZeroed(s);
 		}
 
 		uint32 size_code() const
@@ -314,6 +259,30 @@ namespace mu
 			return 0;
 		}
 
+		inline DATA* get_ptr(const FCacheAddress& at)
+		{
+			if (at.ExecutionIndex == 0 && at.ExecutionOptions == 0)
+			{
+				if (at.At < uint32(m_index0.Num()))
+				{
+					return &m_index0[at.At];
+				}
+				else
+				{
+					return nullptr;
+				}
+			}
+			else
+			{
+				DATA* it = m_otherIndex.Find(at);
+				if (it)
+				{
+					return it;
+				}
+			}
+			return nullptr;
+		}
+
 		inline const DATA* get_ptr(const FCacheAddress& at) const
 		{
 			if (at.ExecutionIndex == 0 && at.ExecutionOptions == 0)
@@ -350,6 +319,18 @@ namespace mu
 			}
 		}
 
+		inline const DATA& operator[](const FCacheAddress& at) const
+		{
+			if (at.ExecutionIndex == 0 && at.ExecutionOptions == 0)
+			{
+				return m_index0[at.At];
+			}
+			else
+			{
+				return m_otherIndex[at];
+			}
+		}
+
 		struct iterator
 		{
 		private:
@@ -455,23 +436,49 @@ namespace mu
 
 		TArray< ExecutionIndex, TInlineAllocator<4> > m_usedRangeIndices;
 
+		/** Runtime data for each program op. */
+		struct FOpExecutionData
+		{
+			uint16 OpHitCount;
+
+			/** Enabled if the op descriptor has been calculated. */
+			uint8 IsDescCacheValid : 1;
+			uint8 IsValueValid : 1;
+			uint8 IsCacheLocked : 1;
+
+			/** The operation DATA_TYPE. */
+			uint8 DataType;
+
+			/** The position in the type-specific array where the result data is stored. 
+			* 0 means index not valid.
+			* For small types like bool, int and float, the actual value is the index, instead of having their own array.
+			*/
+			union
+			{
+				int32 DataTypeIndex;
+				float ScalarResult;
+			};
+		};
+
 		/** Cached resources while the program is executing. 
 		* first value of the pair:
 		* 0 : value not valid (not set).
 		* 1 : valid, not worth freeing for memory
 		* 2 : valid, worth freeing
 		*/
-		CodeContainer< TPair<int32, Ptr<const Resource>> > m_resources;
-
-		/** Addressed with OP::ADDRESS.It is true if value for an image desc is valid. */
-		TArray<bool> m_descCache;
-
-		/** The number operation stages waiting for the output of a specific operation.
-		 * \todo: this could be optimised by merging in other CodeContainers here.
-		 */
-		CodeContainer<int32> m_opHitCount;
+		CodeContainer<FOpExecutionData> OpExecutionData;
 
+		/** */
+		TArray<FVector4f> ColorResults;
+		TArray<Ptr<const Image>> ImageResults;
+		TArray<Ptr<const Layout>> LayoutResults;
+		TArray<Ptr<const Mesh>> MeshResults;
+		TArray<Ptr<const Instance>> InstanceResults;
+		TArray<FProjector> ProjectorResults;
+		TArray<Ptr<const String>> StringResults;
+		TArray<Ptr<const ExtensionData>> ExtensionDataResults;
 
+		/** */
 		inline const ExecutionIndex& GetRangeIndex(uint32_t i)
 		{
 			// Make sure we have the default element.
@@ -509,24 +516,57 @@ namespace mu
 			return uint32_t(m_usedRangeIndices.Num()) - 1;
 		}
 
-		void Init(size_t size)
+		void Init(uint32 Size)
 		{
-			// This prevents live update cache reusal
-			//m_resources.clear();
-			m_resources.resize(size);
-			m_opHitCount.resize(size);
+			// This clear would prevent live update cache reusal
+			//OpExecutionData.clear();
+
+			OpExecutionData.resize(Size);
+
+			if (ColorResults.IsEmpty())
+			{
+				// Insert dafault/null values
+				ColorResults.Add(FVector4f());
+				ImageResults.Add(nullptr);
+				LayoutResults.Add(nullptr);
+				MeshResults.Add(nullptr);
+				InstanceResults.Add(nullptr);
+				ProjectorResults.Add(FProjector());
+				StringResults.Add(nullptr);
+				ExtensionDataResults.Add(nullptr);
+			}
 		}
 
 		void SetUnused(FCacheAddress at)
 		{
-			//UE_LOG(LogMutableCore, Log, TEXT("memory SetUnused : %5d "), at.At);
-			if (m_resources[at].Key >= 2)
+			// Only clear datatypes that have results that use relevant amounts of memory			
+			FOpExecutionData& Data = OpExecutionData[at];
+			uint32 DataTypeIndex = Data.DataTypeIndex;
+			Data.IsValueValid = false;
+
+			if (DataTypeIndex)
 			{
-				// Keep the result anyway if it doesn't use any memory.
-				if (m_resources[at].Value)
+				check(OpExecutionData[at].DataType < DATATYPE::DT_COUNT);
+				switch ((DATATYPE)OpExecutionData[at].DataType)
 				{
-					m_resources[at].Value = nullptr;
-					m_resources[at].Key = 0;
+				case DATATYPE::DT_IMAGE:					
+					ImageResults[DataTypeIndex] = nullptr;
+					break;
+
+				case DATATYPE::DT_MESH:
+					MeshResults[DataTypeIndex] = nullptr;
+					break;
+
+				case DATATYPE::DT_INSTANCE:
+					InstanceResults[DataTypeIndex] = nullptr;
+					break;
+
+				case DATATYPE::DT_EXTENSION_DATA:
+					ExtensionDataResults[DataTypeIndex] = nullptr;
+					break;
+
+				default:
+					break;
 				}
 			}
 		}
@@ -534,30 +574,29 @@ namespace mu
 
 		bool IsValid(FCacheAddress at) const
 		{
-			if (at.At == 0) return false;
+			if (at.At == 0 || at.At>=OpExecutionData.size_code() )
+			{
+				return false;
+			}
+
+			const FOpExecutionData* Data = OpExecutionData.get_ptr(at);
+			if (!Data)
+			{
+				return false;
+			}
 
 			// Is it a desc data query?
 			if (at.Type == FScheduledOp::EType::ImageDesc)
 			{
-				if (uint32(m_descCache.Num()) > at.At)
-				{
-					return m_descCache[at.At];
-				}
-				else
-				{
-					return false;
-				}
+				return Data->IsDescCacheValid;
 			}
 
 			// It's a full data query.
-			auto d = m_resources.get_ptr(at);
-			return d && d->Key != 0;
+			return Data->IsValueValid != 0;
 		}
 
 
-		/** Remove all intermediate data (big and small) from the memory except for the one that has been explicitely
-		* marked as state cache.
-		*/
+		/** */
 		void CheckHitCountsCleared()
 		{
 #if !(UE_BUILD_SHIPPING || UE_BUILD_TEST)
@@ -593,261 +632,382 @@ namespace mu
 		{
 			MUTABLE_CPUPROFILER_SCOPE(ProgramCacheClear);
 
-			size_t codeSize = m_resources.size_code();
-			m_resources.clear();
-			m_resources.resize(codeSize);
-			m_descCache.SetNum(0);
-			m_opHitCount.clear();
-			m_opHitCount.resize(codeSize);
+			uint32 CodeSize = OpExecutionData.size_code();
+			OpExecutionData.clear();
+			OpExecutionData.resize(CodeSize);
 		}
 
 
-		bool GetBool(FCacheAddress at)
+		void ClearDescCache()
 		{
-			if (!at.At) return false;
-			auto d = m_resources.get_ptr(at);
-			if (!d) return false;
+			MUTABLE_CPUPROFILER_SCOPE(ProgramDescCacheClear);
 
-			Ptr<const Bool> pResult;
-			if (at.At)
+			CodeContainer<FProgramCache::FOpExecutionData>::iterator it = OpExecutionData.begin();
+			for (; it.IsValid(); ++it)
 			{
-				pResult = (const Bool*)d->Value.get();
+				(*it).IsDescCacheValid = false;
 			}
-			return pResult ? pResult->m_value : false;
+		}
+
+
+		bool GetBool(FCacheAddress at)
+		{
+			if (!at.At) return false;
+			const FOpExecutionData* Data = OpExecutionData.get_ptr(at);
+			if (!Data) return false;
+
+			check(Data->DataType==DATATYPE::DT_BOOL);
+			return Data->DataTypeIndex!=0;
 		}
 
 		float GetScalar(FCacheAddress at)
 		{
 			if (!at.At) return 0.0f;
-			auto d = m_resources.get_ptr(at);
-			if (!d) return 0.0f;
+			const FOpExecutionData* Data = OpExecutionData.get_ptr(at);
+			if (!Data) return 0.0f;
 
-			Ptr<const Scalar> pResult;
-			if (at.At)
-			{
-				pResult = (const Scalar*)d->Value.get();
-			}
-			return pResult ? pResult->m_value : 0.0f;
+			check(Data->DataType == DATATYPE::DT_SCALAR);
+			return Data->ScalarResult;
 		}
 
-		int GetInt(FCacheAddress at)
+		int32 GetInt(FCacheAddress at)
 		{
 			if (!at.At) return 0;
-			auto d = m_resources.get_ptr(at);
-			if (!d) return 0;
+			const FOpExecutionData* Data = OpExecutionData.get_ptr(at);
+			if (!Data) return 0;
 
-			Ptr<const Int> pResult;
-			if (at.At)
-			{
-				pResult = (const Int*)d->Value.get();
-			}
-			return pResult ? pResult->m_value : 0;
+			check(Data->DataType == DATATYPE::DT_INT);
+			return Data->DataTypeIndex;
 		}
 
 		FVector4f GetColour(FCacheAddress at)
 		{
 			if (!at.At) return FVector4f();
-			auto d = m_resources.get_ptr(at);
-			if (!d) return FVector4f();
+			const FOpExecutionData* Data = OpExecutionData.get_ptr(at);
+			if (!Data) return FVector4f();
 
-			Ptr<const Colour> pResult;
-			if (at.At)
-			{
-				pResult = (const Colour*)d->Value.get();
-			}
-			return pResult ? pResult->m_colour : FVector4f();
+			check(Data->DataType == DATATYPE::DT_COLOUR);
+			return ColorResults[Data->DataTypeIndex];
 		}
 
-		Ptr<const Projector> GetProjector(FCacheAddress at)
+		FProjector GetProjector(FCacheAddress at)
 		{
-			if (!at.At) return nullptr;
-			auto d = m_resources.get_ptr(at);
-			if (!d) return nullptr;
+			if (!at.At) return FProjector();
+			const FOpExecutionData* Data = OpExecutionData.get_ptr(at);
+			if (!Data) return FProjector();
 
-			Ptr<const Projector> pResult;
-			if (at.At)
-			{
-				pResult = (const Projector*)d->Value.get();
-			}
-			return pResult;
+			check(Data->DataType == DATATYPE::DT_PROJECTOR);
+			return ProjectorResults[Data->DataTypeIndex];
 		}
 
 		Ptr<const Instance> GetInstance(FCacheAddress at)
 		{
 			if (!at.At) return nullptr;
-			auto d = m_resources.get_ptr(at);
-			if (!d) return nullptr;
+			FOpExecutionData* Data = OpExecutionData.get_ptr(at);
+			if (!Data) return nullptr;
 
-			Ptr<const Instance> pResult = (const Instance*)d->Value.get();
+			check(Data->DataType == DATATYPE::DT_INSTANCE);
+			Ptr<const Instance> Result = InstanceResults[Data->DataTypeIndex];
 
-			// We need to decrease the hitcount even if the result is null.
-			// Lower hit counts means we shouldn't clear the value
-			if (m_opHitCount[at] > 0)
-			{
-				--m_opHitCount[at];
+			// We need to decrease the hit-count even if the result is null.
+			check(Data->OpHitCount > 0);
 
-				if (m_opHitCount[at] <= 0)
-				{
-					SetUnused(at);
-				}
+			--Data->OpHitCount;
+			if (Data->OpHitCount == 0 && !Data->IsCacheLocked)
+			{
+				SetUnused(at);
 			}
 
-			return pResult;
+			return Result;
 		}
 
 
-		Ptr<const Layout> GetLayout(FCacheAddress at)
+		Ptr<const Image> GetImage(FCacheAddress at, bool& bIsLastReference)
 		{
-			if (!at.At)
-			{
-				return nullptr;
-			}
-			auto d = m_resources.get_ptr(at);
-			if (!d) return nullptr;
+			bIsLastReference = false;
 
-			Ptr<const Layout> pResult;
-			if (at.At)
+			if (!at.At) return nullptr;
+			if (at.At >= OpExecutionData.size_code()) return nullptr;
+			FOpExecutionData* Data = OpExecutionData.get_ptr(at);
+			if (!Data) return nullptr;
+
+			check(Data->DataType == DATATYPE::DT_IMAGE);
+			Ptr<const Image> Result = ImageResults[Data->DataTypeIndex];
+
+			// We need to decrease the hit-count even if the result is null.
+			check(Data->OpHitCount > 0);
+
+			--Data->OpHitCount;
+			if (Data->OpHitCount == 0 && !Data->IsCacheLocked)
 			{
-				pResult = (const Layout*)d->Value.get();
+				SetUnused(at);
+				bIsLastReference = true;
 			}
-			return pResult;
+
+			return Result;
 		}
 
-		Ptr<const String> GetString(FCacheAddress at)
+
+		Ptr<const Mesh> GetMesh(FCacheAddress at, bool& bIsLastReference)
 		{
-			if (!at.At)
-			{
-				return nullptr;
-			}
-			auto d = m_resources.get_ptr(at);
-			if (!d)
-				return nullptr;
+			bIsLastReference = false;
 
-			Ptr<const String> pResult;
-			if (at.At)
+			if (!at.At) return nullptr;
+			if (at.At >= OpExecutionData.size_code()) return nullptr;
+			FOpExecutionData* Data = OpExecutionData.get_ptr(at);
+			if (!Data) return nullptr;
+
+			check(Data->DataType == DATATYPE::DT_MESH);
+			Ptr<const Mesh> Result = MeshResults[Data->DataTypeIndex];
+
+			// We need to decrease the hit-count even if the result is null.
+			check(Data->OpHitCount > 0);
+
+			--Data->OpHitCount;
+			if (Data->OpHitCount == 0 && !Data->IsCacheLocked)
 			{
-				pResult = (const String*)d->Value.get();
+				SetUnused(at);
+				bIsLastReference = true;
 			}
-			return pResult;
+
+			return Result;
 		}
 
-		Ptr<const ExtensionData> GetExtensionData(FCacheAddress at)
+
+		Ptr<const Layout> GetLayout(FCacheAddress at)
 		{
-			if (!at.At)
-			{
-				return nullptr;
-			}
-			
-			const TPair<int, Ptr<const Resource>>* d = m_resources.get_ptr(at);
-			if (!d)
-			{
-				return nullptr;
-			}
+			if (!at.At) return nullptr;
+			FOpExecutionData* Data = OpExecutionData.get_ptr(at);
+			if (!Data) return nullptr;
 
-			Ptr<const ExtensionData> pResult;
-			if (at.At)
-			{
-				pResult = (const ExtensionData*)d->Value.get();
-			}
-			return pResult;
+			check(Data->DataType == DATATYPE::DT_LAYOUT);
+			return LayoutResults[Data->DataTypeIndex];
 		}
 
-		void SetBool(FCacheAddress at, bool v)
+
+		Ptr<const String> GetString(FCacheAddress at)
 		{
-			check(at.At < m_resources.size_code());
+			if (!at.At) return nullptr;
+			FOpExecutionData* Data = OpExecutionData.get_ptr(at);
+			if (!Data) return nullptr;
 
-			Ptr<Bool> pResult = new Bool;
-			pResult->m_value = v;
-			m_resources[at] = TPair<int, Ptr<const Resource>>(1, pResult);
+			check(Data->DataType == DATATYPE::DT_STRING);
+			return StringResults[Data->DataTypeIndex];
 		}
 
-		void SetValidDesc(FCacheAddress at)
+
+		Ptr<const ExtensionData> GetExtensionData(FCacheAddress at)
 		{
-			check(at.Type == FScheduledOp::EType::ImageDesc);
-			check(at.At < uint32(m_descCache.Num()));
+			if (!at.At) return nullptr;
+			FOpExecutionData* Data = OpExecutionData.get_ptr(at);
+			if (!Data) return nullptr;
 
-			m_descCache[at.At] = true;
+			check(Data->DataType == DATATYPE::DT_EXTENSION_DATA);
+			return ExtensionDataResults[Data->DataTypeIndex];
 		}
 
-		void SetScalar(FCacheAddress at, float v)
+		void SetValidDesc(FCacheAddress at)
 		{
-			check(at.At < m_resources.size_code());
+			check(at.Type == FScheduledOp::EType::ImageDesc);
+			check(at.At < OpExecutionData.size_code());
+			FOpExecutionData* Data = OpExecutionData.get_ptr(at);
+			Data->IsDescCacheValid = true;
+		}
 
-			Ptr<Scalar> pResult = new Scalar;
-			pResult->m_value = v;
-			m_resources[at] = TPair<int, Ptr<const Resource>>(1, pResult);
+		void SetBool(FCacheAddress at, bool v)
+		{
+			check(at.At < OpExecutionData.size_code());
+			FOpExecutionData* Data = OpExecutionData.get_ptr(at);
+			check(Data->DataType == DATATYPE::DT_BOOL || Data->DataType == DATATYPE::DT_NONE);
+			Data->DataType = DATATYPE::DT_BOOL;
+			Data->DataTypeIndex = v;
+			Data->IsValueValid = true;
 		}
 
-		void SetInt(FCacheAddress at, int v)
+		void SetInt(FCacheAddress at, int32 v)
 		{
-			check(at.At < m_resources.size_code());
+			check(at.At < OpExecutionData.size_code());
+			FOpExecutionData* Data = OpExecutionData.get_ptr(at);
+			check(Data->DataType == DATATYPE::DT_INT || Data->DataType == DATATYPE::DT_NONE);
+			Data->DataType = DATATYPE::DT_INT;
+			Data->DataTypeIndex = v;
+			Data->IsValueValid = true;
+		}
 
-			Ptr<Int> pResult = new Int;
-			pResult->m_value = v;
-			m_resources[at] = TPair<int, Ptr<const Resource>>(1, pResult);
+		void SetScalar(FCacheAddress at, float v)
+		{
+			check(at.At < OpExecutionData.size_code());
+			FOpExecutionData* Data = OpExecutionData.get_ptr(at);
+			check(Data->DataType == DATATYPE::DT_SCALAR || Data->DataType == DATATYPE::DT_NONE);
+			Data->DataType = DATATYPE::DT_SCALAR;
+			Data->ScalarResult = v;
+			Data->IsValueValid = true;
 		}
 
 		void SetColour(FCacheAddress at, const FVector4f& v)
 		{
-			check(at.At < m_resources.size_code());
+			check(at.At < OpExecutionData.size_code());
+			FOpExecutionData* Data = OpExecutionData.get_ptr(at);
+			check(Data->DataType == DATATYPE::DT_COLOUR || Data->DataType == DATATYPE::DT_NONE);
+			Data->DataType = DATATYPE::DT_COLOUR;
+			Data->IsValueValid = true;
 
-			Ptr<Colour> pResult = new Colour;
-			pResult->m_colour = v;
-			m_resources[at] = TPair<int, Ptr<const Resource>>(1, pResult);
+			if (!Data->DataTypeIndex)
+			{
+				Data->DataTypeIndex = ColorResults.Num();
+				ColorResults.Add(v);
+			}
+			else
+			{
+				ColorResults[Data->DataTypeIndex] = v;
+			}
+			check(Data->DataTypeIndex != 0);
 		}
 
-		void SetProjector(FCacheAddress at, Ptr<const Projector> v)
+		void SetProjector(FCacheAddress at, const FProjector& v)
 		{
-			check(at.At < m_resources.size_code());
-			m_resources[at] = TPair<int, Ptr<const Resource>>(1, v);
+			check(at.At < OpExecutionData.size_code());
+			FOpExecutionData* Data = OpExecutionData.get_ptr(at);
+			check(Data->DataType == DATATYPE::DT_PROJECTOR || Data->DataType == DATATYPE::DT_NONE);
+			Data->DataType = DATATYPE::DT_PROJECTOR;
+			Data->IsValueValid = true;
+
+			if (!Data->DataTypeIndex)
+			{
+				Data->DataTypeIndex = ProjectorResults.Num();
+				ProjectorResults.Add(v);
+			}
+			else
+			{
+				ProjectorResults[Data->DataTypeIndex] = v;
+			}
+			check(Data->DataTypeIndex != 0);
 		}
 
 		void SetInstance(FCacheAddress at, Ptr<const Instance> v)
 		{
-			check(at.At < m_resources.size_code());
-			m_resources[at] = TPair<int, Ptr<const Resource>>(2, v);
+			check(at.At < OpExecutionData.size_code());
+			FOpExecutionData* Data = OpExecutionData.get_ptr(at);
+			check(Data->DataType == DATATYPE::DT_INSTANCE || Data->DataType == DATATYPE::DT_NONE);
+			Data->DataType = DATATYPE::DT_INSTANCE;
+			Data->IsValueValid = true;
+
+			if (!Data->DataTypeIndex)
+			{
+				Data->DataTypeIndex = InstanceResults.Num();
+				InstanceResults.Add(v);
+			}
+			else
+			{
+				InstanceResults[Data->DataTypeIndex] = v;
+			}
+			check(Data->DataTypeIndex != 0);
 		}
 
 		void SetExtensionData(FCacheAddress at, Ptr<const ExtensionData> v)
 		{
-			check(at.At < m_resources.size_code());
-			m_resources[at] = TPair<int, Ptr<const Resource>>(1, v);
+			check(at.At < OpExecutionData.size_code());
+			FOpExecutionData* Data = OpExecutionData.get_ptr(at);
+			check(Data->DataType == DATATYPE::DT_EXTENSION_DATA || Data->DataType == DATATYPE::DT_NONE);
+			Data->DataType = DATATYPE::DT_EXTENSION_DATA;
+			Data->IsValueValid = true;
+
+			if (!Data->DataTypeIndex)
+			{
+				Data->DataTypeIndex = ExtensionDataResults.Num();
+				ExtensionDataResults.Add(v);
+			}
+			else
+			{
+				ExtensionDataResults[Data->DataTypeIndex] = v;
+			}
+			check(Data->DataTypeIndex != 0);
 		}
 
 		void SetImage(FCacheAddress at, Ptr<const Image> v)
 		{
-			check(at.At < m_resources.size_code());
-			m_resources[at] = TPair<int, Ptr<const Resource>>(2, v);
+			check(at.At < OpExecutionData.size_code());
+			FOpExecutionData* Data = OpExecutionData.get_ptr(at);
+			check(Data->DataType == DATATYPE::DT_IMAGE || Data->DataType == DATATYPE::DT_NONE);
+			Data->DataType = DATATYPE::DT_IMAGE;
+			Data->IsValueValid = true;
+
+			if (!Data->DataTypeIndex)
+			{
+				Data->DataTypeIndex = ImageResults.Num();
+				ImageResults.Add(v);
+			}
+			else
+			{
+				ImageResults[Data->DataTypeIndex] = v;
+			}
+			check(Data->DataTypeIndex != 0);
 
 			mu::UpdateLLMStats();
 		}
 
 		void SetMesh(FCacheAddress at, Ptr<const Mesh> v)
 		{
-			// debug
-//            if (v)
-//            {
-//                v->GetPrivate()->CheckIntegrity();
-//            }
+			check(at.At < OpExecutionData.size_code());
+			FOpExecutionData* Data = OpExecutionData.get_ptr(at);
+			check(Data->DataType == DATATYPE::DT_MESH || Data->DataType == DATATYPE::DT_NONE);
+			Data->DataType = DATATYPE::DT_MESH;
+			Data->IsValueValid = true;
 
-			check(at.At < m_resources.size_code());
-			m_resources[at] = TPair<int, Ptr<const Resource>>(2, v);
+			if (!Data->DataTypeIndex)
+			{
+				Data->DataTypeIndex = MeshResults.Num();
+				MeshResults.Add(v);
+			}
+			else
+			{
+				MeshResults[Data->DataTypeIndex] = v;
+			}
+			check(Data->DataTypeIndex != 0);
 
 			mu::UpdateLLMStats();
 		}
 
 		void SetLayout(FCacheAddress at, Ptr<const Layout> v)
 		{
-			check(at.At < m_resources.size_code());
-			m_resources[at] = TPair<int, Ptr<const Resource>>(1, v);
+			check(at.At < OpExecutionData.size_code());
+			FOpExecutionData* Data = OpExecutionData.get_ptr(at);
+			check(Data->DataType == DATATYPE::DT_LAYOUT || Data->DataType == DATATYPE::DT_NONE);
+			Data->DataType = DATATYPE::DT_LAYOUT;
+			Data->IsValueValid = true;
+
+			if (!Data->DataTypeIndex)
+			{
+				Data->DataTypeIndex = LayoutResults.Num();
+				LayoutResults.Add(v);
+			}
+			else
+			{
+				LayoutResults[Data->DataTypeIndex] = v;
+			}
+			check(Data->DataTypeIndex != 0);
 
 			mu::UpdateLLMStats();
 		}
 
 		void SetString(FCacheAddress at, Ptr<const String> v)
 		{
-			check(at.At < m_resources.size_code());
-			m_resources[at] = TPair<int, Ptr<const Resource>>(1, v);
+			check(at.At < OpExecutionData.size_code());
+			FOpExecutionData* Data = OpExecutionData.get_ptr(at);
+			check(Data->DataType == DATATYPE::DT_STRING || Data->DataType == DATATYPE::DT_NONE);
+			Data->DataType = DATATYPE::DT_STRING;
+			Data->IsValueValid = true;
+
+			if (!Data->DataTypeIndex)
+			{
+				Data->DataTypeIndex = StringResults.Num();
+				StringResults.Add(v);
+			}
+			else
+			{
+				StringResults[Data->DataTypeIndex] = v;
+			}
+			check(Data->DataTypeIndex!=0);
 		}
 
 
@@ -858,65 +1018,21 @@ namespace mu
 			// at its position.
 			if (at.At)
 			{
-				m_opHitCount[at] += 1;
+				check(at.At < OpExecutionData.size_code());
+				FOpExecutionData& Data = OpExecutionData[at];
+				Data.OpHitCount++;
 			}
 		}
 
 		inline void SetForceCached(OP::ADDRESS at)
 		{
-			// \TODO: Review the ,0,0
-			m_opHitCount[{at, 0, 0}] = 0xffffff;
-		}
-
-
-		Ptr<const Image> GetImage(FCacheAddress at, bool& bIsLastReference)
-		{
-			bIsLastReference = false;
-
-			if (!at.At) return nullptr;
-			if (size_t(at.At) >= m_resources.size_code()) return nullptr;
-			auto d = m_resources.get_ptr(at);
-			if (!d) return nullptr;
-
-			Ptr<const Image> Result = (const Image*)d->Value.get();
-
-			// We need to decrease the hit count (even if the result is null).
-			// Lower hit counts means we shouldn't clear the value
-			if ( m_opHitCount[at] > 0)
+			// \TODO: It only locks at,0,0
+			if (at)
 			{
-				--m_opHitCount[at];
-				if (m_opHitCount[at] <= 0)
-				{
-					SetUnused(at);
-					bIsLastReference = true;
-				}
-			}
-
-			return Result;
-		}
-
-		Ptr<const Mesh> GetMesh(FCacheAddress at)
-		{
-			if (!at.At) return nullptr;
-			if (size_t(at.At) >= m_resources.size_code()) return nullptr;
-			auto d = m_resources.get_ptr(at);
-			if (!d) return nullptr;
-
-			Ptr<const Mesh> pResult = (const Mesh*)d->Value.get();
-
-			// We need to decrease the hitcount even if the result is null.
-			// Lower hit counts means we shouldn't clear the value
-			if ( m_opHitCount[at] > 0)
-			{
-				--m_opHitCount[at];
-
-				if (m_opHitCount[at] <= 0)
-				{
-					SetUnused(at);
-				}
+				check(at < OpExecutionData.size_code());
+				FOpExecutionData& Data = OpExecutionData[FCacheAddress(at, 0, 0)];
+				Data.IsCacheLocked = true;
 			}
-
-			return pResult;
 		}
 	};
 
@@ -1169,7 +1285,39 @@ namespace mu
 		}
 
 		/** */
-		Ptr<const Image> LoadImage(const FCacheAddress& From, bool bTakeOwnership=false)
+		Ptr<const Mesh> LoadMesh(const FCacheAddress& From, bool bTakeOwnership=false)
+		{
+			bool bIsLastReference = false;
+			Ptr<const Mesh> Result = CurrentInstanceCache->GetMesh(From, bIsLastReference);
+			if (!Result)
+			{
+				return nullptr;
+			}
+
+			// TODO
+			// If we retrieved the last reference to this resource in "From" cache position (it could still be in other cache positions as well)
+			//if (bIsLastReference)
+			//{
+			//	int32* CountPtr = CacheResources.Find(Result);
+			//	check(CountPtr);
+			//	*CountPtr = (*CountPtr) - 1;
+			//	if (!*CountPtr)
+			//	{
+			//		CacheResources.FindAndRemoveChecked(Result);
+			//	}
+			//}
+
+
+			//if (!bTakeOwnership && Result->IsUnique())
+			//{
+			//	TempMesh.Add(Result);
+			//}
+
+			return Result;
+		}
+
+		/** */
+		Ptr<const Image> LoadImage(const FCacheAddress& From, bool bTakeOwnership = false)
 		{
 			bool bIsLastReference = false;
 			Ptr<const Image> Result = CurrentInstanceCache->GetImage(From, bIsLastReference);
@@ -1298,65 +1446,59 @@ namespace mu
 		}
 
 
-		int32 GetCache0Bytes() const
+		/** Calculate the amount of bytes in data cached in the level 0 and 1 cache in all live instances. */
+		int32 GetCacheBytes() const
 		{
-			if (!CurrentInstanceCache)
-			{
-				return 0;
-			}
-
 			int32 Result = 0;
 			TSet<const Resource*> Cache0Unique;
-
-			CodeContainer<int>::iterator it = CurrentInstanceCache->m_opHitCount.begin();
-			for (; it.IsValid(); ++it)
-			{
-				const Resource* Value = CurrentInstanceCache->m_resources[it.get_address()].Value.get();
-				if (!Value)
-				{
-					continue;
-				}
-
-				int32 Count = *it;
-				if (Count < UE_MUTABLE_CACHE_COUNT_LIMIT)
-				{
-					Cache0Unique.Add(Value);
-				}
-			}
-
-			for (const Resource* Value : Cache0Unique)
-			{
-				Result += Value->GetDataSize();
-			}
-
-			return Result;
-		}
-
-		int32 GetCache1Bytes() const
-		{
-			int32 Result = 0;
 			TSet<const Resource*> Cache1Unique;
 
 			for (const FLiveInstance& Instance : LiveInstances)
 			{
-				CodeContainer<int>::iterator it = Instance.Cache->m_opHitCount.begin();
+				CodeContainer<FProgramCache::FOpExecutionData>::iterator it = Instance.Cache->OpExecutionData.begin();
 				for (; it.IsValid(); ++it)
 				{
-					const Resource* Value = Instance.Cache->m_resources[it.get_address()].Value.get();
-					if (!Value)
+					if (!(*it).DataTypeIndex)
 					{
 						continue;
 					}
 
-					int32 Count = *it;
-					if (Count >= UE_MUTABLE_CACHE_COUNT_LIMIT)
+					TSet<const Resource*>* TargetSet = &Cache0Unique;
+					if ((*it).IsCacheLocked)
+					{
+						TargetSet = &Cache1Unique;
+					}
+
+					const Resource* Value = nullptr;
+					switch ((*it).DataType)
 					{
-						Cache1Unique.Add(Value);
+					case DATATYPE::DT_IMAGE:
+						Value = CurrentInstanceCache->ImageResults[(*it).DataTypeIndex].get();
+						if (Value)
+						{
+							TargetSet->Add(Value);
+						}
+						break;
+
+					case DATATYPE::DT_MESH:
+						Value = CurrentInstanceCache->MeshResults[(*it).DataTypeIndex].get();
+						if (Value)
+						{
+							TargetSet->Add(Value);
+						}
+						break;
+
+					default:
+						break;
 					}
 				}
 			}
 
-			for (const Resource* Value : Cache1Unique)
+			// Merge
+			Cache0Unique.Append(Cache1Unique);
+
+			// Count
+			for (const Resource* Value : Cache0Unique)
 			{
 				Result += Value->GetDataSize();
 			}
@@ -1374,24 +1516,49 @@ namespace mu
 
 			MUTABLE_CPUPROFILER_SCOPE(ClearLayer0);
 
-			CodeContainer<int>::iterator it = CurrentInstanceCache->m_opHitCount.begin();
+			CodeContainer<FProgramCache::FOpExecutionData>::iterator it = CurrentInstanceCache->OpExecutionData.begin();
 			for (; it.IsValid(); ++it)
 			{
-				int32 Count = *it;
-				if (Count < UE_MUTABLE_CACHE_COUNT_LIMIT)
+				FProgramCache::FOpExecutionData& Data = *it;
+				if (!Data.DataTypeIndex
+					||
+					Data.IsCacheLocked)
 				{
-					TPair<int32, Ptr<const Resource>>& PairRef = CurrentInstanceCache->m_resources[it.get_address()];
+					continue;
+				}
 
-					//CacheResources.FindAndRemoveChecked(PairRef.Value);
-					CacheResources.Remove(PairRef.Value);
+				const Resource* Value = nullptr;
 
-					// SetUnused only clears if the data size is relevant (mesh or image) but we need to 
-					// clear it in all cases here, because it may have become invalid because of parameter changes.
-					// SetUnused(it.get_address());
-					PairRef.Value = nullptr;
-					PairRef.Key = 0;
-					*it = 0;
+				switch (Data.DataType)
+				{
+				case DATATYPE::DT_IMAGE:
+					Value = CurrentInstanceCache->ImageResults[Data.DataTypeIndex].get();
+					CacheResources.Remove(Value);
+					CurrentInstanceCache->ImageResults[Data.DataTypeIndex] = nullptr;
+					break;
+
+				case DATATYPE::DT_MESH:
+					Value = CurrentInstanceCache->MeshResults[Data.DataTypeIndex].get();
+					CacheResources.Remove(Value);
+					CurrentInstanceCache->MeshResults[Data.DataTypeIndex] = nullptr;
+					break;
+
+				case DATATYPE::DT_LAYOUT:
+					Value = CurrentInstanceCache->LayoutResults[Data.DataTypeIndex].get();
+					CurrentInstanceCache->LayoutResults[Data.DataTypeIndex] = nullptr;
+					break;
+
+				case DATATYPE::DT_INSTANCE:
+					Value = CurrentInstanceCache->InstanceResults[Data.DataTypeIndex].get();
+					CurrentInstanceCache->InstanceResults[Data.DataTypeIndex] = nullptr;
+					break;
+
+				default:
+					break;
 				}
+
+				Data.OpHitCount = 0;
+				Data.IsValueValid = false;
 			}
 		}
 
@@ -1403,17 +1570,48 @@ namespace mu
 		{
 			MUTABLE_CPUPROFILER_SCOPE(ClearLayer1);
 
-			CodeContainer< TPair<int32, Ptr<const Resource>> >::iterator it = CurrentInstanceCache->m_resources.begin();
+			CodeContainer<FProgramCache::FOpExecutionData>::iterator it = CurrentInstanceCache->OpExecutionData.begin();
 			for (; it.IsValid(); ++it)
 			{
-				//CacheResources.FindAndRemoveChecked((*it).Value);
-				CacheResources.Remove((*it).Value);
+				FProgramCache::FOpExecutionData& Data = *it;
+				if (!Data.DataTypeIndex)
+				{
+					continue;
+				}
 
-				(*it).Value = nullptr;
-				(*it).Key = 0;
-			}
+				const Resource* Value = nullptr;
 
-			CurrentInstanceCache->m_descCache.SetNum(0);
+				switch (Data.DataType)
+				{
+				case DATATYPE::DT_IMAGE:
+					Value = CurrentInstanceCache->ImageResults[Data.DataTypeIndex].get();
+					CacheResources.Remove(Value);
+					CurrentInstanceCache->ImageResults[Data.DataTypeIndex] = nullptr;
+					break;
+
+				case DATATYPE::DT_MESH:
+					Value = CurrentInstanceCache->MeshResults[Data.DataTypeIndex].get();
+					CacheResources.Remove(Value);
+					CurrentInstanceCache->MeshResults[Data.DataTypeIndex] = nullptr;
+					break;
+
+				case DATATYPE::DT_LAYOUT:
+					Value = CurrentInstanceCache->LayoutResults[Data.DataTypeIndex].get();
+					CurrentInstanceCache->LayoutResults[Data.DataTypeIndex] = nullptr;
+					break;
+
+				case DATATYPE::DT_INSTANCE:
+					Value = CurrentInstanceCache->InstanceResults[Data.DataTypeIndex].get();
+					CurrentInstanceCache->InstanceResults[Data.DataTypeIndex] = nullptr;
+					break;
+
+				default:
+					break;
+				}
+
+				Data.OpHitCount = 0;
+				Data.IsValueValid = false;
+			}
 		}
 
 
@@ -1496,7 +1694,7 @@ namespace mu
 		MUTABLERUNTIME_API Ptr<const Image> BuildImage(const TSharedPtr<const Model>&, const Parameters*, OP::ADDRESS, int32 MipsToSkip, int32 LOD) ;
 		MUTABLERUNTIME_API Ptr<const Mesh> BuildMesh(const TSharedPtr<const Model>&, const Parameters*, OP::ADDRESS) ;
 		MUTABLERUNTIME_API Ptr<const Layout> BuildLayout(const TSharedPtr<const Model>&, const Parameters*, OP::ADDRESS) ;
-    	MUTABLERUNTIME_API Ptr<const Projector> BuildProjector(const TSharedPtr<const Model>&, const Parameters*, OP::ADDRESS) ;
+    	MUTABLERUNTIME_API FProjector BuildProjector(const TSharedPtr<const Model>&, const Parameters*, OP::ADDRESS) ;
 
 		ExtensionDataStreamer* GetExtensionDataStreamer() const { return ExtensionDataStreamer.Get(); }
 
