[Mutable] Add IsParameterMultidimensional method to the mu::Model.

#rnx

diff --git a/Engine/Plugins/Experimental/Mutable/Source/CustomizableObject/Private/MuCO/CustomizableObject.cpp b/Engine/Plugins/Experimental/Mutable/Source/CustomizableObject/Private/MuCO/CustomizableObject.cpp
--- a/Engine/Plugins/Experimental/Mutable/Source/CustomizableObject/Private/MuCO/CustomizableObject.cpp
+++ b/Engine/Plugins/Experimental/Mutable/Source/CustomizableObject/Private/MuCO/CustomizableObject.cpp
@@ -1686,12 +1686,12 @@ bool UCustomizableObject::IsParameterMultidimensional(const FString& InParameter
 bool UCustomizableObject::IsParameterMultidimensional(const int32& InParamIndex) const
 {
 	check(InParamIndex != INDEX_NONE);
-	
-	const mu::ParametersPtr MutableParameters = mu::Model::NewParameters(this->GetPrivate()->GetModel());
-	check(InParamIndex < MutableParameters->GetCount());
-	const mu::RangeIndexPtr RangeIdxPtr = MutableParameters->NewRangeIndex(InParamIndex);
+	if (PrivateData->GetModel())
+	{
+		return PrivateData->GetModel()->IsParameterMultidimensional(InParamIndex);
+	}
 
-	return RangeIdxPtr.get() != nullptr;
+	return false;
 }
 
 
diff --git a/Engine/Plugins/Experimental/Mutable/Source/MutableRuntime/Private/MuR/Model.cpp b/Engine/Plugins/Experimental/Mutable/Source/MutableRuntime/Private/MuR/Model.cpp
--- a/Engine/Plugins/Experimental/Mutable/Source/MutableRuntime/Private/MuR/Model.cpp
+++ b/Engine/Plugins/Experimental/Mutable/Source/MutableRuntime/Private/MuR/Model.cpp
@@ -545,6 +545,18 @@ namespace mu
     }
 
 
+    //---------------------------------------------------------------------------------------------
+    bool Model::IsParameterMultidimensional(const int32 ParamIndex) const
+    {
+		if (m_pD->m_program.m_parameters.IsValidIndex(ParamIndex))
+		{
+			return m_pD->m_program.m_parameters[ParamIndex].m_ranges.Num() > 0;
+		}
+
+		return false;
+    }
+	
+
     //---------------------------------------------------------------------------------------------
     int Model::GetStateCount() const
     {
diff --git a/Engine/Plugins/Experimental/Mutable/Source/MutableRuntime/Public/MuR/Model.h b/Engine/Plugins/Experimental/Mutable/Source/MutableRuntime/Public/MuR/Model.h
--- a/Engine/Plugins/Experimental/Mutable/Source/MutableRuntime/Public/MuR/Model.h
+++ b/Engine/Plugins/Experimental/Mutable/Source/MutableRuntime/Public/MuR/Model.h
@@ -81,6 +81,9 @@ namespace mu
 		//! default values.
         static ParametersPtr NewParameters( TSharedPtr<const Model> Model, const Parameters* pOld = nullptr );
 
+		/** Return true if the parameter is multi-dimensional */
+		bool IsParameterMultidimensional(int32 ParameterIndex) const;
+
 		//! Get the number of states in the model.
 		int GetStateCount() const;
 
