[Mutable] Minor code improvements.

- Removed empty constructors.
- Added missing virtual.
- Moved implementations from header to source files.

#rnx

diff --git a/Engine/Plugins/Experimental/Mutable/Source/CustomizableObjectEditor/Classes/MuCOE/Nodes/CustomizableObjectNode.h b/Engine/Plugins/Experimental/Mutable/Source/CustomizableObjectEditor/Classes/MuCOE/Nodes/CustomizableObjectNode.h
--- a/Engine/Plugins/Experimental/Mutable/Source/CustomizableObjectEditor/Classes/MuCOE/Nodes/CustomizableObjectNode.h
+++ b/Engine/Plugins/Experimental/Mutable/Source/CustomizableObjectEditor/Classes/MuCOE/Nodes/CustomizableObjectNode.h
@@ -120,8 +120,8 @@ public:
 	 * 
 	 * Some nodes are added manually in the graph editor code and don't need to do this.
 	 */
-	virtual bool ShouldAddToContextMenu(FText& OutCategory) const { return false; }
-	
+	virtual bool ShouldAddToContextMenu(FText& OutCategory) const;
+
 	void GetInputPins(TArray<class UEdGraphPin*>& OutInputPins) const;
 	void GetOutputPins(TArray<class UEdGraphPin*>& OutOutputPins) const;
 	UEdGraphPin* GetOutputPin(int32 OutputIndex) const;
diff --git a/Engine/Plugins/Experimental/Mutable/Source/CustomizableObjectEditor/Classes/MuCOE/Nodes/CustomizableObjectNodeMeshClipDeform.h b/Engine/Plugins/Experimental/Mutable/Source/CustomizableObjectEditor/Classes/MuCOE/Nodes/CustomizableObjectNodeMeshClipDeform.h
--- a/Engine/Plugins/Experimental/Mutable/Source/CustomizableObjectEditor/Classes/MuCOE/Nodes/CustomizableObjectNodeMeshClipDeform.h
+++ b/Engine/Plugins/Experimental/Mutable/Source/CustomizableObjectEditor/Classes/MuCOE/Nodes/CustomizableObjectNodeMeshClipDeform.h
@@ -24,30 +24,22 @@ class CUSTOMIZABLEOBJECTEDITOR_API UCustomizableObjectNodeMeshClipDeform : publi
 public:
 	GENERATED_BODY()
 
-	UCustomizableObjectNodeMeshClipDeform();
-
 	UPROPERTY(EditAnywhere, BlueprintReadWrite, Category = MeshClipDeform)
 	TArray<FString> Tags;
 	
 	UPROPERTY(EditAnywhere, Category = MeshClipDeform)
 	EShapeBindingMethod BindingMethod;
 	
-	// Begin EdGraphNode interface
-	FText GetNodeTitle(ENodeTitleType::Type TitleType) const override;
-	FLinearColor GetNodeTitleColor() const override;
-	FText GetTooltipText() const override;
+	// EdGraphNode interface
+	virtual FText GetNodeTitle(ENodeTitleType::Type TitleType) const override;
+	virtual FLinearColor GetNodeTitleColor() const override;
+	virtual FText GetTooltipText() const override;
 
 	// UCustomizableObjectNode interface
-	void AllocateDefaultPins(UCustomizableObjectNodeRemapPins* RemapPins) override;
-	
-	inline UEdGraphPin* ClipShapePin() const
-	{
-		return FindPin(TEXT("Clip Shape"), EGPD_Input);
-	}
-
-	UEdGraphPin* OutputPin() const override
-	{
-		return FindPin(TEXT("Material"));
-	}
+	virtual void AllocateDefaultPins(UCustomizableObjectNodeRemapPins* RemapPins) override;
+	virtual UEdGraphPin* OutputPin() const override;
+
+	// Own interface
+	UEdGraphPin* ClipShapePin() const;
 };
 
diff --git a/Engine/Plugins/Experimental/Mutable/Source/CustomizableObjectEditor/Classes/MuCOE/Nodes/CustomizableObjectNodeMeshClipMorph.h b/Engine/Plugins/Experimental/Mutable/Source/CustomizableObjectEditor/Classes/MuCOE/Nodes/CustomizableObjectNodeMeshClipMorph.h
--- a/Engine/Plugins/Experimental/Mutable/Source/CustomizableObjectEditor/Classes/MuCOE/Nodes/CustomizableObjectNodeMeshClipMorph.h
+++ b/Engine/Plugins/Experimental/Mutable/Source/CustomizableObjectEditor/Classes/MuCOE/Nodes/CustomizableObjectNodeMeshClipMorph.h
@@ -20,8 +20,6 @@ class CUSTOMIZABLEOBJECTEDITOR_API UCustomizableObjectNodeMeshClipMorph : public
 public:
 	GENERATED_BODY()
 
-	UCustomizableObjectNodeMeshClipMorph();
-
 	UPROPERTY(EditAnywhere, Category = MeshToClipAndMorph)
 	FName BoneName;
 
@@ -31,19 +29,6 @@ public:
 	UPROPERTY(EditAnywhere, Category = MeshToClipAndMorph)
 	uint32 ReferenceSkeletonIndex = 0;
 
-	// UObject interface.
-	void PostEditChangeProperty(FPropertyChangedEvent& PropertyChangedEvent) override;
-	virtual void Serialize(FArchive& Ar) override;
-
-	// Begin EdGraphNode interface
-	FText GetNodeTitle(ENodeTitleType::Type TitleType) const override;
-	FLinearColor GetNodeTitleColor() const override;
-	FText GetTooltipText() const override;
-
-	// UCustomizableObjectNode interface
-	void AllocateDefaultPins(UCustomizableObjectNodeRemapPins* RemapPins) override;
-	void PinConnectionListChanged(UEdGraphPin* Pin) override;
-
 	UPROPERTY(EditAnywhere, Category = MeshClipParameters, meta = (DisplayName="Morph Start Offset", ToolTip="Offset from the origin of the selected bone to the actual start of the morph."))
 	FVector StartOffset;
 
@@ -79,16 +64,29 @@ public:
 
 	bool bUpdateViewportWidget;
 
-	FVector GetOriginWithOffset() const;
+	UCustomizableObjectNodeMeshClipMorph();
 
+	// UObject interface.
+	virtual void PostEditChangeProperty(FPropertyChangedEvent& PropertyChangedEvent) override;
+	virtual void Serialize(FArchive& Ar) override;
+
+	// EdGraphNode interface
+	virtual FText GetNodeTitle(ENodeTitleType::Type TitleType) const override;
+	virtual FLinearColor GetNodeTitleColor() const override;
+	virtual FText GetTooltipText() const override;
+
+	// UCustomizableObjectNode interface
+	virtual void AllocateDefaultPins(UCustomizableObjectNodeRemapPins* RemapPins) override;
+	virtual void PinConnectionListChanged(UEdGraphPin* Pin) override;
+	
+	// UCustomizableObjectNodeMaterialBase interface
+	virtual UEdGraphPin* OutputPin() const override;
+
+	// Own interface
+	FVector GetOriginWithOffset() const;
 	void FindLocalAxes(FVector& Right, FVector& Up, FVector& Forward) const;
 
-	// Change StartOffset from World to Local of the other way arround
+	// Change StartOffset from World to Local of the other way around
 	void ChangeStartOffsetTransform();
-
-	virtual UEdGraphPin* OutputPin() const override
-	{
-		return FindPin(TEXT("Material"));
-	}
 };
 
diff --git a/Engine/Plugins/Experimental/Mutable/Source/CustomizableObjectEditor/Classes/MuCOE/Nodes/CustomizableObjectNodeMeshClipWithMesh.h b/Engine/Plugins/Experimental/Mutable/Source/CustomizableObjectEditor/Classes/MuCOE/Nodes/CustomizableObjectNodeMeshClipWithMesh.h
--- a/Engine/Plugins/Experimental/Mutable/Source/CustomizableObjectEditor/Classes/MuCOE/Nodes/CustomizableObjectNodeMeshClipWithMesh.h
+++ b/Engine/Plugins/Experimental/Mutable/Source/CustomizableObjectEditor/Classes/MuCOE/Nodes/CustomizableObjectNodeMeshClipWithMesh.h
@@ -38,35 +38,28 @@ public:
 	UPROPERTY(EditAnywhere, Category = CustomizableObject)
 	FTransform Transform;
 
-	// Begin EdGraphNode interface
-	FText GetNodeTitle(ENodeTitleType::Type TitleType) const override;
-	FLinearColor GetNodeTitleColor() const override;
-	FText GetTooltipText() const override;
-
-	// UCustomizableObjectNode interface
-	void AllocateDefaultPins(UCustomizableObjectNodeRemapPins* RemapPins) override;
-	void PinConnectionListChanged(UEdGraphPin* Pin) override;
-	
-	virtual UEdGraphPin* OutputPin() const override
-	{
-		return FindPin(TEXT("Material"));
-	}
-
-	UEdGraphPin* ClipMeshPin() const
-	{
-		return FindPin(TEXT("Clip Mesh"));
-	}
-
-	void BeginPostDuplicate(bool bDuplicateForPIE) override;
-
-	// UCustomizableObjectNode interface
-	void UpdateReferencedNodeId(const FGuid& NewGuid) override;
-
 	// Details view variables
 	// The clipping node uses tags
 	bool bUseTags;
 
 	// The clipping node uses a material name
 	bool bUseMaterials;
+
+	// EdGraphNode interface
+	virtual FText GetNodeTitle(ENodeTitleType::Type TitleType) const override;
+	virtual FLinearColor GetNodeTitleColor() const override;
+	virtual FText GetTooltipText() const override;
+
+	// UCustomizableObjectNode interface
+	virtual void AllocateDefaultPins(UCustomizableObjectNodeRemapPins* RemapPins) override;
+	virtual void PinConnectionListChanged(UEdGraphPin* Pin) override;
+	virtual void UpdateReferencedNodeId(const FGuid& NewGuid) override;
+	virtual void BeginPostDuplicate(bool bDuplicateForPIE) override;
+
+	// UCustomizableObjectNodeMaterialBase interface
+	virtual UEdGraphPin* OutputPin() const override;
+
+	// Own interface
+	UEdGraphPin* ClipMeshPin() const;
 };
 
diff --git a/Engine/Plugins/Experimental/Mutable/Source/CustomizableObjectEditor/Private/MuCOE/Nodes/CustomizableObjectNode.cpp b/Engine/Plugins/Experimental/Mutable/Source/CustomizableObjectEditor/Private/MuCOE/Nodes/CustomizableObjectNode.cpp
--- a/Engine/Plugins/Experimental/Mutable/Source/CustomizableObjectEditor/Private/MuCOE/Nodes/CustomizableObjectNode.cpp
+++ b/Engine/Plugins/Experimental/Mutable/Source/CustomizableObjectEditor/Private/MuCOE/Nodes/CustomizableObjectNode.cpp
@@ -111,6 +111,12 @@ bool UCustomizableObjectNode::CustomRemovePin(UEdGraphPin& Pin)
 }
 
 
+bool UCustomizableObjectNode::ShouldAddToContextMenu(FText& OutCategory) const
+{
+	return false;
+}
+
+
 void UCustomizableObjectNode::GetInputPins(TArray<class UEdGraphPin*>& OutInputPins) const
 {
 	OutInputPins.Empty();
diff --git a/Engine/Plugins/Experimental/Mutable/Source/CustomizableObjectEditor/Private/MuCOE/Nodes/CustomizableObjectNodeMeshClipDeform.cpp b/Engine/Plugins/Experimental/Mutable/Source/CustomizableObjectEditor/Private/MuCOE/Nodes/CustomizableObjectNodeMeshClipDeform.cpp
--- a/Engine/Plugins/Experimental/Mutable/Source/CustomizableObjectEditor/Private/MuCOE/Nodes/CustomizableObjectNodeMeshClipDeform.cpp
+++ b/Engine/Plugins/Experimental/Mutable/Source/CustomizableObjectEditor/Private/MuCOE/Nodes/CustomizableObjectNodeMeshClipDeform.cpp
@@ -9,11 +9,6 @@ class UCustomizableObjectNodeRemapPins;
 #define LOCTEXT_NAMESPACE "CustomizableObjectEditor"
 
 
-UCustomizableObjectNodeMeshClipDeform::UCustomizableObjectNodeMeshClipDeform()
-	: Super()
-{
-}
-
 void UCustomizableObjectNodeMeshClipDeform::AllocateDefaultPins(UCustomizableObjectNodeRemapPins* RemapPins)
 {
 	const UEdGraphSchema_CustomizableObject* Schema = GetDefault<UEdGraphSchema_CustomizableObject>();
@@ -25,6 +20,19 @@ void UCustomizableObjectNodeMeshClipDeform::AllocateDefaultPins(UCustomizableObj
 	ClipMeshPin->bDefaultValueIsIgnored = true;
 }
 
+
+UEdGraphPin* UCustomizableObjectNodeMeshClipDeform::ClipShapePin() const
+{
+	return FindPin(TEXT("Clip Shape"), EGPD_Input);
+}
+
+
+UEdGraphPin* UCustomizableObjectNodeMeshClipDeform::OutputPin() const
+{
+	return FindPin(TEXT("Material"));
+}
+
+
 FText UCustomizableObjectNodeMeshClipDeform::GetNodeTitle(ENodeTitleType::Type TitleType) const
 {
 	return LOCTEXT("Clip_Deform_Mesh", "Clip Deform Mesh");
diff --git a/Engine/Plugins/Experimental/Mutable/Source/CustomizableObjectEditor/Private/MuCOE/Nodes/CustomizableObjectNodeMeshClipMorph.cpp b/Engine/Plugins/Experimental/Mutable/Source/CustomizableObjectEditor/Private/MuCOE/Nodes/CustomizableObjectNodeMeshClipMorph.cpp
--- a/Engine/Plugins/Experimental/Mutable/Source/CustomizableObjectEditor/Private/MuCOE/Nodes/CustomizableObjectNodeMeshClipMorph.cpp
+++ b/Engine/Plugins/Experimental/Mutable/Source/CustomizableObjectEditor/Private/MuCOE/Nodes/CustomizableObjectNodeMeshClipMorph.cpp
@@ -87,6 +87,12 @@ void UCustomizableObjectNodeMeshClipMorph::ChangeStartOffsetTransform()
 }
 
 
+UEdGraphPin* UCustomizableObjectNodeMeshClipMorph::OutputPin() const
+{
+	return FindPin(TEXT("Material"));
+}
+
+
 void UCustomizableObjectNodeMeshClipMorph::PostEditChangeProperty(FPropertyChangedEvent & PropertyChangedEvent)
 {
 	FProperty* PropertyThatChanged = PropertyChangedEvent.Property;
diff --git a/Engine/Plugins/Experimental/Mutable/Source/CustomizableObjectEditor/Private/MuCOE/Nodes/CustomizableObjectNodeMeshClipWithMesh.cpp b/Engine/Plugins/Experimental/Mutable/Source/CustomizableObjectEditor/Private/MuCOE/Nodes/CustomizableObjectNodeMeshClipWithMesh.cpp
--- a/Engine/Plugins/Experimental/Mutable/Source/CustomizableObjectEditor/Private/MuCOE/Nodes/CustomizableObjectNodeMeshClipWithMesh.cpp
+++ b/Engine/Plugins/Experimental/Mutable/Source/CustomizableObjectEditor/Private/MuCOE/Nodes/CustomizableObjectNodeMeshClipWithMesh.cpp
@@ -11,9 +11,8 @@ class UCustomizableObjectNodeRemapPins;
 #define LOCTEXT_NAMESPACE "CustomizableObjectEditor"
 
 
-UCustomizableObjectNodeMeshClipWithMesh::UCustomizableObjectNodeMeshClipWithMesh()
-	: Super()
-	, CustomizableObjectToClipWith(nullptr)
+UCustomizableObjectNodeMeshClipWithMesh::UCustomizableObjectNodeMeshClipWithMesh() :
+	CustomizableObjectToClipWith(nullptr)
 {
 	// We allow this node to be shared
 	Transform = FTransform::Identity;
@@ -82,6 +81,19 @@ void UCustomizableObjectNodeMeshClipWithMesh::PinConnectionListChanged(UEdGraphP
 	}
 }
 
+
+UEdGraphPin* UCustomizableObjectNodeMeshClipWithMesh::OutputPin() const
+{
+	return FindPin(TEXT("Material"));
+}
+
+
+UEdGraphPin* UCustomizableObjectNodeMeshClipWithMesh::ClipMeshPin() const
+{
+	return FindPin(TEXT("Clip Mesh"));
+}
+
+
 void UCustomizableObjectNodeMeshClipWithMesh::UpdateReferencedNodeId(const FGuid& NewGuid)
 {
 	ArrayMaterialNodeToClipWithID.Add(NewGuid);
