[Mutable] Fix crash due to invalid access to the cache in the LogWorkingMemory(...) method.

#rnx

diff --git a/Engine/Plugins/Experimental/Mutable/Source/MutableRuntime/Private/MuR/SystemPrivate.h b/Engine/Plugins/Experimental/Mutable/Source/MutableRuntime/Private/MuR/SystemPrivate.h
--- a/Engine/Plugins/Experimental/Mutable/Source/MutableRuntime/Private/MuR/SystemPrivate.h
+++ b/Engine/Plugins/Experimental/Mutable/Source/MutableRuntime/Private/MuR/SystemPrivate.h
@@ -1473,7 +1473,7 @@ namespace mu
 					switch ((*it).DataType)
 					{
 					case DATATYPE::DT_IMAGE:
-						Value = CurrentInstanceCache->ImageResults[(*it).DataTypeIndex].get();
+						Value = Instance.Cache->ImageResults[(*it).DataTypeIndex].get();
 						if (Value)
 						{
 							TargetSet->Add(Value);
@@ -1481,7 +1481,7 @@ namespace mu
 						break;
 
 					case DATATYPE::DT_MESH:
-						Value = CurrentInstanceCache->MeshResults[(*it).DataTypeIndex].get();
+						Value = Instance.Cache->MeshResults[(*it).DataTypeIndex].get();
 						if (Value)
 						{
 							TargetSet->Add(Value);
