[Mutable] Ignore non-normalized UVs in cases where the UVs are irrelevant.

#rnx

diff --git a/Engine/Plugins/Experimental/Mutable/Source/MutableTools/Private/MuT/CodeGenerator.cpp b/Engine/Plugins/Experimental/Mutable/Source/MutableTools/Private/MuT/CodeGenerator.cpp
--- a/Engine/Plugins/Experimental/Mutable/Source/MutableTools/Private/MuT/CodeGenerator.cpp
+++ b/Engine/Plugins/Experimental/Mutable/Source/MutableTools/Private/MuT/CodeGenerator.cpp
@@ -205,6 +205,7 @@ namespace mu
 			FMeshGenerationOptions MeshOptions;
 			MeshOptions.bUniqueVertexIDs = true;
 			MeshOptions.bLayouts = true;
+			MeshOptions.bNormalizeUVs = true;
 			MeshOptions.State = m_currentStateIndex;
 
 			FMeshGenerationResult MeshResult;
@@ -850,6 +851,10 @@ namespace mu
 			const bool bClampUVIslands = node.SharedSurfaceId != INDEX_NONE && m_currentParents.Last().m_lod > 0;
 			MeshOptions.bClampUVIslands = bClampUVIslands;
 
+			// Normalize UVs if we're going to work with images and layouts.
+			const bool bNormalizeUVs = node.m_images.Num() && edits.Num();
+			MeshOptions.bNormalizeUVs = bNormalizeUVs;
+
             GenerateMesh(MeshOptions, meshResults, pMesh );
             lastMeshOp = meshResults.meshOp;
             meshResults.extraMeshLayouts.SetNum( edits.Num() );
@@ -877,6 +882,7 @@ namespace mu
 						MergedMeshOptions.bUniqueVertexIDs = true;
 						MergedMeshOptions.bLayouts = true;
 						MergedMeshOptions.bClampUVIslands = bClampUVIslands;
+						MergedMeshOptions.bNormalizeUVs = bNormalizeUVs;
 						MergedMeshOptions.State = m_currentStateIndex;
 						MergedMeshOptions.ActiveTags = m_activeTags.Last();
 						FMeshGenerationResult addResults;
diff --git a/Engine/Plugins/Experimental/Mutable/Source/MutableTools/Private/MuT/CodeGenerator.h b/Engine/Plugins/Experimental/Mutable/Source/MutableTools/Private/MuT/CodeGenerator.h
--- a/Engine/Plugins/Experimental/Mutable/Source/MutableTools/Private/MuT/CodeGenerator.h
+++ b/Engine/Plugins/Experimental/Mutable/Source/MutableTools/Private/MuT/CodeGenerator.h
@@ -328,13 +328,6 @@ namespace mu
 		// Get the modifiers that have to be applied to elements with a specific tag.
         //void GetSurfacesWithTag(const string& tag, vector<FirstPassGenerator::SURFACE>& surfaces);
 
-        //-----------------------------------------------------------------------------------------
-        void PrepareForLayout( LayoutPtrConst GeneratedLayout,
-                                  MeshPtr currentLayoutMesh,
-                                  size_t currentLayoutChannel,
-                                  const void* errorContext,
-								  const bool bClampUVIslands);
-
         //-----------------------------------------------------------------------------------------
         //!
         Ptr<ASTOp> GenerateTableVariable( TablePtr pTable, const string& strName );
@@ -461,6 +454,9 @@ namespace mu
 			* Mainly used to keep consistent layouts when reusing textures between LODs. */
 			bool bClampUVIslands = false;
 
+			/** If true, UVs will be normalized. Normalize UVs should be done in cases where we operate with Images and Layouts */
+			bool bNormalizeUVs = false;
+
 			/** If this has something the layouts in constant meshes will be ignored, because
 			* they are supposed to match some other set of layouts. If the vector is empty, layouts
 			* are generated normally.
@@ -557,6 +553,13 @@ namespace mu
 		void GenerateMesh_Reshape(const FMeshGenerationOptions&, FMeshGenerationResult&, const NodeMeshReshape*);
 		void GenerateMesh_ClipDeform(const FMeshGenerationOptions&, FMeshGenerationResult&, const NodeMeshClipDeform*);
 
+		//-----------------------------------------------------------------------------------------
+		void PrepareForLayout(LayoutPtrConst GeneratedLayout,
+			MeshPtr currentLayoutMesh,
+			size_t currentLayoutChannel,
+			const void* errorContext,
+			const FMeshGenerationOptions& MeshOptions);
+
 		//!
 		Ptr<const Layout> AddLayout(Ptr<const Layout> SourceLayout);
 
diff --git a/Engine/Plugins/Experimental/Mutable/Source/MutableTools/Private/MuT/CodeGenerator_Image.cpp b/Engine/Plugins/Experimental/Mutable/Source/MutableTools/Private/MuT/CodeGenerator_Image.cpp
--- a/Engine/Plugins/Experimental/Mutable/Source/MutableTools/Private/MuT/CodeGenerator_Image.cpp
+++ b/Engine/Plugins/Experimental/Mutable/Source/MutableTools/Private/MuT/CodeGenerator_Image.cpp
@@ -1382,6 +1382,7 @@ namespace mu
 				MeshOptions.ActiveTags = m_activeTags.Last();
 			}
 			MeshOptions.bLayouts = true;			// We need the layout that we will use to render
+			MeshOptions.bNormalizeUVs = true;		// We need normalized UVs for the projection
 			MeshOptions.bUniqueVertexIDs = false;	// We don't need the IDs at this point.
             GenerateMesh( MeshOptions, MeshResult, node.m_pMesh );
 
diff --git a/Engine/Plugins/Experimental/Mutable/Source/MutableTools/Private/MuT/CodeGenerator_Mesh.cpp b/Engine/Plugins/Experimental/Mutable/Source/MutableTools/Private/MuT/CodeGenerator_Mesh.cpp
--- a/Engine/Plugins/Experimental/Mutable/Source/MutableTools/Private/MuT/CodeGenerator_Mesh.cpp
+++ b/Engine/Plugins/Experimental/Mutable/Source/MutableTools/Private/MuT/CodeGenerator_Mesh.cpp
@@ -270,7 +270,8 @@ class Node;
 		MeshPtr currentLayoutMesh,
 		size_t currentLayoutChannel,
 		const void* errorContext,
-		const bool bClampUVIslands)
+		const FMeshGenerationOptions& MeshOptions
+		)
 	{
 		MUTABLE_CPUPROFILER_SCOPE(LayoutUV_PrepareForLayout);
 
@@ -432,7 +433,7 @@ class Node;
 				}
 
 				// Check that UVs are normalized. If not, clamp the values and throw a warning.
-				if (UV[0] < 0.f || UV[0] > 1.f || UV[1] < 0.f || UV[1] > 1.f)
+				if (MeshOptions.bNormalizeUVs && (UV[0] < 0.f || UV[0] > 1.f || UV[1] < 0.f || UV[1] > 1.f))
 				{
 					UV[0] = FMath::Clamp(UV[0], 0.f, 1.f);
 					UV[1] = FMath::Clamp(UV[1], 0.f, 1.f);
@@ -467,7 +468,7 @@ class Node;
 		TMultiMap<int32, uint32> VertexToFaceMap;
 
 		// Find Unique Vertices
-		if (bClampUVIslands)
+		if (MeshOptions.bClampUVIslands)
 		{
 			VertexToFaceMap.Reserve(NumVertices);
 			Triangles.SetNumUninitialized(NumTriangles);
@@ -478,6 +479,9 @@ class Node;
 		UntypedMeshBufferIteratorConst ItIndices(currentLayoutMesh->GetIndexBuffers(), MBS_VERTEXINDEX);
 		TArray<int32> ConflictiveTriangles;
 
+		const uint32 MaxGridX = MeshOptions.bNormalizeUVs ? MAX_uint32 : Grid.X - 1;
+		const uint32 MaxGridY = MeshOptions.bNormalizeUVs ? MAX_uint32 : Grid.Y - 1;
+
 		// TODO: We could skip this if there's only one block and it fits the entire grid
 		for (int32 TriangleIndex = 0; TriangleIndex < NumTriangles; ++TriangleIndex)
 		{
@@ -493,28 +497,28 @@ class Node;
 			uint16& BlockIndexV0 = *(LayoutData + Index0);
 			if (BlockIndexV0 == NullBlockId)
 			{
-				X = (uint16)FMath::Min<uint32>(Grid.X - 1, FMath::Max<uint32>(0, (uint32)(Grid.X * TempUVs[Index0][0])));
-				Y = (uint16)FMath::Min<uint32>(Grid.Y - 1, FMath::Max<uint32>(0, (uint32)(Grid.Y * TempUVs[Index0][1])));
-				BlockIndexV0 = GridBlockBlockId[Y * Grid.X + X];
+				X = (uint16)FMath::Min<uint32>(MaxGridX, FMath::Max<uint32>(0, (uint32)(Grid.X * TempUVs[Index0][0])));
+				Y = (uint16)FMath::Min<uint32>(MaxGridY, FMath::Max<uint32>(0, (uint32)(Grid.Y * TempUVs[Index0][1])));
+				BlockIndexV0 = (X < Grid.X && Y < Grid.Y) ? GridBlockBlockId[Y * Grid.X + X] : 0;
 			}
 
 			uint16& BlockIndexV1 = *(LayoutData + Index1);
 			if (BlockIndexV1 == NullBlockId)
 			{
-				X = (uint16)FMath::Min<uint32>(Grid.X - 1, FMath::Max<uint32>(0, (uint32)(Grid.X * TempUVs[Index1][0])));
-				Y = (uint16)FMath::Min<uint32>(Grid.Y - 1, FMath::Max<uint32>(0, (uint32)(Grid.Y * TempUVs[Index1][1])));
-				BlockIndexV1 = GridBlockBlockId[Y * Grid.X + X];
+				X = (uint16)FMath::Min<uint32>(MaxGridX, FMath::Max<uint32>(0, (uint32)(Grid.X * TempUVs[Index1][0])));
+				Y = (uint16)FMath::Min<uint32>(MaxGridY, FMath::Max<uint32>(0, (uint32)(Grid.Y * TempUVs[Index1][1])));
+				BlockIndexV1 = (X < Grid.X && Y < Grid.Y) ? GridBlockBlockId[Y * Grid.X + X] : 0;
 			}
 
 			uint16& BlockIndexV2 = *(LayoutData + Index2);
 			if (BlockIndexV2 == NullBlockId)
 			{
-				X = (uint16)FMath::Min<uint32>(Grid.X - 1, FMath::Max<uint32>(0, (uint32)(Grid.X * TempUVs[Index2][0])));
-				Y = (uint16)FMath::Min<uint32>(Grid.Y - 1, FMath::Max<uint32>(0, (uint32)(Grid.Y * TempUVs[Index2][1])));
-				BlockIndexV2 = GridBlockBlockId[Y * Grid.X + X];
+				X = (uint16)FMath::Min<uint32>(MaxGridX, FMath::Max<uint32>(0, (uint32)(Grid.X * TempUVs[Index2][0])));
+				Y = (uint16)FMath::Min<uint32>(MaxGridY, FMath::Max<uint32>(0, (uint32)(Grid.Y * TempUVs[Index2][1])));
+				BlockIndexV2 = (X < Grid.X && Y < Grid.Y) ? GridBlockBlockId[Y * Grid.X + X] : 0;
 			}
 
-			if (bClampUVIslands)
+			if (MeshOptions.bClampUVIslands)
 			{
 				if (BlockIndexV0 != BlockIndexV1 || BlockIndexV0 != BlockIndexV2)
 				{
@@ -683,7 +687,6 @@ class Node;
 				{
 					vec2<float>* pUV = (vec2<float>*)pVertices;
 					*pUV = *UV;
-					check((*pUV)[0] <= 1 && (*pUV)[1] <= 1);
 				}
 				else if (format == MBF_FLOAT16)
 				{
@@ -1384,7 +1387,7 @@ class Node;
 						{
 							Ptr<const Layout> SourceLayout = TypedNode->GetPrivate()->m_pLayout;
 							Ptr<const Layout> GeneratedLayout = AddLayout( SourceLayout );
-							PrepareForLayout(GeneratedLayout, pCloned, LayoutIndex, TypedNode->GetPrivate()->m_errorContext, InOptions.bClampUVIslands);
+							PrepareForLayout(GeneratedLayout, pCloned, LayoutIndex, TypedNode->GetPrivate()->m_errorContext, InOptions);
 
 							OutResult.GeneratedLayouts.Add(GeneratedLayout);
 						}
@@ -1396,7 +1399,7 @@ class Node;
 					for (int32 LayoutIndex = 0; LayoutIndex < InOptions.OverrideLayouts.Num(); ++LayoutIndex)
 					{
 						Ptr<const Layout> GeneratedLayout = InOptions.OverrideLayouts[LayoutIndex];
-						PrepareForLayout(GeneratedLayout, pCloned, LayoutIndex, node.m_errorContext, InOptions.bClampUVIslands);
+						PrepareForLayout(GeneratedLayout, pCloned, LayoutIndex, node.m_errorContext, InOptions);
 
 						OutResult.GeneratedLayouts.Add(GeneratedLayout);
 					}
