[Mutable] Add as compatible the compatible skeletons of the skeletons involved in the merge.

#rnx

diff --git a/Engine/Plugins/Experimental/Mutable/Source/CustomizableObject/Private/MuCO/CustomizableObjectInstance.cpp b/Engine/Plugins/Experimental/Mutable/Source/CustomizableObject/Private/MuCO/CustomizableObjectInstance.cpp
--- a/Engine/Plugins/Experimental/Mutable/Source/CustomizableObject/Private/MuCO/CustomizableObjectInstance.cpp
+++ b/Engine/Plugins/Experimental/Mutable/Source/CustomizableObject/Private/MuCO/CustomizableObjectInstance.cpp
@@ -1097,12 +1097,18 @@ USkeleton* UCustomizableInstancePrivateData::MergeSkeletons(UCustomizableObjectI
 	}
 	else
 	{
-		// The reference skeleton may be used by AnimBp as a target skeleton, add it as a compatible skeleton.
+		// Make the final skeleton compatible with all the merged skeletons and their compatible skeletons.
 		for (USkeleton* Skeleton : Params.SkeletonsToMerge)
 		{
 			if (Skeleton)
 			{
 				FinalSkeleton->AddCompatibleSkeleton(Skeleton);
+
+				const TArray<TSoftObjectPtr<USkeleton>>& CompatibleSkeletons = Skeleton->GetCompatibleSkeletons();
+				for (const TSoftObjectPtr<USkeleton>& CompatibleSkeleton : CompatibleSkeletons)
+				{
+					FinalSkeleton->AddCompatibleSkeletonSoft(CompatibleSkeleton);
+				}
 			}
 		}
 
