Track allocated memory amouts directly, instead of couting it everytime from the tracked resources.
#rnx

diff --git a/Engine/Plugins/Experimental/Mutable/Source/MutableRuntime/Private/MuR/CodeRunnerTasked.cpp b/Engine/Plugins/Experimental/Mutable/Source/MutableRuntime/Private/MuR/CodeRunnerTasked.cpp
--- a/Engine/Plugins/Experimental/Mutable/Source/MutableRuntime/Private/MuR/CodeRunnerTasked.cpp
+++ b/Engine/Plugins/Experimental/Mutable/Source/MutableRuntime/Private/MuR/CodeRunnerTasked.cpp
@@ -80,15 +80,20 @@ namespace mu
 			case DATATYPE::DT_MESH:
 			{
 				check(!program.m_constantMeshes[ResIndex].Value);
-				program.m_constantMeshes[ResIndex].Value = Mesh::StaticUnserialise(arch);
+				Ptr<Mesh> Value = Mesh::StaticUnserialise(arch);
+				program.m_constantMeshes[ResIndex].Value = Value;
 				check(program.m_constantMeshes[ResIndex].Value);
+				m_pSystem->WorkingMemoryManager.TrackedBudgetBytes_Rom += Value->GetDataSize();
 				break;
 			}
 			case DATATYPE::DT_IMAGE:
 			{
 				check(!program.m_constantImageLODs[ResIndex].Value);
-				program.m_constantImageLODs[ResIndex].Value = Image::StaticUnserialise(arch);
+				// TODO: Try to reuse buffer from PooledImages.
+				Ptr<Image> Value = Image::StaticUnserialise(arch);
+				program.m_constantImageLODs[ResIndex].Value = Value;
 				check(program.m_constantImageLODs[ResIndex].Value);
+				m_pSystem->WorkingMemoryManager.TrackedBudgetBytes_Rom += Value->GetDataSize();
 				break;
 			}
 			default:
diff --git a/Engine/Plugins/Experimental/Mutable/Source/MutableRuntime/Private/MuR/ModelPrivate.h b/Engine/Plugins/Experimental/Mutable/Source/MutableRuntime/Private/MuR/ModelPrivate.h
--- a/Engine/Plugins/Experimental/Mutable/Source/MutableRuntime/Private/MuR/ModelPrivate.h
+++ b/Engine/Plugins/Experimental/Mutable/Source/MutableRuntime/Private/MuR/ModelPrivate.h
@@ -346,24 +346,36 @@ namespace mu
 			return false;
 		}
 
-		//! Unload a rom resource
-		inline void UnloadRom(int32 RomIndex)
+		/** Unload a rom resource. Return the size of the unloaded rom.*/
+		inline int32 UnloadRom(int32 RomIndex)
 		{
+			int32 RomSize = 0;
+
 			if (DebugRom && (DebugRomAll||RomIndex == DebugRomIndex))
 				UE_LOG(LogMutableCore, Log, TEXT("Unloading rom %d."), RomIndex);
 
 			switch (m_roms[RomIndex].ResourceType)
 			{
 			case DT_IMAGE:
-				m_constantImageLODs[m_roms[RomIndex].ResourceIndex].Value = nullptr;
+				if (m_constantImageLODs[m_roms[RomIndex].ResourceIndex].Value)
+				{
+					RomSize = m_constantImageLODs[m_roms[RomIndex].ResourceIndex].Value->GetDataSize();
+					m_constantImageLODs[m_roms[RomIndex].ResourceIndex].Value = nullptr;
+				}
 				break;
 			case DT_MESH:
-				m_constantMeshes[m_roms[RomIndex].ResourceIndex].Value = nullptr;
+				if (m_constantMeshes[m_roms[RomIndex].ResourceIndex].Value)
+				{
+					RomSize = m_constantMeshes[m_roms[RomIndex].ResourceIndex].Value->GetDataSize();
+					m_constantMeshes[m_roms[RomIndex].ResourceIndex].Value = nullptr;
+				}
 				break;
 			default:
 				check(false);
 				break;
 			}
+
+			return RomSize;
 		}
 
         //! Adds a constant data and returns its constant index.
diff --git a/Engine/Plugins/Experimental/Mutable/Source/MutableRuntime/Private/MuR/System.cpp b/Engine/Plugins/Experimental/Mutable/Source/MutableRuntime/Private/MuR/System.cpp
--- a/Engine/Plugins/Experimental/Mutable/Source/MutableRuntime/Private/MuR/System.cpp
+++ b/Engine/Plugins/Experimental/Mutable/Source/MutableRuntime/Private/MuR/System.cpp
@@ -166,6 +166,10 @@ namespace mu
 		m_pD->WorkingMemoryManager.PooledImages.Empty();
 		m_pD->WorkingMemoryManager.CacheResources.Empty();
 		check(m_pD->WorkingMemoryManager.TempImages.IsEmpty());
+		m_pD->WorkingMemoryManager.TrackedBudgetBytes_Pooled = 0;
+		m_pD->WorkingMemoryManager.TrackedBudgetBytes_Temp = 0;
+		m_pD->WorkingMemoryManager.TrackedBudgetBytes_Cached = 0;
+		m_pD->WorkingMemoryManager.TrackedBudgetBytes_Rom = 0;
 	}
 
 	
@@ -445,6 +449,7 @@ namespace mu
 		if (m_pD->WorkingMemoryManager.BudgetBytes == 0)
 		{
 			m_pD->WorkingMemoryManager.PooledImages.Empty();
+			m_pD->WorkingMemoryManager.TrackedBudgetBytes_Pooled = 0;
 		}
 
 	}
@@ -466,14 +471,24 @@ namespace mu
 				{
 					if (Data)
 					{
-						m_pD->WorkingMemoryManager.CacheResources.Remove(Data);
+						int32 DataSize = Data->GetDataSize();
+						int32 Removed = m_pD->WorkingMemoryManager.CacheResources.Remove(Data);
+						if (Removed)
+						{
+							m_pD->WorkingMemoryManager.TrackedBudgetBytes_Cached -= DataSize;
+						}
 					}
 				}
 				for (Ptr<const Mesh>& Data : Instance.Cache->MeshResults)
 				{
 					if (Data)
 					{
-						m_pD->WorkingMemoryManager.CacheResources.Remove(Data);
+						int32 DataSize = Data->GetDataSize();
+						int32 Removed = m_pD->WorkingMemoryManager.CacheResources.Remove(Data);
+						if (Removed)
+						{
+							m_pD->WorkingMemoryManager.TrackedBudgetBytes_Cached -= DataSize;
+						}
 					}
 				}
 
@@ -955,67 +970,17 @@ namespace mu
 			+ StreamingBufferBytes
 			+ RomBytes;
 
-		UE_LOG(LogMutableCore, Log, TEXT("Mem KB: Internal %7d   Temp %7d   Pool %7d   Cache0+1 %7d   CacheTracked %7d   Rom %7d   Stream %7d   Total %7d / %7d"),
+		UE_LOG(LogMutableCore, Log, TEXT("Mem KB: Internal %7d   Temp %7d (%7d)   Pool %7d (%7d)   Cache0+1 %7d (%7d) (%7d)   Rom %7d (%7d)   Stream %7d   Total %7d / %7d"),
 			InternalBytes / 1024,
-			TempBytes / 1024,
-			PoolBytes / 1024,
-			CacheBytes / 1024,
-			TrackedCacheBytes / 1024,
-			RomBytes / 1024,
+			TempBytes / 1024, TrackedBudgetBytes_Temp / 1024,
+			PoolBytes / 1024, TrackedBudgetBytes_Pooled / 1024,
+			CacheBytes / 1024, TrackedCacheBytes / 1024, TrackedBudgetBytes_Cached / 1024,
+			RomBytes / 1024, TrackedBudgetBytes_Rom / 1024,
 			StreamingBufferBytes / 1024,
 			TotalBytes / 1024,
 			BudgetBytes / 1024
 			);
 
-		// Detailed log for debug.
-		//if (TotalBytes > 30 * 1024 * 1024)
-		//if (false)
-		//{
-		//	UE_LOG(LogMutableCore, Log, TEXT("Mem detailed report:"));
-
-		//	for (const FLiveInstance& Instance : LiveInstances)
-		//	{
-		//		UE_LOG(LogMutableCore, Log, TEXT("  instance:"));
-
-		//		TSet<const Resource*> InstanceCache0Unique;
-		//		TSet<const Resource*> InstanceCache1Unique;
-
-		//		CodeContainer<int>::iterator it = Instance.Cache->m_opHitCount.begin();
-		//		for (; it.IsValid(); ++it)
-		//		{
-		//			const Resource* Value = Instance.Cache->m_resources[it.get_address()].Value.get();
-		//			if (!Value)
-		//			{
-		//				continue;
-		//			}
-
-		//			int32 Count = *it;
-		//			if (Count < UE_MUTABLE_CACHE_COUNT_LIMIT)
-		//			{
-		//				InstanceCache0Unique.Add(Value);
-		//			}
-		//			else
-		//			{
-		//				InstanceCache1Unique.Add(Value);
-		//			}
-		//		}
-
-		//		UE_LOG(LogMutableCore, Log, TEXT("    cache level 0:"));
-
-		//		for (const Resource* Value : InstanceCache0Unique)
-		//		{
-		//			UE_LOG(LogMutableCore, Log, TEXT("      resource: %d"), Value->GetDataSize());
-		//		}
-
-		//		UE_LOG(LogMutableCore, Log, TEXT("    cache level 1:"));
-
-		//		for (const Resource* Value : InstanceCache1Unique)
-		//		{
-		//			UE_LOG(LogMutableCore, Log, TEXT("      resource: %d"), Value->GetDataSize());
-		//		}
-		//	}
-		//}
-
 		TRACE_COUNTER_SET(MutableRuntime_MemInternal, InternalBytes);
 		TRACE_COUNTER_SET(MutableRuntime_MemTemp, TempBytes);
 		TRACE_COUNTER_SET(MutableRuntime_MemPool, PoolBytes);
@@ -1076,10 +1041,10 @@ namespace mu
 		uint32 TrackedCacheBytes = 0;
 		{
 			MUTABLE_CPUPROFILER_SCOPE(EnsureBudgetBelow_Measure);
-			RomBytes = GetRomBytes();
-			PoolBytes = GetPooledBytes();
-			TempBytes = GetTempBytes();
-			TrackedCacheBytes = GetTrackedCacheBytes();
+			RomBytes = TrackedBudgetBytes_Rom; // GetRomBytes();
+			PoolBytes = TrackedBudgetBytes_Pooled; // GetPooledBytes();
+			TempBytes = TrackedBudgetBytes_Temp; // GetTempBytes();
+			TrackedCacheBytes = TrackedBudgetBytes_Cached; // GetTrackedCacheBytes();
 		}
 
 		uint64 TotalBytes = 0;
@@ -1100,8 +1065,9 @@ namespace mu
 			{
 				// TODO: Actually advancing index if possible after swap may be better to remove the oldest in the pool first.
 				int32 PooledResourceSize = PooledImages[0]->GetDataSize();
-				TotalBytes -= PooledResourceSize;
+				TotalBytes -= PooledResourceSize;				
 				PooledImages.RemoveAtSwap(0);
+				TrackedBudgetBytes_Pooled -= PooledResourceSize;
 				bFinished = TotalBytes <= BudgetBytes;
 			}
 		}
@@ -1162,8 +1128,10 @@ namespace mu
 					MUTABLE_CPUPROFILER_SCOPE(EnsureBudgetBelow_UnloadRom);
 
 					// UE_LOG(LogMutableCore,Log, "Unloading rom because of memory budget: %d.", lowestPriorityRom);
-					LowestPriorityModel->GetPrivate()->m_program.UnloadRom(LowestPriorityRom);
-					TotalBytes -= LowestPriorityModel->GetPrivate()->m_program.m_roms[LowestPriorityRom].Size;
+					int32 UnloadedSize = LowestPriorityModel->GetPrivate()->m_program.UnloadRom(LowestPriorityRom);
+					TotalBytes -= UnloadedSize;
+					TrackedBudgetBytes_Rom -= UnloadedSize;
+
 					bFinished = TotalBytes <= BudgetBytes;
 				}
 			}
@@ -1262,11 +1230,16 @@ namespace mu
 					TotalBytes -= RemovedDataSize;
 					bFinished = TotalBytes <= BudgetBytes;
 				}
+
+				if (bFinished)
+				{
+					break;
+				}
 			}
 
 			// From the current live instances. It is more involved: we have to make sure any data we want to free is not also
 			// in any cache (0 or 1) position with hit-count > 0.
-			if (CurrentInstanceCache)
+			if (CurrentInstanceCache && !bFinished)
 			{
 				// Gather all data in the cache for this instance
 				TArray<const Resource*> CacheUnique;
@@ -1380,7 +1353,11 @@ namespace mu
 
 						if (Value == Removed)
 						{
-							CacheResources.Remove(Removed);
+							int32 RemoveCount = CacheResources.Remove(Removed);
+							if (RemoveCount > 0)
+							{
+								TrackedBudgetBytes_Cached -= RemovedDataSize;
+							}
 							// \TODO: This is not very efficient, since it will search the data again.
 							CurrentInstanceCache->SetUnused(RemIt.get_address());
 						}
diff --git a/Engine/Plugins/Experimental/Mutable/Source/MutableRuntime/Private/MuR/SystemPrivate.h b/Engine/Plugins/Experimental/Mutable/Source/MutableRuntime/Private/MuR/SystemPrivate.h
--- a/Engine/Plugins/Experimental/Mutable/Source/MutableRuntime/Private/MuR/SystemPrivate.h
+++ b/Engine/Plugins/Experimental/Mutable/Source/MutableRuntime/Private/MuR/SystemPrivate.h
@@ -1113,6 +1113,12 @@ namespace mu
 		/** Maximum excess memory reached suring the current operation. */
 		uint64 BudgetExcessBytes = 0;
 
+		/** . */
+		uint32 TrackedBudgetBytes_Rom = 0;
+		uint32 TrackedBudgetBytes_Temp = 0;
+		uint32 TrackedBudgetBytes_Pooled = 0;
+		uint32 TrackedBudgetBytes_Cached = 0;
+
 		/** This value is used to track the order of loading of roms. */
         uint64 RomTick = 0;
 
@@ -1166,6 +1172,7 @@ namespace mu
 				{
 					Ptr<Image> Result = Candidate;
 					PooledImages.RemoveAtSwap(Index);
+					TrackedBudgetBytes_Pooled -= Result->GetDataSize();
 					if (Init == EInitializationType::Black)
 					{
 						Result->InitToBlack();
@@ -1187,6 +1194,7 @@ namespace mu
 			Ptr<Image> Result = new Image(SizeX, SizeY, Lods, Format, Init);
 
 			TempImages.Add(Result);
+			TrackedBudgetBytes_Temp += Result->GetDataSize();
 			return Result;
 		}
 
@@ -1195,7 +1203,12 @@ namespace mu
 		{
 			CheckRunnerThread();
 
-			TempImages.RemoveSingle(Resource);
+			int32 Removed = TempImages.RemoveSingle(Resource);
+			if (Removed)
+			{
+				TrackedBudgetBytes_Temp -= Resource->GetDataSize();
+			}
+
 			check(!TempImages.Contains(Resource));
 			check(!PooledImages.Contains(Resource));
 
@@ -1229,23 +1242,30 @@ namespace mu
 				return;
 			}
 
-			TempImages.RemoveSingle(Resource);
+			int32 ResourceDataSize = Resource->GetDataSize();
+			int32 Removed = TempImages.RemoveSingle(Resource);
+			if (Removed)
+			{
+				TrackedBudgetBytes_Temp -= ResourceDataSize;
+			}
+
 			check(!TempImages.Contains(Resource));
 			check(!PooledImages.Contains(Resource));
 
 			if (IsBudgetTemp(Resource))
 			{
 				// Check if we are exceeding the budget
-				bool bInBudget = EnsureBudgetBelow(Resource->GetDataSize());
+				bool bInBudget = EnsureBudgetBelow(ResourceDataSize);
 				if (bInBudget)
 				{
 					PooledImages.Add(const_cast<Image*>(Resource.get()));
+					TrackedBudgetBytes_Pooled += ResourceDataSize;
 				}
 			}
 			else
 			{
 				// Check if we are exceeding the budget
-				bool bInBudget = EnsureBudgetBelow(0);
+				EnsureBudgetBelow(0);
 			}
 
 			Resource = nullptr;
@@ -1262,23 +1282,30 @@ namespace mu
 				return;
 			}
 
-			TempImages.RemoveSingle(Resource);
+			int32 ResourceDataSize = Resource->GetDataSize();
+			int32 Removed = TempImages.RemoveSingle(Resource);
+			if (Removed)
+			{
+				TrackedBudgetBytes_Temp -= ResourceDataSize;
+			}
+
 			check(!TempImages.Contains(Resource));
 			check(!PooledImages.Contains(Resource));
 
 			if (IsBudgetTemp(Resource))
 			{
 				// Check if we are exceeding the budget
-				bool bInBudget = EnsureBudgetBelow(Resource->GetDataSize());
+				bool bInBudget = EnsureBudgetBelow(ResourceDataSize);
 				if (bInBudget)
 				{
 					PooledImages.Add(Resource.get());
+					TrackedBudgetBytes_Pooled += ResourceDataSize;
 				}
 			}
 			else
 			{
 				// Check if we are exceeding the budget
-				bool bInBudget = EnsureBudgetBelow(0);
+				EnsureBudgetBelow(0);
 			}
 
 			Resource = nullptr;
@@ -1326,6 +1353,8 @@ namespace mu
 				return nullptr;
 			}
 
+			int32 ResultDataSize = Result->GetDataSize();
+
 			// If we retrieved the last reference to this resource in "From" cache position (it could still be in other cache positions as well)
 			if (bIsLastReference)
 			{
@@ -1335,6 +1364,7 @@ namespace mu
 				if (!*CountPtr)
 				{
 					CacheResources.FindAndRemoveChecked(Result);
+					TrackedBudgetBytes_Cached -= ResultDataSize;
 				}
 			}
 
@@ -1342,6 +1372,7 @@ namespace mu
 			if (!bTakeOwnership && Result->IsUnique())
 			{
 				TempImages.Add(Result);
+				TrackedBudgetBytes_Temp += ResultDataSize;
 			}
 
 			return Result;
@@ -1352,10 +1383,21 @@ namespace mu
 		{
 			if (Resource)
 			{
-				TempImages.RemoveSingle(Resource);
+				int32 ResourceDataSize = Resource->GetDataSize();
+				int32 Removed = TempImages.RemoveSingle(Resource);
+				if (Removed)
+				{
+					TrackedBudgetBytes_Temp -= ResourceDataSize;
+				}
+
 				check(!TempImages.Contains(Resource));
 
 				int32& Count = CacheResources.FindOrAdd(Resource, 0);
+				if (!Count)
+				{
+					TrackedBudgetBytes_Cached += ResourceDataSize;
+				}
+
 				++Count;
 			}
 
@@ -1383,6 +1425,8 @@ namespace mu
 				Result += Value->GetDataSize();
 			}
 
+			ensure(Result == TrackedBudgetBytes_Pooled);
+
 			return Result;
 		}
 
@@ -1395,6 +1439,8 @@ namespace mu
 				Result += Value->GetDataSize();
 			}
 
+			ensure(Result==TrackedBudgetBytes_Temp);
+
 			return Result;
 		}
 
@@ -1528,19 +1574,37 @@ namespace mu
 				}
 
 				const Resource* Value = nullptr;
+				int32 ValueDataSize = 0;
+				int32 RemoveCount = 0;
 
 				switch (Data.DataType)
 				{
 				case DATATYPE::DT_IMAGE:
 					Value = CurrentInstanceCache->ImageResults[Data.DataTypeIndex].get();
-					CacheResources.Remove(Value);
-					CurrentInstanceCache->ImageResults[Data.DataTypeIndex] = nullptr;
+					if (Value)
+					{
+						ValueDataSize = Value->GetDataSize();
+						RemoveCount = CacheResources.Remove(Value);
+						if (RemoveCount)
+						{
+							TrackedBudgetBytes_Cached -= ValueDataSize;
+						}
+						CurrentInstanceCache->ImageResults[Data.DataTypeIndex] = nullptr;
+					}
 					break;
 
 				case DATATYPE::DT_MESH:
 					Value = CurrentInstanceCache->MeshResults[Data.DataTypeIndex].get();
-					CacheResources.Remove(Value);
-					CurrentInstanceCache->MeshResults[Data.DataTypeIndex] = nullptr;
+					if (Value)
+					{
+						ValueDataSize = Value->GetDataSize();
+						RemoveCount = CacheResources.Remove(Value);
+						if (RemoveCount)
+						{
+							TrackedBudgetBytes_Cached -= ValueDataSize;
+						}
+						CurrentInstanceCache->MeshResults[Data.DataTypeIndex] = nullptr;
+					}
 					break;
 
 				case DATATYPE::DT_LAYOUT:
