[Mutable] Remove ensure in SystemPrivate::GetTempBytes()

#rnx

diff --git a/Engine/Plugins/Experimental/Mutable/Source/MutableRuntime/Private/MuR/SystemPrivate.h b/Engine/Plugins/Experimental/Mutable/Source/MutableRuntime/Private/MuR/SystemPrivate.h
--- a/Engine/Plugins/Experimental/Mutable/Source/MutableRuntime/Private/MuR/SystemPrivate.h
+++ b/Engine/Plugins/Experimental/Mutable/Source/MutableRuntime/Private/MuR/SystemPrivate.h
@@ -1439,8 +1439,6 @@ namespace mu
 				Result += Value->GetDataSize();
 			}
 
-			ensure(Result==TrackedBudgetBytes_Temp);
-
 			return Result;
 		}
 
