[mutable] Disabled lod-bias testing in asset validation.
#rnx

diff --git a/Engine/Plugins/Experimental/Mutable/Source/CustomizableObjectEditor/Private/MuCOE/SMutableObjectViewer.cpp b/Engine/Plugins/Experimental/Mutable/Source/CustomizableObjectEditor/Private/MuCOE/SMutableObjectViewer.cpp
--- a/Engine/Plugins/Experimental/Mutable/Source/CustomizableObjectEditor/Private/MuCOE/SMutableObjectViewer.cpp
+++ b/Engine/Plugins/Experimental/Mutable/Source/CustomizableObjectEditor/Private/MuCOE/SMutableObjectViewer.cpp
@@ -250,11 +250,12 @@ void SMutableObjectViewer::CompileMutableCodePressed()
 		for (int32 Bias = 0; Bias < MaxBias; ++Bias)
 		{
 			CompileOptions.DebugBias = Bias;
-				
+
 			mu::NodePtr RootNode = Compiler.Export(CustomizableObject, CompileOptions, bDisableLayouts);
 			if (!RootNode)
 			{
 				// TODO: Show errors
+				ensure(false);
 				return;
 			}
 
@@ -272,6 +273,7 @@ void SMutableObjectViewer::CompileMutableCodePressed()
 	if (!RootNode)
 	{
 		// TODO: Show errors
+		ensure(false);
 		return;
 	}
 
diff --git a/Engine/Plugins/Experimental/Mutable/Source/MutableValidation/Private/MuV/AssetValidator_CustomizableObjects.cpp b/Engine/Plugins/Experimental/Mutable/Source/MutableValidation/Private/MuV/AssetValidator_CustomizableObjects.cpp
--- a/Engine/Plugins/Experimental/Mutable/Source/MutableValidation/Private/MuV/AssetValidator_CustomizableObjects.cpp
+++ b/Engine/Plugins/Experimental/Mutable/Source/MutableValidation/Private/MuV/AssetValidator_CustomizableObjects.cpp
@@ -113,16 +113,16 @@ EDataValidationResult UAssetValidator_CustomizableObjects::IsCustomizableObjectV
 
 	// Configuration with LOD bias applied
 	// constexpr int32 MaxBias = 15;	
-	constexpr int32 MaxBias = 6;		// Reduced amount since we have this value in GenerateMutableSource for the max lod bias provided to mutable
-	for (int32 LodBias = 1; LodBias <= MaxBias; LodBias++)
-	{
-		FCompilationOptions ModifiedCompilationOptions = InCustomizableObject->CompileOptions;
-		ModifiedCompilationOptions.bForceLargeLODBias = true;	
-		ModifiedCompilationOptions.DebugBias = LodBias;	
-	
-		// Add one configuration object for each bias setting
-		CompilationOptionsToTest.Add(ModifiedCompilationOptions);
-	}
+	//constexpr int32 MaxBias = 6;		// Reduced amount since we have this value in GenerateMutableSource for the max lod bias provided to mutable
+	//for (int32 LodBias = 1; LodBias <= MaxBias; LodBias++)
+	//{
+	//	FCompilationOptions ModifiedCompilationOptions = InCustomizableObject->CompileOptions;
+	//	ModifiedCompilationOptions.bForceLargeLODBias = true;	
+	//	ModifiedCompilationOptions.DebugBias = LodBias;	
+	//
+	//	// Add one configuration object for each bias setting
+	//	CompilationOptionsToTest.Add(ModifiedCompilationOptions);
+	//}
 	
 	// Caches with all the data produced by the subsequent compilations of the root of this CO
 	TArray<FText> CachedValidationErrors;
