[mutable] Fix compiled CO Population Generator is invalid on packaged games

#10437
#rnx

diff --git a/Engine/Plugins/Experimental/Mutable/Source/CustomizableObjectPopulation/Private/MuCOP/CustomizableObjectPopulationModule.cpp b/Engine/Plugins/Experimental/Mutable/Source/CustomizableObjectPopulation/Private/MuCOP/CustomizableObjectPopulationModule.cpp
--- a/Engine/Plugins/Experimental/Mutable/Source/CustomizableObjectPopulation/Private/MuCOP/CustomizableObjectPopulationModule.cpp
+++ b/Engine/Plugins/Experimental/Mutable/Source/CustomizableObjectPopulation/Private/MuCOP/CustomizableObjectPopulationModule.cpp
@@ -91,12 +91,7 @@ void FCustomizableObjectPopulationModule::RecompilePopulations(UCustomizableObje
 			{
 				if (Population->IsValidPopulation())
 				{
-					//for (int32 j = 0; j < Population->ClassWeights.Num(); ++j)
-					//{
-					//	//Population->ClassWeights[j].Class->CustomizableObject->RefreshVersionId();
-					//	Population->ClassWeights[j].Class->CustomizableObject->LoadDerivedData(true);
-					//}
-					Population->CompilePopulation(NewObject<UCustomizableObjectPopulationGenerator>());
+					Population->CompilePopulation();
 				}
 			}
 		}
