[mutable] Simplify morph operation for the only supported case: one morph with a weight.
#rnx

diff --git a/Engine/Plugins/Experimental/Mutable/Source/CustomizableObject/Classes/MuCO/CustomizableObject.h b/Engine/Plugins/Experimental/Mutable/Source/CustomizableObject/Classes/MuCO/CustomizableObject.h
--- a/Engine/Plugins/Experimental/Mutable/Source/CustomizableObject/Classes/MuCO/CustomizableObject.h
+++ b/Engine/Plugins/Experimental/Mutable/Source/CustomizableObject/Classes/MuCO/CustomizableObject.h
@@ -1449,7 +1449,7 @@ private:
 	// This is a manual version number for the binary blobs in this asset.
 	// Increasing it invalidates all the previously compiled models.
 	// Warning: If while merging code both versions have changed, take the highest+1.
-	static const int32 CurrentSupportedVersion = 387;
+	static const int32 CurrentSupportedVersion = 388;
 
 public:
 
diff --git a/Engine/Plugins/Experimental/Mutable/Source/CustomizableObjectEditor/Private/MuCOE/GenerateMutableSource/GenerateMutableSourceMesh.cpp b/Engine/Plugins/Experimental/Mutable/Source/CustomizableObjectEditor/Private/MuCOE/GenerateMutableSource/GenerateMutableSourceMesh.cpp
--- a/Engine/Plugins/Experimental/Mutable/Source/CustomizableObjectEditor/Private/MuCOE/GenerateMutableSource/GenerateMutableSourceMesh.cpp
+++ b/Engine/Plugins/Experimental/Mutable/Source/CustomizableObjectEditor/Private/MuCOE/GenerateMutableSource/GenerateMutableSourceMesh.cpp
@@ -2111,7 +2111,7 @@ mu::MeshPtr BuildMorphedMutableMesh(const UEdGraphPin* BaseSourcePin, const FStr
 }
 
 
-void GenerateMorphFactor(const UCustomizableObjectNode* Node, const UEdGraphPin& FactorPin, FMutableGraphGenerationContext& GenerationContext, mu::NodeMeshMorphPtr MeshNode)
+void GenerateMorphFactor(const UCustomizableObjectNode* Node, const UEdGraphPin& FactorPin, FMutableGraphGenerationContext& GenerationContext, mu::Ptr<mu::NodeMeshMorph> MeshNode)
 {
 	if (const UEdGraphPin* ConnectedPin = FollowInputPin(FactorPin))
 	{
@@ -2614,8 +2614,7 @@ mu::NodeMeshPtr GenerateMorphMesh(const UEdGraphPin* Pin,
 	const UCustomizableObjectNode* MorphNode = TypedNodeMorphs[MorphIndex].OwningNode;
 	check(MorphNode);
 	
-	mu::NodeMeshMorphPtr Result = new mu::NodeMeshMorph();
-	Result->SetMorphCount(2);
+	mu::Ptr<mu::NodeMeshMorph> Result = new mu::NodeMeshMorph();
 	
 	// Factor
 	GenerateMorphFactor(MorphNode, *TypedNodeMorphs[MorphIndex].FactorPin, GenerationContext, Result);
@@ -2674,16 +2673,12 @@ mu::NodeMeshPtr GenerateMorphMesh(const UEdGraphPin* Pin,
 			MorphedSourceMeshNodeTable->SetParameterName(StringCast<ANSICHAR>(*TypedNodeTable->ParameterName).Get());
 			MorphedSourceMeshNodeTable->SetMessageContext(MorphNode);
 
-			// A null target will leave the base unchanged
-			mu::NodeMeshPtr IdentityMorph = nullptr;
-			Result->SetMorph(0, IdentityMorph);
-
 			mu::NodeMeshMakeMorphPtr Morph = new mu::NodeMeshMakeMorph;
 			Morph->SetBase(BaseSourceMesh.get());
 			Morph->SetTarget(MorphedSourceMeshNodeTable.get());
 			Morph->SetMessageContext(MorphNode);
 
-			Result->SetMorph(1, Morph);
+			Result->SetMorph(Morph);
 		}
 	}
 	else
@@ -2698,19 +2693,12 @@ mu::NodeMeshPtr GenerateMorphMesh(const UEdGraphPin* Pin,
 			MorphedSourceMeshNode->SetValue(MorphedSourceMesh);
 			MorphedSourceMeshNode->SetMessageContext(MorphNode);
 
-			mu::NodeMeshMakeMorphPtr IdentityMorph = new mu::NodeMeshMakeMorph;
-			IdentityMorph->SetBase(BaseSourceMesh.get());
-			IdentityMorph->SetTarget(BaseSourceMesh.get());
-			IdentityMorph->SetMessageContext(MorphNode);
-
-			Result->SetMorph(0, IdentityMorph);
-
 			mu::NodeMeshMakeMorphPtr Morph = new mu::NodeMeshMakeMorph;
 			Morph->SetBase(BaseSourceMesh.get());
 			Morph->SetTarget(MorphedSourceMeshNode.get());
 			Morph->SetMessageContext(MorphNode);
 
-			Result->SetMorph(1, Morph);
+			Result->SetMorph(Morph);
 
 			if (UCustomizableObjectNodeMeshMorph* TypedMorphNode = Cast<UCustomizableObjectNodeMeshMorph>(TypedNodeMorphs[MorphIndex].OwningNode))
 			{
@@ -2771,7 +2759,7 @@ mu::NodeMeshPtr GenerateMorphMesh(const UEdGraphPin* Pin,
 }
 
 
-void GenerateMorphTarget(const UCustomizableObjectNode* Node, const UEdGraphPin* BaseSourcePin, FMutableGraphGenerationContext& GenerationContext, mu::NodeMeshMorphPtr MeshNode, FString MorphName)
+void GenerateMorphTarget(const UCustomizableObjectNode* Node, const UEdGraphPin* BaseSourcePin, FMutableGraphGenerationContext& GenerationContext, mu::Ptr<mu::NodeMeshMorph> MeshNode, FString MorphName)
 {
 	SCOPED_PIN_DATA(GenerationContext, BaseSourcePin)
 
@@ -2785,16 +2773,12 @@ void GenerateMorphTarget(const UCustomizableObjectNode* Node, const UEdGraphPin*
 		MorphedSourceMeshNode->SetValue(MorphedSourceMesh);
 		MorphedSourceMeshNode->SetMessageContext(Node);
 
-		// A null target will leave the base unchanged
-		mu::NodeMeshPtr IdentityMorph = nullptr;
-		MeshNode->SetMorph(0, IdentityMorph);
-
 		mu::NodeMeshMakeMorphPtr Morph = new mu::NodeMeshMakeMorph;
 		Morph->SetBase(BaseSourceMesh.get());
 		Morph->SetTarget(MorphedSourceMeshNode.get());
 		Morph->SetMessageContext(Node);
 
-		MeshNode->SetMorph(1, Morph);
+		MeshNode->SetMorph(Morph);
 	}
 	else
 	{
@@ -3105,7 +3089,7 @@ mu::NodeMeshPtr GenerateMutableSourceMesh(const UEdGraphPin * Pin,
 		}
 		else
 		{
-			mu::NodeMeshMorphPtr MeshNode = new mu::NodeMeshMorph();
+			mu::Ptr<mu::NodeMeshMorph> MeshNode = new mu::NodeMeshMorph();
 			Result = MeshNode;
 		}
 	}
@@ -3123,11 +3107,9 @@ mu::NodeMeshPtr GenerateMutableSourceMesh(const UEdGraphPin * Pin,
 					TypedNodeMeshMorphStackDef->SetRefreshNodeWarning();
 				}
 
-				mu::NodeMeshMorphPtr MeshNode = new mu::NodeMeshMorph();
+				mu::Ptr<mu::NodeMeshMorph> MeshNode = new mu::NodeMeshMorph();
 				Result = MeshNode;
 
-				MeshNode->SetMorphCount(2);
-
 				TArray<UEdGraphPin*> MorphPins = TypedNodeMeshMorphStackDef->GetAllNonOrphanPins();
 
 				int32 AddedMorphs = 0;
diff --git a/Engine/Plugins/Experimental/Mutable/Source/MutableRuntime/Private/MuR/CodeRunner.cpp b/Engine/Plugins/Experimental/Mutable/Source/MutableRuntime/Private/MuR/CodeRunner.cpp
--- a/Engine/Plugins/Experimental/Mutable/Source/MutableRuntime/Private/MuR/CodeRunner.cpp
+++ b/Engine/Plugins/Experimental/Mutable/Source/MutableRuntime/Private/MuR/CodeRunner.cpp
@@ -1091,7 +1091,7 @@ namespace mu
             break;
         }
 
-        case OP_TYPE::ME_MORPH2:
+        case OP_TYPE::ME_MORPH:
         {
 			const uint8* data = pModel->GetPrivate()->m_program.GetOpArgsPointer(item.At);
 
@@ -1101,15 +1101,8 @@ namespace mu
 			OP::ADDRESS BaseAt = 0;
 			FMemory::Memcpy(&BaseAt, data, sizeof(OP::ADDRESS)); data += sizeof(OP::ADDRESS);
 
-			uint8 NumTargets = 0;
-			FMemory::Memcpy(&NumTargets, data, sizeof(uint8)); data += sizeof(uint8);
-
-			TArray<MESH_BUFFER_SEMANTIC, TInlineAllocator<8>> Targets;
-			Targets.SetNum(NumTargets);
-			for (uint8 T = 0; T < NumTargets; ++T)
-			{
-				FMemory::Memcpy(&Targets[T], data, sizeof(OP::ADDRESS)); data += sizeof(OP::ADDRESS);
-			}
+			OP::ADDRESS TargetAt = 0;
+			FMemory::Memcpy(&TargetAt, data, sizeof(OP::ADDRESS)); data += sizeof(OP::ADDRESS);
 
 			switch (item.Stage)
             {
@@ -1117,7 +1110,6 @@ namespace mu
                 if (BaseAt)
                 {
                     AddOp( FScheduledOp(item.At, item, 1),
-                           FScheduledOp(BaseAt, item),
                            FScheduledOp(FactorAt, item) );
                 }
                 else
@@ -1128,72 +1120,30 @@ namespace mu
 
             case 1:
             {
-                MUTABLE_CPUPROFILER_SCOPE(ME_MORPH2_1)
+                MUTABLE_CPUPROFILER_SCOPE(ME_MORPH_1)
 
-                bool baseValid = GetMemory().IsValid( FCacheAddress(BaseAt,item) );
-                float factor = LoadScalar( FCacheAddress(FactorAt,item) );
+                float Factor = LoadScalar( FCacheAddress(FactorAt,item) );
 
                 // Factor goes from -1 to 1 across all targets. [0 - 1] represents positive morphs, while [-1, 0) represent negative morphs.
-				factor = FMath::Clamp(factor, -1.0f, 1.0f); // Is the factor not in range [-1, 1], it will index a non existing morph.
-
-				float absFactor = FMath::Abs(factor);
-                float delta = 1.0f/(NumTargets -1);
-				int max = (int)FMath::CeilToFloat(absFactor/delta);
-				int min = max > 0 ? max - 1 : 0;
-
-                // Factor from 0 to 1 between the two targets
-				float bifactor = absFactor/delta - min;
-
-				// From [0,1] to [-1,0] for negative factors
-				if (factor < 0.0f)
-				{
-					float threshold = -1.0f + UE_SMALL_NUMBER;
+				Factor = FMath::Clamp(Factor, -1.0f, 1.0f); // Is the factor not in range [-1, 1], it will index a non existing morph.
 
-					if (factor <= threshold)
-					{
-						// setting bifactor to -1.0
-						bifactor = -1.0f;
-					}
-					else
-					{
-						// we calculate the bifactor using positive values but is negative here
-						bifactor *= -1.0f;
-					}
-				}
+                FScheduledOpData HeapData;
+				HeapData.Interpolate.Bifactor = Factor;
+				uint32 dataAddress = uint32(m_heapData.Add(HeapData));
 
-                if (baseValid)
+                // No morph
+				if (Factor < UE_SMALL_NUMBER && Factor > -UE_SMALL_NUMBER )
+                {                        
+                    AddOp( FScheduledOp( item.At, item, 2, dataAddress),
+                            FScheduledOp(BaseAt, item) );
+                }
+                // The Morph, partial or full
+                else
                 {
-                    FScheduledOpData HeapData;
-					HeapData.Interpolate.Bifactor = bifactor;
-					HeapData.Interpolate.Min = FMath::Clamp(min, 0, NumTargets - 1);
-					HeapData.Interpolate.Max = FMath::Clamp(max, 0, NumTargets - 1);
-					uint32 dataAddress = uint32(m_heapData.Add(HeapData));
-
-                    // Just the first of the targets
-					if ( bifactor < UE_SMALL_NUMBER && bifactor > -UE_SMALL_NUMBER )
-                    {                        
-                        AddOp( FScheduledOp( item.At, item, 2, dataAddress),
-                                FScheduledOp(BaseAt, item),
-                                FScheduledOp(Targets[min], item) );
-                    }
-                    // Just the second of the targets
-					else if ( bifactor > 1.0f - UE_SMALL_NUMBER || bifactor <= -1.0f + UE_SMALL_NUMBER )
-                    {
-                        check( max>0 );
-
-                        AddOp( FScheduledOp( item.At, item, 2, dataAddress),
-                                FScheduledOp(BaseAt, item),
-                                FScheduledOp(Targets[max], item) );
-                    }
-                    // Mix two targets on the base
-                    else
-                    {
-                        // We will need the base again
-                        AddOp( FScheduledOp( item.At, item, 2, dataAddress),
-                                FScheduledOp(BaseAt, item),
-                                FScheduledOp(Targets[min], item),
-                                FScheduledOp(Targets[max], item) );
-                    }
+                    // We will need the base again
+                    AddOp( FScheduledOp( item.At, item, 2, dataAddress),
+                            FScheduledOp(BaseAt, item),
+                            FScheduledOp(TargetAt, item) );
                 }
 
                 break;
@@ -1201,44 +1151,29 @@ namespace mu
 
             case 2:
             {
-       		    MUTABLE_CPUPROFILER_SCOPE(ME_MORPH2_2)
+       		    MUTABLE_CPUPROFILER_SCOPE(ME_MORPH_2)
 
                 Ptr<const Mesh> pBase = LoadMesh( FCacheAddress(BaseAt,item) );
 
                 // Factor from 0 to 1 between the two targets
                 const FScheduledOpData& HeapData = m_heapData[ (size_t)item.CustomState ];
-                float bifactor = HeapData.Interpolate.Bifactor;
-                int min = HeapData.Interpolate.Min;
-                int max = HeapData.Interpolate.Max;
+                float Factor = HeapData.Interpolate.Bifactor;
 
                 MeshPtrConst pResult;
 
                 if (pBase)
                 {
-                    // Just the first of the targets
-					if ( bifactor < UE_SMALL_NUMBER && bifactor > -UE_SMALL_NUMBER )
+					// No morph
+					if (Factor < UE_SMALL_NUMBER && Factor > -UE_SMALL_NUMBER )
                     {
-						pResult = pBase->Clone();
+						pResult = pBase;
                     }
-                    // Just the second of the targets (positive or negative)
-                    else if ( bifactor >= 1.0f - UE_SMALL_NUMBER || bifactor <= -1.0f + UE_SMALL_NUMBER)
-                    {
-                        check( max>0 );
-                        Ptr<const Mesh> pMorph = LoadMesh( FCacheAddress(Targets[max],item) );
-						pResult = MeshMorph(pBase.get(), pMorph.get(), bifactor);
-                    }
-                    // Mix two targets on the base
-                    else
+					// The Morph, partial or full
+                    else 
                     {
-                        Ptr<const Mesh> pMin = LoadMesh( FCacheAddress(Targets[min],item) );
-                        Ptr<const Mesh> pMax = LoadMesh( FCacheAddress(Targets[max],item) );
-                        pResult = MeshMorph2( pBase.get(), pMin.get(), pMax.get(), bifactor );
+                        Ptr<const Mesh> pMorph = LoadMesh( FCacheAddress(TargetAt,item) );
+						pResult = MeshMorph(pBase.get(), pMorph.get(), Factor);
                     }
-					// Missing branch. With the current system we can never get apply a negative pMin morph, only a negative pMax morph.
-					// This is due to not having a bifactor value that represents a negative pMin morph.
-					// "-0" should represent a negative pMin morph, but it has the contradiction that 0 also represents positive pMax morph.
-					// In other words, we have a discontinuity arround 0.
-					// With the current implementation, the only solutions to get a negative pMin morph is to apply a really small negative value (lim x -> -0) (for example: -0.00001).
                 }
                 StoreMesh( item, pResult );
 
diff --git a/Engine/Plugins/Experimental/Mutable/Source/MutableRuntime/Private/MuR/Operations.cpp b/Engine/Plugins/Experimental/Mutable/Source/MutableRuntime/Private/MuR/Operations.cpp
--- a/Engine/Plugins/Experimental/Mutable/Source/MutableRuntime/Private/MuR/Operations.cpp
+++ b/Engine/Plugins/Experimental/Mutable/Source/MutableRuntime/Private/MuR/Operations.cpp
@@ -845,7 +845,7 @@ namespace mu
 			break;
         }
 
-        case OP_TYPE::ME_MORPH2:
+        case OP_TYPE::ME_MORPH:
         {
 			const uint8_t* data = program.GetOpArgsPointer(at);
 
@@ -857,15 +857,9 @@ namespace mu
 			FMemory::Memcpy(&BaseAt, data, sizeof(OP::ADDRESS)); data += sizeof(OP::ADDRESS);
 			f(BaseAt);
 
-			uint8 NumTargets = 0;
-			FMemory::Memcpy(&NumTargets, data, sizeof(uint8)); data += sizeof(uint8);
-
-			for (uint8 T = 0; T < NumTargets; ++T)
-			{
-				OP::ADDRESS TargetAt = 0;
-				FMemory::Memcpy(&TargetAt, data, sizeof(OP::ADDRESS)); data += sizeof(OP::ADDRESS);
-				f(TargetAt);
-			}
+			OP::ADDRESS TargetAt = 0;
+			FMemory::Memcpy(&TargetAt, data, sizeof(OP::ADDRESS)); data += sizeof(OP::ADDRESS);
+			f(TargetAt);
 			break;
         }
 
diff --git a/Engine/Plugins/Experimental/Mutable/Source/MutableRuntime/Private/MuR/Operations.h b/Engine/Plugins/Experimental/Mutable/Source/MutableRuntime/Private/MuR/Operations.h
--- a/Engine/Plugins/Experimental/Mutable/Source/MutableRuntime/Private/MuR/Operations.h
+++ b/Engine/Plugins/Experimental/Mutable/Source/MutableRuntime/Private/MuR/Operations.h
@@ -233,9 +233,8 @@ namespace mu
         //! The meshes must have the same topology, etc.
         ME_DIFFERENCE,
 
-        //! Apply a linear combination of two morphs on a base. The 2 morphs are taken from a
-        //! sequence.
-        ME_MORPH2,
+        //! Apply a one morphs on a base. 
+        ME_MORPH,
 
         //! Merge a mesh to a mesh
         ME_MERGE,
diff --git a/Engine/Plugins/Experimental/Mutable/Source/MutableRuntime/Public/MuR/Model.h b/Engine/Plugins/Experimental/Mutable/Source/MutableRuntime/Public/MuR/Model.h
--- a/Engine/Plugins/Experimental/Mutable/Source/MutableRuntime/Public/MuR/Model.h
+++ b/Engine/Plugins/Experimental/Mutable/Source/MutableRuntime/Public/MuR/Model.h
@@ -13,7 +13,7 @@
 //! data structures. Compiled models are not necessarily compatible when the runtime is updated,
 //! so this version number can be used externally to verify this. It is not used internally, and
 //! serializing models from different versions than this runtime will probably result in a crash.
-#define MUTABLE_COMPILED_MODEL_CODE_VERSION		uint32( 67 )
+#define MUTABLE_COMPILED_MODEL_CODE_VERSION		uint32( 68 )
 #define MUTABLE_PARAMETERS_VERSION              uint32( 2 )
 
 
diff --git a/Engine/Plugins/Experimental/Mutable/Source/MutableTools/Private/MuT/ASTOpImageRasterMesh.cpp b/Engine/Plugins/Experimental/Mutable/Source/MutableTools/Private/MuT/ASTOpImageRasterMesh.cpp
--- a/Engine/Plugins/Experimental/Mutable/Source/MutableTools/Private/MuT/ASTOpImageRasterMesh.cpp
+++ b/Engine/Plugins/Experimental/Mutable/Source/MutableTools/Private/MuT/ASTOpImageRasterMesh.cpp
@@ -265,7 +265,7 @@ namespace mu
 			break;
 		}
 
-		case OP_TYPE::ME_MORPH2:
+		case OP_TYPE::ME_MORPH:
 		{
 			const ASTOpMeshMorph* typedSource = dynamic_cast<const ASTOpMeshMorph*>(sourceAt.get());
 			Ptr<ASTOpImageRasterMesh> rasterOp = mu::Clone<ASTOpImageRasterMesh>(this);
diff --git a/Engine/Plugins/Experimental/Mutable/Source/MutableTools/Private/MuT/ASTOpLayoutFromMesh.cpp b/Engine/Plugins/Experimental/Mutable/Source/MutableTools/Private/MuT/ASTOpLayoutFromMesh.cpp
--- a/Engine/Plugins/Experimental/Mutable/Source/MutableTools/Private/MuT/ASTOpLayoutFromMesh.cpp
+++ b/Engine/Plugins/Experimental/Mutable/Source/MutableTools/Private/MuT/ASTOpLayoutFromMesh.cpp
@@ -155,7 +155,7 @@ namespace mu
 				switch (at->GetOpType())
 				{
 
-				case OP_TYPE::ME_MORPH2:
+				case OP_TYPE::ME_MORPH:
 				{
 					// Sink, ignoring the op
 					const ASTOpMeshMorph* Typed = dynamic_cast<const ASTOpMeshMorph*>(at.get());
diff --git a/Engine/Plugins/Experimental/Mutable/Source/MutableTools/Private/MuT/ASTOpMeshFormat.cpp b/Engine/Plugins/Experimental/Mutable/Source/MutableTools/Private/MuT/ASTOpMeshFormat.cpp
--- a/Engine/Plugins/Experimental/Mutable/Source/MutableTools/Private/MuT/ASTOpMeshFormat.cpp
+++ b/Engine/Plugins/Experimental/Mutable/Source/MutableTools/Private/MuT/ASTOpMeshFormat.cpp
@@ -264,7 +264,7 @@ mu::Ptr<ASTOp> Sink_MeshFormatAST::Visit(const mu::Ptr<ASTOp>& at, const ASTOpMe
 		break;
 	}
 
-	case OP_TYPE::ME_MORPH2:
+	case OP_TYPE::ME_MORPH:
 	{
 		// Move the format down the base of the morph
 		Ptr<ASTOpMeshMorph> newOp = mu::Clone<ASTOpMeshMorph>(at);
@@ -280,18 +280,15 @@ mu::Ptr<ASTOp> Sink_MeshFormatAST::Visit(const mu::Ptr<ASTOp>& at, const ASTOpMe
 		motaop->SetValue(pTargetMorphFormat, false /* useDiskCache */);
 		auto targetMorphFormatAt = motaop;
 
-		for (int32 t=0; t<newOp->Targets.Num(); ++t)
+		if (newOp->Target)
 		{
-			if (newOp->Targets[t])
-			{
-				mu::Ptr<ASTOpMeshFormat> newFormat = mu::Clone<ASTOpMeshFormat>(currentFormatOp);
-				newFormat->Buffers =
-					OP::MeshFormatArgs::BT_VERTEX
-					| OP::MeshFormatArgs::BT_IGNORE_MISSING;
-				newFormat->Format = targetMorphFormatAt;
+			mu::Ptr<ASTOpMeshFormat> newFormat = mu::Clone<ASTOpMeshFormat>(currentFormatOp);
+			newFormat->Buffers =
+				OP::MeshFormatArgs::BT_VERTEX
+				| OP::MeshFormatArgs::BT_IGNORE_MISSING;
+			newFormat->Format = targetMorphFormatAt;
 
-				newOp->Targets[t] = Visit(newOp->Targets[t].child(), newFormat.get());
-			}
+			newOp->Target = Visit(newOp->Target.child(), newFormat.get());
 		}
 
 		newAt = newOp;
diff --git a/Engine/Plugins/Experimental/Mutable/Source/MutableTools/Private/MuT/ASTOpMeshMorph.cpp b/Engine/Plugins/Experimental/Mutable/Source/MutableTools/Private/MuT/ASTOpMeshMorph.cpp
--- a/Engine/Plugins/Experimental/Mutable/Source/MutableTools/Private/MuT/ASTOpMeshMorph.cpp
+++ b/Engine/Plugins/Experimental/Mutable/Source/MutableTools/Private/MuT/ASTOpMeshMorph.cpp
@@ -15,7 +15,7 @@ namespace mu
 
 	//---------------------------------------------------------------------------------------------
 	ASTOpMeshMorph::ASTOpMeshMorph()
-		: Factor(this), Base(this)
+		: Factor(this), Base(this), Target(this)
 	{
 	}
 
@@ -28,19 +28,12 @@ namespace mu
 	}
 
 
-	//---------------------------------------------------------------------------------------------
-	void ASTOpMeshMorph::AddTarget(const Ptr<ASTOp>& InTarget)
-	{
-		Targets.Add( ASTChild(this,InTarget));
-	}
-
-
 	//---------------------------------------------------------------------------------------------
 	bool ASTOpMeshMorph::IsEqual(const ASTOp& otherUntyped) const
 	{
 		if (const ASTOpMeshMorph* other = dynamic_cast<const ASTOpMeshMorph*>(&otherUntyped))
 		{
-			return Factor == other->Factor && Base == other->Base && Targets == other->Targets;
+			return Factor == other->Factor && Base == other->Base && Target == other->Target;
 		}
 		return false;
 	}
@@ -52,10 +45,7 @@ namespace mu
 		Ptr<ASTOpMeshMorph> n = new ASTOpMeshMorph();
 		n->Factor = mapChild(Factor.child());
 		n->Base = mapChild(Base.child());
-		for (const ASTChild& r : Targets)
-		{
-			n->Targets.Add(ASTChild(n,mapChild(r.child())));
-		}
+		n->Target = mapChild(Target.child());
 		return n;
 	}
 
@@ -65,10 +55,7 @@ namespace mu
 	{
 		f(Factor);
 		f(Base);
-		for (ASTChild& T : Targets)
-		{
-			f(T);
-		}
+		f(Target);
 	}
 
 
@@ -77,10 +64,7 @@ namespace mu
 	{
 		uint64 res = std::hash<ASTOp*>()(Factor.child().get());
 		hash_combine(res, Base.child().get());
-		for (const ASTChild& T : Targets)
-		{
-			hash_combine(res, T.child().get());
-		}
+		hash_combine(res, Target.child().get());
 		return res;
 	}
 
@@ -94,7 +78,7 @@ namespace mu
 			linkedAddress = (OP::ADDRESS)program.m_opAddress.Num();
 
 			program.m_opAddress.Add((uint32_t)program.m_byteCode.Num());
-			AppendCode(program.m_byteCode, OP_TYPE::ME_MORPH2);
+			AppendCode(program.m_byteCode, OP_TYPE::ME_MORPH);
 
 			OP::ADDRESS FactorAt = Factor ? Factor->linkedAddress : 0;
 			AppendCode(program.m_byteCode, FactorAt);
@@ -102,12 +86,8 @@ namespace mu
 			OP::ADDRESS BaseAt = Base ? Base->linkedAddress : 0;
 			AppendCode(program.m_byteCode, BaseAt);
 
-			AppendCode(program.m_byteCode, (uint8)Targets.Num());
-			for (const ASTChild& b : Targets)
-			{
-				OP::ADDRESS TargetAt = b ? b->linkedAddress : 0;
-				AppendCode(program.m_byteCode, TargetAt);
-			}
+			OP::ADDRESS TargetAt = Target ? Target->linkedAddress : 0;
+			AppendCode(program.m_byteCode, TargetAt);
 		}
 	}
 
diff --git a/Engine/Plugins/Experimental/Mutable/Source/MutableTools/Private/MuT/ASTOpMeshMorph.h b/Engine/Plugins/Experimental/Mutable/Source/MutableTools/Private/MuT/ASTOpMeshMorph.h
--- a/Engine/Plugins/Experimental/Mutable/Source/MutableTools/Private/MuT/ASTOpMeshMorph.h
+++ b/Engine/Plugins/Experimental/Mutable/Source/MutableTools/Private/MuT/ASTOpMeshMorph.h
@@ -25,8 +25,8 @@ namespace mu
 		//! Base mesh to morph.
 		ASTChild Base;
 
-		//! Targets to apply on the base depending on the factor
-		TArray<ASTChild> Targets;
+		//! Target to apply on the base depending on the factor
+		ASTChild Target;
 
 	public:
 
@@ -35,7 +35,7 @@ namespace mu
 		~ASTOpMeshMorph() override;
 
 		// ASTOp interface
-		OP_TYPE GetOpType() const override { return OP_TYPE::ME_MORPH2; }
+		OP_TYPE GetOpType() const override { return OP_TYPE::ME_MORPH; }
 		uint64 Hash() const override;
 		void ForEachChild(const TFunctionRef<void(ASTChild&)>) override;
 		bool IsEqual(const ASTOp& otherUntyped) const override;
@@ -43,8 +43,6 @@ namespace mu
 		void Link(FProgram& program, const FLinkerOptions*) override;
 		Ptr<ASTOp> OptimiseSink(const FModelOptimizationOptions&, FOptimizeSinkContext&) const override;
 
-		// Own interface
-		void AddTarget(const Ptr<ASTOp>&);
 	};
 
 }
diff --git a/Engine/Plugins/Experimental/Mutable/Source/MutableTools/Private/MuT/ASTOpMeshRemoveMask.cpp b/Engine/Plugins/Experimental/Mutable/Source/MutableTools/Private/MuT/ASTOpMeshRemoveMask.cpp
--- a/Engine/Plugins/Experimental/Mutable/Source/MutableTools/Private/MuT/ASTOpMeshRemoveMask.cpp
+++ b/Engine/Plugins/Experimental/Mutable/Source/MutableTools/Private/MuT/ASTOpMeshRemoveMask.cpp
@@ -168,7 +168,7 @@ namespace mu
 				switch (at->GetOpType())
 				{
 
-				case OP_TYPE::ME_MORPH2:
+				case OP_TYPE::ME_MORPH:
 				{
 					mu::Ptr<ASTOpMeshMorph> newOp = mu::Clone<ASTOpMeshMorph>(at);
 					newOp->Base = Visit(newOp->Base.child());
diff --git a/Engine/Plugins/Experimental/Mutable/Source/MutableTools/Private/MuT/CodeGenerator.cpp b/Engine/Plugins/Experimental/Mutable/Source/MutableTools/Private/MuT/CodeGenerator.cpp
--- a/Engine/Plugins/Experimental/Mutable/Source/MutableTools/Private/MuT/CodeGenerator.cpp
+++ b/Engine/Plugins/Experimental/Mutable/Source/MutableTools/Private/MuT/CodeGenerator.cpp
@@ -1004,21 +1004,6 @@ namespace mu
                     FMeshGenerationResult morphResult;
                     GenerateMesh(MorphTargetMeshOptions, morphResult, pMorph );
 
-					// BaseMorph generation through mesh diff
-					Ptr<ASTOpMeshDifference> diffBase;
-					{
-						// \TODO: Optimize by setting the identity morph constant instead.
-						Ptr<ASTOpMeshDifference> op = new ASTOpMeshDifference();
-						op->Base = meshResults.meshOp;
-						op->Target = meshResults.meshOp;
-
-						// Morphing tex coords here is not supported:
-						// Generating the homogoneous UVs is difficult since we don't have the base
-						// layout yet.                       
-						op->bIgnoreTextureCoords = true;
-						diffBase = op;
-					}
-
 					// Morph generation through mesh diff
 					Ptr<ASTOp> targetAd = morphResult.meshOp;
                     Ptr<ASTOpMeshDifference> diffAd;
@@ -1057,8 +1042,7 @@ namespace mu
 						op->Base = lastMeshOp;
 
 						// Targets
-						op->AddTarget(diffBase);
-						op->AddTarget(diffAd);
+						op->Target = diffAd;
                         morphAd = op;
                     }
 
diff --git a/Engine/Plugins/Experimental/Mutable/Source/MutableTools/Private/MuT/CodeGenerator_Mesh.cpp b/Engine/Plugins/Experimental/Mutable/Source/MutableTools/Private/MuT/CodeGenerator_Mesh.cpp
--- a/Engine/Plugins/Experimental/Mutable/Source/MutableTools/Private/MuT/CodeGenerator_Mesh.cpp
+++ b/Engine/Plugins/Experimental/Mutable/Source/MutableTools/Private/MuT/CodeGenerator_Mesh.cpp
@@ -765,9 +765,9 @@ class Node;
         Ptr<ASTOpMeshMorph> OpMorph = new ASTOpMeshMorph();
 
         // Factor
-        if ( node.m_pFactor )
+        if ( node.Factor )
         {
-            OpMorph->Factor = Generate( node.m_pFactor.get() );
+            OpMorph->Factor = Generate( node.Factor.get() );
         }
         else
         {
@@ -777,11 +777,11 @@ class Node;
 
         // Base
         FMeshGenerationResult BaseResult;
-        if ( node.m_pBase )
+        if ( node.Base )
         {
 			FMeshGenerationOptions BaseOptions = InOptions;
 			BaseOptions.bUniqueVertexIDs = true;
-            GenerateMesh(BaseOptions,BaseResult, node.m_pBase );
+            GenerateMesh(BaseOptions,BaseResult, node.Base );
             OpMorph->Base = BaseResult.meshOp;
         }
         else
@@ -791,7 +791,7 @@ class Node;
                                             ELMT_ERROR, node.m_errorContext );
         }        
 
-        for ( int32 t=0 ; t<node.m_morphs.Num(); ++t )
+		if (node.Morph)
         {
             FMeshGenerationResult TargetResult;
 			FMeshGenerationOptions TargetOptions = InOptions;
@@ -801,15 +801,15 @@ class Node;
 			// sure that we get the correct mesh when generating the target
 			TargetOptions.OverrideLayouts = BaseResult.GeneratedLayouts;
 			TargetOptions.ActiveTags.Empty();
-            GenerateMesh(TargetOptions, TargetResult, node.m_morphs[t]);
+            GenerateMesh(TargetOptions, TargetResult, node.Morph);
 
             // TODO: Make sure that the target is a mesh with the morph format
             Ptr<ASTOp> target = TargetResult.meshOp;
 
-            OpMorph->AddTarget( target );
+            OpMorph->Target = target;
         }
  
-        const bool bReshapeEnabled = node.m_reshapeSkeleton || node.m_reshapePhysicsVolumes;
+        const bool bReshapeEnabled = node.bReshapeSkeleton || node.bReshapePhysicsVolumes;
         
         Ptr<ASTOpMeshMorphReshape> OpMorphReshape;
         if ( bReshapeEnabled )
@@ -820,10 +820,10 @@ class Node;
 			// Setting bReshapeVertices to false the bind op will remove all mesh members except 
 			// PhysicsBodies and the Skeleton.
             OpBind->bReshapeVertices = false;
-		    OpBind->bReshapeSkeleton = node.m_reshapeSkeleton;
-		    OpBind->BonesToDeform = node.m_bonesToDeform;
-    	    OpBind->bReshapePhysicsVolumes = node.m_reshapePhysicsVolumes; 
-			OpBind->PhysicsToDeform = node.m_physicsToDeform;
+		    OpBind->bReshapeSkeleton = node.bReshapeSkeleton;
+		    OpBind->BonesToDeform = node.BonesToDeform;
+    	    OpBind->bReshapePhysicsVolumes = node.bReshapePhysicsVolumes; 
+			OpBind->PhysicsToDeform = node.PhysicsToDeform;
 			OpBind->BindingMethod = static_cast<uint32>(EShapeBindingMethod::ReshapeClosestProject);
             
 			OpBind->Mesh = BaseResult.meshOp;
diff --git a/Engine/Plugins/Experimental/Mutable/Source/MutableTools/Private/MuT/CodeOptimiser_Parameters.cpp b/Engine/Plugins/Experimental/Mutable/Source/MutableTools/Private/MuT/CodeOptimiser_Parameters.cpp
--- a/Engine/Plugins/Experimental/Mutable/Source/MutableTools/Private/MuT/CodeOptimiser_Parameters.cpp
+++ b/Engine/Plugins/Experimental/Mutable/Source/MutableTools/Private/MuT/CodeOptimiser_Parameters.cpp
@@ -1670,7 +1670,7 @@ namespace mu
 					return false;
 				}
 
-                case OP_TYPE::ME_MORPH2:
+                case OP_TYPE::ME_MORPH:
                 {
                     // Manually choose how to recurse this op
 					const ASTOpMeshMorph* pTyped = dynamic_cast<const ASTOpMeshMorph*>( at.get() );
@@ -1688,12 +1688,9 @@ namespace mu
 							pending.Add({ pTyped->Factor.m_child, state });
                         }
 
-                        for (int32 t=0;t<pTyped->Targets.Num(); ++t)
+                        if ( pTyped->Target )
                         {
-                            if ( pTyped->Targets[t] )
-                            {
-								pending.Add({ pTyped->Targets[t].m_child, state });
-                            }
+							pending.Add({ pTyped->Target.m_child, state });
                         }
                     }
 
diff --git a/Engine/Plugins/Experimental/Mutable/Source/MutableTools/Private/MuT/DataPacker.cpp b/Engine/Plugins/Experimental/Mutable/Source/MutableTools/Private/MuT/DataPacker.cpp
--- a/Engine/Plugins/Experimental/Mutable/Source/MutableTools/Private/MuT/DataPacker.cpp
+++ b/Engine/Plugins/Experimental/Mutable/Source/MutableTools/Private/MuT/DataPacker.cpp
@@ -573,7 +573,7 @@ namespace mu
                 break;
              }
 
-            case OP_TYPE::ME_MORPH2:
+            case OP_TYPE::ME_MORPH:
             {
                 recurse = false;
 
@@ -582,13 +582,7 @@ namespace mu
                 uint64_t newState = currentSemantics;
                 newState |= (UINT64_C(1)<<MBS_VERTEXINDEX);
                 RecurseWithState( op->Base.child(), newState );
-                for ( int32 o=0; o<op->Targets.Num(); ++o)
-                {
-                    if (op->Targets[o])
-                    {
-                        RecurseWithState( op->Targets[o].child(), newState );
-                    }
-                }
+                RecurseWithState( op->Target.child(), newState );
                 break;
              }
 
diff --git a/Engine/Plugins/Experimental/Mutable/Source/MutableTools/Private/MuT/NodeMeshMorph.cpp b/Engine/Plugins/Experimental/Mutable/Source/MutableTools/Private/MuT/NodeMeshMorph.cpp
--- a/Engine/Plugins/Experimental/Mutable/Source/MutableTools/Private/MuT/NodeMeshMorph.cpp
+++ b/Engine/Plugins/Experimental/Mutable/Source/MutableTools/Private/MuT/NodeMeshMorph.cpp
@@ -19,8 +19,7 @@ namespace mu
 	//---------------------------------------------------------------------------------------------
 	// Static initialisation
 	//---------------------------------------------------------------------------------------------
-	NODE_TYPE NodeMeshMorph::Private::s_type =
-			NODE_TYPE( "MeshMorph", NodeMesh::GetStaticType() );
+	NODE_TYPE NodeMeshMorph::Private::s_type = NODE_TYPE( "MeshMorph", NodeMesh::GetStaticType() );
 
 
 	//---------------------------------------------------------------------------------------------
@@ -35,7 +34,7 @@ namespace mu
 	//---------------------------------------------------------------------------------------------
 	int NodeMeshMorph::GetInputCount() const
 	{
-		return 2 + m_pD->m_morphs.Num();
+		return 3;
 	}
 
 
@@ -49,16 +48,19 @@ namespace mu
 		switch (i)
 		{
 		case 0:
-			pResult = m_pD->m_pFactor.get();
+			pResult = m_pD->Factor.get();
 			break;
 
 		case 1:
-			pResult = m_pD->m_pBase.get();
+			pResult = m_pD->Base.get();
 			break;
 
-		default:
-			pResult = m_pD->m_morphs[i-2].get();
+		case 2:
+			pResult = m_pD->Morph.get();
 			break;
+
+		default:
+			check(false);
 		}
 
 		return pResult;
@@ -66,22 +68,25 @@ namespace mu
 
 
 	//---------------------------------------------------------------------------------------------
-	void NodeMeshMorph::SetInputNode( int i, NodePtr pNode )
+	void NodeMeshMorph::SetInputNode( int i, Ptr<Node> pNode )
 	{
 		check( i>=0 && i<GetInputCount() );
 
 		switch (i)
 		{
 		case 0:
-			m_pD->m_pFactor = dynamic_cast<NodeScalar*>(pNode.get());
+			m_pD->Factor = dynamic_cast<NodeScalar*>(pNode.get());
 			break;
 
 		case 1:
-			m_pD->m_pBase = dynamic_cast<NodeMesh*>(pNode.get());
+			m_pD->Base = dynamic_cast<NodeMesh*>(pNode.get());
+			break;
+
+		case 2:
+			m_pD->Morph = dynamic_cast<NodeMesh*>(pNode.get());
 			break;
 
 		default:
-			m_pD->m_morphs[i-2] = dynamic_cast<NodeMesh*>(pNode.get());
 			break;
 		}
 	}
@@ -90,82 +95,67 @@ namespace mu
 	//---------------------------------------------------------------------------------------------
 	// Own Interface
 	//---------------------------------------------------------------------------------------------
-	NodeScalarPtr NodeMeshMorph::GetFactor() const
-	{
-		return m_pD->m_pFactor.get();
-	}
-
-
-	//---------------------------------------------------------------------------------------------
-	void NodeMeshMorph::SetFactor( NodeScalarPtr pNode )
+	Ptr<NodeScalar> NodeMeshMorph::GetFactor() const
 	{
-		m_pD->m_pFactor = pNode;
+		return m_pD->Factor.get();
 	}
 
 
 	//---------------------------------------------------------------------------------------------
-	NodeMeshPtr NodeMeshMorph::GetBase() const
+	void NodeMeshMorph::SetFactor( Ptr<NodeScalar> pNode )
 	{
-		return m_pD->m_pBase.get();
+		m_pD->Factor = pNode;
 	}
 
 
 	//---------------------------------------------------------------------------------------------
-	void NodeMeshMorph::SetBase( NodeMeshPtr pNode )
+	Ptr<NodeMesh> NodeMeshMorph::GetBase() const
 	{
-		m_pD->m_pBase = pNode;
+		return m_pD->Base.get();
 	}
 
 
 	//---------------------------------------------------------------------------------------------
-	NodeMeshPtr NodeMeshMorph::GetMorph( int t ) const
+	void NodeMeshMorph::SetBase( Ptr<NodeMesh> pNode )
 	{
-		check( t>=0 && t<m_pD->m_morphs.Num() );
-
-		return m_pD->m_morphs[t].get();
+		m_pD->Base = pNode;
 	}
 
 
 	//---------------------------------------------------------------------------------------------
-	void NodeMeshMorph::SetMorph( int t, NodeMeshPtr pNode )
+	NodeMeshPtr NodeMeshMorph::GetMorph() const
 	{
-		check( t>=0 && t<m_pD->m_morphs.Num() );
-		m_pD->m_morphs[t] = pNode;
+		return m_pD->Morph.get();
 	}
 
-	//---------------------------------------------------------------------------------------------
-	void NodeMeshMorph::SetMorphCount( int t )
-	{
-		m_pD->m_morphs.SetNum(t);
-	}
 
 	//---------------------------------------------------------------------------------------------
-	int NodeMeshMorph::GetMorphCount() const
+	void NodeMeshMorph::SetMorph( Ptr<NodeMesh> pNode )
 	{
-		return m_pD->m_morphs.Num();
+		m_pD->Morph = pNode;
 	}
 
 	//---------------------------------------------------------------------------------------------
 	void NodeMeshMorph::SetReshapeSkeleton(bool bEnable)
 	{
-		m_pD->m_reshapeSkeleton = bEnable;
+		m_pD->bReshapeSkeleton = bEnable;
 	}	
 
 	//---------------------------------------------------------------------------------------------
 	void NodeMeshMorph::AddBoneToDeform(const char* BoneName)
 	{
-		m_pD->m_bonesToDeform.Emplace(BoneName);
+		m_pD->BonesToDeform.Emplace(BoneName);
 	}
 
 	void NodeMeshMorph::AddPhysicsBodyToDeform(const char* BoneName)
 	{
-		m_pD->m_physicsToDeform.Emplace(BoneName);
+		m_pD->PhysicsToDeform.Emplace(BoneName);
 	}
 
 	//---------------------------------------------------------------------------------------------
 	void NodeMeshMorph::SetReshapePhysicsVolumes(bool bEnable)
 	{
-		m_pD->m_reshapePhysicsVolumes = bEnable;
+		m_pD->bReshapePhysicsVolumes = bEnable;
 	}
 
 	//---------------------------------------------------------------------------------------------
@@ -173,10 +163,10 @@ namespace mu
 	{
 		NodeLayoutPtr pResult;
 
-		if ( m_pBase )
+		if ( Base )
 		{
 			NodeMesh::Private* pPrivate =
-					dynamic_cast<NodeMesh::Private*>( m_pBase->GetBasePrivate() );
+					dynamic_cast<NodeMesh::Private*>( Base->GetBasePrivate() );
 
 			pResult = pPrivate->GetLayout( index );
 		}
diff --git a/Engine/Plugins/Experimental/Mutable/Source/MutableTools/Private/MuT/NodeMeshMorphPrivate.h b/Engine/Plugins/Experimental/Mutable/Source/MutableTools/Private/MuT/NodeMeshMorphPrivate.h
--- a/Engine/Plugins/Experimental/Mutable/Source/MutableTools/Private/MuT/NodeMeshMorphPrivate.h
+++ b/Engine/Plugins/Experimental/Mutable/Source/MutableTools/Private/MuT/NodeMeshMorphPrivate.h
@@ -21,35 +21,30 @@ namespace mu
 
 		static NODE_TYPE s_type;
 
-		NodeScalarPtr m_pFactor;
-		NodeMeshPtr m_pBase;
-		TArray<NodeMeshPtr> m_morphs;
-        bool m_vertexIndicesAreRelative_Deprecated = false;
+		Ptr<NodeScalar> Factor;
+		Ptr<NodeMesh> Base;
+		Ptr<NodeMesh> Morph;
 		
-		bool m_reshapeSkeleton = false;
-		bool m_reshapePhysicsVolumes = false;
+		bool bReshapeSkeleton = false;
+		bool bReshapePhysicsVolumes = false;
 		
-		TArray<string> m_bonesToDeform;
-		TArray<string> m_physicsToDeform;
+		TArray<string> BonesToDeform;
+		TArray<string> PhysicsToDeform;
 
         //!
 		void Serialise( OutputArchive& arch ) const
 		{
-            uint32_t ver = 5;
+            uint32_t ver = 6;
 			arch << ver;
 
-			arch << m_pFactor;
-			arch << m_pBase;
-			arch << m_morphs;
-            arch << m_vertexIndicesAreRelative_Deprecated;
-
-			arch << m_reshapeSkeleton;
-			arch << m_reshapePhysicsVolumes;
-			arch << m_bonesToDeform;
-
-			
-			arch << m_physicsToDeform;
+			arch << Factor;
+			arch << Base;
+			arch << Morph;
 
+			arch << bReshapeSkeleton;
+			arch << bReshapePhysicsVolumes;
+			arch << BonesToDeform;
+			arch << PhysicsToDeform;
         }
 
 		//!
@@ -57,30 +52,50 @@ namespace mu
 		{
             uint32_t ver;
 			arch >> ver;
-			check(ver <= 5);
+			check(ver <= 6);
+
+			arch >> Factor;
+			arch >> Base;
+
+			if (ver < 6)
+			{
+				TArray<Ptr<NodeMesh>> Morphs;
+				arch >> Morphs;
+
+				// The result will not be the same.
+				if (!Morphs.IsEmpty())
+				{
+					Morph = Morphs.Last();
+				}
+			}
+			else
+			{
+				arch >> Morph;
+			}
 
-			arch >> m_pFactor;
-			arch >> m_pBase;
-			arch >> m_morphs;
-            arch >> m_vertexIndicesAreRelative_Deprecated;
+			if (ver < 6)
+			{
+				bool bDeprecated;
+				arch >> bDeprecated;
+			}
 	
 			if (ver >= 2)
 			{
-				arch >> m_reshapeSkeleton;
-				arch >> m_reshapePhysicsVolumes;
+				arch >> bReshapeSkeleton;
+				arch >> bReshapePhysicsVolumes;
 	
 				// This repetition is needed, there was a bug where m_reshapePhysicsVolumes was serialized twice in ver 2.
 				if (ver == 2) 
 				{
-					arch >> m_reshapePhysicsVolumes;
+					arch >> bReshapePhysicsVolumes;
 				}
-				arch >> m_bonesToDeform;			
+				arch >> BonesToDeform;			
 			}
 			else
 			{
-				m_reshapeSkeleton = false;
-				m_reshapePhysicsVolumes = false;
-				m_bonesToDeform.Empty();
+				bReshapeSkeleton = false;
+				bReshapePhysicsVolumes = false;
+				BonesToDeform.Empty();
 			}
 
 			if (ver == 3)
@@ -97,17 +112,17 @@ namespace mu
 
 			if (ver >= 3)
 			{
-				arch >> m_physicsToDeform;
+				arch >> PhysicsToDeform;
 			}
 			else
 			{
-				m_physicsToDeform.Empty();
+				PhysicsToDeform.Empty();
 			}
 
 		}
 
 		// NodeMesh::Private interface
-        NodeLayoutPtr GetLayout( int index ) const override;
+        Ptr<NodeLayout> GetLayout( int index ) const override;
 
 	};
 
diff --git a/Engine/Plugins/Experimental/Mutable/Source/MutableTools/Public/MuT/NodeMeshMorph.h b/Engine/Plugins/Experimental/Mutable/Source/MutableTools/Public/MuT/NodeMeshMorph.h
--- a/Engine/Plugins/Experimental/Mutable/Source/MutableTools/Public/MuT/NodeMeshMorph.h
+++ b/Engine/Plugins/Experimental/Mutable/Source/MutableTools/Public/MuT/NodeMeshMorph.h
@@ -13,16 +13,9 @@ namespace mu
 
 	// Forward definitions
 	class NodeScalar;
-	typedef Ptr<NodeScalar> NodeScalarPtr;
-	typedef Ptr<const NodeScalar> NodeScalarPtrConst;
 
-	class NodeMeshMorph;
-	typedef Ptr<NodeMeshMorph> NodeMeshMorphPtr;
-	typedef Ptr<const NodeMeshMorph> NodeMeshMorphPtrConst;
 
-
-	//! Node that morphs a base mesh with one or two weighted targets from a sequence.
-	//! \ingroup model
+	/** Node that morphs a base mesh with one weighted target. */
 	class MUTABLETOOLS_API NodeMeshMorph : public NodeMesh
 	{
 	public:
@@ -35,47 +28,35 @@ namespace mu
 
 		void SerialiseWrapper(OutputArchive& arch) const override;
 		static void Serialise( const NodeMeshMorph* pNode, OutputArchive& arch );
-		static NodeMeshMorphPtr StaticUnserialise( InputArchive& arch );
+		static Ptr<NodeMeshMorph> StaticUnserialise( InputArchive& arch );
 
 
 		//-----------------------------------------------------------------------------------------
 		// Node Interface
-		//-----------------------------------------------------------------------------------------
-
-		
+		//-----------------------------------------------------------------------------------------		
 
 		const NODE_TYPE* GetType() const override;
 		static const NODE_TYPE* GetStaticType();
 
 		int GetInputCount() const override;
 		Node* GetInputNode( int i ) const override;
-		void SetInputNode( int i, NodePtr pNode ) override;
+		void SetInputNode( int i, Ptr<Node> pNode ) override;
 
 		//-----------------------------------------------------------------------------------------
 		// Own Interface
 		//-----------------------------------------------------------------------------------------
 
-		//! Get the node generating the weight used to select (and combine) the morphs to apply.
-		NodeScalarPtr GetFactor() const;
-		void SetFactor( NodeScalarPtr );
-
-		//! Get the nodes generating the base mesh to be morphed
-		NodeMeshPtr GetBase() const;
-		void SetBase( NodeMeshPtr );
-
-		//! Set the number of target meshes. It will keep the currently set targets and initialise
-		//! the new ones as 0.
-		void SetMorphCount( int );
-		int GetMorphCount() const;
+		/** Get the node generating the weight used to select(and combine) the morphs to apply. */
+		Ptr<NodeScalar> GetFactor() const;
+		void SetFactor( Ptr<NodeScalar> );
 
-		//! Get the nodes generating the several morph targets
-		NodeMeshPtr GetMorph( int t ) const;
-		void SetMorph( int t, NodeMeshPtr );
+		/** Get the nodes generating the base mesh to be morphed. */
+		Ptr<NodeMesh> GetBase() const;
+		void SetBase( Ptr<NodeMesh> );
 
-		//! Set the domain of the vertex indices in the morph connections. By default they are 
-		//! absolute as generated by the MakeMeshMorph node. But if they are provided externally, 
-		//! they may be relative to the main mesh connected in the Base node.
-        //void SetMorphIndicesAreRelative( bool relative );
+		/** Get the nodes generating the several morph targets. */
+		Ptr<NodeMesh> GetMorph() const;
+		void SetMorph( Ptr<NodeMesh> );
 
 		/** Also deform the mesh skeleton. Disabled by default. */
 		void SetReshapeSkeleton(bool);	
