[Mutable] Better Load/Save descriptor documentation.

- Changed API to make clear when bUseCompactDescriptor is true, and its implications.
- Improved documentation.

#rnx

diff --git a/Engine/Plugins/Experimental/Mutable/Source/CustomizableObject/Classes/MuCO/CustomizableObjectInstance.h b/Engine/Plugins/Experimental/Mutable/Source/CustomizableObject/Classes/MuCO/CustomizableObjectInstance.h
--- a/Engine/Plugins/Experimental/Mutable/Source/CustomizableObject/Classes/MuCO/CustomizableObjectInstance.h
+++ b/Engine/Plugins/Experimental/Mutable/Source/CustomizableObject/Classes/MuCO/CustomizableObjectInstance.h
@@ -584,8 +584,10 @@ public:
 	bool AnimInstanceNeedsFixup(TSubclassOf<UAnimInstance> AnimInstance) const;
 	void AnimInstanceFixup(UAnimInstance* AnimInstance) const;
 
-	/** Serializes/Deserializes all the customization parameters to/from a descriptor, ready to be sent/read, works like a typical UE4 Two-Way Save System Function  */
-	void SaveDescriptor(FArchive &CustomizableObjectDescriptor);
+	/** See FCustomizableObjectInstanceDescriptor::SaveDescriptor. */
+	void SaveDescriptor(FArchive &CustomizableObjectDescriptor,bool bUseCompactDescriptor);
+
+	/** See FCustomizableObjectInstanceDescriptor::LoadDescriptor. */
 	void LoadDescriptor(FArchive &CustomizableObjectDescriptor);
 
 	// Enable physics asset replacement so that generated skeletal meshes have the merged physics assets of their skeletal mesh parts and reference mesh
diff --git a/Engine/Plugins/Experimental/Mutable/Source/CustomizableObject/Classes/MuCO/CustomizableObjectSystem.h b/Engine/Plugins/Experimental/Mutable/Source/CustomizableObject/Classes/MuCO/CustomizableObjectSystem.h
--- a/Engine/Plugins/Experimental/Mutable/Source/CustomizableObject/Classes/MuCO/CustomizableObjectSystem.h
+++ b/Engine/Plugins/Experimental/Mutable/Source/CustomizableObject/Classes/MuCO/CustomizableObjectSystem.h
@@ -249,9 +249,6 @@ public:
 	FEditorCompileSettings EditorSettings;
 
 #endif
-
-	bool IsCompactSerializationEnabled() const;
-
 	bool IsSupport16BitBoneIndexEnabled() const;
 
 	bool IsProgressiveMipStreamingEnabled() const;
diff --git a/Engine/Plugins/Experimental/Mutable/Source/CustomizableObject/Private/MuCO/CustomizableObjectInstance.cpp b/Engine/Plugins/Experimental/Mutable/Source/CustomizableObject/Private/MuCO/CustomizableObjectInstance.cpp
--- a/Engine/Plugins/Experimental/Mutable/Source/CustomizableObject/Private/MuCO/CustomizableObjectInstance.cpp
+++ b/Engine/Plugins/Experimental/Mutable/Source/CustomizableObject/Private/MuCO/CustomizableObjectInstance.cpp
@@ -2456,9 +2456,9 @@ void UCustomizableObjectInstance::MultilayerProjectorUpdateVirtualLayer(const FN
 }
 
 
-void UCustomizableObjectInstance::SaveDescriptor(FArchive &Ar)
+void UCustomizableObjectInstance::SaveDescriptor(FArchive &Ar, bool bUseCompactDescriptor)
 {
-	Descriptor.SaveDescriptor(Ar);
+	Descriptor.SaveDescriptor(Ar, bUseCompactDescriptor);
 }
 
 
diff --git a/Engine/Plugins/Experimental/Mutable/Source/CustomizableObject/Private/MuCO/CustomizableObjectInstanceDescriptor.cpp b/Engine/Plugins/Experimental/Mutable/Source/CustomizableObject/Private/MuCO/CustomizableObjectInstanceDescriptor.cpp
--- a/Engine/Plugins/Experimental/Mutable/Source/CustomizableObject/Private/MuCO/CustomizableObjectInstanceDescriptor.cpp
+++ b/Engine/Plugins/Experimental/Mutable/Source/CustomizableObject/Private/MuCO/CustomizableObjectInstanceDescriptor.cpp
@@ -39,16 +39,10 @@ FCustomizableObjectInstanceDescriptor::FCustomizableObjectInstanceDescriptor(UCu
 }
 
 
-void FCustomizableObjectInstanceDescriptor::SaveDescriptor(FArchive& Ar)
+void FCustomizableObjectInstanceDescriptor::SaveDescriptor(FArchive& Ar, bool bUseCompactDescriptor)
 {
 	check(CustomizableObject);
 
-	// This is a non-portable but very compact descriptor if bUseCompactDescriptor is true. It assumes the compiled objects are the same
-	// on both ends of the serialisation. That's why we iterate the parameters in the actual compiled
-	// model, instead of the arrays in this class.
-
-	bool bUseCompactDescriptor = UCustomizableObjectSystem::GetInstance()->IsCompactSerializationEnabled();
-
 	Ar << bUseCompactDescriptor;
 
 	// Not sure if this is needed, but it is small.
@@ -119,7 +113,7 @@ void FCustomizableObjectInstanceDescriptor::SaveDescriptor(FArchive& Ar)
 
 			int32 IntParameterIndex = FindIntParameterNameIndex(Name);
 			if (IntParameterIndex >= 0)
-				{					
+			{					
 				const FCustomizableObjectIntParameterValue & P = IntParameters[IntParameterIndex];
 					Value = CustomizableObject->FindIntParameterValue(ModelParameterIndex,P.ParameterValueName);
 
@@ -231,12 +225,7 @@ void FCustomizableObjectInstanceDescriptor::LoadDescriptor(FArchive& Ar)
 {
 	check(CustomizableObject);
 
-	// This is a non-portable but very compact descriptor if bUseCompactDescriptor is true. It assumes the compiled objects are the same
-	// on both ends of the serialisation. That's why we iterate the parameters in the actual compiled
-	// model, instead of the arrays in this class.
-
-	bool bUseCompactDescriptor = UCustomizableObjectSystem::GetInstance()->IsCompactSerializationEnabled();
-
+	bool bUseCompactDescriptor;
 	Ar << bUseCompactDescriptor;
 
 	// Not sure if this is needed, but it is small.
@@ -249,7 +238,7 @@ void FCustomizableObjectInstanceDescriptor::LoadDescriptor(FArchive& Ar)
 		Ar << ModelParameterCount;
 	}
 
-	for (int32 i = 0; i < ModelParameterCount; ++i)
+	for (int32 ParameterIndex = 0; ParameterIndex < ModelParameterCount; ++ParameterIndex)
 	{
 		FString Name;
 		EMutableParameterType Type;
@@ -257,7 +246,7 @@ void FCustomizableObjectInstanceDescriptor::LoadDescriptor(FArchive& Ar)
 
 		if (bUseCompactDescriptor)
 		{
-			ModelParameterIndex = i;
+			ModelParameterIndex = ParameterIndex;
 			Name = CustomizableObject->GetParameterName(ModelParameterIndex);
 			Type = CustomizableObject->GetParameterType(ModelParameterIndex);
 		}
diff --git a/Engine/Plugins/Experimental/Mutable/Source/CustomizableObject/Private/MuCO/CustomizableObjectSystem.cpp b/Engine/Plugins/Experimental/Mutable/Source/CustomizableObject/Private/MuCO/CustomizableObjectSystem.cpp
--- a/Engine/Plugins/Experimental/Mutable/Source/CustomizableObject/Private/MuCO/CustomizableObjectSystem.cpp
+++ b/Engine/Plugins/Experimental/Mutable/Source/CustomizableObject/Private/MuCO/CustomizableObjectSystem.cpp
@@ -380,11 +380,6 @@ void UCustomizableObjectSystem::InitSystem()
 	}
 #endif
 
-	if (FParse::Param(FCommandLine::Get(), TEXT("MutablePortableSerialization")))
-	{
-		Private->bCompactSerialization = false;
-	}
-
 	DefaultInstanceLODManagement = NewObject<UCustomizableInstanceLODManagement>();
 	check(DefaultInstanceLODManagement != nullptr);
 	CurrentInstanceLODManagement = DefaultInstanceLODManagement;
@@ -2854,13 +2849,6 @@ int32 UCustomizableObjectSystem::GetAverageBuildTime() const
 }
 
 
-bool UCustomizableObjectSystem::IsCompactSerializationEnabled() const
-{
-	check(Private != nullptr);
-	return Private->bCompactSerialization;
-}
-
-
 bool UCustomizableObjectSystem::IsSupport16BitBoneIndexEnabled() const
 {
 	check(Private != nullptr);
diff --git a/Engine/Plugins/Experimental/Mutable/Source/CustomizableObject/Private/MuCO/CustomizableObjectSystemPrivate.h b/Engine/Plugins/Experimental/Mutable/Source/CustomizableObject/Private/MuCO/CustomizableObjectSystemPrivate.h
--- a/Engine/Plugins/Experimental/Mutable/Source/CustomizableObject/Private/MuCO/CustomizableObjectSystemPrivate.h
+++ b/Engine/Plugins/Experimental/Mutable/Source/CustomizableObject/Private/MuCO/CustomizableObjectSystemPrivate.h
@@ -686,8 +686,6 @@ public:
 
 	int32 GetCountAllocatedSkeletalMesh() { return CountAllocatedSkeletalMesh; }
 
-	bool bCompactSerialization = true;
-
 	bool bReplaceDiscardedWithReferenceMesh = false;
 	bool bReleaseTexturesImmediately = false;
 
diff --git a/Engine/Plugins/Experimental/Mutable/Source/CustomizableObject/Public/MuCO/CustomizableObjectInstanceDescriptor.h b/Engine/Plugins/Experimental/Mutable/Source/CustomizableObject/Public/MuCO/CustomizableObjectInstanceDescriptor.h
--- a/Engine/Plugins/Experimental/Mutable/Source/CustomizableObject/Public/MuCO/CustomizableObjectInstanceDescriptor.h
+++ b/Engine/Plugins/Experimental/Mutable/Source/CustomizableObject/Public/MuCO/CustomizableObjectInstanceDescriptor.h
@@ -39,10 +39,20 @@ struct CUSTOMIZABLEOBJECT_API FCustomizableObjectInstanceDescriptor
 	
 	explicit FCustomizableObjectInstanceDescriptor(UCustomizableObject& Object);
 
-	/** Serialize this object. Does not support Multilayer Projectors! */
-	void SaveDescriptor(FArchive &Ar);
+	/** Serialize this object. 
+	 *
+	 * Backwards compatibility is not guaranteed.
+ 	 * Multilayer Projectors not supported.
+	 *
+  	 * @param bUseCompactDescriptor If true it assumes the compiled objects are the same on both ends of the serialisation */
+	void SaveDescriptor(FArchive &Ar, bool bUseCompactDescriptor);
 
 	/** Deserialize this object. Does not support Multilayer Projectors! */
+
+	/** Deserialize this object.
+     *
+	 * Backwards compatibility is not guaranteed.
+	 * Multilayer Projectors not supported */
 	void LoadDescriptor(FArchive &Ar);
 
 	UCustomizableObject* GetCustomizableObject() const;
diff --git a/Engine/Plugins/Experimental/Mutable/Source/CustomizableObjectEditor/Private/MuCOE/SCustomizableInstanceProperties.cpp b/Engine/Plugins/Experimental/Mutable/Source/CustomizableObjectEditor/Private/MuCOE/SCustomizableInstanceProperties.cpp
--- a/Engine/Plugins/Experimental/Mutable/Source/CustomizableObjectEditor/Private/MuCOE/SCustomizableInstanceProperties.cpp
+++ b/Engine/Plugins/Experimental/Mutable/Source/CustomizableObjectEditor/Private/MuCOE/SCustomizableInstanceProperties.cpp
@@ -1400,7 +1400,7 @@ FReply SCustomizableInstanceProperties::OnCopyAllParameters()
 
 	UE_LOG(LogMutable, Log, TEXT("Parameters have been copied."));
 
-	CustomInstance->SaveDescriptor(*ToBinary);
+	CustomInstance->SaveDescriptor(*ToBinary, true);
 
 	// FString StringifyedData = FString::FromHexBlob(ToBinary->GetData(), ToBinary->Num());
 	// FPlatformApplicationMisc::ClipboardCopy(StringifyedData.GetCharArray().GetData());
