[mutable] Reenabled code to remove UVs that do not belong to any layout block of an extend material.
+Added null checker that avoids the ExtractLayoutBlock operation of a null mesh.

#rnx

diff --git a/Engine/Plugins/Experimental/Mutable/Source/CustomizableObjectEditor/Private/MuCOE/GenerateMutableSource/GenerateMutableSourceSurface.cpp b/Engine/Plugins/Experimental/Mutable/Source/CustomizableObjectEditor/Private/MuCOE/GenerateMutableSource/GenerateMutableSourceSurface.cpp
--- a/Engine/Plugins/Experimental/Mutable/Source/CustomizableObjectEditor/Private/MuCOE/GenerateMutableSource/GenerateMutableSourceSurface.cpp
+++ b/Engine/Plugins/Experimental/Mutable/Source/CustomizableObjectEditor/Private/MuCOE/GenerateMutableSource/GenerateMutableSourceSurface.cpp
@@ -244,7 +244,7 @@ void UpdateSharedSurfaceId(FMutableGraphGenerationContext& GenerationContext, UC
 	}
 }
 
-
+UE_DISABLE_OPTIMIZATION
 mu::NodeSurfacePtr GenerateMutableSourceSurface(const UEdGraphPin * Pin, FMutableGraphGenerationContext & GenerationContext, FMutableGraphSurfaceGenerationData& SurfaceData)
 {
 	check(Pin)
@@ -1100,29 +1100,29 @@ mu::NodeSurfacePtr GenerateMutableSourceSurface(const UEdGraphPin * Pin, FMutabl
 			{
 				mu::NodeMeshPtr MeshPtr = AddMeshNode;
 
-				//const TArray<UCustomizableObjectLayout*> Layouts = TypedNodeExt->GetLayouts();
-				//
-				//if (Layouts.Num())
-				//{
-				//	//TODO: Implement support for multiple UV channels (e.g. Add warning for vertices which have a block in a layout but not in the other)
-				//	mu::NodeMeshFragmentPtr MeshFrag = new mu::NodeMeshFragment();
-				//
-				//	MeshFrag->SetMesh(MeshPtr);
-				//	MeshFrag->SetLayoutOrGroup(0);
-				//	MeshFrag->SetFragmentType(mu::NodeMeshFragment::FT_LAYOUT_BLOCKS);
-				//	MeshFrag->SetBlockCount(Layouts[0]->Blocks.Num());
-				//
-				//	for (int i = 0; i < Layouts[0]->Blocks.Num(); ++i)
-				//	{
-				//		MeshFrag->SetBlock(i, i);
-				//	}
-				//
-				//	MeshPtr = MeshFrag;
-				//}
-				//else
-				//{
-				//	GenerationContext.Compiler->CompilerLog(LOCTEXT("ExtendMaterialLayoutMissing","Skeletal Mesh without Layout Node linked to an Extend Material. A 4x4 layout will be added as default layout."), Node);
-				//}
+				const TArray<UCustomizableObjectLayout*> Layouts = TypedNodeExt->GetLayouts();
+				
+				if (Layouts.Num())
+				{
+					//TODO: Implement support for multiple UV channels (e.g. Add warning for vertices which have a block in a layout but not in the other)
+					mu::NodeMeshFragmentPtr MeshFrag = new mu::NodeMeshFragment();
+				
+					MeshFrag->SetMesh(MeshPtr);
+					MeshFrag->SetLayoutOrGroup(0);
+					MeshFrag->SetFragmentType(mu::NodeMeshFragment::FT_LAYOUT_BLOCKS);
+					MeshFrag->SetBlockCount(Layouts[0]->Blocks.Num());
+				
+					for (int i = 0; i < Layouts[0]->Blocks.Num(); ++i)
+					{
+						MeshFrag->SetBlock(i, i);
+					}
+				
+					MeshPtr = MeshFrag;
+				}
+				else
+				{
+					GenerationContext.Compiler->CompilerLog(LOCTEXT("ExtendMaterialLayoutMissing","Skeletal Mesh without Layout Node linked to an Extend Material. A 4x4 layout will be added as default layout."), Node);
+				}
 
 				mu::NodeMeshFormatPtr MeshFormat = new mu::NodeMeshFormat();
 				SetSurfaceFormat( GenerationContext,
@@ -1624,6 +1624,6 @@ mu::NodeSurfacePtr GenerateMutableSourceSurface(const UEdGraphPin * Pin, FMutabl
 
 	return Result;
 }
-
+UE_ENABLE_OPTIMIZATION
 #undef LOCTEXT_NAMESPACE
 
diff --git a/Engine/Plugins/Experimental/Mutable/Source/MutableRuntime/Private/MuR/CodeRunner.cpp b/Engine/Plugins/Experimental/Mutable/Source/MutableRuntime/Private/MuR/CodeRunner.cpp
--- a/Engine/Plugins/Experimental/Mutable/Source/MutableRuntime/Private/MuR/CodeRunner.cpp
+++ b/Engine/Plugins/Experimental/Mutable/Source/MutableRuntime/Private/MuR/CodeRunner.cpp
@@ -1589,9 +1589,9 @@ namespace mu
                 {
                         AddOp( FScheduledOp( item.At, item, 1),
                                FScheduledOp( source, item) );
-                    }
-                    else
-                    {
+                }
+                else
+                {
                     StoreMesh( item, nullptr );
                 }
                 break;
@@ -1601,18 +1601,22 @@ namespace mu
 				MUTABLE_CPUPROFILER_SCOPE(ME_EXTRACTLAYOUTBLOCK_1)
 
                 Ptr<const Mesh> Source = LoadMesh( FCacheAddress(source,item) );
+				MeshPtr pResult;
 
-                // Access with memcpy necessary for unaligned arm issues.
-                uint32 blocks[1024];
-				FMemory::Memcpy( blocks, data, sizeof(uint32)*FMath::Min(1024,int(blockCount)) );
+				if (Source)
+				{
+					// Access with memcpy necessary for unaligned arm issues.
+					uint32 blocks[1024];
+					FMemory::Memcpy(blocks, data, sizeof(uint32)* FMath::Min(1024, int(blockCount)));
+
+					pResult = MeshExtractLayoutBlock(Source.get(),
+						layout,
+						blockCount,
+						blocks);
+				}
 
-                MeshPtr pResult;
-                pResult = MeshExtractLayoutBlock( Source.get(),
-                                                  layout,
-                                                  blockCount,
-                                                  blocks );
+				StoreMesh(item, pResult);
 
-                StoreMesh( item, pResult );
                 break;
             }
 
