[mutable] Removed UE_DISABLE_OPTIMIZATION

#rnx

diff --git a/Engine/Plugins/Experimental/Mutable/Source/CustomizableObjectEditor/Private/MuCOE/GenerateMutableSource/GenerateMutableSourceSurface.cpp b/Engine/Plugins/Experimental/Mutable/Source/CustomizableObjectEditor/Private/MuCOE/GenerateMutableSource/GenerateMutableSourceSurface.cpp
--- a/Engine/Plugins/Experimental/Mutable/Source/CustomizableObjectEditor/Private/MuCOE/GenerateMutableSource/GenerateMutableSourceSurface.cpp
+++ b/Engine/Plugins/Experimental/Mutable/Source/CustomizableObjectEditor/Private/MuCOE/GenerateMutableSource/GenerateMutableSourceSurface.cpp
@@ -244,7 +244,7 @@ void UpdateSharedSurfaceId(FMutableGraphGenerationContext& GenerationContext, UC
 	}
 }
 
-UE_DISABLE_OPTIMIZATION
+
 mu::NodeSurfacePtr GenerateMutableSourceSurface(const UEdGraphPin * Pin, FMutableGraphGenerationContext & GenerationContext, FMutableGraphSurfaceGenerationData& SurfaceData)
 {
 	check(Pin)
@@ -1624,6 +1624,6 @@ mu::NodeSurfacePtr GenerateMutableSourceSurface(const UEdGraphPin * Pin, FMutabl
 
 	return Result;
 }
-UE_ENABLE_OPTIMIZATION
+
 #undef LOCTEXT_NAMESPACE
 
