[Mutable] Ensure relevant parameters are always generated in the editor.

#rnx

diff --git a/Engine/Plugins/Experimental/Mutable/Source/CustomizableObject/Private/MuCO/CustomizableObjectInstance.cpp b/Engine/Plugins/Experimental/Mutable/Source/CustomizableObject/Private/MuCO/CustomizableObjectInstance.cpp
--- a/Engine/Plugins/Experimental/Mutable/Source/CustomizableObject/Private/MuCO/CustomizableObjectInstance.cpp
+++ b/Engine/Plugins/Experimental/Mutable/Source/CustomizableObject/Private/MuCO/CustomizableObjectInstance.cpp
@@ -636,6 +636,9 @@ void UCustomizableObjectInstance::SetBuildParameterDecorations(const bool Value)
 
 bool UCustomizableObjectInstance::GetBuildParameterRelevancy() const
 {
+#if WITH_EDITOR
+	return true;
+#endif
 	return bBuildParameterRelevancy;
 }
 
