[mutable] Fixed unserialize check version of mutable tables.

#rnx

diff --git a/Engine/Plugins/Experimental/Mutable/Source/MutableTools/Private/MuT/TablePrivate.h b/Engine/Plugins/Experimental/Mutable/Source/MutableTools/Private/MuT/TablePrivate.h
--- a/Engine/Plugins/Experimental/Mutable/Source/MutableTools/Private/MuT/TablePrivate.h
+++ b/Engine/Plugins/Experimental/Mutable/Source/MutableTools/Private/MuT/TablePrivate.h
@@ -103,7 +103,7 @@ namespace mu
 		{
             uint32_t ver;
 			arch >> ver;
-			check(ver<=1);
+			check(ver<=2);
 
             uint32_t columnCount;
 			arch >> columnCount;
