[Mutable] Fix BuildMaterial not linking the correct material instance on skipped LODs.

#rnx

diff --git a/Engine/Plugins/Experimental/Mutable/Source/CustomizableObject/Private/MuCO/CustomizableObjectInstance.cpp b/Engine/Plugins/Experimental/Mutable/Source/CustomizableObject/Private/MuCO/CustomizableObjectInstance.cpp
--- a/Engine/Plugins/Experimental/Mutable/Source/CustomizableObject/Private/MuCO/CustomizableObjectInstance.cpp
+++ b/Engine/Plugins/Experimental/Mutable/Source/CustomizableObject/Private/MuCO/CustomizableObjectInstance.cpp
@@ -5588,7 +5588,7 @@ void UCustomizableInstancePrivateData::BuildMaterials(const TSharedPtr<FMutableO
 
 	// Find skipped LODs. The following valid LOD will be copied into them. 
 	TArray<bool> LODsSkipped;
-	LODsSkipped.SetNum(OperationData->CurrentMaxLOD);
+	LODsSkipped.SetNum(OperationData->NumLODsAvailable);
 
 	TArray<FGeneratedTexture> NewGeneratedTextures;
 
@@ -6090,7 +6090,7 @@ void UCustomizableInstancePrivateData::BuildMaterials(const TSharedPtr<FMutableO
 		{
 			// Copy data from valid LODs into the skipped ones.
 			int32 LastValidLODIndex = OperationData->CurrentMaxLOD;
-			for (int32 LODIndex = OperationData->CurrentMaxLOD; LODIndex < FirstLODAvailable; ++LODIndex)
+			for (int32 LODIndex = OperationData->CurrentMaxLOD; LODIndex >= FirstLODAvailable; --LODIndex)
 			{
 				// Copy information from the LastValidLODIndex
 				if (LODsSkipped[LODIndex])
