[Mutable] Added logging to most API methods to get and set parameter values of the descriptor.
- Warn about uncompiled CO and missing parameters.
- Small changes in naming some variables to keep them consistent between methods.

#rnx

diff --git a/Engine/Plugins/Experimental/Mutable/Source/CustomizableObject/Classes/MuCO/CustomizableObject.h b/Engine/Plugins/Experimental/Mutable/Source/CustomizableObject/Classes/MuCO/CustomizableObject.h
--- a/Engine/Plugins/Experimental/Mutable/Source/CustomizableObject/Classes/MuCO/CustomizableObject.h
+++ b/Engine/Plugins/Experimental/Mutable/Source/CustomizableObject/Classes/MuCO/CustomizableObject.h
@@ -1453,6 +1453,8 @@ private:
 
 public:
 
+	void AddUncompiledCOWarning(const FString& AdditionalLoggingInfo);
+
 #if WITH_EDITOR
 	
 	// Compile the object if Automatic Compilation is enabled and the object can be compiled.
diff --git a/Engine/Plugins/Experimental/Mutable/Source/CustomizableObject/Classes/MuCO/CustomizableObjectSystem.h b/Engine/Plugins/Experimental/Mutable/Source/CustomizableObject/Classes/MuCO/CustomizableObjectSystem.h
--- a/Engine/Plugins/Experimental/Mutable/Source/CustomizableObject/Classes/MuCO/CustomizableObjectSystem.h
+++ b/Engine/Plugins/Experimental/Mutable/Source/CustomizableObject/Classes/MuCO/CustomizableObjectSystem.h
@@ -302,7 +302,7 @@ public:
     
 	// Show a warning on-screen and via a notification (if in Editor) and log an error when a CustomizableObject is
 	// being used and it's not compiled.  Callers can add additional information to the error log.
-	void AddUncompiledCOWarning(UCustomizableObject* InObject, FString const* OptionalLogInfo = nullptr);
+	void AddUncompiledCOWarning(const UCustomizableObject& InObject, FString const* OptionalLogInfo = nullptr);
 
 	// Enables the collection of internal mutabe performance data. It has a performance cost.
 	void EnableBenchmark();
diff --git a/Engine/Plugins/Experimental/Mutable/Source/CustomizableObject/Private/MuCO/CustomizableObject.cpp b/Engine/Plugins/Experimental/Mutable/Source/CustomizableObject/Private/MuCO/CustomizableObject.cpp
--- a/Engine/Plugins/Experimental/Mutable/Source/CustomizableObject/Private/MuCO/CustomizableObject.cpp
+++ b/Engine/Plugins/Experimental/Mutable/Source/CustomizableObject/Private/MuCO/CustomizableObject.cpp
@@ -749,7 +749,7 @@ bool UCustomizableObject::ConditionalAutoCompile()
 	// Don't compile if we're running game, Compile and/or if AutoCompile is disabled
 	if (IsRunningGame() || System->IsCompilationDisabled() || !System->IsAutoCompileEnabled())
 	{
-		System->AddUncompiledCOWarning(this);
+		System->AddUncompiledCOWarning(*this);
 		return false;
 	}
 
@@ -932,6 +932,19 @@ bool UCustomizableObject::IsCompiled() const
 }
 
 
+void UCustomizableObject::AddUncompiledCOWarning(const FString& AdditionalLoggingInfo)
+{
+	// Send a warning (on-screen notification, log error, and in-editor notification)
+	UCustomizableObjectSystem* System = UCustomizableObjectSystem::GetInstance();
+	if (!System || !System->IsValidLowLevel() || System->HasAnyFlags(RF_BeginDestroyed))
+	{
+		return;
+	}
+
+	System->AddUncompiledCOWarning(*this, &AdditionalLoggingInfo);
+}
+
+
 USkeletalMesh* UCustomizableObject::GetRefSkeletalMesh(int32 ComponentIndex)
 {
 #if WITH_EDITORONLY_DATA
@@ -1369,8 +1382,8 @@ void UCustomizableObject::UpdateParameterPropertiesFromModel()
 	}
 	else
 	{
-		ParameterProperties.Reset();
-		ParameterPropertiesLookupTable.Reset();
+		ParameterProperties.Empty();
+		ParameterPropertiesLookupTable.Empty();
 	}
 }
 
diff --git a/Engine/Plugins/Experimental/Mutable/Source/CustomizableObject/Private/MuCO/CustomizableObjectInstanceDescriptor.cpp b/Engine/Plugins/Experimental/Mutable/Source/CustomizableObject/Private/MuCO/CustomizableObjectInstanceDescriptor.cpp
--- a/Engine/Plugins/Experimental/Mutable/Source/CustomizableObject/Private/MuCO/CustomizableObjectInstanceDescriptor.cpp
+++ b/Engine/Plugins/Experimental/Mutable/Source/CustomizableObject/Private/MuCO/CustomizableObjectInstanceDescriptor.cpp
@@ -1114,161 +1114,147 @@ bool FCustomizableObjectInstanceDescriptor::HasAnyParameters() const
 		!VectorParameters.IsEmpty();
 }
 
+#if WITH_EDITOR 
+
+#define RETURN_ON_UNCOMPILED_CO(CustomizableObject, ErrorMessage) \
+	if (!CustomizableObject->IsCompiled()) \
+	{ \
+		FString AdditionalLoggingInfo = FString::Printf(TEXT("Calling function: %hs.  %s"), __FUNCTION__, *ErrorMessage); \
+		CustomizableObject->AddUncompiledCOWarning(AdditionalLoggingInfo);\
+		return; \
+	} \
+
+#else
+
+#define RETURN_ON_UNCOMPILED_CO(CustomizableObject, ErrorMessage) \
+	if (!ensureMsgf(CustomizableObject->IsCompiled(), TEXT("Customizable Object (%s) was not compiled."), *GetNameSafe(CustomizableObject))) \
+	{ \
+		FString AdditionalLoggingInfo = FString::Printf(TEXT("Calling function: %hs.  %s"), __FUNCTION__, *ErrorMessage); \
+		CustomizableObject->AddUncompiledCOWarning(AdditionalLoggingInfo);\
+		return; \
+	} \
+ 
+#endif
+	
+
+
+void LogParameterNotFoundWarning(const FString& ParameterName, const int32 ObjectParameterIndex, const int32 InstanceParameterIndex, const UCustomizableObject* CustomizableObject, char const* CallingFunction)
+{
+#if !WITH_EDITOR
+	// Ensuring to help make sure we catch these critical data mismatches, but since we can handle these gracefully,
+	// don't ensure in Shipping builds.  We need to always log an error afterwards since the ensure will only fire
+	// once per session, and we don't want to always ensure since that might interfere with unrelated debugging.
+	ensureMsgf(
+		TEXT("Failed to find parameter (%s) on CO (%s). CO parameter index: (%d). COI parameter index: (%d)"),
+		*ParameterName, *GetNameSafe(CustomizableObject), ObjectParameterIndex, InstanceParameterIndex
+	);
+#endif
+	UE_LOG(LogMutable, Error,
+		TEXT("%hs: Failed to find parameter (%s) on CO (%s). CO parameter index: (%d). COI parameter index: (%d)s"),
+		CallingFunction, *ParameterName, *GetNameSafe(CustomizableObject), ObjectParameterIndex, InstanceParameterIndex
+	);
+}
+
 
 const FString& FCustomizableObjectInstanceDescriptor::GetIntParameterSelectedOption(const FString& ParamName, const int32 RangeIndex) const
 {
 	check(CustomizableObject);
 	
 	const int32 ParameterIndexInObject = CustomizableObject->FindParameter(ParamName);
+	const int32 ParameterIndexInInstance = FindIntParameterNameIndex(ParamName);
 
-	const int32 IntParamIndex = FindIntParameterNameIndex(ParamName);
-
-	if (ParameterIndexInObject >= 0 && IntParamIndex >= 0)
+	if (ParameterIndexInObject >= 0 && ParameterIndexInInstance >= 0)
 	{
 		if (RangeIndex == -1)
 		{
 			check(!CustomizableObject->IsParameterMultidimensional(ParameterIndexInObject)); // This param is multidimensional, it must have a RangeIndex of 0 or more
-			return IntParameters[IntParamIndex].ParameterValueName;
+			return IntParameters[ParameterIndexInInstance].ParameterValueName;
 		}
 		else
 		{
 			check(CustomizableObject->IsParameterMultidimensional(ParameterIndexInObject)); // This param is not multidimensional, it must have a RangeIndex of -1
 
-			if (IntParameters[IntParamIndex].ParameterRangeValueNames.IsValidIndex(RangeIndex))
+			if (IntParameters[ParameterIndexInInstance].ParameterRangeValueNames.IsValidIndex(RangeIndex))
 			{
-				return IntParameters[IntParamIndex].ParameterRangeValueNames[RangeIndex];
+				return IntParameters[ParameterIndexInInstance].ParameterRangeValueNames[RangeIndex];
 			}
 		}
 	}
 
+	LogParameterNotFoundWarning(ParamName, ParameterIndexInObject, ParameterIndexInInstance, CustomizableObject, __FUNCTION__);
+
 	return FCustomizableObjectBoolParameterValue::DEFAULT_PARAMETER_VALUE_NAME;
 }
 
-void FCustomizableObjectInstanceDescriptor::AddUncompiledCOWarning(const FString& AdditionalLoggingInfo)
-{
-	// Send a warning (on-screen notification, log error, and in-editor notification)
-	UCustomizableObjectSystem* System = UCustomizableObjectSystem::GetInstance();
-	if (!System || !System->IsValidLowLevel() || System->HasAnyFlags(RF_BeginDestroyed))
-	{
-		return;
-	}
 
-	System->AddUncompiledCOWarning(CustomizableObject.Get(), &AdditionalLoggingInfo);
-}
-
-void FCustomizableObjectInstanceDescriptor::SetIntParameterSelectedOption(const int32 IntParamIndex, const FString& SelectedOption, const int32 RangeIndex)
+void FCustomizableObjectInstanceDescriptor::SetIntParameterSelectedOption(const int32 ParameterIndexInInstance, const FString& SelectedOption, const int32 RangeIndex)
 {
 	check(CustomizableObject);
+	RETURN_ON_UNCOMPILED_CO(CustomizableObject, TEXT("Error: Cannot set Int parameter "));
 
-	if (!ensureMsgf(CustomizableObject->IsCompiled(), TEXT("Customizable Object (%s) was not compiled."),
-					*GetNameSafe(CustomizableObject)
-				   )
-	   )
-	{
-		FString AdditionalLoggingInfo = FString::Printf(
-			TEXT("Calling function: %hs.  Cannot set parameter (Parameter Index: %d, SelectedOptionName: %s, RangeIndex: %d)."),
-			__FUNCTION__, IntParamIndex, *SelectedOption, RangeIndex
-													   );
+	const int32 ParameterIndexInObject = IntParameters.IsValidIndex(ParameterIndexInInstance) ? CustomizableObject->FindParameter(IntParameters[ParameterIndexInInstance].ParameterName) : INDEX_NONE;
 
-		// Send a warning (on-screen notification, log error, and in-editor notification) for each occurrence.
-		AddUncompiledCOWarning(AdditionalLoggingInfo);
-
-		// Early out since we can't set any parameters on an uncompiled CO.
+	if (ParameterIndexInObject < 0 || ParameterIndexInInstance < 0)
+	{
+		// Warn and early out since we could not find the parameter to set.
+		LogParameterNotFoundWarning(TEXT("Unknown Int Parmeter"), ParameterIndexInObject, ParameterIndexInInstance, CustomizableObject, __FUNCTION__);
 		return;
 	}
 
-	if (!IntParameters.IsValidIndex(IntParamIndex))
+	const bool bValid = SelectedOption == TEXT("None") || CustomizableObject->FindIntParameterValue(ParameterIndexInObject, SelectedOption) >= 0;
+	if (!bValid)
 	{
-		UE_LOG(LogMutable, Error,
-				TEXT("%hs: IntParamIndex (%d) is not valid on CO (%s).  (SelectedOptionName: %s, RangeIndex: %d).  Cannot set parameter."),
-				__FUNCTION__, IntParamIndex, *GetNameSafe(CustomizableObject), *SelectedOption, RangeIndex
-			  );
-
+#if !UE_BUILD_SHIPPING
+		const FString Message = FString::Printf(
+			TEXT("Tried to set the invalid value [%s] to parameter [%d, %s]! Value index=[%d]. Correct values=[%s]."),
+			*SelectedOption, ParameterIndexInObject,
+			*IntParameters[ParameterIndexInInstance].ParameterName,
+			CustomizableObject->FindIntParameterValue(ParameterIndexInObject, SelectedOption),
+			*GetAvailableOptionsString(*CustomizableObject, ParameterIndexInObject)
+		);
+		UE_LOG(LogMutable, Error, TEXT("%s"), *Message);
+#endif
 		return;
 	}
 
-	const int32 ParameterIndexInObject = CustomizableObject->FindParameter(IntParameters[IntParamIndex].ParameterName);
-	if (ParameterIndexInObject >= 0)
+	if (RangeIndex == -1)
 	{
-		const bool bValid = SelectedOption == TEXT("None") || CustomizableObject->FindIntParameterValue(ParameterIndexInObject, SelectedOption) >= 0;
-		if (!bValid)
-		{
-			const FString Message = FString::Printf(
-				TEXT("LogMutable: Tried to set the invalid value [%s] to parameter [%d, %s]! Value index=[%d]. Correct values=[%s]."), 
-				*SelectedOption, ParameterIndexInObject,
-				*IntParameters[IntParamIndex].ParameterName, 
-				CustomizableObject->FindIntParameterValue(ParameterIndexInObject, SelectedOption), 
-				*GetAvailableOptionsString(*CustomizableObject, ParameterIndexInObject)
-			);
-			UE_LOG(LogMutable, Error, TEXT("%s"), *Message);
-		}
-			
-		if (RangeIndex == -1)
+		// If this param were multidimensional, it must have a RangeIndex of 0 or more
+		check(!CustomizableObject->IsParameterMultidimensional(ParameterIndexInObject));
+		IntParameters[ParameterIndexInInstance].ParameterValueName = SelectedOption;
+	}
+	else
+	{
+		// If this param were not multidimensional, it must have a RangeIndex of -1
+		check(CustomizableObject->IsParameterMultidimensional(ParameterIndexInObject));
+
+		if (!IntParameters[ParameterIndexInInstance].ParameterRangeValueNames.IsValidIndex(RangeIndex))
 		{
-			// If this param were multidimensional, it must have a RangeIndex of 0 or more
-			check(!CustomizableObject->IsParameterMultidimensional(ParameterIndexInObject)); 
-			IntParameters[IntParamIndex].ParameterValueName = SelectedOption;
+			const int32 InsertionIndex = IntParameters[ParameterIndexInInstance].ParameterRangeValueNames.Num();
+			const int32 NumInsertedElements = RangeIndex + 1 - IntParameters[ParameterIndexInInstance].ParameterRangeValueNames.Num();
+			IntParameters[ParameterIndexInInstance].ParameterRangeValueNames.InsertDefaulted(InsertionIndex, NumInsertedElements);
 		}
-		else
-		{
-			// If this param were not multidimensional, it must have a RangeIndex of -1
-			check(CustomizableObject->IsParameterMultidimensional(ParameterIndexInObject));
 
-			if (!IntParameters[IntParamIndex].ParameterRangeValueNames.IsValidIndex(RangeIndex))
-			{
-				const int32 InsertionIndex = IntParameters[IntParamIndex].ParameterRangeValueNames.Num();
-				const int32 NumInsertedElements = RangeIndex + 1 - IntParameters[IntParamIndex].ParameterRangeValueNames.Num();
-				IntParameters[IntParamIndex].ParameterRangeValueNames.InsertDefaulted(InsertionIndex, NumInsertedElements);
-			}
-
-			check(IntParameters[IntParamIndex].ParameterRangeValueNames.IsValidIndex(RangeIndex));
-			IntParameters[IntParamIndex].ParameterRangeValueNames[RangeIndex] = SelectedOption;
-		}
+		check(IntParameters[ParameterIndexInInstance].ParameterRangeValueNames.IsValidIndex(RangeIndex));
+		IntParameters[ParameterIndexInInstance].ParameterRangeValueNames[RangeIndex] = SelectedOption;
 	}
 }
 
+
 void FCustomizableObjectInstanceDescriptor::SetIntParameterSelectedOption(const FString& ParamName, const FString& SelectedOptionName, const int32 RangeIndex)
 {
 	check(CustomizableObject);
+	RETURN_ON_UNCOMPILED_CO(CustomizableObject, TEXT("Error: Cannot set Int parameter "));
 
-	if (!ensureMsgf(CustomizableObject->IsCompiled(), TEXT("Customizable Object (%s) was not compiled."),
-					*GetNameSafe(CustomizableObject)
-				   )
-	   )
+	const int32 ParamIndexInInstance = FindIntParameterNameIndex(ParamName);
+	if (ParamIndexInInstance == INDEX_NONE)
 	{
-		FString AdditionalLoggingInfo = FString::Printf(
-			TEXT("Calling function: %hs.  Cannot set parameter (Parameter Name: %s, SelectedOptionName: %s, RangeIndex: %d)."),
-			__FUNCTION__, *ParamName, *SelectedOptionName, RangeIndex
-													   );
-
-		// Send a warning (on-screen notification, log error, and in-editor notification) for each occurrence.
-		AddUncompiledCOWarning(AdditionalLoggingInfo);
-
-		// Early out since we can't set any parameters on an uncompiled CO.
+		// Warn and early out since we could not find the parameter to set.
+		LogParameterNotFoundWarning(ParamName, ParamIndexInInstance, ParamIndexInInstance, CustomizableObject, __FUNCTION__);
 		return;
 	}
 
-	const int32 IntParamIndex = FindIntParameterNameIndex(ParamName);
-	if (IntParamIndex != INDEX_NONE)
-	{
-		SetIntParameterSelectedOption(IntParamIndex, SelectedOptionName, RangeIndex);
-	}
-	else
-	{
-#if !UE_BUILD_SHIPPING
-		// Ensuring to help make sure we catch these critical data mismatches, but since we can handle these gracefully,
-		// don't ensure in Shipping builds.  We need to always log an error afterwards since the ensure will only fire
-		// once per session, and we don't want to always ensure since that might interfere with unrelated debugging.
-		ensureMsgf(false, 
-				TEXT("Failed to find valid parameter index for ParamName (%s) on CO (%s).  (SelectedOptionName: %s, RangeIndex: %d).  Cannot set parameter."),
-				*ParamName, *GetNameSafe(CustomizableObject), *SelectedOptionName, RangeIndex
-				  );
-#endif
-		UE_LOG(LogMutable, Error,
-				TEXT("%hs: Failed to find valid parameter index for ParamName (%s) on CO (%s).  (SelectedOptionName: %s, RangeIndex: %d).  Cannot set parameter."),
-				__FUNCTION__, *ParamName, *GetNameSafe(CustomizableObject), *SelectedOptionName, RangeIndex
-			  );
-	}
+	SetIntParameterSelectedOption(ParamIndexInInstance, SelectedOptionName, RangeIndex);
 }
 
 
@@ -1277,7 +1263,6 @@ float FCustomizableObjectInstanceDescriptor::GetFloatParameterSelectedOption(con
 	check(CustomizableObject);
 
 	const int32 ParameterIndexInObject = CustomizableObject->FindParameter(FloatParamName);
-
 	const int32 FloatParamIndex = FindFloatParameterNameIndex(FloatParamName);
 
 	if (ParameterIndexInObject >= 0 && FloatParamIndex >= 0)
@@ -1298,6 +1283,8 @@ float FCustomizableObjectInstanceDescriptor::GetFloatParameterSelectedOption(con
 		}
 	}
 
+	LogParameterNotFoundWarning(FloatParamName, ParameterIndexInObject, FloatParamIndex, CustomizableObject, __FUNCTION__);
+	
 	return 	FCustomizableObjectFloatParameterValue::DEFAULT_PARAMETER_VALUE; 
 }
 
@@ -1305,33 +1292,38 @@ float FCustomizableObjectInstanceDescriptor::GetFloatParameterSelectedOption(con
 void FCustomizableObjectInstanceDescriptor::SetFloatParameterSelectedOption(const FString& FloatParamName, const float FloatValue, const int32 RangeIndex)
 {
 	check(CustomizableObject);
+	RETURN_ON_UNCOMPILED_CO(CustomizableObject, TEXT("Error: Cannot set Int parameter "));
 
 	const int32 ParameterIndexInObject = CustomizableObject->FindParameter(FloatParamName);
+	const int32 ParameterIndexInInstance = FindFloatParameterNameIndex(FloatParamName);
 
-	const int32 FloatParamIndex = FindFloatParameterNameIndex(FloatParamName);
+	if (ParameterIndexInObject < 0 || ParameterIndexInInstance < 0)
+	{
+		// Warn and early out since we could not find the parameter to set.
+		LogParameterNotFoundWarning(FloatParamName, ParameterIndexInObject, ParameterIndexInInstance, CustomizableObject, __FUNCTION__);
+		return;
+	}
 
-	if (ParameterIndexInObject >= 0 && FloatParamIndex >= 0)
+	if (RangeIndex == INDEX_NONE)
 	{
-		if (RangeIndex == -1)
+		check(!CustomizableObject->IsParameterMultidimensional(ParameterIndexInObject)); // This param is multidimensional, it must have a RangeIndex of 0 or more
+		FloatParameters[ParameterIndexInInstance].ParameterValue = FloatValue;
+	}
+	else
+	{
+		check(CustomizableObject->IsParameterMultidimensional(ParameterIndexInObject)); // This param is not multidimensional, it must have a RangeIndex of -1
+
+		if (!FloatParameters[ParameterIndexInInstance].ParameterRangeValues.IsValidIndex(RangeIndex))
 		{
-			check(!CustomizableObject->IsParameterMultidimensional(ParameterIndexInObject)); // This param is multidimensional, it must have a RangeIndex of 0 or more
-			FloatParameters[FloatParamIndex].ParameterValue = FloatValue;
+			const int32 InsertionIndex = FloatParameters[ParameterIndexInInstance].ParameterRangeValues.Num();
+			const int32 NumInsertedElements = RangeIndex + 1 - FloatParameters[ParameterIndexInInstance].ParameterRangeValues.Num();
+			FloatParameters[ParameterIndexInInstance].ParameterRangeValues.InsertDefaulted(InsertionIndex, NumInsertedElements);
 		}
-		else
-		{
-			check(CustomizableObject->IsParameterMultidimensional(ParameterIndexInObject)); // This param is not multidimensional, it must have a RangeIndex of -1
 
-			if (!FloatParameters[FloatParamIndex].ParameterRangeValues.IsValidIndex(RangeIndex))
-			{
-				const int32 InsertionIndex = FloatParameters[FloatParamIndex].ParameterRangeValues.Num();
-				const int32 NumInsertedElements = RangeIndex + 1 - FloatParameters[FloatParamIndex].ParameterRangeValues.Num();
-				FloatParameters[FloatParamIndex].ParameterRangeValues.InsertDefaulted(InsertionIndex, NumInsertedElements);
-			}
-
-			check(FloatParameters[FloatParamIndex].ParameterRangeValues.IsValidIndex(RangeIndex));
-			FloatParameters[FloatParamIndex].ParameterRangeValues[RangeIndex] = FloatValue;
-		}
+		check(FloatParameters[ParameterIndexInInstance].ParameterRangeValues.IsValidIndex(RangeIndex));
+		FloatParameters[ParameterIndexInInstance].ParameterRangeValues[RangeIndex] = FloatValue;
 	}
+
 }
 
 
@@ -1340,26 +1332,27 @@ FString FCustomizableObjectInstanceDescriptor::GetTextureParameterSelectedOption
 	check(CustomizableObject);
 
 	const int32 ParameterIndexInObject = CustomizableObject->FindParameter(TextureParamName);
+	const int32 ParameterIndexInInstance = FindTextureParameterNameIndex(TextureParamName);
 
-	const int32 TextureParamIndex = FindTextureParameterNameIndex(TextureParamName);
-
-	if (ParameterIndexInObject >= 0 && TextureParamIndex >= 0)
+	if (ParameterIndexInObject >= 0 && ParameterIndexInInstance >= 0)
 	{
 		if (RangeIndex == -1)
 		{
 			check(!CustomizableObject->IsParameterMultidimensional(ParameterIndexInObject)); // This param is multidimensional, it must have a RangeIndex of 0 or more
-			return TextureParameters[TextureParamIndex].ParameterValue;
+			return TextureParameters[ParameterIndexInInstance].ParameterValue;
 		}
 		else
 		{
 			check(CustomizableObject->IsParameterMultidimensional(ParameterIndexInObject)); // This param is not multidimensional, it must have a RangeIndex of -1
 
-			if (TextureParameters[TextureParamIndex].ParameterRangeValues.IsValidIndex(RangeIndex))
+			if (TextureParameters[ParameterIndexInInstance].ParameterRangeValues.IsValidIndex(RangeIndex))
 			{
-				return TextureParameters[TextureParamIndex].ParameterRangeValues[RangeIndex];
+				return TextureParameters[ParameterIndexInInstance].ParameterRangeValues[RangeIndex];
 			}
 		}
 	}
+	
+	LogParameterNotFoundWarning(TextureParamName, ParameterIndexInObject, ParameterIndexInInstance, CustomizableObject, __FUNCTION__);
 
 	return FCustomizableObjectTextureParameterValue::DEFAULT_PARAMETER_VALUE; 
 }
@@ -1368,38 +1361,45 @@ FString FCustomizableObjectInstanceDescriptor::GetTextureParameterSelectedOption
 void FCustomizableObjectInstanceDescriptor::SetTextureParameterSelectedOption(const FString& TextureParamName, const FString& TextureValue, const int32 RangeIndex)
 {
 	check(CustomizableObject);
+	RETURN_ON_UNCOMPILED_CO(CustomizableObject, TEXT("Error: Cannot set Int parameter "));
 
 	const int32 ParameterIndexInObject = CustomizableObject->FindParameter(TextureParamName);
+	const int32 ParameterIndexInInstance = FindTextureParameterNameIndex(TextureParamName);
 
-	const int32 TextureParamIndex = FindTextureParameterNameIndex(TextureParamName);
+	if (ParameterIndexInObject < 0 || ParameterIndexInInstance < 0)
+	{
+		// Early out since we could not find the parameter to set.
+		LogParameterNotFoundWarning(TextureParamName, ParameterIndexInObject, ParameterIndexInInstance, CustomizableObject, __FUNCTION__);
+		return;
+	}
 
-	if (ParameterIndexInObject >= 0 && TextureParamIndex >= 0)
+	if (RangeIndex == INDEX_NONE)
 	{
-		if (RangeIndex == -1)
+		check(!CustomizableObject->IsParameterMultidimensional(ParameterIndexInObject)); // This param is multidimensional, it must have a RangeIndex of 0 or more
+		TextureParameters[ParameterIndexInInstance].ParameterValue = TextureValue;
+	}
+	else
+	{
+		check(CustomizableObject->IsParameterMultidimensional(ParameterIndexInObject)); // This param is not multidimensional, it must have a RangeIndex of -1
+
+		if (!TextureParameters[ParameterIndexInInstance].ParameterRangeValues.IsValidIndex(RangeIndex))
 		{
-			check(!CustomizableObject->IsParameterMultidimensional(ParameterIndexInObject)); // This param is multidimensional, it must have a RangeIndex of 0 or more
-			TextureParameters[TextureParamIndex].ParameterValue = TextureValue;
+			const int32 InsertionIndex = TextureParameters[ParameterIndexInInstance].ParameterRangeValues.Num();
+			const int32 NumInsertedElements = RangeIndex + 1 - TextureParameters[ParameterIndexInInstance].ParameterRangeValues.Num();
+			TextureParameters[ParameterIndexInInstance].ParameterRangeValues.InsertDefaulted(InsertionIndex, NumInsertedElements);
 		}
-		else
-		{
-			check(CustomizableObject->IsParameterMultidimensional(ParameterIndexInObject)); // This param is not multidimensional, it must have a RangeIndex of -1
-
-			if (!TextureParameters[TextureParamIndex].ParameterRangeValues.IsValidIndex(RangeIndex))
-			{
-				const int32 InsertionIndex = TextureParameters[TextureParamIndex].ParameterRangeValues.Num();
-				const int32 NumInsertedElements = RangeIndex + 1 - TextureParameters[TextureParamIndex].ParameterRangeValues.Num();
-				TextureParameters[TextureParamIndex].ParameterRangeValues.InsertDefaulted(InsertionIndex, NumInsertedElements);
-			}
 
-			check(TextureParameters[TextureParamIndex].ParameterRangeValues.IsValidIndex(RangeIndex));
-			TextureParameters[TextureParamIndex].ParameterRangeValues[RangeIndex] = TextureValue;
-		}
+		check(TextureParameters[ParameterIndexInInstance].ParameterRangeValues.IsValidIndex(RangeIndex));
+		TextureParameters[ParameterIndexInInstance].ParameterRangeValues[RangeIndex] = TextureValue;
 	}
 }
 
 
 void FCustomizableObjectInstanceDescriptor::SetTextureParameterSelectedOptionT(const FString& TextureParamName,	UTexture2D* TextureValue, int32 RangeIndex)
 {
+	check(CustomizableObject);
+	RETURN_ON_UNCOMPILED_CO(CustomizableObject, TEXT("Error: Cannot set Int parameter "));
+
 	UDefaultImageProvider& UDefaultImageProvider = UCustomizableObjectSystem::GetInstance()->GetOrCreateDefaultImageProvider();
 	const FString Id = UDefaultImageProvider.Add(TextureValue);
 	
@@ -1412,20 +1412,23 @@ FLinearColor FCustomizableObjectInstanceDescriptor::GetColorParameterSelectedOpt
 	check(CustomizableObject);
 
 	const int32 ParameterIndexInObject = CustomizableObject->FindParameter(ColorParamName);
-
 	const int32 ColorParamIndex = FindVectorParameterNameIndex(ColorParamName);
 
-	if (ParameterIndexInObject >= 0 && ColorParamIndex >= 0)
+	if (ColorParamIndex == INDEX_NONE)
 	{
-		return VectorParameters[ColorParamIndex].ParameterValue;
+		LogParameterNotFoundWarning(ColorParamName, ParameterIndexInObject, ColorParamIndex, CustomizableObject, __FUNCTION__);
+		return FCustomizableObjectVectorParameterValue::DEFAULT_PARAMETER_VALUE;
 	}
-
-	return FCustomizableObjectVectorParameterValue::DEFAULT_PARAMETER_VALUE;
+	
+	return VectorParameters[ColorParamIndex].ParameterValue;
 }
 
 
 void FCustomizableObjectInstanceDescriptor::SetColorParameterSelectedOption(const FString& ColorParamName, const FLinearColor& ColorValue)
 {
+	check(CustomizableObject);
+	RETURN_ON_UNCOMPILED_CO(CustomizableObject, TEXT("Error: Cannot set Int parameter "));
+
 	SetVectorParameterSelectedOption(ColorParamName, ColorValue);
 }
 
@@ -1435,45 +1438,53 @@ bool FCustomizableObjectInstanceDescriptor::GetBoolParameterSelectedOption(const
 	check(CustomizableObject);
 
 	const int32 ParameterIndexInObject = CustomizableObject->FindParameter(BoolParamName);
-
 	const int32 BoolParamIndex = FindBoolParameterNameIndex(BoolParamName);
 
-	if (ParameterIndexInObject >= 0 && BoolParamIndex >= 0)
+	if (BoolParamIndex == INDEX_NONE)
 	{
-		return BoolParameters[BoolParamIndex].ParameterValue;
+		LogParameterNotFoundWarning(BoolParamName, ParameterIndexInObject, BoolParamIndex, CustomizableObject, __FUNCTION__);
+		return FCustomizableObjectBoolParameterValue::DEFAULT_PARAMETER_VALUE;
 	}
 
-	return FCustomizableObjectBoolParameterValue::DEFAULT_PARAMETER_VALUE;
+	return BoolParameters[BoolParamIndex].ParameterValue;
 }
 
 
 void FCustomizableObjectInstanceDescriptor::SetBoolParameterSelectedOption(const FString& BoolParamName, const bool BoolValue)
 {
 	check(CustomizableObject);
+	RETURN_ON_UNCOMPILED_CO(CustomizableObject, TEXT("Error: Cannot set Int parameter "));
 
 	const int32 ParameterIndexInObject = CustomizableObject->FindParameter(BoolParamName);
+	const int32 ParameterIndexInInstance = FindBoolParameterNameIndex(BoolParamName);
 
-	const int32 BoolParamIndex = FindBoolParameterNameIndex(BoolParamName);
-
-	if (ParameterIndexInObject >= 0 && BoolParamIndex >= 0)
+	if (ParameterIndexInObject < 0 || ParameterIndexInInstance < 0)
 	{
-		BoolParameters[BoolParamIndex].ParameterValue = BoolValue;
+		// Early out since we could not find the parameter to set.
+		LogParameterNotFoundWarning(BoolParamName, ParameterIndexInObject, ParameterIndexInInstance, CustomizableObject, __FUNCTION__);
+		return;
 	}
+
+	BoolParameters[ParameterIndexInInstance].ParameterValue = BoolValue;
 }
 
 
 void FCustomizableObjectInstanceDescriptor::SetVectorParameterSelectedOption(const FString& VectorParamName, const FLinearColor& VectorValue)
 {
 	check(CustomizableObject);
+	RETURN_ON_UNCOMPILED_CO(CustomizableObject, FString("Error: Cannot set Int parameter "));
 
 	const int32 ParameterIndexInObject = CustomizableObject->FindParameter(VectorParamName);
+	const int32 ParameterIndexInInstance = FindVectorParameterNameIndex(VectorParamName);
 
-	const int32 VectorParamIndex = FindVectorParameterNameIndex(VectorParamName);
-
-	if (ParameterIndexInObject >= 0	&& VectorParamIndex >= 0)
+	if (ParameterIndexInObject < 0 || ParameterIndexInInstance < 0)
 	{
-		VectorParameters[VectorParamIndex].ParameterValue = VectorValue;
+		// Early out since we could not find the parameter to set.
+		LogParameterNotFoundWarning(VectorParamName, ParameterIndexInObject, ParameterIndexInInstance, CustomizableObject, __FUNCTION__);
+		return;
 	}
+
+	VectorParameters[ParameterIndexInInstance].ParameterValue = VectorValue;
 }
 
 
@@ -1483,40 +1494,48 @@ void FCustomizableObjectInstanceDescriptor::SetProjectorValue(const FString& Pro
 	const int32 RangeIndex)
 {
 	check(CustomizableObject);
+	RETURN_ON_UNCOMPILED_CO(CustomizableObject, TEXT("Error: Cannot set Projector parameter "))
 
 	const int32 ParameterIndexInObject = CustomizableObject->FindParameter(ProjectorParamName);
+	const int32 ParameterIndexInInstance = FindProjectorParameterNameIndex(ProjectorParamName);
 
-	const int32 ProjectorParamIndex = FindProjectorParameterNameIndex(ProjectorParamName);
+	if (ParameterIndexInObject < 0 || ParameterIndexInInstance < 0)
+	{
+		// Early out since we could not find the parameter to set.
+		LogParameterNotFoundWarning(ProjectorParamName, ParameterIndexInObject, ParameterIndexInInstance, CustomizableObject, __FUNCTION__);
+		return;
+	}
+
+	// Parameter to modify
+	FCustomizableObjectProjectorParameterValue& ProjectorParameter = ProjectorParameters[ParameterIndexInInstance];
 
-	if (ParameterIndexInObject >= 0 && ProjectorParamIndex >= 0)
+	// New Value
+	FCustomizableObjectProjector ProjectorData;
+	ProjectorData.Position = static_cast<FVector3f>(Pos);
+	ProjectorData.Direction = static_cast<FVector3f>(Direction);
+	ProjectorData.Up = static_cast<FVector3f>(Up);
+	ProjectorData.Scale = static_cast<FVector3f>(Scale);
+	ProjectorData.Angle = Angle;
+	ProjectorData.ProjectionType = ProjectorParameter.Value.ProjectionType;
+
+	if (RangeIndex == -1)
 	{
-		FCustomizableObjectProjector ProjectorData;
-		ProjectorData.Position = static_cast<FVector3f>(Pos);
-		ProjectorData.Direction = static_cast<FVector3f>(Direction);
-		ProjectorData.Up = static_cast<FVector3f>(Up);
-		ProjectorData.Scale = static_cast<FVector3f>(Scale);
-		ProjectorData.Angle = Angle;
-		ProjectorData.ProjectionType = ProjectorParameters[ProjectorParamIndex].Value.ProjectionType;
+		check(!CustomizableObject->IsParameterMultidimensional(ParameterIndexInObject)); // This param is multidimensional, it must have a RangeIndex of 0 or more
+		ProjectorParameter.Value = ProjectorData;
+	}
+	else
+	{
+		check(CustomizableObject->IsParameterMultidimensional(ParameterIndexInObject)); // This param is not multidimensional, it must have a RangeIndex of -1
 
-		if (RangeIndex == -1)
+		if (!ProjectorParameter.RangeValues.IsValidIndex(RangeIndex))
 		{
-			check(!CustomizableObject->IsParameterMultidimensional(ParameterIndexInObject)); // This param is multidimensional, it must have a RangeIndex of 0 or more
-			ProjectorParameters[ProjectorParamIndex].Value = ProjectorData;
+			const int32 InsertionIndex = ProjectorParameter.RangeValues.Num();
+			const int32 NumInsertedElements = RangeIndex + 1 - ProjectorParameter.RangeValues.Num();
+			ProjectorParameter.RangeValues.InsertDefaulted(InsertionIndex, NumInsertedElements);
 		}
-		else
-		{
-			check(CustomizableObject->IsParameterMultidimensional(ParameterIndexInObject)); // This param is not multidimensional, it must have a RangeIndex of -1
-
-			if (!ProjectorParameters[ProjectorParamIndex].RangeValues.IsValidIndex(RangeIndex))
-			{
-				const int32 InsertionIndex = ProjectorParameters[ProjectorParamIndex].RangeValues.Num();
-				const int32 NumInsertedElements = RangeIndex + 1 - ProjectorParameters[ProjectorParamIndex].RangeValues.Num();
-				ProjectorParameters[ProjectorParamIndex].RangeValues.InsertDefaulted(InsertionIndex, NumInsertedElements);
-			}
 
-			check(ProjectorParameters[ProjectorParamIndex].RangeValues.IsValidIndex(RangeIndex));
-			ProjectorParameters[ProjectorParamIndex].RangeValues[RangeIndex] = ProjectorData;
-		}
+		check(ProjectorParameter.RangeValues.IsValidIndex(RangeIndex));
+		ProjectorParameter.RangeValues[RangeIndex] = ProjectorData;
 	}
 }
 
@@ -1524,35 +1543,42 @@ void FCustomizableObjectInstanceDescriptor::SetProjectorValue(const FString& Pro
 void FCustomizableObjectInstanceDescriptor::SetProjectorPosition(const FString& ProjectorParamName, const FVector3f& Pos, const int32 RangeIndex)
 {
 	check(CustomizableObject);
+	RETURN_ON_UNCOMPILED_CO(CustomizableObject, TEXT("Error: Cannot set Projector parameter "))
 
 	const int32 ParameterIndexInObject = CustomizableObject->FindParameter(ProjectorParamName);
+	const int32 ParameterIndexInInstance = FindProjectorParameterNameIndex(ProjectorParamName);
+
+	if (ParameterIndexInObject < 0 || ParameterIndexInInstance < 0)
+	{
+		// Early out since we could not find the parameter to set.
+		LogParameterNotFoundWarning(ProjectorParamName, ParameterIndexInObject, ParameterIndexInInstance, CustomizableObject, __FUNCTION__);
+		return;
+	}
 
-	const int32 ProjectorParamIndex = FindProjectorParameterNameIndex(ProjectorParamName);
+	// Parameter to modify
+	FCustomizableObjectProjectorParameterValue& ProjectorParameter = ProjectorParameters[ParameterIndexInInstance];
 
-	if (ParameterIndexInObject >= 0	&& ProjectorParamIndex >= 0)
+	FCustomizableObjectProjector ProjectorData = ProjectorParameter.Value;
+	ProjectorData.Position = static_cast<FVector3f>(Pos);
+
+	if (RangeIndex == -1)
 	{
-		FCustomizableObjectProjector ProjectorData = ProjectorParameters[ProjectorParamIndex].Value;
-		ProjectorData.Position = static_cast<FVector3f>(Pos);
+		check(!CustomizableObject->IsParameterMultidimensional(ParameterIndexInObject)); // This param is multidimensional, it must have a RangeIndex of 0 or more
+		ProjectorParameter.Value = ProjectorData;
+	}
+	else
+	{
+		check(CustomizableObject->IsParameterMultidimensional(ParameterIndexInObject)); // This param is not multidimensional, it must have a RangeIndex of -1
 
-		if (RangeIndex == -1)
+		if (!ProjectorParameter.RangeValues.IsValidIndex(RangeIndex))
 		{
-			check(!CustomizableObject->IsParameterMultidimensional(ParameterIndexInObject)); // This param is multidimensional, it must have a RangeIndex of 0 or more
-			ProjectorParameters[ProjectorParamIndex].Value = ProjectorData;
+			const int32 InsertionIndex = ProjectorParameter.RangeValues.Num();
+			const int32 NumInsertedElements = RangeIndex + 1 - ProjectorParameter.RangeValues.Num();
+			ProjectorParameter.RangeValues.InsertDefaulted(InsertionIndex, NumInsertedElements);
 		}
-		else
-		{
-			check(CustomizableObject->IsParameterMultidimensional(ParameterIndexInObject)); // This param is not multidimensional, it must have a RangeIndex of -1
 
-			if (!ProjectorParameters[ProjectorParamIndex].RangeValues.IsValidIndex(RangeIndex))
-			{
-				const int32 InsertionIndex = ProjectorParameters[ProjectorParamIndex].RangeValues.Num();
-				const int32 NumInsertedElements = RangeIndex + 1 - ProjectorParameters[ProjectorParamIndex].RangeValues.Num();
-				ProjectorParameters[ProjectorParamIndex].RangeValues.InsertDefaulted(InsertionIndex, NumInsertedElements);
-			}
-
-			check(ProjectorParameters[ProjectorParamIndex].RangeValues.IsValidIndex(RangeIndex));
-			ProjectorParameters[ProjectorParamIndex].RangeValues[RangeIndex] = ProjectorData;
-		}
+		check(ProjectorParameter.RangeValues.IsValidIndex(RangeIndex));
+		ProjectorParameter.RangeValues[RangeIndex] = ProjectorData;
 	}
 }
 
@@ -1562,40 +1588,13 @@ void FCustomizableObjectInstanceDescriptor::GetProjectorValue(const FString& Pro
 	float& OutAngle, ECustomizableObjectProjectorType& OutType,
 	const int32 RangeIndex) const
 {
-	check(CustomizableObject);
+	FVector3f Pos, Direction, Up, Scale;
+	GetProjectorValueF(ProjectorParamName, Pos, Direction, Up, Scale, OutAngle, OutType, RangeIndex);
 
-	const int32 ParameterIndexInObject = CustomizableObject->FindParameter(ProjectorParamName);
-
-	const int32 ProjectorParamIndex = FindProjectorParameterNameIndex(ProjectorParamName);
-
-	if (ParameterIndexInObject >= 0 && ProjectorParamIndex >= 0)
-	{
-		const FCustomizableObjectProjector* Projector;
-
-		if (RangeIndex == -1)
-		{
-			check(!CustomizableObject->IsParameterMultidimensional(ParameterIndexInObject)); // This param is multidimensional, it must have a RangeIndex of 0 or more
-
-			Projector = &ProjectorParameters[ProjectorParamIndex].Value;
-		}
-		else
-		{
-			check(CustomizableObject->IsParameterMultidimensional(ParameterIndexInObject)); // This param is not multidimensional, it must have a RangeIndex of -1
-			check(ProjectorParameters[ProjectorParamIndex].RangeValues.IsValidIndex(RangeIndex));
-
-			Projector = &ProjectorParameters[ProjectorParamIndex].RangeValues[RangeIndex];
-		}
-
-		if (Projector)
-		{
-			OutPos = static_cast<FVector>(Projector->Position);
-			OutDirection = static_cast<FVector>(Projector->Direction);
-			OutUp = static_cast<FVector>(Projector->Up);
-			OutScale = static_cast<FVector>(Projector->Scale);
-			OutAngle = Projector->Angle;
-			OutType = Projector->ProjectionType;
-		}
-	}
+	OutPos = static_cast<FVector>(Pos);
+	OutDirection = static_cast<FVector>(Direction);
+	OutUp = static_cast<FVector>(Up);
+	OutScale = static_cast<FVector>(Scale);
 }
 
 
@@ -1607,37 +1606,41 @@ void FCustomizableObjectInstanceDescriptor::GetProjectorValueF(const FString& Pr
 	check(CustomizableObject);
 
 	const int32 ParameterIndexInObject = CustomizableObject->FindParameter(ProjectorParamName);
+	const int32 ParameterIndexInInstance = FindProjectorParameterNameIndex(ProjectorParamName);
 
-	const int32 ProjectorParamIndex = FindProjectorParameterNameIndex(ProjectorParamName);
-
-	if (ParameterIndexInObject >= 0 && ProjectorParamIndex >= 0)
+	if (ParameterIndexInObject >= 0 && ParameterIndexInInstance >= 0)
 	{
-		const FCustomizableObjectProjector* Projector;
+		const FCustomizableObjectProjectorParameterValue& ProjectorParameter = ProjectorParameters[ParameterIndexInInstance];
+		const FCustomizableObjectProjector* ProjectorData;
 
 		if (RangeIndex == -1)
 		{
 			check(!CustomizableObject->IsParameterMultidimensional(ParameterIndexInObject)); // This param is multidimensional, it must have a RangeIndex of 0 or more
 
-			Projector = &ProjectorParameters[ProjectorParamIndex].Value;
+			ProjectorData = &ProjectorParameter.Value;
 		}
 		else
 		{
 			check(CustomizableObject->IsParameterMultidimensional(ParameterIndexInObject)); // This param is not multidimensional, it must have a RangeIndex of -1
-			check(ProjectorParameters[ProjectorParamIndex].RangeValues.IsValidIndex(RangeIndex));
+			check(ProjectorParameter.RangeValues.IsValidIndex(RangeIndex));
 
-			Projector = &ProjectorParameters[ProjectorParamIndex].RangeValues[RangeIndex];
+			ProjectorData = &ProjectorParameter.RangeValues[RangeIndex];
 		}
 
-		if (Projector)
+		if (ProjectorData)
 		{
-			OutPos = Projector->Position;
-			OutDirection = Projector->Direction;
-			OutUp = Projector->Up;
-			OutScale = Projector->Scale;
-			OutAngle = Projector->Angle;
-			OutType = Projector->ProjectionType;
+			OutPos = ProjectorData->Position;
+			OutDirection = ProjectorData->Direction;
+			OutUp = ProjectorData->Up;
+			OutScale = ProjectorData->Scale;
+			OutAngle = ProjectorData->Angle;
+			OutType = ProjectorData->ProjectionType;
 		}
 	}
+	else
+	{
+		LogParameterNotFoundWarning(ProjectorParamName, ParameterIndexInObject, ParameterIndexInInstance, CustomizableObject, __FUNCTION__);
+	}
 }
 
 
@@ -1646,21 +1649,22 @@ FVector FCustomizableObjectInstanceDescriptor::GetProjectorPosition(const FStrin
 	check(CustomizableObject);
 
 	const int32 ParameterIndexInObject = CustomizableObject->FindParameter(ParamName);
+	const int32 ParameterIndexInInstance = FindProjectorParameterNameIndex(ParamName);
 
-	const int32 ProjectorParamIndex = FindProjectorParameterNameIndex(ParamName);
-
-	if (ParameterIndexInObject >= 0 && ProjectorParamIndex >= 0)
+	if (ParameterIndexInObject >= 0 && ParameterIndexInInstance >= 0)
 	{
 		if (RangeIndex == -1)
 		{
-			return static_cast<FVector>(ProjectorParameters[ProjectorParamIndex].Value.Position);
+			return static_cast<FVector>(ProjectorParameters[ParameterIndexInInstance].Value.Position);
 		}
-		else if (ProjectorParameters[ProjectorParamIndex].RangeValues.IsValidIndex(RangeIndex))
+		else if (ProjectorParameters[ParameterIndexInInstance].RangeValues.IsValidIndex(RangeIndex))
 		{
-			return static_cast<FVector>(ProjectorParameters[ProjectorParamIndex].RangeValues[RangeIndex].Position);
+			return static_cast<FVector>(ProjectorParameters[ParameterIndexInInstance].RangeValues[RangeIndex].Position);
 		}
 	}
 
+	LogParameterNotFoundWarning(ParamName, ParameterIndexInObject, ParameterIndexInInstance, CustomizableObject, __FUNCTION__);
+
 	return FVector(-0.0,-0.0,-0.0);
 }
 
@@ -1670,20 +1674,21 @@ FVector FCustomizableObjectInstanceDescriptor::GetProjectorDirection(const FStri
 	check(CustomizableObject);
 
 	const int32 ParameterIndexInObject = CustomizableObject->FindParameter(ParamName);
+	const int32 ParameterIndexInInstance = FindProjectorParameterNameIndex(ParamName);
 
-	const int32 ProjectorParamIndex = FindProjectorParameterNameIndex(ParamName);
-
-	if (ParameterIndexInObject >= 0 && ProjectorParamIndex >= 0)
+	if (ParameterIndexInObject >= 0 && ParameterIndexInInstance >= 0)
 	{
 		if (RangeIndex == -1)
 		{
-			return static_cast<FVector>(ProjectorParameters[ProjectorParamIndex].Value.Direction);
+			return static_cast<FVector>(ProjectorParameters[ParameterIndexInInstance].Value.Direction);
 		}
-		else if (ProjectorParameters[ProjectorParamIndex].RangeValues.IsValidIndex(RangeIndex))
+		else if (ProjectorParameters[ParameterIndexInInstance].RangeValues.IsValidIndex(RangeIndex))
 		{
-			return static_cast<FVector>(ProjectorParameters[ProjectorParamIndex].RangeValues[RangeIndex].Direction);
+			return static_cast<FVector>(ProjectorParameters[ParameterIndexInInstance].RangeValues[RangeIndex].Direction);
 		}
 	}
+	
+	LogParameterNotFoundWarning(ParamName, ParameterIndexInObject, ParameterIndexInInstance, CustomizableObject, __FUNCTION__);
 
 	return FVector(-0.0, -0.0, -0.0);
 }
@@ -1694,20 +1699,21 @@ FVector FCustomizableObjectInstanceDescriptor::GetProjectorUp(const FString& Par
 	check(CustomizableObject);
 
 	const int32 ParameterIndexInObject = CustomizableObject->FindParameter(ParamName);
+	const int32 ParameterIndexInInstance = FindProjectorParameterNameIndex(ParamName);
 
-	const int32 ProjectorParamIndex = FindProjectorParameterNameIndex(ParamName);
-
-	if (ParameterIndexInObject >= 0 && ProjectorParamIndex >= 0)
+	if (ParameterIndexInObject >= 0 && ParameterIndexInInstance >= 0)
 	{
 		if (RangeIndex == -1)
 		{
-			return static_cast<FVector>(ProjectorParameters[ProjectorParamIndex].Value.Up);
+			return static_cast<FVector>(ProjectorParameters[ParameterIndexInInstance].Value.Up);
 		}
-		else if (ProjectorParameters[ProjectorParamIndex].RangeValues.IsValidIndex(RangeIndex))
+		else if (ProjectorParameters[ParameterIndexInInstance].RangeValues.IsValidIndex(RangeIndex))
 		{
-			return static_cast<FVector>(ProjectorParameters[ProjectorParamIndex].RangeValues[RangeIndex].Up);
+			return static_cast<FVector>(ProjectorParameters[ParameterIndexInInstance].RangeValues[RangeIndex].Up);
 		}
 	}
+	
+	LogParameterNotFoundWarning(ParamName, ParameterIndexInObject, ParameterIndexInInstance, CustomizableObject, __FUNCTION__);
 
 	return FVector(-0.0, -0.0, -0.0);	
 }
@@ -1718,20 +1724,21 @@ FVector FCustomizableObjectInstanceDescriptor::GetProjectorScale(const FString&
 	check(CustomizableObject);
 
 	const int32 ParameterIndexInObject = CustomizableObject->FindParameter(ParamName);
+	const int32 ParameterIndexInInstance = FindProjectorParameterNameIndex(ParamName);
 
-	const int32 ProjectorParamIndex = FindProjectorParameterNameIndex(ParamName);
-
-	if ((ParameterIndexInObject >= 0) && (ProjectorParamIndex >= 0))
+	if (ParameterIndexInObject >= 0 && ParameterIndexInInstance >= 0)
 	{
 		if (RangeIndex == -1)
 		{
-			return static_cast<FVector>(ProjectorParameters[ProjectorParamIndex].Value.Scale);
+			return static_cast<FVector>(ProjectorParameters[ParameterIndexInInstance].Value.Scale);
 		}
-		else if (ProjectorParameters[ProjectorParamIndex].RangeValues.IsValidIndex(RangeIndex))
+		else if (ProjectorParameters[ParameterIndexInInstance].RangeValues.IsValidIndex(RangeIndex))
 		{
-			return static_cast<FVector>(ProjectorParameters[ProjectorParamIndex].RangeValues[RangeIndex].Scale);
+			return static_cast<FVector>(ProjectorParameters[ParameterIndexInInstance].RangeValues[RangeIndex].Scale);
 		}
 	}
+	
+	LogParameterNotFoundWarning(ParamName, ParameterIndexInObject, ParameterIndexInInstance, CustomizableObject, __FUNCTION__);
 
 	return FVector(-0.0, -0.0, -0.0);
 }
@@ -1742,21 +1749,22 @@ float FCustomizableObjectInstanceDescriptor::GetProjectorAngle(const FString& Pa
 	check(CustomizableObject);
 
 	const int32 ParameterIndexInObject = CustomizableObject->FindParameter(ParamName);
+	const int32 ParameterIndexInInstance = FindProjectorParameterNameIndex(ParamName);
 
-	const int32 ProjectorParamIndex = FindProjectorParameterNameIndex(ParamName);
-
-	if (ParameterIndexInObject >= 0 && ProjectorParamIndex >= 0)
+	if (ParameterIndexInObject >= 0 && ParameterIndexInInstance >= 0)
 	{
 		if (RangeIndex == -1)
 		{
-			return ProjectorParameters[ProjectorParamIndex].Value.Angle;
+			return ProjectorParameters[ParameterIndexInInstance].Value.Angle;
 		}
-		else if (ProjectorParameters[ProjectorParamIndex].RangeValues.IsValidIndex(RangeIndex))
+		else if (ProjectorParameters[ParameterIndexInInstance].RangeValues.IsValidIndex(RangeIndex))
 		{
-			return ProjectorParameters[ProjectorParamIndex].RangeValues[RangeIndex].Angle;
+			return ProjectorParameters[ParameterIndexInInstance].RangeValues[RangeIndex].Angle;
 		}
 	}
 
+	LogParameterNotFoundWarning(ParamName, ParameterIndexInObject, ParameterIndexInInstance, CustomizableObject, __FUNCTION__);
+
 	return 0.0;
 }
 
@@ -1766,21 +1774,22 @@ ECustomizableObjectProjectorType FCustomizableObjectInstanceDescriptor::GetProje
 	check(CustomizableObject);
 
 	const int32 ParameterIndexInObject = CustomizableObject->FindParameter(ParamName);
+	const int32 ParameterIndexInInstance = FindProjectorParameterNameIndex(ParamName);
 
-	const int32 ProjectorParamIndex = FindProjectorParameterNameIndex(ParamName);
-
-	if (ParameterIndexInObject >= 0 && ProjectorParamIndex >= 0)
+	if (ParameterIndexInObject >= 0 && ParameterIndexInInstance >= 0)
 	{
 		if (RangeIndex == -1)
 		{
-			return ProjectorParameters[ProjectorParamIndex].Value.ProjectionType;
+			return ProjectorParameters[ParameterIndexInInstance].Value.ProjectionType;
 		}
-		else if (ProjectorParameters[ProjectorParamIndex].RangeValues.IsValidIndex(RangeIndex))
+		else if (ProjectorParameters[ParameterIndexInInstance].RangeValues.IsValidIndex(RangeIndex))
 		{
-			return ProjectorParameters[ProjectorParamIndex].RangeValues[RangeIndex].ProjectionType;
+			return ProjectorParameters[ParameterIndexInInstance].RangeValues[RangeIndex].ProjectionType;
 		}
 	}
 
+	LogParameterNotFoundWarning(ParamName, ParameterIndexInObject, ParameterIndexInInstance, CustomizableObject, __FUNCTION__);
+
 	return ECustomizableObjectProjectorType::Planar;
 }
 
@@ -1801,6 +1810,8 @@ FCustomizableObjectProjector FCustomizableObjectInstanceDescriptor::GetProjector
 		}
 	}
 
+	LogParameterNotFoundWarning(ParamName, Index, Index, CustomizableObject, __FUNCTION__);
+
 	return FCustomizableObjectProjectorParameterValue::DEFAULT_PARAMETER_VALUE;
 }
 
@@ -2153,15 +2164,20 @@ void FCustomizableObjectInstanceDescriptor::SetState(const int32 InState)
 void FCustomizableObjectInstanceDescriptor::SetCurrentState(const FString& StateName)
 {
 	check(CustomizableObject);
+	RETURN_ON_UNCOMPILED_CO(CustomizableObject, TEXT("Error: Cannot set state"))
 
 	const int32 Result = CustomizableObject->FindState(StateName);
-	if (ensureMsgf(Result != -1, TEXT("Unknown %s state."), *StateName))
+#if WITH_EDITOR
+	if (Result != INDEX_NONE)
+#else
+	if (ensureMsgf(Result != INDEX_NONE, TEXT("Unknown %s state."), *StateName))
+#endif
 	{
 		SetState(Result);
 	}
 	else
 	{
-		UE_LOG(LogMutable, Error, TEXT("%s: Unknown %s state."), StringCast<TCHAR>(__FUNCTION__).Get(), *StateName);
+		UE_LOG(LogMutable, Error, TEXT("%hs: Unknown %s state."), __FUNCTION__, *StateName);
 	}
 }
 
@@ -2169,6 +2185,7 @@ void FCustomizableObjectInstanceDescriptor::SetCurrentState(const FString& State
 void FCustomizableObjectInstanceDescriptor::SetRandomValues()
 {
 	check(CustomizableObject);
+	RETURN_ON_UNCOMPILED_CO(CustomizableObject, TEXT("Error: Cannot set random values"))
 
 	for (int32 i = 0; i < FloatParameters.Num(); ++i)
 	{
@@ -2203,7 +2220,11 @@ bool FCustomizableObjectInstanceDescriptor::CreateMultiLayerProjector(const FNam
 {
 	if (!FMultilayerProjector::AreDescriptorParametersValid(*this, ProjectorParamName.ToString()))
 	{
+#if WITH_EDITOR
+		UE_LOG(LogMutable, Error, TEXT("%s"), *FMultilayerProjector::DESCRIPTOR_PARAMETERS_INVALID);
+#else
 		ensureAlwaysMsgf(false, TEXT("%s"), *FMultilayerProjector::DESCRIPTOR_PARAMETERS_INVALID);
+#endif
 		return false;
 	}
 
@@ -2471,3 +2492,5 @@ const TArray<uint16>& FDescriptorRuntimeHash::GetRequestedLODs() const
 {
 	return RequestedLODsPerComponent;
 }
+
+#undef RETURN_ON_CO_UNCOMPILED
\ No newline at end of file
diff --git a/Engine/Plugins/Experimental/Mutable/Source/CustomizableObject/Private/MuCO/CustomizableObjectSystem.cpp b/Engine/Plugins/Experimental/Mutable/Source/CustomizableObject/Private/MuCO/CustomizableObjectSystem.cpp
--- a/Engine/Plugins/Experimental/Mutable/Source/CustomizableObject/Private/MuCO/CustomizableObjectSystem.cpp
+++ b/Engine/Plugins/Experimental/Mutable/Source/CustomizableObject/Private/MuCO/CustomizableObjectSystem.cpp
@@ -2884,34 +2884,29 @@ void UCustomizableObjectSystem::SetOnlyGenerateRequestedLODsEnabled(bool bIsEnab
 }
 
 
-void UCustomizableObjectSystem::AddUncompiledCOWarning(UCustomizableObject* InObject, FString const* OptionalLogInfo)
+void UCustomizableObjectSystem::AddUncompiledCOWarning(const UCustomizableObject& InObject, FString const* OptionalLogInfo)
 {
-	if (!InObject)
-	{
-		return;
-	}
-
 	FString Msg;
-	Msg += FString::Printf(TEXT("Warning: Customizable Object [%s] not compiled. Please compile and save the object."), *InObject->GetName());
-	GEngine->AddOnScreenDebugMessage((uint64)((PTRINT)InObject), 10.0f, FColor::Red, Msg);
+	Msg += FString::Printf(TEXT("Warning: Customizable Object [%s] not compiled. Please compile and save the object."), *InObject.GetName());
+	GEngine->AddOnScreenDebugMessage((uint64)((PTRINT)&InObject), 10.0f, FColor::Red, Msg);
 
 #if WITH_EDITOR
 	const FString ErrorString = FString::Printf(
 		TEXT("Customizable Object [%s] not compiled.  Compile via the editor or via code before instancing.  %s"),
-						   *InObject->GetName(), OptionalLogInfo ? **OptionalLogInfo : TEXT(""));
+						   *InObject.GetName(), OptionalLogInfo ? **OptionalLogInfo : TEXT(""));
 #else // !WITH_EDITOR
 	const FString ErrorString = FString::Printf(
 		TEXT("Customizable Object [%s] not compiled.  This is not an Editor build, so this is an unrecoverable bad state; could be due to code or a cook failure.  %s"),
-						   *InObject->GetName(), OptionalLogInfo ? **OptionalLogInfo : TEXT(""));
+						   *InObject.GetName(), OptionalLogInfo ? **OptionalLogInfo : TEXT(""));
 #endif
 
 	// Also log an error so if this happens as part of a bug report we'll have this info.
 	UE_LOG(LogMutable, Error, TEXT("%s"), *ErrorString);
 
 #if WITH_EDITOR
-	if (UncompiledCustomizableObjectIds.Find(InObject->GetVersionId()) == INDEX_NONE)
+	if (UncompiledCustomizableObjectIds.Find(InObject.GetVersionId()) == INDEX_NONE)
 	{
-		UncompiledCustomizableObjectIds.Add(InObject->GetVersionId());
+		UncompiledCustomizableObjectIds.Add(InObject.GetVersionId());
 
 		FMessageLog MessageLog("Mutable");
 		MessageLog.Warning(FText::FromString(Msg));
diff --git a/Engine/Plugins/Experimental/Mutable/Source/CustomizableObject/Public/MuCO/CustomizableObjectInstanceDescriptor.h b/Engine/Plugins/Experimental/Mutable/Source/CustomizableObject/Public/MuCO/CustomizableObjectInstanceDescriptor.h
--- a/Engine/Plugins/Experimental/Mutable/Source/CustomizableObject/Public/MuCO/CustomizableObjectInstanceDescriptor.h
+++ b/Engine/Plugins/Experimental/Mutable/Source/CustomizableObject/Public/MuCO/CustomizableObjectInstanceDescriptor.h
@@ -333,7 +333,6 @@ struct CUSTOMIZABLEOBJECT_API FCustomizableObjectInstanceDescriptor
 	void MultilayerProjectorUpdateVirtualLayer(const FName& ProjectorParamName, const FName& Id, const FMultilayerProjectorVirtualLayer& Layer);
 
 private:
-	void AddUncompiledCOWarning(const FString& AdditionalLoggingInfo);
 
 	UPROPERTY()
 	TObjectPtr<UCustomizableObject> CustomizableObject = nullptr;
