#rnx

diff --git a/Engine/Plugins/Experimental/Mutable/Source/CustomizableObject/Private/MuCO/CustomizableObjectSystem.cpp b/Engine/Plugins/Experimental/Mutable/Source/CustomizableObject/Private/MuCO/CustomizableObjectSystem.cpp
--- a/Engine/Plugins/Experimental/Mutable/Source/CustomizableObject/Private/MuCO/CustomizableObjectSystem.cpp
+++ b/Engine/Plugins/Experimental/Mutable/Source/CustomizableObject/Private/MuCO/CustomizableObjectSystem.cpp
@@ -993,58 +993,68 @@ bool UCustomizableObjectSystem::LockObject(const class UCustomizableObject* InOb
 	check(!InObject->GetPrivate()->bLocked);
 	check(IsInGameThread() && !IsInParallelGameThread());
 
-	// If the current instance is for this object, make the lock fail by returning false
-	if (Private->CurrentInstanceBeingUpdated &&
-		Private->CurrentInstanceBeingUpdated->GetCustomizableObject() == InObject)
+	if (InObject && InObject->GetPrivate() && Private)
 	{
-		UE_LOG(LogMutable, Warning, TEXT("---- failed to lock object %s"), *InObject->GetName());
+		// If the current instance is for this object, make the lock fail by returning false
+		if (Private->CurrentInstanceBeingUpdated &&
+			Private->CurrentInstanceBeingUpdated->GetCustomizableObject() == InObject)
+		{
+			UE_LOG(LogMutable, Warning, TEXT("---- failed to lock object %s"), *InObject->GetName());
 
-		return false;
-	}
+			return false;
+		}
 
-	FString Message = FString::Printf(TEXT("Customizable Object %s has pending texture streaming operations. Please wait a few seconds and try again."), 
-		*InObject->GetName());
+		FString Message = FString::Printf(TEXT("Customizable Object %s has pending texture streaming operations. Please wait a few seconds and try again."),
+			*InObject->GetName());
 
-	// Pre-check pending operations before locking. This check is redundant and incomplete because it's checked again after locking 
-	// and some operations may start between here and the actual lock. But in the CO Editor preview it will prevent some 
-	// textures getting stuck at low resolution when they try to update mips and are cancelled when the user presses 
-	// the compile button but the compilation quits anyway because there are pending operations
-	if (CheckIfDiskOrMipUpdateOperationsPending(*InObject))
-	{
-		UE_LOG(LogMutable, Warning, TEXT("%s"), *Message);
+		// Pre-check pending operations before locking. This check is redundant and incomplete because it's checked again after locking 
+		// and some operations may start between here and the actual lock. But in the CO Editor preview it will prevent some 
+		// textures getting stuck at low resolution when they try to update mips and are cancelled when the user presses 
+		// the compile button but the compilation quits anyway because there are pending operations
+		if (CheckIfDiskOrMipUpdateOperationsPending(*InObject))
+		{
+			UE_LOG(LogMutable, Warning, TEXT("%s"), *Message);
 
-		return false;
-	}
+			return false;
+		}
 
-	// Lock the object, no new file or mip streaming operations should start from this point
-	InObject->GetPrivate()->bLocked = true;
+		// Lock the object, no new file or mip streaming operations should start from this point
+		InObject->GetPrivate()->bLocked = true;
 
-	// But some could have started between the first CheckIfDiskOrMipUpdateOperationsPending and the lock a few lines back, so check again
-	if (CheckIfDiskOrMipUpdateOperationsPending(*InObject))
-	{
-		UE_LOG(LogMutable, Warning, TEXT("%s"), *Message);
+		// But some could have started between the first CheckIfDiskOrMipUpdateOperationsPending and the lock a few lines back, so check again
+		if (CheckIfDiskOrMipUpdateOperationsPending(*InObject))
+		{
+			UE_LOG(LogMutable, Warning, TEXT("%s"), *Message);
 
-		// Unlock and return because the pending operations cannot be easily stopped now, the compilation hasn't started and the CO
-		// hasn't changed state yet. It's simpler to quit the compilation, unlock and let the user try to compile again
-		InObject->GetPrivate()->bLocked = false;
+			// Unlock and return because the pending operations cannot be easily stopped now, the compilation hasn't started and the CO
+			// hasn't changed state yet. It's simpler to quit the compilation, unlock and let the user try to compile again
+			InObject->GetPrivate()->bLocked = false;
 
-		return false;
-	}
+			return false;
+		}
 
-	// Ensure that we don't try to handle any further streaming operations for this object
-	check(GetPrivate() != nullptr);
-	if (GetPrivate()->Streamer)
-	{
-		GetPrivate()->Streamer->CancelStreamingForObject(InObject);
-	}
+		// Ensure that we don't try to handle any further streaming operations for this object
+		check(GetPrivate() != nullptr);
+		if (GetPrivate()->Streamer)
+		{
+			GetPrivate()->Streamer->CancelStreamingForObject(InObject);
+		}
 
-	// Clear the cache for the instance, since we will remake it
-	FMutableResourceCache& Cache = GetPrivate()->GetObjectCache(InObject);
-	Cache.Clear();
+		// Clear the cache for the instance, since we will remake it
+		FMutableResourceCache& Cache = GetPrivate()->GetObjectCache(InObject);
+		Cache.Clear();
 
-	check(InObject->GetPrivate()->bLocked);
+		check(InObject->GetPrivate()->bLocked);
 
-	return true;
+		return true;
+	}
+	else
+	{
+		FString ObjectName = InObject ? InObject->GetName() : FString("null");
+		UE_LOG(LogMutable, Warning, TEXT("Failed to lock the object [%s] because it was null or the system was null or partially destroyed."), *ObjectName);
+
+		return false;
+	}
 }
 
 
@@ -1053,7 +1063,11 @@ void UCustomizableObjectSystem::UnlockObject(const class UCustomizableObject* Ob
 	check(Obj != nullptr);
 	check(Obj->GetPrivate()->bLocked);
 	check(IsInGameThread() && !IsInParallelGameThread());
-	Obj->GetPrivate()->bLocked = false;
+	
+	if (Obj && Obj->GetPrivate())
+	{
+		Obj->GetPrivate()->bLocked = false;
+	}
 }
 
 
diff --git a/Engine/Plugins/Experimental/Mutable/Source/CustomizableObjectEditor/Private/MuCOE/CustomizableObjectCompiler.cpp b/Engine/Plugins/Experimental/Mutable/Source/CustomizableObjectEditor/Private/MuCOE/CustomizableObjectCompiler.cpp
--- a/Engine/Plugins/Experimental/Mutable/Source/CustomizableObjectEditor/Private/MuCOE/CustomizableObjectCompiler.cpp
+++ b/Engine/Plugins/Experimental/Mutable/Source/CustomizableObjectEditor/Private/MuCOE/CustomizableObjectCompiler.cpp
@@ -293,6 +293,11 @@ void FCustomizableObjectCompiler::AddReferencedObjects(FReferenceCollector& Coll
 	{
 		Collector.AddReferencedObject(ArrayGCProtect[i]);
 	}
+
+	if (CurrentObject)
+	{
+		Collector.AddReferencedObject(CurrentObject);
+	}
 }
 
 
@@ -1481,14 +1486,25 @@ void FCustomizableObjectCompiler::ForceFinishCompilation()
 		CleanCachedReferencers();
 		UpdateArrayGCProtect();
 
-		UCustomizableObjectSystem::GetInstance()->UnlockObject(CurrentObject);
+		UCustomizableObjectSystem* System = UCustomizableObjectSystem::GetInstance();
+
+		// Check the system has not started BeginDestroy, it can crash during a GC during a compilation forced shutdown without this check
+		if (!System->HasAnyFlags(EObjectFlags::RF_BeginDestroyed))
+		{
+			System->UnlockObject(CurrentObject);
+		}
 	}
 
 	else if (SaveDDTask.IsValid())
 	{
 		SaveDDThread->WaitForCompletion();
 
-		UCustomizableObjectSystem::GetInstance()->UnlockObject(CurrentObject);
+		UCustomizableObjectSystem* System = UCustomizableObjectSystem::GetInstance();
+
+		if (!System->HasAnyFlags(EObjectFlags::RF_BeginDestroyed))
+		{
+			System->UnlockObject(CurrentObject);
+		}
 	}
 
 	RemoveCompileNotification();
diff --git a/Engine/Plugins/Experimental/Mutable/Source/CustomizableObjectEditor/Private/MuCOE/CustomizableObjectCompiler.h b/Engine/Plugins/Experimental/Mutable/Source/CustomizableObjectEditor/Private/MuCOE/CustomizableObjectCompiler.h
--- a/Engine/Plugins/Experimental/Mutable/Source/CustomizableObjectEditor/Private/MuCOE/CustomizableObjectCompiler.h
+++ b/Engine/Plugins/Experimental/Mutable/Source/CustomizableObjectEditor/Private/MuCOE/CustomizableObjectCompiler.h
@@ -125,7 +125,7 @@ private:
 	// Cache configuration settings from ini files 
 	ECustomizableObjectNumBoneInfluences CustomizableObjectNumBoneInfluences = ECustomizableObjectNumBoneInfluences::Four;
 
-	// TODO: add references
+	// Protected from GC with FCustomizableObjectCompiler::AddReferencedObjects
 	UCustomizableObject* CurrentObject = nullptr;
 
 	/** Array where to put the names of the already processed child in ProcessChildObjectsRecursively */
