Fix uninitialized property warnings
#rnx

diff --git a/Engine/Plugins/Experimental/Chooser/Source/Chooser/Public/ChooserPropertyAccess.h b/Engine/Plugins/Experimental/Chooser/Source/Chooser/Public/ChooserPropertyAccess.h
--- a/Engine/Plugins/Experimental/Chooser/Source/Chooser/Public/ChooserPropertyAccess.h
+++ b/Engine/Plugins/Experimental/Chooser/Source/Chooser/Public/ChooserPropertyAccess.h
@@ -46,7 +46,7 @@ struct FChooserPropertyBinding
 	TArray<FName> PropertyBindingChain;
 	
 	UPROPERTY()
-	int ContextIndex;
+	int ContextIndex = 0;
 
 #if WITH_EDITORONLY_DATA
 	UPROPERTY()
@@ -102,7 +102,7 @@ struct FContextObjectTypeBase
 	GENERATED_BODY()
 	
 	UPROPERTY(EditAnywhere, Category="Type")
-	EContextObjectDirection Direction;
+	EContextObjectDirection Direction = EContextObjectDirection::Read;
 };
 
 
diff --git a/Engine/Plugins/Experimental/Mutable/Source/CustomizableObject/Classes/MuCO/CustomizableObject.h b/Engine/Plugins/Experimental/Mutable/Source/CustomizableObject/Classes/MuCO/CustomizableObject.h
--- a/Engine/Plugins/Experimental/Mutable/Source/CustomizableObject/Classes/MuCO/CustomizableObject.h
+++ b/Engine/Plugins/Experimental/Mutable/Source/CustomizableObject/Classes/MuCO/CustomizableObject.h
@@ -843,11 +843,11 @@ struct FMutableRefSocket
 	FName BoneName;
 
 	UPROPERTY()
-	FVector RelativeLocation;
+	FVector RelativeLocation = FVector::ZeroVector;
 	UPROPERTY()
-	FRotator RelativeRotation;
+	FRotator RelativeRotation = FRotator::ZeroRotator;
 	UPROPERTY()
-	FVector RelativeScale;
+	FVector RelativeScale = FVector::ZeroVector;;
 
 	UPROPERTY()
 	bool bForceAlwaysAnimated = false;
@@ -956,7 +956,7 @@ struct FMutableRefSkeletalMeshData
 
 	// Bounding Box
 	UPROPERTY()
-	FBoxSphereBounds Bounds;
+	FBoxSphereBounds Bounds = FBoxSphereBounds(ForceInitToZero);
 
 	// Settings
 	UPROPERTY()
diff --git a/Engine/Plugins/Experimental/Mutable/Source/CustomizableObject/Classes/MuCO/CustomizableObjectInstance.h b/Engine/Plugins/Experimental/Mutable/Source/CustomizableObject/Classes/MuCO/CustomizableObjectInstance.h
--- a/Engine/Plugins/Experimental/Mutable/Source/CustomizableObject/Classes/MuCO/CustomizableObjectInstance.h
+++ b/Engine/Plugins/Experimental/Mutable/Source/CustomizableObject/Classes/MuCO/CustomizableObjectInstance.h
@@ -80,7 +80,7 @@ struct FUpdateContext
 	GENERATED_BODY()
 
 	UPROPERTY()
-	EUpdateResult UpdateResult;
+	EUpdateResult UpdateResult = EUpdateResult::Success;
 };
 
 
diff --git a/Engine/Plugins/Experimental/Mutable/Source/CustomizableObject/Private/MuCO/CustomizableInstancePrivateData.h b/Engine/Plugins/Experimental/Mutable/Source/CustomizableObject/Private/MuCO/CustomizableInstancePrivateData.h
--- a/Engine/Plugins/Experimental/Mutable/Source/CustomizableObject/Private/MuCO/CustomizableInstancePrivateData.h
+++ b/Engine/Plugins/Experimental/Mutable/Source/CustomizableObject/Private/MuCO/CustomizableInstancePrivateData.h
@@ -149,7 +149,7 @@ struct FAnimInstanceOverridePhysicsAsset
 	GENERATED_BODY()
 
 	UPROPERTY(Transient)
-	int32 PropertyIndex;
+	int32 PropertyIndex = 0;
 
 	UPROPERTY(Transient)
 	TObjectPtr<UPhysicsAsset> PhysicsAsset;
diff --git a/Engine/Plugins/Experimental/Mutable/Source/CustomizableObject/Public/MuCO/CustomizableObjectClothingTypes.h b/Engine/Plugins/Experimental/Mutable/Source/CustomizableObject/Public/MuCO/CustomizableObjectClothingTypes.h
--- a/Engine/Plugins/Experimental/Mutable/Source/CustomizableObject/Public/MuCO/CustomizableObjectClothingTypes.h
+++ b/Engine/Plugins/Experimental/Mutable/Source/CustomizableObject/Public/MuCO/CustomizableObjectClothingTypes.h
@@ -53,7 +53,7 @@ struct FCustomizableObjectClothingAssetData
 	TArray<int32> UsedBoneIndices;
 
 	UPROPERTY()
-	int32 ReferenceBoneIndex;
+	int32 ReferenceBoneIndex = 0;
 
 	UPROPERTY()
 	TArray<FCustomizableObjectClothConfigData> ConfigsData;
diff --git a/Engine/Plugins/Runtime/MassGameplay/Source/MassSpawner/Public/MassEntityConfigAsset.h b/Engine/Plugins/Runtime/MassGameplay/Source/MassSpawner/Public/MassEntityConfigAsset.h
--- a/Engine/Plugins/Runtime/MassGameplay/Source/MassSpawner/Public/MassEntityConfigAsset.h
+++ b/Engine/Plugins/Runtime/MassGameplay/Source/MassSpawner/Public/MassEntityConfigAsset.h
@@ -93,7 +93,7 @@ protected:
 	TObjectPtr<UObject> ConfigOwner = nullptr;
 
 private:
-	UPROPERTY(VisibleAnywhere, Category="Mass")
+	UPROPERTY(VisibleAnywhere, Category="Mass", meta = (IgnoreForMemberInitializationTest))
 	FGuid ConfigGuid;
 
 private:
