[mutable] Fix ignored parameter relevancy build setting

#rnx

diff --git a/Engine/Plugins/Experimental/Mutable/Source/CustomizableObject/Classes/MuCO/CustomizableObjectInstance.h b/Engine/Plugins/Experimental/Mutable/Source/CustomizableObject/Classes/MuCO/CustomizableObjectInstance.h
--- a/Engine/Plugins/Experimental/Mutable/Source/CustomizableObject/Classes/MuCO/CustomizableObjectInstance.h
+++ b/Engine/Plugins/Experimental/Mutable/Source/CustomizableObject/Classes/MuCO/CustomizableObjectInstance.h
@@ -765,7 +765,7 @@ private:
    	TMap<FName, FMultilayerProjector> MultilayerProjectors_DEPRECATED;
 
 	/** If this is set to true, when updating the instance an additional step will be performed to calculate the list of instance parameters that are relevant for the current parameter vaules. */
-	bool bBuildParameterRelevancy = false;
+	bool bBuildParameterRelevancy_DEPRECATED = false;
 };
 
 #if WITH_EDITOR
diff --git a/Engine/Plugins/Experimental/Mutable/Source/CustomizableObject/Private/MuCO/CustomizableObjectInstance.cpp b/Engine/Plugins/Experimental/Mutable/Source/CustomizableObject/Private/MuCO/CustomizableObjectInstance.cpp
--- a/Engine/Plugins/Experimental/Mutable/Source/CustomizableObject/Private/MuCO/CustomizableObjectInstance.cpp
+++ b/Engine/Plugins/Experimental/Mutable/Source/CustomizableObject/Private/MuCO/CustomizableObjectInstance.cpp
@@ -639,13 +639,13 @@ bool UCustomizableObjectInstance::GetBuildParameterRelevancy() const
 #if WITH_EDITOR
 	return true;
 #endif
-	return bBuildParameterRelevancy;
+	return Descriptor.bBuildParameterRelevancy;
 }
 
 
 void UCustomizableObjectInstance::SetBuildParameterRelevancy(bool Value)
 {
-	bBuildParameterRelevancy = Value;
+	Descriptor.SetBuildParameterRelevancy(Value);
 }
 
 
