[Mutable] Implemented commandlet to run tests over selected Customizable Objects.
- Compiles the CO and then generates an X amount of random (deterministic) instances.
- Added deterministic way of setting random values for the instance's parameters.

diff --git a/Engine/Plugins/Experimental/Mutable/Source/CustomizableObject/Classes/MuCO/CustomizableObjectInstance.h b/Engine/Plugins/Experimental/Mutable/Source/CustomizableObject/Classes/MuCO/CustomizableObjectInstance.h
--- a/Engine/Plugins/Experimental/Mutable/Source/CustomizableObject/Classes/MuCO/CustomizableObjectInstance.h
+++ b/Engine/Plugins/Experimental/Mutable/Source/CustomizableObject/Classes/MuCO/CustomizableObjectInstance.h
@@ -255,10 +255,13 @@ public:
 	UFUNCTION(BlueprintCallable, Category = CustomizableObjectInstance)
 	bool HasAnyParameters() const;
 	
-	//! Set random values to the parameters. Useful for testing only.
+	/** Set random values to the parameters. Useful for testing only. */
 	UFUNCTION(BlueprintCallable, Category = CustomizableObjectInstance)
 	void SetRandomValues();
 
+	/**  Set random values to the parameters using a seed. Useful for testing only. */
+	void SetRandomValues(const int32 InRandomizationSeed);
+
 	// Utilities to manage saving and loading parameters from profiles.
 	bool LoadParametersFromProfile(int32 ProfileIndex);
 	bool SaveParametersToProfile(int32 ProfileIndex);
diff --git a/Engine/Plugins/Experimental/Mutable/Source/CustomizableObject/Private/MuCO/CustomizableObjectInstance.cpp b/Engine/Plugins/Experimental/Mutable/Source/CustomizableObject/Private/MuCO/CustomizableObjectInstance.cpp
--- a/Engine/Plugins/Experimental/Mutable/Source/CustomizableObject/Private/MuCO/CustomizableObjectInstance.cpp
+++ b/Engine/Plugins/Experimental/Mutable/Source/CustomizableObject/Private/MuCO/CustomizableObjectInstance.cpp
@@ -2653,11 +2653,14 @@ int32 UCustomizableObjectInstance::FindProjectorParameterNameIndex(const FString
 	return Descriptor.FindProjectorParameterNameIndex(ParamName);
 }
 
-
 void UCustomizableObjectInstance::SetRandomValues()
 {
-	Descriptor.SetRandomValues();
+	Descriptor.SetRandomValues(FMath::Rand());
+}
 
+void UCustomizableObjectInstance::SetRandomValues(const int32 InRandomizationSeed)
+{
+	Descriptor.SetRandomValues(InRandomizationSeed);
 }
 
 
diff --git a/Engine/Plugins/Experimental/Mutable/Source/CustomizableObject/Private/MuCO/CustomizableObjectInstanceDescriptor.cpp b/Engine/Plugins/Experimental/Mutable/Source/CustomizableObject/Private/MuCO/CustomizableObjectInstanceDescriptor.cpp
--- a/Engine/Plugins/Experimental/Mutable/Source/CustomizableObject/Private/MuCO/CustomizableObjectInstanceDescriptor.cpp
+++ b/Engine/Plugins/Experimental/Mutable/Source/CustomizableObject/Private/MuCO/CustomizableObjectInstanceDescriptor.cpp
@@ -2182,11 +2182,14 @@ void FCustomizableObjectInstanceDescriptor::SetCurrentState(const FString& State
 }
 
 
-void FCustomizableObjectInstanceDescriptor::SetRandomValues()
+void FCustomizableObjectInstanceDescriptor::SetRandomValues(const int32& InRandomizationSeed)
 {
 	check(CustomizableObject);
 	RETURN_ON_UNCOMPILED_CO(CustomizableObject, TEXT("Error: Cannot set random values"))
 
+	// Set seed for deterministic behaviour
+	FMath::SRandInit(InRandomizationSeed);
+	
 	for (int32 i = 0; i < FloatParameters.Num(); ++i)
 	{
 		FloatParameters[i].ParameterValue = FMath::SRand();
@@ -2194,7 +2197,7 @@ void FCustomizableObjectInstanceDescriptor::SetRandomValues()
 
 	for (int32 i = 0; i < BoolParameters.Num(); ++i)
 	{
-		BoolParameters[i].ParameterValue = FMath::Rand() % 2 == 0;
+		BoolParameters[i].ParameterValue = (int32)FMath::SRand() % 2 == 0;
 	}
 
 	for (int32 i = 0; i < IntParameters.Num(); ++i)
@@ -2207,7 +2210,7 @@ void FCustomizableObjectInstanceDescriptor::SetRandomValues()
 			const int32 NumValues = CustomizableObject->GetIntParameterNumOptions(ParameterIndexInCO);
 			if (NumValues > 0)
 			{
-				const int32 Index = FMath::Rand() % NumValues;
+				const int32 Index = (int32)FMath::SRand() % NumValues;
 				FString Option = CustomizableObject->GetIntParameterAvailableOption(ParameterIndexInCO, Index);
 				SetIntParameterSelectedOption(i, Option);
 			}
diff --git a/Engine/Plugins/Experimental/Mutable/Source/CustomizableObject/Private/MuCO/UnrealMutableImageProvider.cpp b/Engine/Plugins/Experimental/Mutable/Source/CustomizableObject/Private/MuCO/UnrealMutableImageProvider.cpp
--- a/Engine/Plugins/Experimental/Mutable/Source/CustomizableObject/Private/MuCO/UnrealMutableImageProvider.cpp
+++ b/Engine/Plugins/Experimental/Mutable/Source/CustomizableObject/Private/MuCO/UnrealMutableImageProvider.cpp
@@ -164,7 +164,7 @@ mu::EImageFormat GetMutablePixelFormat(EPixelFormat InTextureFormat)
 		if (!GlobalExternalImages.Contains(id))
 		{
 			// Null case, no image was provided
-			UE_LOG(LogMutable, Warning, TEXT("Failed to get external image. GlobalExternalImage not found."));
+			UE_LOG(LogMutable, Warning, TEXT("Failed to get external image. GlobalExternalImage not found with id %s."), *id);
 
 			ResultCallback(CreateDummy());
 			return Invoke(TrivialReturn);
diff --git a/Engine/Plugins/Experimental/Mutable/Source/CustomizableObject/Public/MuCO/CustomizableObjectInstanceDescriptor.h b/Engine/Plugins/Experimental/Mutable/Source/CustomizableObject/Public/MuCO/CustomizableObjectInstanceDescriptor.h
--- a/Engine/Plugins/Experimental/Mutable/Source/CustomizableObject/Public/MuCO/CustomizableObjectInstanceDescriptor.h
+++ b/Engine/Plugins/Experimental/Mutable/Source/CustomizableObject/Public/MuCO/CustomizableObjectInstanceDescriptor.h
@@ -280,7 +280,7 @@ struct CUSTOMIZABLEOBJECT_API FCustomizableObjectInstanceDescriptor
 
 	// ------------------------------------------------------------
 	
-	void SetRandomValues();
+	void SetRandomValues(const int32& InRandomizationSeed);
 
 	// ------------------------------------------------------------
 	// Multilayer Projectors
diff --git a/Engine/Plugins/Experimental/Mutable/Source/MutableValidation/MutableValidation.Build.cs b/Engine/Plugins/Experimental/Mutable/Source/MutableValidation/MutableValidation.Build.cs
--- a/Engine/Plugins/Experimental/Mutable/Source/MutableValidation/MutableValidation.Build.cs
+++ b/Engine/Plugins/Experimental/Mutable/Source/MutableValidation/MutableValidation.Build.cs
@@ -13,7 +13,11 @@ namespace UnrealBuildTool.Rules
 
 			DefaultBuildSettings = BuildSettingsVersion.V2;
 
-			PublicDependencyModuleNames.AddRange(new string[] { "Settings" });
+			PublicDependencyModuleNames.AddRange(
+				new string[] {
+					"Settings", 
+					"Engine"
+				});
 			
 			PrivateDependencyModuleNames.AddRange(
 				new string[] {
diff --git a/Engine/Plugins/Experimental/Mutable/Source/MutableValidation/Private/MuV/CustomizableObjectValidationCommandlet.cpp b/Engine/Plugins/Experimental/Mutable/Source/MutableValidation/Private/MuV/CustomizableObjectValidationCommandlet.cpp
new file mode 100644
--- /dev/null
+++ b/Engine/Plugins/Experimental/Mutable/Source/MutableValidation/Private/MuV/CustomizableObjectValidationCommandlet.cpp
@@ -0,0 +1,177 @@
+// Copyright Epic Games, Inc. All Rights Reserved.
+
+#include "MuV/CustomizableObjectValidationCommandlet.h"
+
+#include "Containers/Array.h"
+#include "MuCO/CustomizableObject.h"
+#include "MuCO/CustomizableObjectSystem.h"
+
+
+int32 UCustomizableObjectValidationCommandlet::Main(const FString& Params)
+{
+	// Execution arguments for commandlet from IDE
+	// $(LocalDebuggerCommandArguments) -run=CustomizableObjectValidation -CustomizableObject=(PathToCO)
+	
+	// Get the package name of the CO to test
+	FString CustomizableObjectAssetPath = "";
+	if (!FParse::Value(*Params, TEXT("CustomizableObject="), CustomizableObjectAssetPath))
+	{
+		UE_LOG(LogMutable,Error,TEXT("Failed to parse Customizable Object package name from provided argument : %s"),*Params)
+		return 1;
+	}
+	
+	// Get the amount of instances to generate if parameter was provided
+	uint32 InstancesToGenerate = 16;
+	if (!FParse::Value(*Params, TEXT("InstanceGenerationCount="),InstancesToGenerate))
+	{
+		UE_LOG(LogMutable,Display,TEXT("Instance generation count not specified. Using default value : %u"),InstancesToGenerate);
+	}
+	
+	// Load the resource
+	UObject* FoundObject = FSoftObjectPath(CustomizableObjectAssetPath).TryLoad();
+	if (!FoundObject)
+	{
+		UE_LOG(LogMutable,Error,TEXT("Failed to retrieve UObject from path %s"),*CustomizableObjectAssetPath);
+		return 1;
+	}
+	
+	// Get the CustomizableObject.
+	ToTestCustomizableObject = Cast<UCustomizableObject>(FoundObject);
+	if (!ToTestCustomizableObject)
+	{
+		UE_LOG(LogMutable,Error,TEXT("Failed to cast found UObject to UCustomizableObject."));
+		return 1;
+	}
+	
+	// Compile the Customizable Object ------------------------------------------------------------------------------ //
+	bool bWasCoCompilationSuccesfull = false;
+	{
+		UE_LOG(LogMutable,Display,TEXT("Compiling Customizable Object..."));
+    	
+    	// Request a compiler to be able to locate the root and to compile it
+    	const TUniquePtr<FCustomizableObjectCompilerBase> Compiler =
+    		TUniquePtr<FCustomizableObjectCompilerBase>(UCustomizableObjectSystem::GetInstance()->GetNewCompiler());
+    		
+    	// Compile the CO with the provided compilation options
+    	// Run Sync compilation -> Warning : Potentially long operation -------------
+    	Compiler->Compile(*ToTestCustomizableObject, ToTestCustomizableObject->CompileOptions, false);
+    	// --------------------------------------------------------------------------
+    		
+    	// Get the compilation result
+    	const ECustomizableObjectCompilationState CompilationEndResult = Compiler->GetCompilationState();
+    	check(CompilationEndResult != ECustomizableObjectCompilationState::None);
+    	check(CompilationEndResult != ECustomizableObjectCompilationState::InProgress);
+    	
+    	bWasCoCompilationSuccesfull = CompilationEndResult == ECustomizableObjectCompilationState::Completed;
+	}
+	// -------------------------------------------------------------------------------------------------------------- //
+	
+	if (bWasCoCompilationSuccesfull)
+	{
+		UE_LOG(LogMutable,Display,TEXT("Customizable Object was compiled succesfully."));
+				
+		// Generate target random instances to be tested ------------------------------------------------------------ //
+		bool bWasInstancesCreationSuccessful = true;
+		{
+			UE_LOG(LogMutable,Display,TEXT("Generating %i random instances..."), InstancesToGenerate);	
+			
+			// Generate a series of instances to later update
+			for (uint32 CustomizableObjectInstanceIndex = 0; CustomizableObjectInstanceIndex < InstancesToGenerate; CustomizableObjectInstanceIndex++)
+			{
+				UCustomizableObjectInstance* GeneratedInstance = ToTestCustomizableObject->CreateInstance();
+				if (GeneratedInstance)
+				{
+					// Randomize instance values
+					GeneratedInstance->SetRandomValues(0);
+					InstancesToProcess.Push(GeneratedInstance);
+				}
+				else
+				{
+					UE_LOG(LogMutable,Error,TEXT("Failed to generate COI for Customizable Object with name : %s ."),*ToTestCustomizableObject->GetName());
+					bWasInstancesCreationSuccessful = false;
+				}
+			}
+		}
+		// ---------------------------------------------------------------------------------------------------------- //
+		
+		// Update the instances generated --------------------------------------------------------------------------- //
+		{
+			UE_LOG(LogMutable,Display,TEXT("Updating generated instances..."));
+            	
+            // Now update the instances one by one
+            while (!InstancesToProcess.IsEmpty() || InstanceBeingUpdated)
+            {
+            	// Tick the engine
+            	CommandletHelpers::TickEngine();
+    
+            	// Stop if exit was requested
+            	if (IsEngineExitRequested())
+            	{
+            		break;
+            	}
+            
+            	// Wait until current instance turns invalid
+            	if (InstanceBeingUpdated)
+            	{
+            		continue;
+            	}
+           
+            	// We may have already updated the instance so ensure we have more instances to work with
+            	if (InstancesToProcess.IsEmpty())
+            	{
+            		continue;
+            	}
+            	
+            	InstanceBeingUpdated = InstancesToProcess[0];
+            	InstancesToProcess.RemoveAt(0);
+            	if (InstanceBeingUpdated)
+            	{
+            		UE_LOG(LogMutable,Display,TEXT("Invoking update for %s instance."),*InstanceBeingUpdated->GetName());
+            		InstanceUpdateDelegate.BindDynamic(this, &UCustomizableObjectValidationCommandlet::OnInstanceUpdate);
+            		InstanceBeingUpdated->UpdateSkeletalMeshAsyncResult(InstanceUpdateDelegate);
+            	}
+            }
+		}	
+		// ---------------------------------------------------------------------------------------------------------- //
+
+		// Compute instance update result
+		const bool bInstancesTestedSuccessfully = !bInstanceFailedUpdate && bWasInstancesCreationSuccessful;
+        if (bInstancesTestedSuccessfully)
+        {
+        	UE_LOG(LogMutable,Display,TEXT("Generation of Customizable object instances was succesfull."));
+        }
+        else
+        {
+        	UE_LOG(LogMutable,Error,TEXT("The generation of Customizable object instances  was not succesfull."));
+        }
+	}
+	else
+	{
+		UE_LOG(LogMutable,Error,TEXT("The compilation of the Customizable object was not succesfull."));
+	}
+	
+	
+	// If something failed then fail the commandlet execution
+	UE_LOG(LogMutable,Display,TEXT("Mutable commandlet finished."));
+	return 0;
+}
+
+
+void UCustomizableObjectValidationCommandlet::OnInstanceUpdate(const FUpdateContext& Result)
+{
+	const EUpdateResult InstanceUpdateResult = Result.UpdateResult;
+	if (InstanceUpdateResult == EUpdateResult::Success)
+	{
+		UE_LOG(LogMutable,Display,TEXT("Instance %s finished update succesfully."),*InstanceBeingUpdated->GetName());
+	}
+	else
+	{
+		const FString OutputStatus = UEnum::GetValueAsString(Result.UpdateResult);
+		UE_LOG(LogMutable,Error,TEXT("Instance %s finished update with anomalous state : %s."),*InstanceBeingUpdated->GetName(),*OutputStatus);
+		bInstanceFailedUpdate = true;
+	}
+	
+	InstanceUpdateDelegate.Unbind();
+	InstanceBeingUpdated = nullptr;
+}
+
diff --git a/Engine/Plugins/Experimental/Mutable/Source/MutableValidation/Public/MuV/CustomizableObjectValidationCommandlet.h b/Engine/Plugins/Experimental/Mutable/Source/MutableValidation/Public/MuV/CustomizableObjectValidationCommandlet.h
new file mode 100644
--- /dev/null
+++ b/Engine/Plugins/Experimental/Mutable/Source/MutableValidation/Public/MuV/CustomizableObjectValidationCommandlet.h
@@ -0,0 +1,49 @@
+// Copyright Epic Games, Inc. All Rights Reserved.
+
+#pragma once
+
+#include "Commandlets/Commandlet.h"
+#include "Containers/Array.h"
+#include "MuCO/CustomizableObjectInstance.h"
+
+#include "CustomizableObjectValidationCommandlet.generated.h"
+
+class UCustomizableObject;
+
+UCLASS()
+class UCustomizableObjectValidationCommandlet : public UCommandlet 
+{
+	GENERATED_BODY()
+
+public:
+	virtual int32 Main(const FString& Params) override;
+	
+private:
+
+	/** Customizable Object to be tested */
+	UPROPERTY()
+	UCustomizableObject* ToTestCustomizableObject = nullptr;
+	
+	/** Customizable Object Instance currently being updated*/
+	UPROPERTY()
+	UCustomizableObjectInstance* InstanceBeingUpdated = nullptr;
+
+	/** Array of COI to be generated with randomized parameter values*/
+	UPROPERTY()
+	TArray<UCustomizableObjectInstance*> InstancesToProcess;
+	
+	/** Handle to be able to unbind OnInstanceUpdated(...) from instance end of update delegate.*/
+	FDelegateHandle OnInstanceUpdateHandle;
+
+	/** Did any of the instances fail the UpdateSkeletalMesh process?*/
+	bool bInstanceFailedUpdate = false;
+
+	// Instance update delegate
+	FInstanceUpdateDelegate InstanceUpdateDelegate;
+
+	/** Callback invoked once the currently updating instance has done the update.
+	 * @paragm Result is a container that provides us with data related with the instance updating process.
+	 */
+	UFUNCTION()
+	void OnInstanceUpdate(const FUpdateContext& Result);
+};
\ No newline at end of file
